Move all interchange code files under experimental folder.
#rnx

diff --git a/Engine/Plugins/Experimental/Interchange/Editor/Source/ResultsBrowser/InterchangeResultsBrowser.Build.cs b/Engine/Plugins/Experimental/Interchange/Editor/Source/ResultsBrowser/InterchangeResultsBrowser.Build.cs
new file mode 100644
--- /dev/null
+++ b/Engine/Plugins/Experimental/Interchange/Editor/Source/ResultsBrowser/InterchangeResultsBrowser.Build.cs
@@ -0,0 +1,33 @@
+// Copyright Epic Games, Inc. All Rights Reserved.
+
+namespace UnrealBuildTool.Rules
+{
+	public class InterchangeResultsBrowser : ModuleRules
+	{
+		public InterchangeResultsBrowser(ReadOnlyTargetRules Target) : base(Target)
+		{
+			PublicDependencyModuleNames.AddRange(
+				new string[]
+				{
+					"Core",
+					"CoreUObject",
+					"Engine",
+					"InterchangeCore",
+					"InterchangeEngine",
+				}
+			);
+
+			PrivateDependencyModuleNames.AddRange(
+				new string[]
+				{
+					"EditorStyle",
+					"InputCore",
+					"Projects",
+					"Slate",
+					"SlateCore",
+					"WorkspaceMenuStructure"
+				}
+			);
+		}
+	}
+}
diff --git a/Engine/Plugins/Experimental/Interchange/Editor/Source/ResultsBrowser/Private/InterchangeResultsBrowserModule.cpp b/Engine/Plugins/Experimental/Interchange/Editor/Source/ResultsBrowser/Private/InterchangeResultsBrowserModule.cpp
new file mode 100644
--- /dev/null
+++ b/Engine/Plugins/Experimental/Interchange/Editor/Source/ResultsBrowser/Private/InterchangeResultsBrowserModule.cpp
@@ -0,0 +1,117 @@
+// Copyright Epic Games, Inc. All Rights Reserved.
+
+#include "InterchangeResultsBrowserModule.h"
+
+#include "CoreMinimal.h"
+#include "Engine/Engine.h"
+#include "Framework/Docking/TabManager.h"
+#include "InterchangeResultsBrowserStyle.h"
+#include "InterchangeManager.h"
+#include "SInterchangeResultsBrowserWindow.h"
+#include "Styling/SlateIconFinder.h"
+#include "Widgets/Docking/SDockTab.h"
+#include "WorkspaceMenuStructure.h"
+#include "WorkspaceMenuStructureModule.h"
+
+
+#define LOCTEXT_NAMESPACE "InterchangeResultsBrowser"
+
+class FInterchangeResultsBrowserModule : public IInterchangeResultsBrowserModule
+{
+	/** IModuleInterface implementation */
+	virtual void StartupModule() override;
+	virtual void ShutdownModule() override;
+
+private:
+
+	TSharedRef<SDockTab> SpawnInterchangeResults(const FSpawnTabArgs& Args);
+	void OpenErrorBrowser(TStrongObjectPtr<UInterchangeResultsContainer> InResultsContainer);
+
+	/** Pointer to the style set to use for the UI. */
+	TSharedPtr<ISlateStyle> InterchangeResultsBrowserStyle = nullptr;
+
+	/** Owning pointer to the InterchangeResultsContainer being displayed */
+	TStrongObjectPtr<UInterchangeResultsContainer> ResultsContainer;
+
+	/** Pointer to the Slate widget */
+	TSharedPtr<SInterchangeResultsBrowserWindow> InterchangeResultsBrowserWindow;
+
+	/** Whether the results list is filtered */
+	// @TODO: load this from config
+	bool bIsFiltered = false;
+};
+
+
+IMPLEMENT_MODULE(FInterchangeResultsBrowserModule, InterchangeResultsBrowser)
+
+
+void FInterchangeResultsBrowserModule::StartupModule()
+{
+	auto RegisterItems = [this]()
+	{
+		UInterchangeManager& InterchangeManager = UInterchangeManager::GetInterchangeManager();
+		InterchangeManager.OnBatchImportComplete.AddRaw(this, &FInterchangeResultsBrowserModule::OpenErrorBrowser);
+
+		auto UnregisterItems = [this]()
+		{
+			UInterchangeManager& InterchangeManager = UInterchangeManager::GetInterchangeManager();
+			InterchangeManager.OnBatchImportComplete.RemoveAll(this);
+		};
+
+		InterchangeManager.OnPreDestroyInterchangeManager.AddLambda(UnregisterItems);
+	};
+
+	if (GEngine)
+	{
+		RegisterItems();
+	}
+	else
+	{
+		FCoreDelegates::OnPostEngineInit.AddLambda(RegisterItems);
+	}
+	
+	if (!InterchangeResultsBrowserStyle.IsValid())
+	{
+		InterchangeResultsBrowserStyle = MakeShared<FInterchangeResultsBrowserStyle>();
+	}
+
+	FGlobalTabmanager::Get()->RegisterNomadTabSpawner("InterchangeResults", FOnSpawnTab::CreateRaw(this, &FInterchangeResultsBrowserModule::SpawnInterchangeResults))
+		.SetDisplayName(LOCTEXT("InterchangeResultsBrowser", "Interchange Results Browser"))
+		.SetTooltipText(LOCTEXT("InterchangeResultsBrowserTooltipText", "Open the Interchange Results Browser tab."))
+		.SetGroup(WorkspaceMenu::GetMenuStructure().GetDeveloperToolsLogCategory())
+		.SetIcon(FSlateIconFinder::FindIcon("InterchangeResultsBrowser.TabIcon"));
+}
+
+
+void FInterchangeResultsBrowserModule::ShutdownModule()
+{
+	InterchangeResultsBrowserStyle = nullptr;
+}
+
+
+TSharedRef<SDockTab> FInterchangeResultsBrowserModule::SpawnInterchangeResults(const FSpawnTabArgs& Args)
+{
+	TSharedRef<SDockTab> DockTab = SNew(SDockTab)
+		.TabRole(ETabRole::NomadTab);
+
+	DockTab->SetContent(
+		SAssignNew(InterchangeResultsBrowserWindow, SInterchangeResultsBrowserWindow)
+		.OwnerTab(DockTab)
+		.InterchangeResultsContainer(ResultsContainer.Get())
+		.IsFiltered(bIsFiltered)
+		.OnFilterChangedState_Lambda([this](bool bNewState) { this->bIsFiltered = bNewState; })
+	);
+
+	return DockTab;
+}
+
+
+void FInterchangeResultsBrowserModule::OpenErrorBrowser(TStrongObjectPtr<UInterchangeResultsContainer> InResultsContainer)
+{
+	ResultsContainer = InResultsContainer;
+	ResultsContainer->Finalize();
+	FGlobalTabmanager::Get()->TryInvokeTab(FName("InterchangeResults"));
+	InterchangeResultsBrowserWindow->Set(ResultsContainer.Get());
+}
+
+#undef LOCTEXT_NAMESPACE
diff --git a/Engine/Plugins/Experimental/Interchange/Editor/Source/ResultsBrowser/Private/InterchangeResultsBrowserStyle.h b/Engine/Plugins/Experimental/Interchange/Editor/Source/ResultsBrowser/Private/InterchangeResultsBrowserStyle.h
new file mode 100644
--- /dev/null
+++ b/Engine/Plugins/Experimental/Interchange/Editor/Source/ResultsBrowser/Private/InterchangeResultsBrowserStyle.h
@@ -0,0 +1,44 @@
+// Copyright Epic Games, Inc. All Rights Reserved.
+
+#pragma once
+
+#include "Interfaces/IPluginManager.h"
+#include "Brushes/SlateImageBrush.h"
+#include "Styling/SlateStyle.h"
+#include "Styling/SlateStyleRegistry.h"
+
+/**
+ * Implements the visual style of the text asset editor UI.
+ */
+class FInterchangeResultsBrowserStyle final : public FSlateStyleSet
+{
+public:
+
+	FInterchangeResultsBrowserStyle()
+		: FSlateStyleSet("InterchangeResultsBrowserStyle")
+	{
+		const FVector2D Icon16x16(16.0f, 16.0f);
+
+		const FString BaseDir = IPluginManager::Get().FindPlugin("InterchangeEditor")->GetBaseDir();
+		SetContentRoot(BaseDir / TEXT("Content"));
+
+		FSlateImageBrush* SuccessBrush = new FSlateImageBrush(RootToContentDir(TEXT("Resources/Result_Success"), TEXT(".png")), Icon16x16);
+		Set("InterchangeResultsBrowser.ResultType.Success", SuccessBrush);
+		
+		FSlateImageBrush* WarningBrush = new FSlateImageBrush(RootToContentDir(TEXT("Resources/Result_Warning"), TEXT(".png")), Icon16x16);
+		Set("InterchangeResultsBrowser.ResultType.Warning", WarningBrush);
+
+		FSlateImageBrush* ErrorBrush = new FSlateImageBrush(RootToContentDir(TEXT("Resources/Result_Fail"), TEXT(".png")), Icon16x16);
+		Set("InterchangeResultsBrowser.ResultType.Error", ErrorBrush);
+
+		FSlateImageBrush* TabIcon = new FSlateImageBrush(RootToContentDir(TEXT("Resources/TabIcon"), TEXT(".png")), Icon16x16);
+		Set("InterchangeResultsBrowser.TabIcon", TabIcon);
+
+		FSlateStyleRegistry::RegisterSlateStyle(*this);
+	}
+
+	~FInterchangeResultsBrowserStyle()
+	{
+		FSlateStyleRegistry::UnRegisterSlateStyle(*this);
+	}
+};
diff --git a/Engine/Plugins/Experimental/Interchange/Editor/Source/ResultsBrowser/Private/SInterchangeResultsBrowserWindow.cpp b/Engine/Plugins/Experimental/Interchange/Editor/Source/ResultsBrowser/Private/SInterchangeResultsBrowserWindow.cpp
new file mode 100644
--- /dev/null
+++ b/Engine/Plugins/Experimental/Interchange/Editor/Source/ResultsBrowser/Private/SInterchangeResultsBrowserWindow.cpp
@@ -0,0 +1,433 @@
+// Copyright Epic Games, Inc. All Rights Reserved.
+
+#include "SInterchangeResultsBrowserWindow.h"
+
+#include "EditorStyleSet.h"
+#include "Framework/MultiBox/MultiBoxBuilder.h"
+#include "IDocumentation.h"
+#include "Modules/ModuleManager.h"
+#include "Styling/SlateIconFinder.h"
+#include "Widgets/Images/SImage.h"
+#include "Widgets/Input/SButton.h"
+#include "Widgets/Input/SCheckBox.h"
+#include "Widgets/Layout/SSplitter.h"
+#include "Widgets/Layout/SUniformGridPanel.h"
+#include "Widgets/Layout/SSpacer.h"
+
+#define LOCTEXT_NAMESPACE "InterchangeResultsBrowser"
+
+
+namespace SResultsBrowserWidgetDefs
+{
+	const FName ColumnID_ResultType("ResultType");
+	const FName ColumnID_Source("Source");
+	const FName ColumnID_Dest("Dest");
+	const FName ColumnID_AssetType("AssetType");
+	const FName ColumnID_Message("Message");
+}
+
+
+SInterchangeResultsBrowserWindow::SInterchangeResultsBrowserWindow()
+	: OwnerTab(nullptr),
+	  bIsFiltered(false)
+{
+}
+
+
+SInterchangeResultsBrowserWindow::~SInterchangeResultsBrowserWindow()
+{
+}
+
+
+void SInterchangeResultsBrowserWindow::Construct(const FArguments& InArgs)
+{
+	SortByColumn = SResultsBrowserWidgetDefs::ColumnID_ResultType;
+	SortMode = EColumnSortMode::Ascending;
+	OnFilterChangedState = InArgs._OnFilterChangedState;
+	bIsFiltered = InArgs._IsFiltered;
+	ResultsContainer = InArgs._InterchangeResultsContainer;
+
+	OwnerTab = InArgs._OwnerTab;
+	check(OwnerTab.IsValid());
+
+	TSharedRef<SHeaderRow> HeaderRowWidget = SNew(SHeaderRow);
+
+	HeaderRowWidget->AddColumn(
+		SHeaderRow::Column(SResultsBrowserWidgetDefs::ColumnID_ResultType)
+		.SortMode(this, &SInterchangeResultsBrowserWindow::GetColumnSortMode, SResultsBrowserWidgetDefs::ColumnID_ResultType)
+		.OnSort(this, &SInterchangeResultsBrowserWindow::OnColumnSortModeChanged)
+		[
+			SNew(SSpacer)
+		]
+		.FixedWidth(24.0f)
+		);
+
+	HeaderRowWidget->AddColumn(
+		SHeaderRow::Column(SResultsBrowserWidgetDefs::ColumnID_Source)
+		.DefaultLabel(LOCTEXT("ColumnLabel_SourceFile", "Source File"))
+		.SortMode(this, &SInterchangeResultsBrowserWindow::GetColumnSortMode, SResultsBrowserWidgetDefs::ColumnID_Source)
+		.OnSort(this, &SInterchangeResultsBrowserWindow::OnColumnSortModeChanged)
+		.FillWidth(5.0f)
+	);
+
+	HeaderRowWidget->AddColumn(
+		SHeaderRow::Column(SResultsBrowserWidgetDefs::ColumnID_Dest)
+		.DefaultLabel(LOCTEXT("ColumnLabel_DestinationAsset", "Destination Asset"))
+		.SortMode(this, &SInterchangeResultsBrowserWindow::GetColumnSortMode, SResultsBrowserWidgetDefs::ColumnID_Dest)
+		.OnSort(this, &SInterchangeResultsBrowserWindow::OnColumnSortModeChanged)
+		.FillWidth(5.0f)
+	);
+
+	HeaderRowWidget->AddColumn(
+		SHeaderRow::Column(SResultsBrowserWidgetDefs::ColumnID_AssetType)
+		.DefaultLabel(LOCTEXT("ColumnLabel_AssetType", "Asset Type"))
+		.SortMode(this, &SInterchangeResultsBrowserWindow::GetColumnSortMode, SResultsBrowserWidgetDefs::ColumnID_AssetType)
+		.OnSort(this, &SInterchangeResultsBrowserWindow::OnColumnSortModeChanged)
+		.FillWidth(3.0f)
+	);
+
+	HeaderRowWidget->AddColumn(
+		SHeaderRow::Column(SResultsBrowserWidgetDefs::ColumnID_Message)
+		.DefaultLabel(LOCTEXT("ColumnLabel_Details", "Details"))
+		.FillWidth(9.0f)
+	);
+
+	this->ChildSlot
+	[
+		SNew(SBorder)
+		.Padding(FMargin(10.0f, 3.0f))
+		.BorderImage(FEditorStyle::GetBrush("ToolPanel.DarkGroupBorder"))
+		[
+			SNew(SVerticalBox)
+			+ SVerticalBox::Slot()
+			.AutoHeight()
+			.HAlign(HAlign_Left)
+			.Padding(10)
+			[
+				SNew(SHorizontalBox)
+				+ SHorizontalBox::Slot()
+				.AutoWidth()
+				[
+					SNew(SCheckBox)
+					.IsChecked(bIsFiltered ? ECheckBoxState::Checked : ECheckBoxState::Unchecked)
+					.OnCheckStateChanged(this, &SInterchangeResultsBrowserWindow::OnFilterStateChanged)
+				]
+				+ SHorizontalBox::Slot()
+				.AutoWidth()
+				[
+					SNew(STextBlock)
+					.Text(LOCTEXT("OnlyShowErrors", "Only display errors/warnings"))
+				]
+			]
+			+ SVerticalBox::Slot()
+			.FillHeight(1.0f)
+			.Padding(2)
+			[
+				SNew(SSplitter)
+				.Orientation(Orient_Vertical)
+				+ SSplitter::Slot()
+				.Value(0.4f)
+				[
+					SAssignNew(ListView, SListView<UInterchangeResult*>)
+					.ItemHeight(30)
+					.ListItemsSource(&FilteredResults)
+					.HeaderRow(HeaderRowWidget)
+					.SelectionMode(ESelectionMode::Single)
+					.OnGenerateRow(this, &SInterchangeResultsBrowserWindow::OnGenerateRowForList)
+					.OnSelectionChanged(this, &SInterchangeResultsBrowserWindow::OnSelectionChanged)
+				]
+				+ SSplitter::Slot()
+				.Value(0.2f)
+				[
+					SNew(SBorder)
+					.Padding(3)
+					.BorderImage(FEditorStyle::GetBrush("ToolPanel.GroupBorder"))
+					[
+						SNew(SBox)
+						.Padding(10)
+						[
+							SAssignNew(Description, STextBlock)
+							.AutoWrapText(true)
+						]
+					]
+				]
+			]
+			+ SVerticalBox::Slot()
+			.AutoHeight()
+			.HAlign(HAlign_Right)
+			.Padding(2)
+			[
+				SNew(SUniformGridPanel)
+				.SlotPadding(2)
+				+ SUniformGridPanel::Slot(0, 0)
+				[
+					IDocumentation::Get()->CreateAnchor(FString("Engine/Content/Interchange/ErrorBrowser"))
+				]
+				+ SUniformGridPanel::Slot(1, 0)
+				[
+					SNew(SButton)
+					.HAlign(HAlign_Center)
+					.Text(LOCTEXT("ErrorBrowser_Clear", "Clear"))
+					.IsEnabled(false)
+					.OnClicked(this, &SInterchangeResultsBrowserWindow::OnCloseDialog)
+				]
+				+ SUniformGridPanel::Slot(2, 0)
+				[
+					SNew(SButton)
+					.HAlign(HAlign_Center)
+					.Text(LOCTEXT("ErrorBrowser_Close", "Close"))
+					.OnClicked(this, &SInterchangeResultsBrowserWindow::OnCloseDialog)
+				]
+			]
+		]
+	];
+
+	Set(InArgs._InterchangeResultsContainer);
+}
+
+
+void SInterchangeResultsBrowserWindow::Set(UInterchangeResultsContainer* Data)
+{
+	ResultsContainer = Data;
+	RepopulateItems();
+}
+
+
+void SInterchangeResultsBrowserWindow::RepopulateItems()
+{
+	FilteredResults.Empty();
+
+	if (ResultsContainer != nullptr)
+	{
+		for (UInterchangeResult* Item : ResultsContainer->GetResults())
+		{
+			if (Item->GetResultType() != EInterchangeResultType::Success || !bIsFiltered)
+			{
+				FilteredResults.Add(Item);
+			}
+		}
+
+		RequestSort();
+	}
+}
+
+
+void SInterchangeResultsBrowserWindow::CloseErrorBrowser()
+{
+	if (TSharedPtr<SDockTab> OwnerTabPin = OwnerTab.Pin())
+	{
+		OwnerTabPin->RequestCloseTab();
+	}
+	OwnerTab = nullptr;
+}
+
+
+TSharedRef<ITableRow> SInterchangeResultsBrowserWindow::OnGenerateRowForList(UInterchangeResult* InItem, const TSharedRef<STableViewBase>& OwnerTable)
+{
+	TSharedRef<ITableRow> Row =
+		SNew(SInterchangeResultsBrowserListRow, OwnerTable)
+		.InterchangeResultsBrowserWidget(SharedThis(this))
+		.Item(InItem);
+
+	return Row;
+}
+
+
+EColumnSortMode::Type SInterchangeResultsBrowserWindow::GetColumnSortMode(const FName ColumnId) const
+{
+	if (SortByColumn != ColumnId)
+	{
+		return EColumnSortMode::None;
+	}
+
+	return SortMode;
+}
+
+
+void SInterchangeResultsBrowserWindow::OnColumnSortModeChanged(const EColumnSortPriority::Type SortPriority, const FName& ColumnId, const EColumnSortMode::Type InSortMode)
+{
+	SortByColumn = ColumnId;
+	SortMode = InSortMode;
+
+	RequestSort();
+}
+
+
+void SInterchangeResultsBrowserWindow::OnFilterStateChanged(ECheckBoxState NewState)
+{
+	bIsFiltered = (NewState == ECheckBoxState::Checked);
+	OnFilterChangedState.ExecuteIfBound(bIsFiltered);
+	RepopulateItems();
+}
+
+
+void SInterchangeResultsBrowserWindow::RequestSort()
+{
+	if (SortByColumn == SResultsBrowserWidgetDefs::ColumnID_ResultType)
+	{
+		if (SortMode == EColumnSortMode::Ascending)
+		{
+			FilteredResults.Sort([](const UInterchangeResult& A, const UInterchangeResult& B) {
+				return static_cast<int>(A.GetResultType()) < static_cast<int>(B.GetResultType()); });
+		}
+		else if (SortMode == EColumnSortMode::Descending)
+		{
+			FilteredResults.Sort([](const UInterchangeResult& A, const UInterchangeResult& B) {
+				return static_cast<int>(A.GetResultType()) >= static_cast<int>(B.GetResultType()); });
+		}
+	}
+	else if (SortByColumn == SResultsBrowserWidgetDefs::ColumnID_Source)
+	{
+		if (SortMode == EColumnSortMode::Ascending)
+		{
+			FilteredResults.Sort([](const UInterchangeResult& A, const UInterchangeResult& B) {
+				return A.SourceAssetName < B.SourceAssetName; });
+		}
+		else if (SortMode == EColumnSortMode::Descending)
+		{
+			FilteredResults.Sort([](const UInterchangeResult& A, const UInterchangeResult& B) {
+				return A.SourceAssetName >= B.SourceAssetName; });
+		}
+	}
+	else if (SortByColumn == SResultsBrowserWidgetDefs::ColumnID_Dest)
+	{
+		if (SortMode == EColumnSortMode::Ascending)
+		{
+			FilteredResults.Sort([](const UInterchangeResult& A, const UInterchangeResult& B) {
+				return A.DestinationAssetName < B.DestinationAssetName; });
+		}
+		else if (SortMode == EColumnSortMode::Descending)
+		{
+			FilteredResults.Sort([](const UInterchangeResult& A, const UInterchangeResult& B) {
+				return A.DestinationAssetName >= B.DestinationAssetName; });
+		}
+	}
+	else if (SortByColumn == SResultsBrowserWidgetDefs::ColumnID_AssetType)
+	{
+		if (SortMode == EColumnSortMode::Ascending)
+		{
+			FilteredResults.Sort([](const UInterchangeResult& A, const UInterchangeResult& B) {
+				FString NameA = A.AssetType ? A.AssetType->GetName() : FString();
+				FString NameB = B.AssetType ? B.AssetType->GetName() : FString();
+				return NameA < NameB;
+				});
+		}
+		else if (SortMode == EColumnSortMode::Descending)
+		{
+			FilteredResults.Sort([](const UInterchangeResult& A, const UInterchangeResult& B) {
+				FString NameA = A.AssetType ? A.AssetType->GetName() : FString();
+				FString NameB = B.AssetType ? B.AssetType->GetName() : FString();
+				return NameA >= NameB;
+				});
+		}
+	}
+
+	ListView->RequestListRefresh();
+}
+
+
+
+void SInterchangeResultsBrowserWindow::OnSelectionChanged(UInterchangeResult* InItem, ESelectInfo::Type SelectInfo)
+{
+	if (InItem != nullptr)
+	{
+		Description->SetText(InItem->GetText());
+	}
+	else
+	{
+		Description->SetText(FText());
+	}
+}
+
+
+
+
+void SInterchangeResultsBrowserListRow::Construct(const FArguments& InArgs, const TSharedRef<STableViewBase>& InOwnerTableView)
+{
+	InterchangeResultsBrowserWidgetPtr = InArgs._InterchangeResultsBrowserWidget;
+	Item = InArgs._Item;
+
+	SMultiColumnTableRow<UInterchangeResult*>::Construct(FSuperRowType::FArguments(), InOwnerTableView);
+}
+
+
+static FName GetImageFromResultType(EInterchangeResultType ResultType)
+{
+	switch (ResultType)
+	{
+	case EInterchangeResultType::Success: return FName("InterchangeResultsBrowser.ResultType.Success");
+	case EInterchangeResultType::Warning: return FName("InterchangeResultsBrowser.ResultType.Warning");
+	case EInterchangeResultType::Error: return FName("InterchangeResultsBrowser.ResultType.Error");
+	}
+
+	return FName();
+}
+
+
+TSharedRef<SWidget> SInterchangeResultsBrowserListRow::GenerateWidgetForColumn(const FName& ColumnName)
+{
+	if (TSharedPtr<SInterchangeResultsBrowserWindow> InterchangeResultsBrowserWidget = InterchangeResultsBrowserWidgetPtr.Pin())
+	{
+		const FMargin RowPadding(8, 2, 2, 0);
+
+		if (ColumnName == SResultsBrowserWidgetDefs::ColumnID_ResultType)
+		{
+			return SNew(SHorizontalBox)
+				+ SHorizontalBox::Slot()
+				.Padding(RowPadding)
+				.AutoWidth()
+				[
+					SNew(SImage)
+					.Image(FSlateIconFinder::FindIcon(GetImageFromResultType(Item->GetResultType())).GetOptionalIcon())
+				];
+		}
+		else if (ColumnName == SResultsBrowserWidgetDefs::ColumnID_Source)
+		{
+			return SNew(SHorizontalBox)
+				+ SHorizontalBox::Slot()
+				.Padding(RowPadding)
+				[
+					SNew(STextBlock)
+					.Text(FText::FromString(FPaths::GetCleanFilename(Item->SourceAssetName)))
+					.ToolTipText(FText::FromString(Item->SourceAssetName))
+				];
+		}
+		else if (ColumnName == SResultsBrowserWidgetDefs::ColumnID_Dest)
+		{
+			return SNew(SHorizontalBox)
+				+ SHorizontalBox::Slot()
+				.Padding(RowPadding)
+				[
+					SNew(STextBlock)
+					.Text(FText::FromString(FPackageName::ObjectPathToObjectName(Item->DestinationAssetName)))
+					.ToolTipText(FText::FromString(Item->DestinationAssetName))
+				];
+		}
+		else if (ColumnName == SResultsBrowserWidgetDefs::ColumnID_AssetType)
+		{
+			return SNew(SHorizontalBox)
+				+ SHorizontalBox::Slot()
+				.Padding(RowPadding)
+				[
+					SNew(STextBlock)
+					.Text(FText::FromString(Item->AssetType ? Item->AssetType->GetName() : FString()))
+				];
+		}
+		else if (ColumnName == SResultsBrowserWidgetDefs::ColumnID_Message)
+		{
+			return SNew(SHorizontalBox)
+				+ SHorizontalBox::Slot()
+				.Padding(RowPadding)
+				[
+					SNew(STextBlock)
+					.Text(Item->GetText())
+					.ToolTipText(Item->GetText())
+				];
+		}
+	}
+
+	return SNullWidget::NullWidget;
+}
+
+
+#undef LOCTEXT_NAMESPACE
diff --git a/Engine/Plugins/Experimental/Interchange/Editor/Source/ResultsBrowser/Private/SInterchangeResultsBrowserWindow.h b/Engine/Plugins/Experimental/Interchange/Editor/Source/ResultsBrowser/Private/SInterchangeResultsBrowserWindow.h
new file mode 100644
--- /dev/null
+++ b/Engine/Plugins/Experimental/Interchange/Editor/Source/ResultsBrowser/Private/SInterchangeResultsBrowserWindow.h
@@ -0,0 +1,103 @@
+// Copyright Epic Games, Inc. All Rights Reserved.
+
+#pragma once
+
+#include "CoreMinimal.h"
+#include "InputCoreTypes.h"
+#include "Input/Reply.h"
+#include "InterchangeResultsContainer.h"
+#include "Styling/SlateColor.h"
+#include "Widgets/DeclarativeSyntaxSupport.h"
+#include "Widgets/SCompoundWidget.h"
+#include "Widgets/Docking/SDockTab.h"
+#include "Widgets/Views/SListView.h"
+#include "Widgets/Views/STableViewBase.h"
+#include "Widgets/Views/STableRow.h"
+
+
+class SInterchangeResultsBrowserWindow : public SCompoundWidget
+{
+public:
+	DECLARE_DELEGATE_OneParam(FOnFilterChangedState, bool)
+
+	SLATE_BEGIN_ARGS(SInterchangeResultsBrowserWindow)
+		: _OwnerTab()
+		, _InterchangeResultsContainer()
+		, _IsFiltered(false)
+		{}
+
+		SLATE_ARGUMENT(TWeakPtr<SDockTab>, OwnerTab)
+		SLATE_ARGUMENT(UInterchangeResultsContainer*, InterchangeResultsContainer)
+		SLATE_ARGUMENT(bool, IsFiltered)
+		SLATE_EVENT(FOnFilterChangedState, OnFilterChangedState)
+	SLATE_END_ARGS()
+
+public:
+
+	SInterchangeResultsBrowserWindow();
+	~SInterchangeResultsBrowserWindow();
+
+	void Construct(const FArguments& InArgs);
+	virtual bool SupportsKeyboardFocus() const override { return true; }
+
+	void Set(UInterchangeResultsContainer* Data);
+	void CloseErrorBrowser();
+
+	FReply OnCloseDialog()
+	{
+		CloseErrorBrowser();
+		return FReply::Handled();
+	}
+
+	virtual FReply OnKeyDown(const FGeometry& MyGeometry, const FKeyEvent& InKeyEvent) override
+	{
+		if (InKeyEvent.GetKey() == EKeys::Escape)
+		{
+			return OnCloseDialog();
+		}
+
+		return FReply::Unhandled();
+	}
+
+private:
+
+	/** Called by SListView to get a widget corresponding to the supplied item */
+	TSharedRef<ITableRow> OnGenerateRowForList(UInterchangeResult* Item, const TSharedRef<STableViewBase>& OwnerTable);
+
+	EColumnSortMode::Type GetColumnSortMode(const FName ColumnId) const;
+
+	void OnColumnSortModeChanged(const EColumnSortPriority::Type SortPriority, const FName& ColumnId, const EColumnSortMode::Type InSortMode);
+	void OnSelectionChanged(UInterchangeResult* InItem, ESelectInfo::Type SelectInfo);
+	void OnFilterStateChanged(ECheckBoxState NewState);
+	void RepopulateItems();
+	void RequestSort();
+
+
+	TWeakPtr<SDockTab> OwnerTab;
+	UInterchangeResultsContainer* ResultsContainer;
+	TArray<UInterchangeResult*> FilteredResults;
+	TSharedPtr<SListView<UInterchangeResult*>> ListView;
+	TSharedPtr<STextBlock> Description;
+	FName SortByColumn;
+	EColumnSortMode::Type SortMode;
+	FOnFilterChangedState OnFilterChangedState;
+	bool bIsFiltered;
+};
+
+
+class SInterchangeResultsBrowserListRow : public SMultiColumnTableRow<UInterchangeResult*>
+{
+public:
+	SLATE_BEGIN_ARGS(SInterchangeResultsBrowserListRow) {}
+		SLATE_ARGUMENT(TSharedPtr<SInterchangeResultsBrowserWindow>, InterchangeResultsBrowserWidget)
+		SLATE_ARGUMENT(UInterchangeResult*, Item)
+	SLATE_END_ARGS()
+
+	void Construct(const FArguments& InArgs, const TSharedRef<STableViewBase>& InOwnerTableView);
+	virtual TSharedRef<SWidget> GenerateWidgetForColumn(const FName& ColumnName) override;
+
+private:
+
+	TWeakPtr<SInterchangeResultsBrowserWindow> InterchangeResultsBrowserWidgetPtr;
+	UInterchangeResult* Item;
+};
diff --git a/Engine/Plugins/Experimental/Interchange/Editor/Source/ResultsBrowser/Public/InterchangeResultsBrowserModule.h b/Engine/Plugins/Experimental/Interchange/Editor/Source/ResultsBrowser/Public/InterchangeResultsBrowserModule.h
new file mode 100644
--- /dev/null
+++ b/Engine/Plugins/Experimental/Interchange/Editor/Source/ResultsBrowser/Public/InterchangeResultsBrowserModule.h
@@ -0,0 +1,39 @@
+// Copyright Epic Games, Inc. All Rights Reserved.
+
+#pragma once
+
+#include "CoreMinimal.h"
+#include "Modules/ModuleInterface.h"
+#include "Modules/ModuleManager.h"
+#include "Styling/ISlateStyle.h"
+
+
+/**
+ * The public interface to this module
+ */
+class IInterchangeResultsBrowserModule : public IModuleInterface
+{
+
+public:
+
+	/**
+	 * Singleton-like access to this module's interface.  This is just for convenience!
+	 * Beware of calling this during the shutdown phase, though.  Your module might have been unloaded already.
+	 *
+	 * @return Returns singleton instance, loading the module on demand if needed
+	 */
+	static inline IInterchangeResultsBrowserModule& Get()
+	{
+		return FModuleManager::LoadModuleChecked<IInterchangeResultsBrowserModule>("InterchangeResultsBrowser");
+	}
+
+	/**
+	 * Checks to see if this module is loaded and ready.  It is only valid to call Get() if IsAvailable() returns true.
+	 *
+	 * @return True if the module is loaded and ready to use
+	 */
+	static inline bool IsAvailable()
+	{
+		return FModuleManager::Get().IsModuleLoaded("InterchangeResultsBrowser");
+	}
+};
diff --git a/Engine/Plugins/Experimental/Interchange/Runtime/Source/Import/Private/Mesh/InterchangeStaticMeshFactory.cpp b/Engine/Plugins/Experimental/Interchange/Runtime/Source/Import/Private/Mesh/InterchangeStaticMeshFactory.cpp
new file mode 100644
--- /dev/null
+++ b/Engine/Plugins/Experimental/Interchange/Runtime/Source/Import/Private/Mesh/InterchangeStaticMeshFactory.cpp
@@ -0,0 +1,408 @@
+// Copyright Epic Games, Inc. All Rights Reserved. 
+#include "Mesh/InterchangeStaticMeshFactory.h"
+
+#include "Components.h"
+#include "Engine/StaticMesh.h"
+#include "InterchangeImportCommon.h"
+#include "InterchangeImportLog.h"
+#include "InterchangeMaterialFactoryNode.h"
+#include "InterchangeMaterialNode.h"
+#include "InterchangeMeshNode.h"
+#include "InterchangeSceneNode.h"
+#include "InterchangeStaticMeshLodDataNode.h"
+#include "InterchangeStaticMeshFactoryNode.h"
+#include "InterchangeSourceData.h"
+#include "InterchangeTranslatorBase.h"
+#include "Mesh/InterchangeStaticMeshPayload.h"
+#include "Mesh/InterchangeStaticMeshPayloadInterface.h"
+#include "Nodes/InterchangeBaseNode.h"
+#include "Nodes/InterchangeBaseNodeContainer.h"
+#include "StaticMeshAttributes.h"
+#include "StaticMeshOperations.h"
+
+#if WITH_EDITORONLY_DATA
+
+#include "EditorFramework/AssetImportData.h"
+
+#endif //WITH_EDITORONLY_DATA
+
+
+UClass* UInterchangeStaticMeshFactory::GetFactoryClass() const
+{
+	return UStaticMesh::StaticClass();
+}
+
+
+UObject* UInterchangeStaticMeshFactory::CreateEmptyAsset(const FCreateAssetParams& Arguments)
+{
+#if !WITH_EDITOR || !WITH_EDITORONLY_DATA
+
+	UE_LOG(LogInterchangeImport, Error, TEXT("Cannot import StaticMesh asset in runtime, this is an editor only feature."));
+	return nullptr;
+
+#else
+	UStaticMesh* StaticMesh = nullptr;
+	if (!Arguments.AssetNode || !Arguments.AssetNode->GetObjectClass()->IsChildOf(GetFactoryClass()))
+	{
+		return nullptr;
+	}
+
+	const UInterchangeStaticMeshFactoryNode* StaticMeshFactoryNode = Cast<UInterchangeStaticMeshFactoryNode>(Arguments.AssetNode);
+	if (StaticMeshFactoryNode == nullptr)
+	{
+		return nullptr;
+	}
+
+	// create an asset if it doesn't exist
+	UObject* ExistingAsset = StaticFindObject(nullptr, Arguments.Parent, *Arguments.AssetName);
+
+	// create a new static mesh or overwrite existing asset, if possible
+	if (!ExistingAsset)
+	{
+		StaticMesh = NewObject<UStaticMesh>(Arguments.Parent, *Arguments.AssetName, RF_Public | RF_Standalone);
+	}
+	else if (ExistingAsset->GetClass()->IsChildOf(UStaticMesh::StaticClass()))
+	{
+		//This is a reimport, we are just re-updating the source data
+		StaticMesh = Cast<UStaticMesh>(ExistingAsset);
+	}
+	
+	if (!StaticMesh)
+	{
+		UE_LOG(LogInterchangeImport, Warning, TEXT("Could not create StaticMesh asset %s"), *Arguments.AssetName);
+		return nullptr;
+	}
+	
+	StaticMesh->PreEditChange(nullptr);	
+	return StaticMesh;
+#endif //else !WITH_EDITOR || !WITH_EDITORONLY_DATA
+}
+
+
+UObject* UInterchangeStaticMeshFactory::CreateAsset(const FCreateAssetParams& Arguments)
+{
+#if !WITH_EDITOR || !WITH_EDITORONLY_DATA
+
+	UE_LOG(LogInterchangeImport, Error, TEXT("Cannot import static mesh asset in runtime, this is an editor only feature."));
+	return nullptr;
+
+#else
+
+	if (!Arguments.AssetNode || !Arguments.AssetNode->GetObjectClass()->IsChildOf(GetFactoryClass()))
+	{
+		return nullptr;
+	}
+
+	const UInterchangeStaticMeshFactoryNode* StaticMeshFactoryNode = Cast<UInterchangeStaticMeshFactoryNode>(Arguments.AssetNode);
+	if (StaticMeshFactoryNode == nullptr)
+	{
+		return nullptr;
+	}
+
+	const IInterchangeStaticMeshPayloadInterface* StaticMeshTranslatorPayloadInterface = Cast<IInterchangeStaticMeshPayloadInterface>(Arguments.Translator);
+	if (!StaticMeshTranslatorPayloadInterface)
+	{
+		UE_LOG(LogInterchangeImport, Error, TEXT("Cannot import static mesh, the translator does not implement the IInterchangeStaticMeshPayloadInterface."));
+		return nullptr;
+	}
+
+	const UClass* StaticMeshClass = StaticMeshFactoryNode->GetObjectClass();
+	check(StaticMeshClass && StaticMeshClass->IsChildOf(GetFactoryClass()));
+
+	// create an asset if it doesn't exist
+	UObject* ExistingAsset = StaticFindObject(nullptr, Arguments.Parent, *Arguments.AssetName);
+
+	UObject* StaticMeshObject = nullptr;
+
+	// create a new static mesh or overwrite existing asset, if possible
+	if (!ExistingAsset)
+	{
+		// NewObject is not thread safe, the asset registry directory watcher tick on the main thread can trig before we finish initializing the UObject and will crash
+		// The UObject should have been created by calling CreateEmptyAsset on the main thread.
+		check(IsInGameThread());
+		StaticMeshObject = NewObject<UObject>(Arguments.Parent, StaticMeshClass, *Arguments.AssetName, RF_Public | RF_Standalone);
+	}
+	else if(ExistingAsset->GetClass()->IsChildOf(StaticMeshClass))
+	{
+		//This is a reimport, we are just re-updating the source data
+		StaticMeshObject = ExistingAsset;
+	}
+
+	if (!StaticMeshObject)
+	{
+		UE_LOG(LogInterchangeImport, Error, TEXT("Could not create StaticMesh asset %s"), *Arguments.AssetName);
+		return nullptr;
+	}
+
+	if (StaticMeshObject)
+	{
+		// Currently static mesh re-import will not touch the static mesh at all
+		// TODO design a re-import process for the static mesh (expressions and input connections)
+		if (!Arguments.ReimportObject)
+		{
+			UStaticMesh* StaticMesh = Cast<UStaticMesh>(StaticMeshObject);
+			if (!ensure(StaticMesh))
+			{
+				UE_LOG(LogInterchangeImport, Error, TEXT("Could not create StaticMesh asset %s"), *Arguments.AssetName);
+				return nullptr;
+			}
+			
+			int32 LodCount = StaticMeshFactoryNode->GetLodDataCount();
+			StaticMesh->SetNumSourceModels(LodCount);
+
+			TArray<FString> LodDataUniqueIds;
+			StaticMeshFactoryNode->GetLodDataUniqueIds(LodDataUniqueIds);
+			ensure(LodDataUniqueIds.Num() == LodCount);
+
+			int32 CurrentLodIndex = 0;
+			for (int32 LodIndex = 0; LodIndex < LodCount; ++LodIndex)
+			{
+				FString LodUniqueId = LodDataUniqueIds[LodIndex];
+				const UInterchangeStaticMeshLodDataNode* LodDataNode = Cast<UInterchangeStaticMeshLodDataNode>(Arguments.NodeContainer->GetNode(LodUniqueId));
+				if (!LodDataNode)
+				{
+					UE_LOG(LogInterchangeImport, Warning, TEXT("Invalid LOD when importing StaticMesh asset %s"), *Arguments.AssetName);
+					continue;
+				}
+
+				// Get the mesh node context for each MeshUids
+				struct FMeshNodeContext
+				{
+					const UInterchangeMeshNode* MeshNode = nullptr;
+					const UInterchangeSceneNode* SceneNode = nullptr;
+					TOptional<FTransform> SceneGlobalTransform;
+					FString TranslatorPayloadKey;
+				};
+				TArray<FMeshNodeContext> MeshReferences;
+
+				// Scope to query the mesh node
+				{
+					TArray<FString> MeshUids;
+					LodDataNode->GetMeshUids(MeshUids);
+					MeshReferences.Reserve(MeshUids.Num());
+
+					for (const FString& MeshUid : MeshUids)
+					{
+						FMeshNodeContext MeshReference;
+						MeshReference.MeshNode = Cast<UInterchangeMeshNode>(Arguments.NodeContainer->GetNode(MeshUid));
+						if (!MeshReference.MeshNode)
+						{
+							// The reference is a scene node and we need to bake the geometry
+							MeshReference.SceneNode = Cast<UInterchangeSceneNode>(Arguments.NodeContainer->GetNode(MeshUid));
+							if (!ensure(MeshReference.SceneNode != nullptr))
+							{
+								UE_LOG(LogInterchangeImport, Warning, TEXT("Invalid LOD mesh reference when importing StaticMesh asset %s"), *Arguments.AssetName);
+								continue;
+							}
+
+							FString MeshDependencyUid;
+							MeshReference.SceneNode->GetCustomAssetInstanceUid(MeshDependencyUid);
+							MeshReference.MeshNode = Cast<UInterchangeMeshNode>(Arguments.NodeContainer->GetNode(MeshDependencyUid));
+
+							// Cache the scene node global matrix, we will use this matrix to bake the vertices
+							FTransform SceneNodeGlobalTransform;
+							if (MeshReference.SceneNode->GetCustomGlobalTransform(SceneNodeGlobalTransform))
+							{
+								MeshReference.SceneGlobalTransform = SceneNodeGlobalTransform;
+							}
+						}
+
+						if (!ensure(MeshReference.MeshNode != nullptr))
+						{
+							UE_LOG(LogInterchangeImport, Warning, TEXT("Invalid LOD mesh reference when importing StaticMesh asset %s"), *Arguments.AssetName);
+							continue;
+						}
+
+						TOptional<FString> MeshPayloadKey = MeshReference.MeshNode->GetPayLoadKey();
+						if (MeshPayloadKey.IsSet())
+						{
+							MeshReference.TranslatorPayloadKey = MeshPayloadKey.GetValue();
+						}
+						else
+						{
+							UE_LOG(LogInterchangeImport, Warning, TEXT("Empty LOD mesh reference payload when importing StaticMesh asset %s"), *Arguments.AssetName);
+							continue;
+						}
+
+						MeshReferences.Add(MeshReference);
+					}
+				}
+
+				// Add the lod mesh data to the static mesh
+				FMeshDescription* LodMeshDescription = StaticMesh->CreateMeshDescription(CurrentLodIndex);
+
+				FStaticMeshOperations::FAppendSettings AppendSettings;
+				for (int32 ChannelIdx = 0; ChannelIdx < FStaticMeshOperations::FAppendSettings::MAX_NUM_UV_CHANNELS; ++ChannelIdx)
+				{
+					AppendSettings.bMergeUVChannels[ChannelIdx] = true;
+				}
+
+				// Fetch the payloads for all mesh references in parallel
+				TMap<FString, TFuture<TOptional<UE::Interchange::FStaticMeshPayloadData>>> MeshPayloads;
+				for (const FMeshNodeContext& MeshNodeContext : MeshReferences)
+				{
+					const FString& PayloadKey = MeshNodeContext.TranslatorPayloadKey;
+					MeshPayloads.Add(PayloadKey, StaticMeshTranslatorPayloadInterface->GetStaticMeshPayloadData(PayloadKey));
+				}
+				
+				// Fill the lod mesh description using all combined mesh parts
+				for (const FMeshNodeContext& MeshNodeContext : MeshReferences)
+				{
+					TOptional<UE::Interchange::FStaticMeshPayloadData> LodMeshPayload = MeshPayloads.FindChecked(MeshNodeContext.TranslatorPayloadKey).Get();
+					if (!LodMeshPayload.IsSet())
+					{
+						UE_LOG(LogInterchangeImport, Warning, TEXT("Invalid static mesh payload key [%s] StaticMesh asset %s"), *MeshNodeContext.TranslatorPayloadKey, *Arguments.AssetName);
+						continue;
+					}
+
+					// Bake the payload, with the provided transform
+					if (MeshNodeContext.SceneGlobalTransform.IsSet())
+					{
+						AppendSettings.MeshTransform = MeshNodeContext.SceneGlobalTransform;
+					}
+					else
+					{
+						AppendSettings.MeshTransform.Reset();
+					}
+
+					FStaticMeshOperations::AppendMeshDescription(LodMeshPayload->MeshDescription, *LodMeshDescription, AppendSettings);
+				}
+
+				//////////////////////////////////////////////////////////////////////////
+				//Manage vertex color
+				//Replace -> do nothing, we want to use the translated source data
+				//Ignore -> remove vertex color from import data (when we re-import, ignore have to put back the current mesh vertex color)
+				//Override -> replace the vertex color by the override color
+				{
+					FStaticMeshAttributes Attributes(*LodMeshDescription);
+					TVertexInstanceAttributesRef<FVector4f> VertexInstanceColors = Attributes.GetVertexInstanceColors();
+					bool bReplaceVertexColor = false;
+					StaticMeshFactoryNode->GetCustomVertexColorReplace(bReplaceVertexColor);
+					if (!bReplaceVertexColor)
+					{
+						bool bIgnoreVertexColor = false;
+						StaticMeshFactoryNode->GetCustomVertexColorIgnore(bIgnoreVertexColor);
+						if (bIgnoreVertexColor)
+						{
+							//Flush the vertex color, if we re-import we have to fill it with the old data
+							for (const FVertexInstanceID& VertexInstanceID : LodMeshDescription->VertexInstances().GetElementIDs())
+							{
+								VertexInstanceColors[VertexInstanceID] = FVector4f(FLinearColor(FColor::White));
+							}
+						}
+						else
+						{
+							FColor OverrideVertexColor;
+							if (StaticMeshFactoryNode->GetCustomVertexColorOverride(OverrideVertexColor))
+							{
+								for (const FVertexInstanceID& VertexInstanceID : LodMeshDescription->VertexInstances().GetElementIDs())
+								{
+									VertexInstanceColors[VertexInstanceID] = FVector4f(FLinearColor(OverrideVertexColor));
+								}
+							}
+						}
+					}
+				}
+
+				StaticMesh->CommitMeshDescription(CurrentLodIndex);
+
+				// Handle materials
+				// @TODO: move this outside of the LOD loop?
+				TArray<FString> FactoryDependencies;
+				StaticMeshFactoryNode->GetFactoryDependencies(FactoryDependencies);
+
+				for (int32 DependencyIndex = 0; DependencyIndex < FactoryDependencies.Num(); ++DependencyIndex)
+				{
+					const UInterchangeMaterialFactoryNode* MaterialFactoryNode = Cast<UInterchangeMaterialFactoryNode>(Arguments.NodeContainer->GetNode(FactoryDependencies[DependencyIndex]));
+					if (!MaterialFactoryNode || !MaterialFactoryNode->ReferenceObject.IsValid())
+					{
+						continue;
+					}
+
+					FName MaterialSlotName = *MaterialFactoryNode->GetDisplayLabel();
+					UMaterialInterface* MaterialInterface = Cast<UMaterialInterface>(MaterialFactoryNode->ReferenceObject.ResolveObject());
+					if (!MaterialInterface)
+					{
+						MaterialInterface = UMaterial::GetDefaultMaterial(MD_Surface);
+					}
+
+					int32 MaterialSlotIndex = StaticMesh->GetMaterialIndexFromImportedMaterialSlotName(MaterialSlotName);
+					if (MaterialSlotIndex == INDEX_NONE)
+					{
+						StaticMesh->GetStaticMaterials().Emplace(MaterialInterface, MaterialSlotName);
+					}
+				}
+
+				FStaticMeshConstAttributes StaticMeshAttributes(*LodMeshDescription);
+				TPolygonGroupAttributesRef<const FName> SlotNames = StaticMeshAttributes.GetPolygonGroupMaterialSlotNames();
+				int32 SectionIndex = 0;
+				for (FPolygonGroupID PolygonGroupID : LodMeshDescription->PolygonGroups().GetElementIDs())
+				{
+					int32 MaterialSlotIndex = StaticMesh->GetMaterialIndexFromImportedMaterialSlotName(SlotNames[PolygonGroupID]);
+					
+					if (MaterialSlotIndex == INDEX_NONE)
+					{
+						// If no material was found with this slot name, it is probably because the pipeline is configured to not import materials.
+						// Fill out a blank slot instead.
+						MaterialSlotIndex = StaticMesh->GetStaticMaterials().Emplace(nullptr, SlotNames[PolygonGroupID]);
+					}
+
+					FMeshSectionInfo Info = StaticMesh->GetSectionInfoMap().Get(CurrentLodIndex, SectionIndex);
+					Info.MaterialIndex = MaterialSlotIndex;
+					StaticMesh->GetSectionInfoMap().Remove(CurrentLodIndex, SectionIndex);
+					StaticMesh->GetSectionInfoMap().Set(CurrentLodIndex, SectionIndex, Info);
+
+					SectionIndex++;
+				}
+
+				// Set static mesh build settings
+				FStaticMeshSourceModel& SrcModel = StaticMesh->GetSourceModel(CurrentLodIndex);
+				SrcModel.ReductionSettings.MaxDeviation = 0.0f;
+				SrcModel.ReductionSettings.PercentTriangles = 1.0f;
+				SrcModel.ReductionSettings.PercentVertices = 1.0f;
+
+				CurrentLodIndex++;
+			}
+
+			// Apply all StaticMeshFactoryNode custom attributes to the static mesh asset
+			StaticMeshFactoryNode->ApplyAllCustomAttributeToObject(StaticMesh);
+		}
+		
+		// Getting the file Hash will cache it into the source data
+		Arguments.SourceData->GetFileContentHash();
+
+		// The interchange completion task (called in the GameThread after the factories pass), will call PostEditChange
+		// which will trigger another asynchronous system that will build all static meshes in parallel
+	}
+	else
+	{
+		// The static mesh is not a UStaticMesh
+		StaticMeshObject->RemoveFromRoot();
+		StaticMeshObject->MarkPendingKill();
+	}
+	return StaticMeshObject;
+
+#endif //else !WITH_EDITOR || !WITH_EDITORONLY_DATA
+}
+
+
+/* This function is call in the completion task on the main thread, use it to call main thread post creation step for your assets */
+void UInterchangeStaticMeshFactory::PreImportPreCompletedCallback(const FImportPreCompletedCallbackParams& Arguments)
+{
+	check(IsInGameThread());
+	Super::PreImportPreCompletedCallback(Arguments);
+
+	// TODO: make sure this works at runtime
+#if WITH_EDITORONLY_DATA
+	if (ensure(Arguments.ImportedObject && Arguments.SourceData))
+	{
+		// We must call the Update of the asset source file in the main thread because UAssetImportData::Update execute some delegate we do not control
+		UStaticMesh* StaticMesh = CastChecked<UStaticMesh>(Arguments.ImportedObject);
+
+		UAssetImportData* ImportDataPtr = StaticMesh->GetAssetImportData();
+		UE::Interchange::FFactoryCommon::FUpdateImportAssetDataParameters UpdateImportAssetDataParameters(StaticMesh, ImportDataPtr, Arguments.SourceData, Arguments.NodeUniqueID, Arguments.NodeContainer, Arguments.Pipelines);
+		ImportDataPtr = UE::Interchange::FFactoryCommon::UpdateImportAssetData(UpdateImportAssetDataParameters);
+		StaticMesh->SetAssetImportData(ImportDataPtr);
+	}
+#endif
+}
\ No newline at end of file
diff --git a/Engine/Plugins/Experimental/Interchange/Runtime/Source/Import/Private/Texture/InterchangeBMPTranslator.cpp b/Engine/Plugins/Experimental/Interchange/Runtime/Source/Import/Private/Texture/InterchangeBMPTranslator.cpp
new file mode 100644
--- /dev/null
+++ b/Engine/Plugins/Experimental/Interchange/Runtime/Source/Import/Private/Texture/InterchangeBMPTranslator.cpp
@@ -0,0 +1,109 @@
+// Copyright Epic Games, Inc. All Rights Reserved. 
+
+#include "Texture/InterchangeBMPTranslator.h"
+
+#include "Engine/Texture.h"
+#include "Engine/Texture2D.h"
+#include "IImageWrapper.h"
+#include "IImageWrapperModule.h"
+#include "InterchangeImportLog.h"
+#include "InterchangeTextureNode.h"
+#include "Misc/ConfigCacheIni.h"
+#include "Misc/FileHelper.h"
+#include "Misc/Paths.h"
+#include "Misc/ScopedSlowTask.h"
+#include "Modules/ModuleManager.h"
+#include "Nodes/InterchangeBaseNodeContainer.h"
+#include "Texture/TextureTranslatorUtilities.h"
+
+
+bool UInterchangeBMPTranslator::CanImportSourceData(const UInterchangeSourceData* InSourceData) const
+{
+	FString Extension = FPaths::GetExtension(InSourceData->GetFilename());
+	FString BMPExtension = (TEXT("bmp;Texture"));
+	return BMPExtension.StartsWith(Extension);
+}
+
+bool UInterchangeBMPTranslator::Translate(UInterchangeBaseNodeContainer& BaseNodeContainer) const
+{
+	return UE::Interchange::FTextureTranslatorUtilities::Generic2DTextureTranslate(GetSourceData(), BaseNodeContainer);
+}
+
+TOptional<UE::Interchange::FImportImage> UInterchangeBMPTranslator::GetTexturePayloadData(const UInterchangeSourceData* PayloadSourceData, const FString& PayLoadKey) const
+{
+	check(PayloadSourceData == GetSourceData());
+
+	if (!GetSourceData())
+	{
+		UE_LOG(LogInterchangeImport, Error, TEXT("Failed to import BMP, bad source data."));
+		return TOptional<UE::Interchange::FImportImage>();
+	}
+
+	TArray64<uint8> SourceDataBuffer;
+	FString Filename = GetSourceData()->GetFilename();
+
+	//Make sure the key fit the filename, The key should always be valid
+	if (!Filename.Equals(PayLoadKey))
+	{
+		UE_LOG(LogInterchangeImport, Error, TEXT("Failed to import BMP, wrong payload key. [%s]"), *Filename);
+		return TOptional<UE::Interchange::FImportImage>();
+	}
+
+	if (!FPaths::FileExists(Filename))
+	{
+		UE_LOG(LogInterchangeImport, Error, TEXT("Failed to import BMP, cannot open file. [%s]"), *Filename);
+		return TOptional<UE::Interchange::FImportImage>();
+	}
+
+	if (!FFileHelper::LoadFileToArray(SourceDataBuffer, *Filename))
+	{
+		UE_LOG(LogInterchangeImport, Error, TEXT("Failed to import BMP, cannot load file content into an array. [%s]"), *Filename);
+		return TOptional<UE::Interchange::FImportImage>();
+	}
+
+	const uint8* Buffer = SourceDataBuffer.GetData();
+	const uint8* BufferEnd = Buffer + SourceDataBuffer.Num();
+
+	bool bAllowNonPowerOfTwo = false;
+	GConfig->GetBool(TEXT("TextureImporter"), TEXT("AllowNonPowerOfTwoTextures"), bAllowNonPowerOfTwo, GEditorIni);
+
+	// Validate it.
+	const int32 Length = BufferEnd - Buffer;
+
+	IImageWrapperModule& ImageWrapperModule = FModuleManager::LoadModuleChecked<IImageWrapperModule>(FName("ImageWrapper"));
+
+	//
+	// BMP
+	//
+	TSharedPtr<IImageWrapper> BmpImageWrapper = ImageWrapperModule.CreateImageWrapper(EImageFormat::BMP);
+	if (!BmpImageWrapper.IsValid() || !BmpImageWrapper->SetCompressed(Buffer, Length))
+	{
+		UE_LOG(LogInterchangeImport, Error, TEXT("Failed to decode BMP. [%s]"), *Filename);
+		return TOptional<UE::Interchange::FImportImage>();
+	}
+	if (!UE::Interchange::FImportImageHelper::IsImportResolutionValid(BmpImageWrapper->GetWidth(), BmpImageWrapper->GetHeight(), bAllowNonPowerOfTwo))
+	{
+		UE_LOG(LogInterchangeImport, Error, TEXT("Failed to import BMP, invalid resolution. Resolution[%d, %d], AllowPowerOfTwo[%s], [%s]"), BmpImageWrapper->GetWidth(), BmpImageWrapper->GetHeight(), bAllowNonPowerOfTwo ? TEXT("True") : TEXT("false"), *Filename);
+		return TOptional<UE::Interchange::FImportImage>();
+	}
+
+	int32 BitDepth = BmpImageWrapper->GetBitDepth();
+	ERGBFormat Format = BmpImageWrapper->GetFormat();
+
+	UE::Interchange::FImportImage PayloadData;
+
+	PayloadData.Init2DWithParams(
+		BmpImageWrapper->GetWidth(),
+		BmpImageWrapper->GetHeight(),
+		TSF_BGRA8,
+		false
+	);
+
+	if (!BmpImageWrapper->GetRaw(Format, BitDepth, PayloadData.GetArrayViewOfRawData()))
+	{
+		UE_LOG(LogInterchangeImport, Error, TEXT("Failed to decode BMP. [%s]"), *Filename);
+		return TOptional<UE::Interchange::FImportImage>();
+	}
+
+	return PayloadData;
+}
diff --git a/Engine/Plugins/Experimental/Interchange/Runtime/Source/Import/Private/Texture/InterchangeEXRTranslator.cpp b/Engine/Plugins/Experimental/Interchange/Runtime/Source/Import/Private/Texture/InterchangeEXRTranslator.cpp
new file mode 100644
--- /dev/null
+++ b/Engine/Plugins/Experimental/Interchange/Runtime/Source/Import/Private/Texture/InterchangeEXRTranslator.cpp
@@ -0,0 +1,115 @@
+// Copyright Epic Games, Inc. All Rights Reserved. 
+#include "Texture/InterchangeEXRTranslator.h"
+
+#include "Engine/Texture.h"
+#include "Engine/Texture2D.h"
+#include "IImageWrapper.h"
+#include "IImageWrapperModule.h"
+#include "InterchangeImportLog.h"
+#include "InterchangeTextureNode.h"
+#include "Misc/ConfigCacheIni.h"
+#include "Misc/FileHelper.h"
+#include "Misc/Paths.h"
+#include "Misc/ScopedSlowTask.h"
+#include "Modules/ModuleManager.h"
+#include "Nodes/InterchangeBaseNodeContainer.h"
+#include "Texture/TextureTranslatorUtilities.h"
+
+bool UInterchangeEXRTranslator::CanImportSourceData(const UInterchangeSourceData* InSourceData) const
+{
+	FString Extension = FPaths::GetExtension(InSourceData->GetFilename());
+	FString EXRExtension = (TEXT("exr;Texture"));
+	return EXRExtension.StartsWith(Extension);
+}
+
+bool UInterchangeEXRTranslator::Translate(UInterchangeBaseNodeContainer& BaseNodeContainer) const
+{
+	return UE::Interchange::FTextureTranslatorUtilities::Generic2DTextureTranslate(GetSourceData(), BaseNodeContainer);
+}
+
+TOptional<UE::Interchange::FImportImage> UInterchangeEXRTranslator::GetTexturePayloadData(const UInterchangeSourceData* PayloadSourceData, const FString& PayLoadKey) const
+{
+	check(PayloadSourceData == GetSourceData());
+
+	if (!GetSourceData())
+	{
+		UE_LOG(LogInterchangeImport, Error, TEXT("Failed to import EXR, bad source data."));
+		return TOptional<UE::Interchange::FImportImage>();
+	}
+
+	TArray64<uint8> SourceDataBuffer;
+	FString Filename = GetSourceData()->GetFilename();
+	
+	//Make sure the key fit the filename, The key should always be valid
+	if (!Filename.Equals(PayLoadKey))
+	{
+		UE_LOG(LogInterchangeImport, Error, TEXT("Failed to import EXR, wrong payload key. [%s]"), *Filename);
+		return TOptional<UE::Interchange::FImportImage>();
+	}
+
+	if (!FPaths::FileExists(Filename))
+	{
+		UE_LOG(LogInterchangeImport, Error, TEXT("Failed to import EXR, cannot open file. [%s]"), *Filename);
+		return TOptional<UE::Interchange::FImportImage>();
+	}
+
+	if (!FFileHelper::LoadFileToArray(SourceDataBuffer, *Filename))
+	{
+		UE_LOG(LogInterchangeImport, Error, TEXT("Failed to import EXR, cannot load file content into an array. [%s]"), *Filename);
+		return TOptional<UE::Interchange::FImportImage>();
+	}
+
+	const uint8* Buffer = SourceDataBuffer.GetData();
+	const uint8* BufferEnd = Buffer + SourceDataBuffer.Num();
+
+	bool bAllowNonPowerOfTwo = false;
+	GConfig->GetBool(TEXT("TextureImporter"), TEXT("AllowNonPowerOfTwoTextures"), bAllowNonPowerOfTwo, GEditorIni);
+
+	// Validate it.
+	const int32 Length = BufferEnd - Buffer;
+
+	//
+	// EXR
+	//
+	IImageWrapperModule& ImageWrapperModule = FModuleManager::LoadModuleChecked<IImageWrapperModule>(FName("ImageWrapper"));
+
+	TSharedPtr<IImageWrapper> ExrImageWrapper = ImageWrapperModule.CreateImageWrapper(EImageFormat::EXR);
+	if (!ExrImageWrapper.IsValid() || !ExrImageWrapper->SetCompressed(Buffer, Length))
+	{
+		UE_LOG(LogInterchangeImport, Error, TEXT("Failed to import EXR, unsupported format. [%s]"), *Filename);
+		return TOptional<UE::Interchange::FImportImage>();
+	}
+
+	int32 Width = ExrImageWrapper->GetWidth();
+	int32 Height = ExrImageWrapper->GetHeight();
+
+	if (!UE::Interchange::FImportImageHelper::IsImportResolutionValid(Width, Height, bAllowNonPowerOfTwo))
+	{
+		UE_LOG(LogInterchangeImport, Error, TEXT("Failed to import EXR, invalid resolution. Resolution[%d, %d], AllowPowerOfTwo[%s], [%s]"), Width, Height, bAllowNonPowerOfTwo ? TEXT("True") : TEXT("false"), *Filename);
+		return TOptional<UE::Interchange::FImportImage>();
+	}
+
+	// Currently the only texture source format compatible with EXR image formats is TSF_RGBA16F.
+	// EXR decoder automatically converts imported image channels into the requested float format.
+	// In case if the imported image is a gray scale image, its content will be stored in the green channel of the created texture.
+	ETextureSourceFormat TextureFormat = TSF_RGBA16F;
+	ERGBFormat Format = ERGBFormat::RGBAF;
+	int32 BitDepth = 16;
+
+	UE::Interchange::FImportImage PayloadData;
+	PayloadData.Init2DWithParams(
+		Width,
+		Height,
+		TextureFormat,
+		false
+	);
+	PayloadData.CompressionSettings = TC_HDR;
+
+	if (!ExrImageWrapper->GetRaw(Format, BitDepth, PayloadData.GetArrayViewOfRawData()))
+	{
+		UE_LOG(LogInterchangeImport, Error, TEXT("Failed to decode EXR. [%s]"), *Filename);
+		return TOptional<UE::Interchange::FImportImage>();
+	}
+
+	return PayloadData;
+}
\ No newline at end of file
diff --git a/Engine/Plugins/Experimental/Interchange/Runtime/Source/Import/Private/Texture/InterchangeIESTranslator.cpp b/Engine/Plugins/Experimental/Interchange/Runtime/Source/Import/Private/Texture/InterchangeIESTranslator.cpp
new file mode 100644
--- /dev/null
+++ b/Engine/Plugins/Experimental/Interchange/Runtime/Source/Import/Private/Texture/InterchangeIESTranslator.cpp
@@ -0,0 +1,86 @@
+// Copyright Epic Games, Inc. All Rights Reserved. 
+
+#include "Texture/InterchangeIESTranslator.h"
+
+#include "Engine/Texture.h"
+#include "Engine/TextureLightProfile.h"
+#include "IESConverter.h"
+#include "InterchangeImportLog.h"
+#include "Misc/ConfigCacheIni.h"
+#include "Misc/FileHelper.h"
+#include "Misc/Paths.h"
+#include "Texture/TextureTranslatorUtilities.h"
+
+bool UInterchangeIESTranslator::CanImportSourceData(const UInterchangeSourceData* InSourceData) const
+{
+	FString Extension = FPaths::GetExtension(InSourceData->GetFilename());
+	FString IESExtension = (TEXT("ies"));
+	return IESExtension.StartsWith(Extension);
+}
+
+bool UInterchangeIESTranslator::Translate(UInterchangeBaseNodeContainer& BaseNodeContainer) const
+{
+	return UE::Interchange::FTextureTranslatorUtilities::GenericTextureLightProfileTranslate(GetSourceData(), BaseNodeContainer);
+}
+
+TOptional<UE::Interchange::FImportLightProfile> UInterchangeIESTranslator::GetLightProfilePayloadData(const UInterchangeSourceData* PayloadSourceData, const FString& PayLoadKey) const
+{
+	check(GetSourceData() == PayloadSourceData);
+
+	if (!GetSourceData())
+	{
+		UE_LOG(LogInterchangeImport, Error, TEXT("Failed to import IES, bad source data."));
+		return {};
+	}
+
+	TArray64<uint8> SourceDataBuffer;
+	FString Filename = GetSourceData()->GetFilename();
+
+	//Make sure the key fit the filename, The key should always be valid
+	if (!Filename.Equals(PayLoadKey))
+	{
+		UE_LOG(LogInterchangeImport, Error, TEXT("Failed to import IES, wrong payload key. [%s]"), *Filename);
+		return {};
+	}
+
+	if (!FPaths::FileExists(Filename))
+	{
+		UE_LOG(LogInterchangeImport, Error, TEXT("Failed to import IES, cannot open file. [%s]"), *Filename);
+		return {};
+	}
+
+	if (!FFileHelper::LoadFileToArray(SourceDataBuffer, *Filename))
+	{
+		UE_LOG(LogInterchangeImport, Error, TEXT("Failed to import IES, cannot load file content into an array. [%s]"), *Filename);
+		return {};
+	}
+
+	const uint8* Buffer = SourceDataBuffer.GetData();
+
+	FIESConverter IESConverter(Buffer, SourceDataBuffer.Num());
+
+	if(IESConverter.IsValid())
+	{
+		UE::Interchange::FImportLightProfile Payload;
+
+		Payload.Init2DWithParams(
+			IESConverter.GetWidth(),
+			IESConverter.GetHeight(),
+			TSF_RGBA16F,
+			false
+			);
+
+		Payload.CompressionSettings = TC_HDR;
+		Payload.Brightness = IESConverter.GetBrightness();
+		Payload.TextureMultiplier = IESConverter.GetMultiplier();
+
+
+		const TArray<uint8>&  RawData = IESConverter.GetRawData();;
+
+		FPlatformMemory::Memcpy(Payload.RawData.GetData(), RawData.GetData(), Payload.RawData.GetSize());
+
+		return Payload;
+	}
+
+	return {};
+}
diff --git a/Engine/Plugins/Experimental/Interchange/Runtime/Source/Import/Private/Texture/InterchangeJPGTranslator.cpp b/Engine/Plugins/Experimental/Interchange/Runtime/Source/Import/Private/Texture/InterchangeJPGTranslator.cpp
new file mode 100644
--- /dev/null
+++ b/Engine/Plugins/Experimental/Interchange/Runtime/Source/Import/Private/Texture/InterchangeJPGTranslator.cpp
@@ -0,0 +1,133 @@
+// Copyright Epic Games, Inc. All Rights Reserved. 
+#include "Texture/InterchangeJPGTranslator.h"
+
+#include "Engine/Texture.h"
+#include "Engine/Texture2D.h"
+#include "IImageWrapper.h"
+#include "IImageWrapperModule.h"
+#include "InterchangeImportLog.h"
+#include "InterchangeTextureNode.h"
+#include "Misc/ConfigCacheIni.h"
+#include "Misc/FileHelper.h"
+#include "Misc/Paths.h"
+#include "Misc/ScopedSlowTask.h"
+#include "Modules/ModuleManager.h"
+#include "Nodes/InterchangeBaseNodeContainer.h"
+#include "Texture/TextureTranslatorUtilities.h"
+
+bool UInterchangeJPGTranslator::CanImportSourceData(const UInterchangeSourceData* InSourceData) const
+{
+	FString Extension = FPaths::GetExtension(InSourceData->GetFilename());
+	FString JPGExtension = (TEXT("jpg;Texture"));
+	return JPGExtension.StartsWith(Extension);
+}
+
+bool UInterchangeJPGTranslator::Translate(UInterchangeBaseNodeContainer& BaseNodeContainer) const
+{
+	return UE::Interchange::FTextureTranslatorUtilities::Generic2DTextureTranslate(GetSourceData(), BaseNodeContainer);
+}
+
+TOptional<UE::Interchange::FImportImage> UInterchangeJPGTranslator::GetTexturePayloadData(const UInterchangeSourceData* PayloadSourceData, const FString& PayLoadKey) const
+{
+	check(PayloadSourceData == GetSourceData());
+
+	if (!GetSourceData())
+	{
+		UE_LOG(LogInterchangeImport, Error, TEXT("Failed to import JPEG, bad source data."));
+		return TOptional<UE::Interchange::FImportImage>();
+	}
+
+	TArray64<uint8> SourceDataBuffer;
+	FString Filename = GetSourceData()->GetFilename();
+	
+	//Make sure the key fit the filename, The key should always be valid
+	if (!Filename.Equals(PayLoadKey))
+	{
+		UE_LOG(LogInterchangeImport, Error, TEXT("Failed to import JPEG, wrong payload key. [%s]"), *Filename);
+		return TOptional<UE::Interchange::FImportImage>();
+	}
+
+	if (!FPaths::FileExists(Filename))
+	{
+		UE_LOG(LogInterchangeImport, Error, TEXT("Failed to import JPEG, cannot open file. [%s]"), *Filename);
+		return TOptional<UE::Interchange::FImportImage>();
+	}
+
+	if (!FFileHelper::LoadFileToArray(SourceDataBuffer, *Filename))
+	{
+		UE_LOG(LogInterchangeImport, Error, TEXT("Failed to import JPEG, cannot load file content into an array. [%s]"), *Filename);
+		return TOptional<UE::Interchange::FImportImage>();
+	}
+
+	const uint8* Buffer = SourceDataBuffer.GetData();
+	const uint8* BufferEnd = Buffer + SourceDataBuffer.Num();
+
+	bool bAllowNonPowerOfTwo = false;
+	GConfig->GetBool(TEXT("TextureImporter"), TEXT("AllowNonPowerOfTwoTextures"), bAllowNonPowerOfTwo, GEditorIni);
+
+	// Validate it.
+	const int32 Length = BufferEnd - Buffer;
+
+	IImageWrapperModule& ImageWrapperModule = FModuleManager::LoadModuleChecked<IImageWrapperModule>(FName("ImageWrapper"));
+
+	//
+	// JPG
+	//
+	TSharedPtr<IImageWrapper> JpegImageWrapper = ImageWrapperModule.CreateImageWrapper(EImageFormat::JPEG);
+	if (!JpegImageWrapper.IsValid() || !JpegImageWrapper->SetCompressed(Buffer, Length))
+	{
+		UE_LOG(LogInterchangeImport, Error, TEXT("Failed to decode JPEG. [%s]"), *Filename);
+		return TOptional<UE::Interchange::FImportImage>();
+	}
+	if (!UE::Interchange::FImportImageHelper::IsImportResolutionValid(JpegImageWrapper->GetWidth(), JpegImageWrapper->GetHeight(), bAllowNonPowerOfTwo))
+	{
+		UE_LOG(LogInterchangeImport, Error, TEXT("Failed to import JPEG, invalid resolution. Resolution[%d, %d], AllowPowerOfTwo[%s], [%s]"), JpegImageWrapper->GetWidth(), JpegImageWrapper->GetHeight(), bAllowNonPowerOfTwo ? TEXT("True") : TEXT("false"), *Filename);
+		return TOptional<UE::Interchange::FImportImage>();
+	}
+
+	// Select the texture's source format
+	ETextureSourceFormat TextureFormat = TSF_Invalid;
+	int32 BitDepth = JpegImageWrapper->GetBitDepth();
+	ERGBFormat Format = JpegImageWrapper->GetFormat();
+
+	if (Format == ERGBFormat::Gray)
+	{
+		if (BitDepth <= 8)
+		{
+			TextureFormat = TSF_G8;
+			Format = ERGBFormat::Gray;
+			BitDepth = 8;
+		}
+	}
+	else if (Format == ERGBFormat::RGBA)
+	{
+		if (BitDepth <= 8)
+		{
+			TextureFormat = TSF_BGRA8;
+			Format = ERGBFormat::BGRA;
+			BitDepth = 8;
+		}
+	}
+
+	if (TextureFormat == TSF_Invalid)
+	{
+		UE_LOG(LogInterchangeImport, Error, TEXT("JPEG file [%s] contains data in an unsupported format"), *Filename);
+		return TOptional<UE::Interchange::FImportImage>();
+	}
+
+	UE::Interchange::FImportImage PayloadData;
+
+	PayloadData.Init2DWithParams(
+		JpegImageWrapper->GetWidth(),
+		JpegImageWrapper->GetHeight(),
+		TextureFormat,
+		BitDepth < 16
+	);
+
+	if (!JpegImageWrapper->GetRaw(Format, BitDepth, PayloadData.GetArrayViewOfRawData()))
+	{
+		UE_LOG(LogInterchangeImport, Error, TEXT("Failed to decode JPEG. [%s]"), *Filename);
+		return TOptional<UE::Interchange::FImportImage>();
+	}
+	return PayloadData;
+}
\ No newline at end of file
diff --git a/Engine/Plugins/Experimental/Interchange/Runtime/Source/Import/Private/Texture/InterchangePSDTranslator.cpp b/Engine/Plugins/Experimental/Interchange/Runtime/Source/Import/Private/Texture/InterchangePSDTranslator.cpp
new file mode 100644
--- /dev/null
+++ b/Engine/Plugins/Experimental/Interchange/Runtime/Source/Import/Private/Texture/InterchangePSDTranslator.cpp
@@ -0,0 +1,485 @@
+// Copyright Epic Games, Inc. All Rights Reserved. 
+#include "Texture/InterchangePSDTranslator.h"
+
+#include "Engine/Texture.h"
+#include "Engine/Texture2D.h"
+#include "IImageWrapper.h"
+#include "IImageWrapperModule.h"
+#include "InterchangeImportLog.h"
+#include "InterchangeTextureNode.h"
+#include "Misc/ConfigCacheIni.h"
+#include "Misc/FileHelper.h"
+#include "Misc/Paths.h"
+#include "Misc/ScopedSlowTask.h"
+#include "Modules/ModuleManager.h"
+#include "Nodes/InterchangeBaseNodeContainer.h"
+#include "Texture/TextureTranslatorUtilities.h"
+
+//////////////////////////////////////////////////////////////////////////
+// PSD helper local function
+namespace UE
+{
+	namespace Interchange
+	{
+		namespace Private
+		{
+#pragma pack(push,1)
+			struct FPSDFileHeader
+			{
+				int32	Signature;		// 8BPS
+				int16	Version;		// Version
+				int16	nChannels;		// Number of Channels (3=RGB) (4=RGBA)
+				int32	Height;			// Number of Image Rows
+				int32	Width;			// Number of Image Columns
+				int16	Depth;			// Number of Bits per Channel
+				int16	Mode;			// Image Mode (0=Bitmap)(1=Grayscale)(2=Indexed)(3=RGB)(4=CYMK)(7=Multichannel)
+				uint8	Pad[6];			// Padding
+
+				/**
+				 * @return Whether file has a valid signature
+				 */
+				bool IsValid() const
+				{
+					// Fail on bad signature
+					return (Signature == 0x38425053);
+				}
+
+				/**
+				 * @return Whether file has a supported version
+				 */
+				bool IsSupported() const
+				{
+					// Fail on bad version
+					if (Version != 1)
+						return false;
+
+					// Fail on anything other than 1, 3 or 4 channels
+					if ((nChannels != 1) && (nChannels != 3) && (nChannels != 4))
+						return false;
+
+					// Fail on anything other than 8 Bits/channel or 16 Bits/channel  
+					if ((Depth != 8) && (Depth != 16))
+						return false;
+
+					// Fail on anything other than Grayscale and RGB
+					// We can add support for indexed later if needed.
+					if (Mode != 1 && Mode != 3)
+						return false;
+
+					return true;
+				}
+			};
+#pragma pack(pop)
+
+
+			static bool ReadData(uint8* pOut, const uint8*& pBuffer, FPSDFileHeader& Info)
+			{
+				const uint8* pPlane = nullptr;
+				const uint8* pRowTable = nullptr;
+				int32 iPlane;
+				int16 CompressionType;
+				int32 iPixel;
+				int32 iRow;
+				int32 CompressedBytes;
+				int32 iByte;
+				int32 Count;
+				uint8 Value;
+
+				// Double check to make sure this is a valid request
+				if (!Info.IsValid() || !Info.IsSupported())
+				{
+					return false;
+				}
+
+				const uint8* pCur = pBuffer + sizeof(FPSDFileHeader);
+				int32 NPixels = Info.Width * Info.Height;
+
+				int32 ClutSize = ((int32)pCur[0] << 24) +
+					((int32)pCur[1] << 16) +
+					((int32)pCur[2] << 8) +
+					((int32)pCur[3] << 0);
+				pCur += 4 + ClutSize;
+
+				// Skip Image Resource Section
+				int32 ImageResourceSize = ((int32)pCur[0] << 24) +
+					((int32)pCur[1] << 16) +
+					((int32)pCur[2] << 8) +
+					((int32)pCur[3] << 0);
+				pCur += 4 + ImageResourceSize;
+
+				// Skip Layer and Mask Section
+				int32 LayerAndMaskSize = ((int32)pCur[0] << 24) +
+					((int32)pCur[1] << 16) +
+					((int32)pCur[2] << 8) +
+					((int32)pCur[3] << 0);
+				pCur += 4 + LayerAndMaskSize;
+
+				// Determine number of bytes per pixel
+				int32 BytesPerPixel = 3;
+				const int32 BytesPerChannel = (Info.Depth / 8);
+				switch (Info.Mode)
+				{
+					case 1: // 'GrayScale'
+						BytesPerPixel = BytesPerChannel;
+						break;
+					case 2:
+						BytesPerPixel = 1;
+						return false;  // until we support indexed...
+						break;
+					case 3: // 'RGBColor'
+						if (Info.nChannels == 3)
+							BytesPerPixel = 3 * BytesPerChannel;
+						else
+							BytesPerPixel = 4 * BytesPerChannel;
+						break;
+					default:
+						return false;
+						break;
+				}
+
+				// Get Compression Type
+				CompressionType = ((int32)pCur[0] << 8) + ((int32)pCur[1] << 0);
+				pCur += 2;
+
+				// Fail on 16 Bits/channel with RLE. This can occur when the file is not saved with 'Maximize Compatibility'. Compression doesn't appear to be standard.
+				if (CompressionType == 1 && Info.Depth == 16)
+				{
+					return false;
+				}
+
+				// If no alpha channel, set alpha to opaque (255 or 65536).
+				if (Info.nChannels != 4)
+				{
+					if (Info.Depth == 8)
+					{
+						const uint32 Channels = 4;
+						const uint32 BufferSize = Info.Width * Info.Height * Channels * sizeof(uint8);
+						FMemory::Memset(pOut, 0xff, BufferSize);
+					}
+					else if (Info.Depth == 16)
+					{
+						const uint32 Channels = 4;
+						const uint32 BufferSize = Info.Width * Info.Height * Channels * sizeof(uint16);
+						FMemory::Memset(pOut, 0xff, BufferSize);
+					}
+				}
+
+				// Uncompressed?
+				if (CompressionType == 0)
+				{
+					if (Info.Depth == 8)
+					{
+						FColor* Dest = (FColor*)pOut;
+						for (int32 Pixel = 0; Pixel < NPixels; Pixel++)
+						{
+							if (Info.nChannels == 1)
+							{
+								Dest[Pixel].R = pCur[NPixels + Pixel];
+								Dest[Pixel].G = pCur[NPixels + Pixel];
+								Dest[Pixel].B = pCur[NPixels + Pixel];
+							}
+							else
+							{
+								Dest[Pixel].R = pCur[NPixels * 0 + Pixel];
+								Dest[Pixel].G = pCur[NPixels * 1 + Pixel];
+								Dest[Pixel].B = pCur[NPixels * 2 + Pixel];
+								if (Info.nChannels == 4)
+								{
+									Dest[Pixel].A = pCur[NPixels * 3 + Pixel];
+								}
+							}
+						}
+					}
+					else if (Info.Depth == 16)
+					{
+						uint32 SrcOffset = 0;
+
+						if (Info.nChannels == 1)
+						{
+							uint16* Dest = (uint16*)pOut;
+							uint32 ChannelOffset = 0;
+
+							for (int32 Pixel = 0; Pixel < NPixels; Pixel++)
+							{
+								Dest[ChannelOffset + 0] = ((pCur[SrcOffset] << 8) + (pCur[SrcOffset + 1] << 0));
+								Dest[ChannelOffset + 1] = ((pCur[SrcOffset] << 8) + (pCur[SrcOffset + 1] << 0));
+								Dest[ChannelOffset + 2] = ((pCur[SrcOffset] << 8) + (pCur[SrcOffset + 1] << 0));
+
+								//Increment offsets
+								ChannelOffset += 4;
+								SrcOffset += BytesPerChannel;
+							}
+						}
+						else
+						{
+							// Loop through the planes	
+							for (iPlane = 0; iPlane < Info.nChannels; iPlane++)
+							{
+								uint16* Dest = (uint16*)pOut;
+								uint32 ChannelOffset = iPlane;
+
+								for (int32 Pixel = 0; Pixel < NPixels; Pixel++)
+								{
+									Dest[ChannelOffset] = ((pCur[SrcOffset] << 8) + (pCur[SrcOffset + 1] << 0));
+
+									//Increment offsets
+									ChannelOffset += 4;
+									SrcOffset += BytesPerChannel;
+								}
+							}
+						}
+					}
+				}
+				// RLE?
+				else if (CompressionType == 1)
+				{
+					// Setup RowTable
+					pRowTable = pCur;
+					pCur += Info.nChannels * Info.Height * 2;
+
+					FColor* Dest = (FColor*)pOut;
+
+					// Loop through the planes
+					for (iPlane = 0; iPlane < Info.nChannels; iPlane++)
+					{
+						int32 iWritePlane = iPlane;
+						if (iWritePlane > BytesPerPixel - 1) iWritePlane = BytesPerPixel - 1;
+
+						// Loop through the rows
+						for (iRow = 0; iRow < Info.Height; iRow++)
+						{
+							// Load a row
+							CompressedBytes = (pRowTable[(iPlane * Info.Height + iRow) * 2] << 8) +
+								(pRowTable[(iPlane * Info.Height + iRow) * 2 + 1] << 0);
+
+							// Setup Plane
+							pPlane = pCur;
+							pCur += CompressedBytes;
+
+							// Decompress Row
+							iPixel = 0;
+							iByte = 0;
+							while ((iPixel < Info.Width) && (iByte < CompressedBytes))
+							{
+								int8 code = (int8)pPlane[iByte++];
+
+								// Is it a repeat?
+								if (code < 0)
+								{
+									Count = -(int32)code + 1;
+									Value = pPlane[iByte++];
+									while (Count-- > 0)
+									{
+										int32 idx = (iPixel)+(iRow * Info.Width);
+										if (Info.nChannels == 1)
+										{
+											Dest[idx].R = Value;
+											Dest[idx].G = Value;
+											Dest[idx].B = Value;
+										}
+										else
+										{
+											switch (iWritePlane)
+											{
+												case 0: Dest[idx].R = Value; break;
+												case 1: Dest[idx].G = Value; break;
+												case 2: Dest[idx].B = Value; break;
+												case 3: Dest[idx].A = Value; break;
+											}
+										}
+										iPixel++;
+									}
+								}
+								// Must be a literal then
+								else
+								{
+									Count = (int32)code + 1;
+									while (Count-- > 0)
+									{
+										Value = pPlane[iByte++];
+										int32 idx = (iPixel)+(iRow * Info.Width);
+
+										if (Info.nChannels == 1)
+										{
+											Dest[idx].R = Value;
+											Dest[idx].G = Value;
+											Dest[idx].B = Value;
+										}
+										else
+										{
+											switch (iWritePlane)
+											{
+												case 0: Dest[idx].R = Value; break;
+												case 1: Dest[idx].G = Value; break;
+												case 2: Dest[idx].B = Value; break;
+												case 3: Dest[idx].A = Value; break;
+											}
+										}
+										iPixel++;
+									}
+								}
+							}
+
+							// Confirm that we decoded the right number of bytes
+							check(iByte == CompressedBytes);
+							check(iPixel == Info.Width);
+						}
+					}
+				}
+				else
+					return false;
+
+				// Success!
+				return true;
+			}
+
+			static void GetPSDHeader(const uint8* Buffer, FPSDFileHeader& Info)
+			{
+				Info.Signature = ((int32)Buffer[0] << 24) +
+					((int32)Buffer[1] << 16) +
+					((int32)Buffer[2] << 8) +
+					((int32)Buffer[3] << 0);
+
+				Info.Version = ((int32)Buffer[4] << 8) +
+					((int32)Buffer[5] << 0);
+
+				Info.nChannels = ((int32)Buffer[12] << 8) +
+					((int32)Buffer[13] << 0);
+
+				Info.Height = ((int32)Buffer[14] << 24) +
+					((int32)Buffer[15] << 16) +
+					((int32)Buffer[16] << 8) +
+					((int32)Buffer[17] << 0);
+
+				Info.Width = ((int32)Buffer[18] << 24) +
+					((int32)Buffer[19] << 16) +
+					((int32)Buffer[20] << 8) +
+					((int32)Buffer[21] << 0);
+
+				Info.Depth = ((int32)Buffer[22] << 8) +
+					((int32)Buffer[23] << 0);
+
+				Info.Mode = ((int32)Buffer[24] << 8) +
+					((int32)Buffer[25] << 0);
+			}
+
+		}//ns Private
+	}//ns Interchange
+}//ns UE
+
+
+bool UInterchangePSDTranslator::CanImportSourceData(const UInterchangeSourceData* InSourceData) const
+{
+	FString Extension = FPaths::GetExtension(InSourceData->GetFilename());
+	FString PNGExtension = (TEXT("psd;Texture"));
+	return PNGExtension.StartsWith(Extension);
+}
+
+bool UInterchangePSDTranslator::Translate(UInterchangeBaseNodeContainer& BaseNodeContainer) const
+{
+	return UE::Interchange::FTextureTranslatorUtilities::Generic2DTextureTranslate(GetSourceData(), BaseNodeContainer);
+}
+
+TOptional<UE::Interchange::FImportImage> UInterchangePSDTranslator::GetTexturePayloadData(const UInterchangeSourceData* PayloadSourceData, const FString& PayLoadKey) const
+{
+	check(PayloadSourceData == GetSourceData());
+
+	if (!GetSourceData())
+	{
+		UE_LOG(LogInterchangeImport, Error, TEXT("Failed to import PSD, bad source data."));
+		return TOptional<UE::Interchange::FImportImage>();
+	}
+
+	TArray64<uint8> SourceDataBuffer;
+	FString Filename = GetSourceData()->GetFilename();
+
+	//Make sure the key fit the filename, The key should always be valid
+	if (!Filename.Equals(PayLoadKey))
+	{
+		UE_LOG(LogInterchangeImport, Error, TEXT("Failed to import PSD, wrong payload key. [%s]"), *Filename);
+		return TOptional<UE::Interchange::FImportImage>();
+	}
+
+	if (!FPaths::FileExists(Filename))
+	{
+		UE_LOG(LogInterchangeImport, Error, TEXT("Failed to import PSD, cannot open file. [%s]"), *Filename);
+		return TOptional<UE::Interchange::FImportImage>();
+	}
+
+	if (!FFileHelper::LoadFileToArray(SourceDataBuffer, *Filename))
+	{
+		UE_LOG(LogInterchangeImport, Error, TEXT("Failed to import PSD, cannot load file content into an array. [%s]"), *Filename);
+		return TOptional<UE::Interchange::FImportImage>();
+	}
+
+	const uint8* Buffer = SourceDataBuffer.GetData();
+	const uint8* BufferEnd = Buffer + SourceDataBuffer.Num();
+
+	bool bAllowNonPowerOfTwo = false;
+	GConfig->GetBool(TEXT("TextureImporter"), TEXT("AllowNonPowerOfTwoTextures"), bAllowNonPowerOfTwo, GEditorIni);
+
+	// Validate it.
+	const int32 Length = BufferEnd - Buffer;
+
+	//
+	// PSD File
+	//
+	UE::Interchange::Private::FPSDFileHeader PSDHeader;
+	if (Length < sizeof(UE::Interchange::Private::FPSDFileHeader))
+	{
+		return TOptional<UE::Interchange::FImportImage>();
+	}
+
+	UE::Interchange::Private::GetPSDHeader(Buffer, PSDHeader);
+
+	if (!PSDHeader.IsValid())
+	{
+		return TOptional<UE::Interchange::FImportImage>();
+	}
+
+	// Check the resolution of the imported texture to ensure validity
+	if (!UE::Interchange::FImportImageHelper::IsImportResolutionValid(PSDHeader.Width, PSDHeader.Height, bAllowNonPowerOfTwo))
+	{
+		return TOptional<UE::Interchange::FImportImage>();
+	}
+
+	if (!PSDHeader.IsSupported())
+	{
+		UE_LOG(LogInterchangeImport, Error, TEXT("Format of this PSD is not supported. Only Grayscale and RGBColor PSD images are currently supported, in 8-bit or 16-bit."));
+		return TOptional<UE::Interchange::FImportImage>();
+	}
+
+	// Select the texture's source format
+	ETextureSourceFormat TextureFormat = TSF_Invalid;
+	if (PSDHeader.Depth == 8)
+	{
+		TextureFormat = TSF_BGRA8;
+	}
+	else if (PSDHeader.Depth == 16)
+	{
+		TextureFormat = TSF_RGBA16;
+	}
+
+	if (TextureFormat == TSF_Invalid)
+	{
+		UE_LOG(LogInterchangeImport, Error, TEXT("PSD file contains data in an unsupported format."));
+		return TOptional<UE::Interchange::FImportImage>();
+	}
+
+	UE::Interchange::FImportImage PayloadData;
+
+	// The PSD is supported. Load it up.        
+	PayloadData.Init2DWithOneMip(
+		PSDHeader.Width,
+		PSDHeader.Height,
+		TextureFormat
+	);
+
+	if (!UE::Interchange::Private::ReadData(static_cast<uint8*>(PayloadData.RawData.GetData()), Buffer, PSDHeader))
+	{
+		UE_LOG(LogInterchangeImport, Error, TEXT("Failed to read this PSD"));
+		return TOptional<UE::Interchange::FImportImage>();
+	}
+
+	return PayloadData;
+}
diff --git a/Engine/Plugins/Experimental/Interchange/Runtime/Source/Import/Public/Mesh/InterchangeStaticMeshFactory.h b/Engine/Plugins/Experimental/Interchange/Runtime/Source/Import/Public/Mesh/InterchangeStaticMeshFactory.h
new file mode 100644
--- /dev/null
+++ b/Engine/Plugins/Experimental/Interchange/Runtime/Source/Import/Public/Mesh/InterchangeStaticMeshFactory.h
@@ -0,0 +1,29 @@
+// Copyright Epic Games, Inc. All Rights Reserved.
+
+#pragma once
+
+#include "CoreMinimal.h"
+#include "InterchangeFactoryBase.h"
+#include "UObject/Object.h"
+#include "UObject/ObjectMacros.h"
+
+#include "InterchangeStaticMeshFactory.generated.h"
+
+UCLASS(BlueprintType, Experimental)
+class INTERCHANGEIMPORT_API UInterchangeStaticMeshFactory : public UInterchangeFactoryBase
+{
+	GENERATED_BODY()
+public:
+
+	//////////////////////////////////////////////////////////////////////////
+	// Interchange factory base interface begin
+
+	virtual UClass* GetFactoryClass() const override;
+	virtual UObject* CreateEmptyAsset(const FCreateAssetParams& Arguments) override;
+	virtual UObject* CreateAsset(const FCreateAssetParams& Arguments) override;
+	virtual void PreImportPreCompletedCallback(const FImportPreCompletedCallbackParams& Arguments) override;
+	//virtual void PostImportPreCompletedCallback(const FImportPreCompletedCallbackParams& Arguments) const override;
+
+	// Interchange factory base interface end
+	//////////////////////////////////////////////////////////////////////////
+};
diff --git a/Engine/Plugins/Experimental/Interchange/Runtime/Source/Import/Public/Mesh/InterchangeStaticMeshPayload.h b/Engine/Plugins/Experimental/Interchange/Runtime/Source/Import/Public/Mesh/InterchangeStaticMeshPayload.h
new file mode 100644
--- /dev/null
+++ b/Engine/Plugins/Experimental/Interchange/Runtime/Source/Import/Public/Mesh/InterchangeStaticMeshPayload.h
@@ -0,0 +1,18 @@
+// Copyright Epic Games, Inc. All Rights Reserved.
+
+#pragma once
+
+#include "CoreMinimal.h"
+#include "MeshDescription.h"
+
+namespace UE
+{
+	namespace Interchange
+	{
+		struct FStaticMeshPayloadData
+		{
+			FMeshDescription MeshDescription;
+		};
+
+	} // namespace Interchange
+} // namespace UE
diff --git a/Engine/Plugins/Experimental/Interchange/Runtime/Source/Import/Public/Texture/InterchangeBMPTranslator.h b/Engine/Plugins/Experimental/Interchange/Runtime/Source/Import/Public/Texture/InterchangeBMPTranslator.h
new file mode 100644
--- /dev/null
+++ b/Engine/Plugins/Experimental/Interchange/Runtime/Source/Import/Public/Texture/InterchangeBMPTranslator.h
@@ -0,0 +1,48 @@
+// Copyright Epic Games, Inc. All Rights Reserved.
+
+#pragma once
+
+#include "CoreMinimal.h"
+#include "UObject/Object.h"
+#include "UObject/ObjectMacros.h"
+#include "InterchangeTranslatorBase.h"
+#include "Nodes/InterchangeBaseNodeContainer.h"
+#include "Texture/InterchangeTexturePayloadData.h"
+#include "Texture/InterchangeTexturePayloadInterface.h"
+
+#include "InterchangeBMPTranslator.generated.h"
+
+UCLASS(BlueprintType, Experimental)
+class INTERCHANGEIMPORT_API UInterchangeBMPTranslator : public UInterchangeTranslatorBase, public IInterchangeTexturePayloadInterface
+{
+	GENERATED_BODY()
+public:
+
+	/*
+	 * return true if the translator can translate the specified source data.
+	 */
+	virtual bool CanImportSourceData(const UInterchangeSourceData* InSourceData) const override;
+
+	/**
+	 * Translate the associated source data into a node hold by the specified nodes container.
+	 *
+	 * @param BaseNodeContainer - The unreal objects descriptions container where to put the translated source data.
+	 * @return true if the translator can translate the source data, false otherwise.
+	 */
+	virtual bool Translate(UInterchangeBaseNodeContainer& BaseNodeContainer) const override;
+
+
+	/* IInterchangeTexturePayloadInterface Begin */
+
+	/**
+	 * Once the translation is done, the import process need a way to retrieve payload data.
+	 * This payload will be use by the factories to create the asset.
+	 *
+	 * @param PayloadSourceData - The source data containing the data to translate
+	 * @param PayloadKey - The key to retrieve the a particular payload contain into the specified source data.
+	 * @return a PayloadData containing the import image data. The TOptional will not be set if there is an error.
+	 */
+	virtual TOptional<UE::Interchange::FImportImage> GetTexturePayloadData(const UInterchangeSourceData* PayloadSourceData, const FString& PayLoadKey) const override;
+
+	/* IInterchangeTexturePayloadInterface End */
+};
diff --git a/Engine/Plugins/Experimental/Interchange/Runtime/Source/Import/Public/Texture/InterchangeEXRTranslator.h b/Engine/Plugins/Experimental/Interchange/Runtime/Source/Import/Public/Texture/InterchangeEXRTranslator.h
new file mode 100644
--- /dev/null
+++ b/Engine/Plugins/Experimental/Interchange/Runtime/Source/Import/Public/Texture/InterchangeEXRTranslator.h
@@ -0,0 +1,49 @@
+// Copyright Epic Games, Inc. All Rights Reserved.
+
+#pragma once
+
+#include "CoreMinimal.h"
+#include "UObject/Object.h"
+#include "UObject/ObjectMacros.h"
+#include "InterchangeTranslatorBase.h"
+#include "Nodes/InterchangeBaseNodeContainer.h"
+#include "Texture/InterchangeTexturePayloadData.h"
+#include "Texture/InterchangeTexturePayloadInterface.h"
+
+#include "InterchangeEXRTranslator.generated.h"
+
+UCLASS(BlueprintType, Experimental)
+class INTERCHANGEIMPORT_API UInterchangeEXRTranslator : public UInterchangeTranslatorBase, public IInterchangeTexturePayloadInterface
+{
+	GENERATED_BODY()
+public:
+	
+	/*
+	 * return true if the translator can translate the specified source data.
+	 */
+	virtual bool CanImportSourceData(const UInterchangeSourceData* InSourceData) const override;
+
+	/**
+	 * Translate the associated source data into a node hold by the specified nodes container.
+	 *
+	 * @param BaseNodeContainer - The unreal objects descriptions container where to put the translated source data.
+	 * @return true if the translator can translate the source data, false otherwise.
+	 */
+	virtual bool Translate(UInterchangeBaseNodeContainer& BaseNodeContainer) const override;
+
+	/* IInterchangeTexturePayloadInterface Begin */
+
+	/**
+	 * Once the translation is done, the import process need a way to retrieve payload data.
+	 * This payload will be use by the factories to create the asset.
+	 *
+	 * @param PayloadSourceData - The source data containing the data to translate
+	 * @param PayloadKey - The key to retrieve the a particular payload contain into the specified source data.
+	 * @return a PayloadData containing the import image data. The TOptional will not be set if there is an error.
+	 */
+	virtual TOptional<UE::Interchange::FImportImage> GetTexturePayloadData(const UInterchangeSourceData* PayloadSourceData, const FString& PayLoadKey) const override;
+
+	/* IInterchangeTexturePayloadInterface End */
+};
+
+
diff --git a/Engine/Plugins/Experimental/Interchange/Runtime/Source/Import/Public/Texture/InterchangeHDRTranslator.h b/Engine/Plugins/Experimental/Interchange/Runtime/Source/Import/Public/Texture/InterchangeHDRTranslator.h
new file mode 100644
--- /dev/null
+++ b/Engine/Plugins/Experimental/Interchange/Runtime/Source/Import/Public/Texture/InterchangeHDRTranslator.h
@@ -0,0 +1,48 @@
+// Copyright Epic Games, Inc. All Rights Reserved.
+
+#pragma once
+
+#include "CoreMinimal.h"
+#include "UObject/Object.h"
+#include "UObject/ObjectMacros.h"
+#include "InterchangeTranslatorBase.h"
+#include "Nodes/InterchangeBaseNodeContainer.h"
+#include "Texture/InterchangeTexturePayloadData.h"
+#include "Texture/InterchangeTexturePayloadInterface.h"
+
+#include "InterchangeHDRTranslator.generated.h"
+
+UCLASS(BlueprintType, Experimental)
+class INTERCHANGEIMPORT_API UInterchangeHDRTranslator : public UInterchangeTranslatorBase, public IInterchangeTexturePayloadInterface
+{
+	GENERATED_BODY()
+public:
+
+	/*
+	 * return true if the translator can translate the specified file.
+	 */
+	virtual bool CanImportSourceData(const UInterchangeSourceData* InSourceData) const override;
+
+	/**
+	 * Translate the associated source data into a node hold by the specified nodes container.
+	 *
+	 * @param BaseNodeContainer - The unreal objects descriptions container where to put the translated source data.
+	 * @return true if the translator can translate the source data, false otherwise.
+	 */
+	virtual bool Translate(UInterchangeBaseNodeContainer& BaseNodeContainer) const override;
+
+
+	/* IInterchangeTexturePayloadInterface Begin */
+
+	/**
+	 * Once the translation is done, the import process need a way to retrieve payload data.
+	 * This payload will be use by the factories to create the asset.
+	 *
+	 * @param PayloadSourceData - The source data containing the data to translate
+	 * @param PayloadKey - The key to retrieve the a particular payload contain into the specified source data.
+	 * @return a PayloadData containing the import image data. The TOptional will not be set if there is an error.
+	 */
+	virtual TOptional<UE::Interchange::FImportImage> GetTexturePayloadData(const UInterchangeSourceData* PayloadSourceData, const FString& PayLoadKey) const override;
+
+	/* IInterchangeTexturePayloadInterface End */
+};
diff --git a/Engine/Plugins/Experimental/Interchange/Runtime/Source/Import/Public/Texture/InterchangeIESTranslator.h b/Engine/Plugins/Experimental/Interchange/Runtime/Source/Import/Public/Texture/InterchangeIESTranslator.h
new file mode 100644
--- /dev/null
+++ b/Engine/Plugins/Experimental/Interchange/Runtime/Source/Import/Public/Texture/InterchangeIESTranslator.h
@@ -0,0 +1,48 @@
+// Copyright Epic Games, Inc. All Rights Reserved.
+
+#pragma once
+
+#include "CoreMinimal.h"
+#include "UObject/Object.h"
+#include "UObject/ObjectMacros.h"
+#include "InterchangeTranslatorBase.h"
+#include "Nodes/InterchangeBaseNodeContainer.h"
+#include "Texture/InterchangeTextureLightProfilePayloadData.h"
+#include "Texture/InterchangeTextureLightProfilePayloadInterface.h"
+
+#include "InterchangeIESTranslator.generated.h"
+
+UCLASS(BlueprintType, Experimental)
+class INTERCHANGEIMPORT_API UInterchangeIESTranslator : public UInterchangeTranslatorBase, public IInterchangeTextureLightProfilePayloadInterface
+{
+	GENERATED_BODY()
+public:
+
+	/*
+	 * return true if the translator can translate the specified file.
+	 */
+	virtual bool CanImportSourceData(const UInterchangeSourceData* InSourceData) const override;
+
+	/**
+	 * Translate the associated source data into a node hold by the specified nodes container.
+	 *
+	 * @param BaseNodeContainer - The unreal objects descriptions container where to put the translated source data.
+	 * @return true if the translator can translate the source data, false otherwise.
+	 */
+	virtual bool Translate(UInterchangeBaseNodeContainer& BaseNodeContainer) const override;
+
+
+	/* IInterchangeTexturePayloadInterface Begin */
+
+	/**
+	 * Once the translation is done, the import process need a way to retrieve payload data.
+	 * This payload will be use by the factories to create the asset.
+	 *
+	 * @param SourceData - The source data containing the data to translate
+	 * @param PayloadKey - The key to retrieve the a particular payload contain into the specified source data.
+	 * @return a PayloadData containing the import light profile data. The TOptional will not be set if there is an error.
+	 */
+	virtual TOptional<UE::Interchange::FImportLightProfile> GetLightProfilePayloadData(const UInterchangeSourceData* PayloadSourceData, const FString& PayLoadKey) const override;
+
+	/* IInterchangeTexturePayloadInterface End */
+};
\ No newline at end of file
diff --git a/Engine/Plugins/Experimental/Interchange/Runtime/Source/Import/Public/Texture/InterchangeJPGTranslator.h b/Engine/Plugins/Experimental/Interchange/Runtime/Source/Import/Public/Texture/InterchangeJPGTranslator.h
new file mode 100644
--- /dev/null
+++ b/Engine/Plugins/Experimental/Interchange/Runtime/Source/Import/Public/Texture/InterchangeJPGTranslator.h
@@ -0,0 +1,49 @@
+// Copyright Epic Games, Inc. All Rights Reserved.
+
+#pragma once
+
+#include "CoreMinimal.h"
+#include "UObject/Object.h"
+#include "UObject/ObjectMacros.h"
+#include "InterchangeTranslatorBase.h"
+#include "Nodes/InterchangeBaseNodeContainer.h"
+#include "Texture/InterchangeTexturePayloadData.h"
+#include "Texture/InterchangeTexturePayloadInterface.h"
+
+#include "InterchangeJPGTranslator.generated.h"
+
+UCLASS(BlueprintType, Experimental)
+class INTERCHANGEIMPORT_API UInterchangeJPGTranslator : public UInterchangeTranslatorBase, public IInterchangeTexturePayloadInterface
+{
+	GENERATED_BODY()
+public:
+	
+	/**
+	 * return true if the translator can translate the specified source data.
+	 */
+	virtual bool CanImportSourceData(const UInterchangeSourceData* InSourceData) const override;
+
+	/**
+	 * Translate the associated source data into a node hold by the specified nodes container.
+	 *
+	 * @param BaseNodeContainer - The unreal objects descriptions container where to put the translated source data.
+	 * @return true if the translator can translate the source data, false otherwise.
+	 */
+	virtual bool Translate(UInterchangeBaseNodeContainer& BaseNodeContainer) const override;
+
+	/* IInterchangeTexturePayloadInterface Begin */
+
+	/**
+	 * Once the translation is done, the import process need a way to retrieve payload data.
+	 * This payload will be use by the factories to create the asset.
+	 *
+	 * @param PayloadSourceData - The source data containing the data to translate
+	 * @param PayloadKey - The key to retrieve the a particular payload contain into the specified source data.
+	 * @return a PayloadData containing the import image data. The TOptional will not be set if there is an error.
+	 */
+	virtual TOptional<UE::Interchange::FImportImage> GetTexturePayloadData(const UInterchangeSourceData* PayloadSourceData, const FString& PayLoadKey) const override;
+
+	/* IInterchangeTexturePayloadInterface End */
+};
+
+
diff --git a/Engine/Plugins/Experimental/Interchange/Runtime/Source/Import/Public/Texture/InterchangePCXTranslator.h b/Engine/Plugins/Experimental/Interchange/Runtime/Source/Import/Public/Texture/InterchangePCXTranslator.h
new file mode 100644
--- /dev/null
+++ b/Engine/Plugins/Experimental/Interchange/Runtime/Source/Import/Public/Texture/InterchangePCXTranslator.h
@@ -0,0 +1,49 @@
+// Copyright Epic Games, Inc. All Rights Reserved.
+
+#pragma once
+
+#include "CoreMinimal.h"
+#include "UObject/Object.h"
+#include "UObject/ObjectMacros.h"
+#include "InterchangeTranslatorBase.h"
+#include "Nodes/InterchangeBaseNodeContainer.h"
+#include "Texture/InterchangeTexturePayloadData.h"
+#include "Texture/InterchangeTexturePayloadInterface.h"
+
+#include "InterchangePCXTranslator.generated.h"
+
+UCLASS(BlueprintType, Experimental)
+class INTERCHANGEIMPORT_API UInterchangePCXTranslator : public UInterchangeTranslatorBase, public IInterchangeTexturePayloadInterface
+{
+	GENERATED_BODY()
+public:
+	
+	/*
+	 * return true if the translator can translate the specified source data.
+	 */
+	virtual bool CanImportSourceData(const UInterchangeSourceData* InSourceData) const override;
+
+	/**
+	 * Translates the associated source data into a node hold by the specified nodes container.
+	 *
+	 * @param BaseNodeContainer - The unreal objects descriptions container where to put the translated source data.
+	 * @return true if the translator can translate the source data, false otherwise.
+	 */
+	virtual bool Translate(UInterchangeBaseNodeContainer& BaseNodeContainer) const override;
+
+	/* IInterchangeTexturePayloadInterface Begin */
+
+	/**
+	 * Once the translation is done, the import process need a way to retrieve payload data.
+	 * This payload will be use by the factories to create the asset.
+	 *
+	 * @param PayloadSourceData - The source data containing the data to translate
+	 * @param PayloadKey - The key to retrieve the a particular payload contain into the specified source data.
+	 * @return a PayloadData containing the import image data. The TOptional will not be set if there is an error.
+	 */
+	virtual TOptional<UE::Interchange::FImportImage> GetTexturePayloadData(const UInterchangeSourceData* PayloadSourceData, const FString& PayLoadKey) const override;
+
+	/* IInterchangeTexturePayloadInterface End */
+};
+
+
diff --git a/Engine/Plugins/Experimental/Interchange/Runtime/Source/Import/Public/Texture/InterchangePNGTranslator.h b/Engine/Plugins/Experimental/Interchange/Runtime/Source/Import/Public/Texture/InterchangePNGTranslator.h
new file mode 100644
--- /dev/null
+++ b/Engine/Plugins/Experimental/Interchange/Runtime/Source/Import/Public/Texture/InterchangePNGTranslator.h
@@ -0,0 +1,49 @@
+// Copyright Epic Games, Inc. All Rights Reserved.
+
+#pragma once
+
+#include "CoreMinimal.h"
+#include "UObject/Object.h"
+#include "UObject/ObjectMacros.h"
+#include "InterchangeTranslatorBase.h"
+#include "Nodes/InterchangeBaseNodeContainer.h"
+#include "Texture/InterchangeTexturePayloadData.h"
+#include "Texture/InterchangeTexturePayloadInterface.h"
+
+#include "InterchangePNGTranslator.generated.h"
+
+UCLASS(BlueprintType, Experimental)
+class INTERCHANGEIMPORT_API UInterchangePNGTranslator : public UInterchangeTranslatorBase, public IInterchangeTexturePayloadInterface
+{
+	GENERATED_BODY()
+public:
+	
+	/*
+	 * return true if the translator can translate the specified source data.
+	 */
+	virtual bool CanImportSourceData(const UInterchangeSourceData* InSourceData) const override;
+
+	/**
+	 * Translate the associated source data into a node hold by the specified nodes container.
+	 *
+	 * @param BaseNodeContainer - The unreal objects descriptions container where to put the translated source data.
+	 * @return true if the translator can translate the source data, false otherwise.
+	 */
+	virtual bool Translate(UInterchangeBaseNodeContainer& BaseNodeContainer) const override;
+
+	/* IInterchangeTexturePayloadInterface Begin */
+
+	/**
+	 * Once the translation is done, the import process need a way to retrieve payload data.
+	 * This payload will be use by the factories to create the asset.
+	 *
+	 * @param PayloadSourceData - The source data containing the data to translate
+	 * @param PayloadKey - The key to retrieve the a particular payload contain into the specified source data.
+	 * @return a PayloadData containing the import image data. The TOptional will not be set if there is an error.
+	 */
+	virtual TOptional<UE::Interchange::FImportImage> GetTexturePayloadData(const UInterchangeSourceData* PayloadSourceData, const FString& PayLoadKey) const override;
+
+	/* IInterchangeTexturePayloadInterface End */
+};
+
+
diff --git a/Engine/Plugins/Experimental/Interchange/Runtime/Source/Import/Public/Texture/InterchangePSDTranslator.h b/Engine/Plugins/Experimental/Interchange/Runtime/Source/Import/Public/Texture/InterchangePSDTranslator.h
new file mode 100644
--- /dev/null
+++ b/Engine/Plugins/Experimental/Interchange/Runtime/Source/Import/Public/Texture/InterchangePSDTranslator.h
@@ -0,0 +1,49 @@
+// Copyright Epic Games, Inc. All Rights Reserved.
+
+#pragma once
+
+#include "CoreMinimal.h"
+#include "UObject/Object.h"
+#include "UObject/ObjectMacros.h"
+#include "InterchangeTranslatorBase.h"
+#include "Nodes/InterchangeBaseNodeContainer.h"
+#include "Texture/InterchangeTexturePayloadData.h"
+#include "Texture/InterchangeTexturePayloadInterface.h"
+
+#include "InterchangePSDTranslator.generated.h"
+
+UCLASS(BlueprintType, Experimental)
+class INTERCHANGEIMPORT_API UInterchangePSDTranslator : public UInterchangeTranslatorBase, public IInterchangeTexturePayloadInterface
+{
+	GENERATED_BODY()
+public:
+
+	/*
+	 * return true if the translator can translate the specified source data.
+	 */
+	virtual bool CanImportSourceData(const UInterchangeSourceData* InSourceData) const override;
+
+	/**
+	 * Translate the associated source data into a node hold by the specified nodes container.
+	 *
+	 * @param BaseNodeContainer - The unreal objects descriptions container where to put the translated source data.
+	 * @return true if the translator can translate the source data, false otherwise.
+	 */
+	virtual bool Translate(UInterchangeBaseNodeContainer& BaseNodeContainer) const override;
+
+	/* IInterchangeTexturePayloadInterface Begin */
+
+	/**
+	 * Once the translation is done, the import process need a way to retrieve payload data.
+	 * This payload will be use by the factories to create the asset.
+	 *
+	 * @param PayloadSourceData - The source data containing the data to translate
+	 * @param PayloadKey - The key to retrieve the a particular payload contain into the specified source data.
+	 * @return a PayloadData containing the import image data. The TOptional will not be set if there is an error.
+	 */
+	virtual TOptional<UE::Interchange::FImportImage> GetTexturePayloadData(const UInterchangeSourceData* PayloadSourceData, const FString& PayLoadKey) const override;
+
+	/* IInterchangeTexturePayloadInterface End */
+};
+
+
diff --git a/Engine/Plugins/Experimental/Interchange/Runtime/Source/Import/Public/Texture/InterchangeTGATranslator.h b/Engine/Plugins/Experimental/Interchange/Runtime/Source/Import/Public/Texture/InterchangeTGATranslator.h
new file mode 100644
--- /dev/null
+++ b/Engine/Plugins/Experimental/Interchange/Runtime/Source/Import/Public/Texture/InterchangeTGATranslator.h
@@ -0,0 +1,48 @@
+// Copyright Epic Games, Inc. All Rights Reserved.
+#pragma once
+
+#include "CoreMinimal.h"
+#include "UObject/Object.h"
+#include "UObject/ObjectMacros.h"
+#include "InterchangeTranslatorBase.h"
+#include "Nodes/InterchangeBaseNodeContainer.h"
+#include "Texture/InterchangeTexturePayloadData.h"
+#include "Texture/InterchangeTexturePayloadInterface.h"
+
+#include "InterchangeTGATranslator.generated.h"
+
+UCLASS(BlueprintType, Experimental)
+class INTERCHANGEIMPORT_API UInterchangeTGATranslator : public UInterchangeTranslatorBase, public IInterchangeTexturePayloadInterface
+{
+	GENERATED_BODY()
+public:
+	
+	/**
+	 * return true if the translator can translate the specified source data.
+	 */
+	virtual bool CanImportSourceData(const UInterchangeSourceData* InSourceData) const override;
+
+	/**
+	 * Translates the associated source data into a node hold by the specified nodes container.
+	 *
+	 * @param BaseNodeContainer - The unreal objects descriptions container where to put the translated source data.
+	 * @return true if the translator can translate the source data, false otherwise.
+	 */
+	virtual bool Translate(UInterchangeBaseNodeContainer& BaseNodeContainer) const override;
+
+	/* IInterchangeTexturePayloadInterface Begin */
+
+	/**
+	 * Once the translation is done, the import process need a way to retrieve payload data.
+	 * This payload will be use by the factories to create the asset.
+	 *
+	 * @param PayloadSourceData - The source data containing the data to translate
+	 * @param PayloadKey - The key to retrieve the a particular payload contain into the specified source data.
+	 * @return a PayloadData containing the import image data. The TOptional will not be set if there is an error.
+	 */
+	virtual TOptional<UE::Interchange::FImportImage> GetTexturePayloadData(const UInterchangeSourceData* PayloadSourceData, const FString& PayLoadKey) const override;
+
+	/* IInterchangeTexturePayloadInterface End */
+};
+
+
diff --git a/Engine/Plugins/Experimental/Interchange/Runtime/Source/Messages/Public/Fbx/InterchangeFbxMessages.h b/Engine/Plugins/Experimental/Interchange/Runtime/Source/Messages/Public/Fbx/InterchangeFbxMessages.h
new file mode 100644
--- /dev/null
+++ b/Engine/Plugins/Experimental/Interchange/Runtime/Source/Messages/Public/Fbx/InterchangeFbxMessages.h
@@ -0,0 +1,87 @@
+// Copyright Epic Games, Inc. All Rights Reserved.
+
+#pragma once
+
+#include "CoreMinimal.h"
+#include "InterchangeResult.h"
+
+#include "InterchangeFbxMessages.generated.h"
+
+
+/**
+ * Base class for FBX parser warnings
+ */
+UCLASS()
+class INTERCHANGEMESSAGES_API UInterchangeResultMeshWarning : public UInterchangeResultWarning
+{
+	GENERATED_BODY()
+
+public:
+
+	UPROPERTY()
+	FString MeshName;
+};
+
+
+/**
+ * Base class for FBX parser errors
+ */
+UCLASS()
+class INTERCHANGEMESSAGES_API UInterchangeResultMeshError : public UInterchangeResultError
+{
+	GENERATED_BODY()
+
+public:
+
+	UPROPERTY()
+	FString MeshName;
+};
+
+
+/**
+ * A generic class for FBX parser warnings, with no additional metadata, and where the text is specified by the user
+ */
+UCLASS()
+class INTERCHANGEMESSAGES_API UInterchangeResultMeshWarning_Generic : public UInterchangeResultMeshWarning
+{
+	GENERATED_BODY()
+
+public:
+	virtual FText GetText() const override;
+
+	UPROPERTY()
+	FText Text;
+};
+
+
+/**
+ * A generic class for FBX parser errors, with no additional metadata, and where the text is specified by the user
+ */
+UCLASS()
+class INTERCHANGEMESSAGES_API UInterchangeResultMeshError_Generic : public UInterchangeResultMeshError
+{
+	GENERATED_BODY()
+
+public:
+	virtual FText GetText() const override;
+
+	UPROPERTY()
+	FText Text;
+};
+
+
+/**
+ * 
+ */
+UCLASS()
+class INTERCHANGEMESSAGES_API UInterchangeResultMeshWarning_TooManyUVs : public UInterchangeResultMeshWarning
+{
+	GENERATED_BODY()
+
+public:
+	virtual FText GetText() const override;
+
+	UPROPERTY()
+	int32 ExcessUVs;
+};
+
