Move all interchange code files under experimental folder.

diff --git a/Engine/Plugins/Experimental/Interchange/Runtime/Source/Nodes/Public/InterchangeStaticMeshFactoryNode.h b/Engine/Plugins/Experimental/Interchange/Runtime/Source/Nodes/Public/InterchangeStaticMeshFactoryNode.h
new file mode 100644
--- /dev/null
+++ b/Engine/Plugins/Experimental/Interchange/Runtime/Source/Nodes/Public/InterchangeStaticMeshFactoryNode.h
@@ -0,0 +1,232 @@
+// Copyright Epic Games, Inc. All Rights Reserved.
+
+#pragma once
+
+#include "CoreMinimal.h"
+#include "Nodes/InterchangeBaseNode.h"
+
+#if WITH_ENGINE
+#include "Engine/StaticMesh.h"
+#endif
+
+#include "InterchangeStaticMeshFactoryNode.generated.h"
+
+
+namespace UE
+{
+	namespace Interchange
+	{
+		struct FStaticMeshNodeStaticData : public FBaseNodeStaticData
+		{
+			static const FString& GetLodDependenciesBaseKey()
+			{
+				static FString LodDependencies_BaseKey = TEXT("Lod_Dependencies");
+				return LodDependencies_BaseKey;
+			}
+		};
+
+	} // namespace Interchange
+} // namespace UE
+
+
+UCLASS(BlueprintType, Experimental)
+class INTERCHANGENODES_API UInterchangeStaticMeshFactoryNode : public UInterchangeBaseNode
+{
+	GENERATED_BODY()
+
+public:
+	UInterchangeStaticMeshFactoryNode()
+	{
+#if WITH_ENGINE
+		AssetClass = nullptr;
+#endif
+		LodDependencies.Initialize(Attributes, UE::Interchange::FStaticMeshNodeStaticData::GetLodDependenciesBaseKey());
+	}
+
+	/**
+	 * Initialize node data
+	 * @param UniqueID - The uniqueId for this node
+	 * @param DisplayLabel - The name of the node
+	 * @param InAssetClass - The class the StaticMesh factory will create for this node.
+	 */
+	UFUNCTION(BlueprintCallable, Category = "Interchange | Node | StaticMesh")
+	void InitializeStaticMeshNode(const FString& UniqueID, const FString& DisplayLabel, const FString& InAssetClass)
+	{
+		bIsNodeClassInitialized = false;
+		InitializeNode(UniqueID, DisplayLabel, EInterchangeNodeContainerType::NodeContainerType_FactoryData);
+
+		FString OperationName = GetTypeName() + TEXT(".SetAssetClassName");
+		InterchangePrivateNodeBase::SetCustomAttribute<FString>(*Attributes, ClassNameAttributeKey, OperationName, InAssetClass);
+		FillAssetClassFromAttribute();
+	}
+
+	virtual void Serialize(FArchive& Ar) override
+	{
+		Super::Serialize(Ar);
+#if WITH_ENGINE
+		if (Ar.IsLoading())
+		{
+			// Make sure the class is properly set when we compile with engine, this will set the bIsNodeClassInitialized to true.
+			SetNodeClassFromClassAttribute();
+		}
+#endif //#if WITH_ENGINE
+	}
+
+	/**
+	 * Return the node type name of the class, we use this when reporting error
+	 */
+	virtual FString GetTypeName() const override
+	{
+		const FString TypeName = TEXT("StaticMeshNode");
+		return TypeName;
+	}
+
+	virtual FString GetKeyDisplayName(const UE::Interchange::FAttributeKey& NodeAttributeKey) const override
+	{
+		FString KeyDisplayName = NodeAttributeKey.Key;
+		if (NodeAttributeKey.Key.Equals(UE::Interchange::FStaticMeshNodeStaticData::GetLodDependenciesBaseKey()))
+		{
+			KeyDisplayName = TEXT("LOD Dependencies Count");
+			return KeyDisplayName;
+		}
+		else if (NodeAttributeKey.Key.StartsWith(UE::Interchange::FStaticMeshNodeStaticData::GetLodDependenciesBaseKey()))
+		{
+			KeyDisplayName = TEXT("LOD Dependencies Index ");
+			const FString IndexKey = UE::Interchange::FNameAttributeArrayHelper::IndexKey();
+			int32 IndexPosition = NodeAttributeKey.Key.Find(IndexKey) + IndexKey.Len();
+			if (IndexPosition < NodeAttributeKey.Key.Len())
+			{
+				KeyDisplayName += NodeAttributeKey.Key.RightChop(IndexPosition);
+			}
+			return KeyDisplayName;
+		}
+		return Super::GetKeyDisplayName(NodeAttributeKey);
+	}
+
+	/** Get the class this node want to create */
+	UFUNCTION(BlueprintCallable, Category = "Interchange | Node | StaticMesh")
+	virtual class UClass* GetObjectClass() const override
+	{
+		ensure(bIsNodeClassInitialized);
+#if WITH_ENGINE
+		return AssetClass.Get() != nullptr ? AssetClass.Get() : UStaticMesh::StaticClass();
+#else
+		return nullptr;
+#endif
+	}
+
+	virtual FGuid GetHash() const override
+	{
+		return Attributes->GetStorageHash();
+	}
+
+public:
+	/** Return The number of LOD this static mesh has.*/
+	UFUNCTION(BlueprintCallable, Category = "Interchange | Node | StaticMesh")
+	int32 GetLodDataCount() const
+	{
+		return LodDependencies.GetCount();
+	}
+
+	UFUNCTION(BlueprintCallable, Category = "Interchange | Node | StaticMesh")
+	void GetLodDataUniqueIds(TArray<FString>& OutLodDataUniqueIds) const
+	{
+		LodDependencies.GetNames(OutLodDataUniqueIds);
+	}
+
+	UFUNCTION(BlueprintCallable, Category = "Interchange | Node | StaticMesh")
+	bool AddLodDataUniqueId(const FString& LodDataUniqueId)
+	{
+		return LodDependencies.AddName(LodDataUniqueId);
+	}
+
+	UFUNCTION(BlueprintCallable, Category = "Interchange | Node | StaticMesh")
+	bool RemoveLodDataUniqueId(const FString& LodDataUniqueId)
+	{
+		return LodDependencies.RemoveName(LodDataUniqueId);
+	}
+
+	/** Return false if the Attribute was not set previously.*/
+	UFUNCTION(BlueprintCallable, Category = "Interchange | Node | StaticMesh")
+	bool GetCustomVertexColorReplace(bool& AttributeValue) const
+	{
+		IMPLEMENT_NODE_ATTRIBUTE_GETTER(VertexColorReplace, bool)
+	}
+
+	UFUNCTION(BlueprintCallable, Category = "Interchange | Node | StaticMesh")
+	bool SetCustomVertexColorReplace(const bool& AttributeValue)
+	{
+		IMPLEMENT_NODE_ATTRIBUTE_SETTER_NODELEGATE(VertexColorReplace, bool)
+	}
+
+	/** Return false if the Attribute was not set previously.*/
+	UFUNCTION(BlueprintCallable, Category = "Interchange | Node | StaticMesh")
+	bool GetCustomVertexColorIgnore(bool& AttributeValue) const
+	{
+		IMPLEMENT_NODE_ATTRIBUTE_GETTER(VertexColorIgnore, bool)
+	}
+
+	UFUNCTION(BlueprintCallable, Category = "Interchange | Node | StaticMesh")
+	bool SetCustomVertexColorIgnore(const bool& AttributeValue)
+	{
+		IMPLEMENT_NODE_ATTRIBUTE_SETTER_NODELEGATE(VertexColorIgnore, bool)
+	}
+
+	/** Return false if the Attribute was not set previously.*/
+	UFUNCTION(BlueprintCallable, Category = "Interchange | Node | StaticMesh")
+	bool GetCustomVertexColorOverride(FColor& AttributeValue) const
+	{
+		IMPLEMENT_NODE_ATTRIBUTE_GETTER(VertexColorOverride, FColor)
+	}
+
+	UFUNCTION(BlueprintCallable, Category = "Interchange | Node | StaticMesh")
+	bool SetCustomVertexColorOverride(const FColor& AttributeValue)
+	{
+		IMPLEMENT_NODE_ATTRIBUTE_SETTER_NODELEGATE(VertexColorOverride, FColor)
+	}
+private:
+
+	void FillAssetClassFromAttribute()
+	{
+#if WITH_ENGINE
+		FString OperationName = GetTypeName() + TEXT(".GetAssetClassName");
+		FString ClassName;
+		InterchangePrivateNodeBase::GetCustomAttribute<FString>(*Attributes, ClassNameAttributeKey, OperationName, ClassName);
+		if (ClassName.Equals(UStaticMesh::StaticClass()->GetName()))
+		{
+			AssetClass = UStaticMesh::StaticClass();
+			bIsNodeClassInitialized = true;
+		}
+#endif
+	}
+
+	bool SetNodeClassFromClassAttribute()
+	{
+		if (!bIsNodeClassInitialized)
+		{
+			FillAssetClassFromAttribute();
+		}
+		return bIsNodeClassInitialized;
+	}
+
+	bool IsEditorOnlyDataDefined()
+	{
+#if WITH_EDITORONLY_DATA
+		return true;
+#else
+		return false;
+#endif
+	}
+
+	const UE::Interchange::FAttributeKey ClassNameAttributeKey = UE::Interchange::FBaseNodeStaticData::ClassTypeAttributeKey();
+	const UE::Interchange::FAttributeKey Macro_CustomVertexColorReplaceKey = UE::Interchange::FAttributeKey(TEXT("VertexColorReplace"));
+	const UE::Interchange::FAttributeKey Macro_CustomVertexColorIgnoreKey = UE::Interchange::FAttributeKey(TEXT("VertexColorIgnore"));
+	const UE::Interchange::FAttributeKey Macro_CustomVertexColorOverrideKey = UE::Interchange::FAttributeKey(TEXT("VertexColorOverride"));
+
+	UE::Interchange::FNameAttributeArrayHelper LodDependencies;
+protected:
+#if WITH_ENGINE
+	TSubclassOf<UStaticMesh> AssetClass = nullptr;
+#endif
+	bool bIsNodeClassInitialized = false;
+};
diff --git a/Engine/Plugins/Experimental/Interchange/Runtime/Source/Nodes/Public/InterchangeStaticMeshLodDataNode.h b/Engine/Plugins/Experimental/Interchange/Runtime/Source/Nodes/Public/InterchangeStaticMeshLodDataNode.h
new file mode 100644
--- /dev/null
+++ b/Engine/Plugins/Experimental/Interchange/Runtime/Source/Nodes/Public/InterchangeStaticMeshLodDataNode.h
@@ -0,0 +1,132 @@
+// Copyright Epic Games, Inc. All Rights Reserved.
+
+#pragma once
+
+#include "CoreMinimal.h"
+#include "Nodes/InterchangeBaseNode.h"
+
+#include "InterchangeStaticMeshLodDataNode.generated.h"
+
+
+namespace UE
+{
+	namespace Interchange
+	{
+		struct FStaticMeshNodeLodDataStaticData : public FBaseNodeStaticData
+		{
+			static const FString& GetMeshUidsBaseKey()
+			{
+				static FString MeshUids_BaseKey = TEXT("__MeshUids__Key");
+				return MeshUids_BaseKey;
+			}
+		};
+
+	} // namespace Interchange
+} // namespace UE
+
+
+UCLASS(BlueprintType, Experimental)
+class INTERCHANGENODES_API UInterchangeStaticMeshLodDataNode : public UInterchangeBaseNode
+{
+	GENERATED_BODY()
+
+public:
+	UInterchangeStaticMeshLodDataNode()
+	{
+		MeshUids.Initialize(Attributes, UE::Interchange::FStaticMeshNodeLodDataStaticData::GetMeshUidsBaseKey());
+	}
+
+	/**
+	 * Return the node type name of the class, we use this when reporting error
+	 */
+	virtual FString GetTypeName() const override
+	{
+		const FString TypeName = TEXT("StaticMeshLodDataNode");
+		return TypeName;
+	}
+
+	virtual FString GetKeyDisplayName(const UE::Interchange::FAttributeKey& NodeAttributeKey) const override
+	{
+		FString KeyDisplayName = NodeAttributeKey.Key;
+		if (NodeAttributeKey.Key.Equals(UE::Interchange::FStaticMeshNodeLodDataStaticData::GetMeshUidsBaseKey()))
+		{
+			KeyDisplayName = TEXT("Mesh count");
+			return KeyDisplayName;
+		}
+		else if (NodeAttributeKey.Key.StartsWith(UE::Interchange::FStaticMeshNodeLodDataStaticData::GetMeshUidsBaseKey()))
+		{
+			KeyDisplayName = TEXT("Mesh index ");
+			const FString IndexKey = UE::Interchange::FNameAttributeArrayHelper::IndexKey();
+			int32 IndexPosition = NodeAttributeKey.Key.Find(IndexKey) + IndexKey.Len();
+			if (IndexPosition < NodeAttributeKey.Key.Len())
+			{
+				KeyDisplayName += NodeAttributeKey.Key.RightChop(IndexPosition);
+			}
+			return KeyDisplayName;
+		}
+
+		return Super::GetKeyDisplayName(NodeAttributeKey);
+	}
+
+	virtual FString GetAttributeCategory(const UE::Interchange::FAttributeKey& NodeAttributeKey) const override
+	{
+		if (NodeAttributeKey.Key.StartsWith(UE::Interchange::FStaticMeshNodeLodDataStaticData::GetMeshUidsBaseKey()))
+		{
+			return FString(TEXT("Meshes"));
+		}
+		return Super::GetAttributeCategory(NodeAttributeKey);
+	}
+
+	virtual FGuid GetHash() const override
+	{
+		return Attributes->GetStorageHash();
+	}
+
+
+public:
+	/* Mesh Uids: It can be either a scene or a mesh node uid. If its a scene it mean we want the mesh factory to bake the geo payload with the global transform of the scene node. */
+
+	UFUNCTION(BlueprintCallable, Category = "Interchange | Node | StaticMeshLodData")
+	int32 GetMeshUidsCount() const
+	{
+		return MeshUids.GetCount();
+	}
+
+	UFUNCTION(BlueprintCallable, Category = "Interchange | Node | StaticMeshLodData")
+	void GetMeshUids(TArray<FString>& OutMeshNames) const
+	{
+		MeshUids.GetNames(OutMeshNames);
+	}
+
+	UFUNCTION(BlueprintCallable, Category = "Interchange | Node | StaticMeshLodData")
+	bool AddMeshUid(const FString& MeshName)
+	{
+		return MeshUids.AddName(MeshName);
+	}
+
+	UFUNCTION(BlueprintCallable, Category = "Interchange | Node | StaticMeshLodData")
+	bool RemoveMeshUid(const FString& MeshName)
+	{
+		return MeshUids.RemoveName(MeshName);
+	}
+
+	UFUNCTION(BlueprintCallable, Category = "Interchange | Node | StaticMeshLodData")
+	bool RemoveAllMeshes()
+	{
+		return MeshUids.RemoveAllNames();
+	}
+
+private:
+
+	bool IsEditorOnlyDataDefined()
+	{
+#if WITH_EDITORONLY_DATA
+		return true;
+#else
+		return false;
+#endif
+	}
+
+	UE::Interchange::FNameAttributeArrayHelper MeshUids;
+protected:
+};
diff --git a/Engine/Plugins/Experimental/Interchange/Runtime/Source/Parsers/Fbx/Private/FbxAPI.cpp b/Engine/Plugins/Experimental/Interchange/Runtime/Source/Parsers/Fbx/Private/FbxAPI.cpp
new file mode 100644
--- /dev/null
+++ b/Engine/Plugins/Experimental/Interchange/Runtime/Source/Parsers/Fbx/Private/FbxAPI.cpp
@@ -0,0 +1,161 @@
+// Copyright Epic Games, Inc. All Rights Reserved.
+
+#include "FbxAPI.h"
+
+#include "CoreMinimal.h"
+#include "FbxCamera.h"
+#include "FbxConvert.h"
+#include "FbxHelper.h"
+#include "FbxInclude.h"
+#include "FbxLight.h"
+#include "FbxMaterial.h"
+#include "FbxMesh.h"
+#include "FbxScene.h"
+#include "InterchangeMaterialNode.h"
+#include "InterchangeTextureNode.h"
+#include "Nodes/InterchangeBaseNodeContainer.h"
+
+#define LOCTEXT_NAMESPACE "InterchangeFbxParser"
+
+namespace UE
+{
+	namespace Interchange
+	{
+		namespace Private
+		{
+			FFbxParser::~FFbxParser()
+			{
+				PayloadContexts.Empty();
+			}
+
+			bool FFbxParser::LoadFbxFile(const FString& Filename)
+			{
+				SourceFilename = Filename;
+				int32 SDKMajor, SDKMinor, SDKRevision;
+				
+				//The first thing to do is to create the FBX Manager which is the object allocator for almost all the classes in the SDK
+				SDKManager = FbxManager::Create();
+				if (!SDKManager)
+				{
+					UInterchangeResultError_Generic* Message = AddMessage<UInterchangeResultError_Generic>();
+					Message->Text = LOCTEXT("CannotCreateFBXManager", "Cannot create FBX SDK manager.");
+					return false;
+				}
+
+				//Create an IOSettings object. This object holds all import/export settings.
+				FbxIOSettings* ios = FbxIOSettings::Create(SDKManager, IOSROOT);
+				ios->SetBoolProp(IMP_FBX_MATERIAL, true);
+				ios->SetBoolProp(IMP_FBX_TEXTURE, true);
+				ios->SetBoolProp(IMP_FBX_LINK, true);
+				ios->SetBoolProp(IMP_FBX_SHAPE, true);
+				ios->SetBoolProp(IMP_FBX_GOBO, true);
+				ios->SetBoolProp(IMP_FBX_ANIMATION, true);
+				ios->SetBoolProp(IMP_SKINS, true);
+				ios->SetBoolProp(IMP_DEFORMATION, true);
+				ios->SetBoolProp(IMP_FBX_GLOBAL_SETTINGS, true);
+				ios->SetBoolProp(IMP_TAKE, true);
+				SDKManager->SetIOSettings(ios);
+
+				SDKGeometryConverter = new FbxGeometryConverter(SDKManager);
+
+				//Create an FBX scene. This object holds most objects imported/exported from/to files.
+				SDKScene = FbxScene::Create(SDKManager, "My Scene");
+				if (!SDKScene)
+				{
+					UInterchangeResultError_Generic* Message = AddMessage<UInterchangeResultError_Generic>();
+					Message->Text = LOCTEXT("CannotCreateFBXScene", "Cannot create FBX SDK scene.");
+					return false;
+				}
+				
+				// Create an importer.
+				SDKImporter = FbxImporter::Create(SDKManager, "");
+
+				// Get the version number of the FBX files generated by the
+				// version of FBX SDK that you are using.
+				FbxManager::GetFileFormatVersion(SDKMajor, SDKMinor, SDKRevision);
+
+				// Initialize the importer by providing a filename.
+				const bool bImportStatus = SDKImporter->Initialize(TCHAR_TO_UTF8(*Filename));
+				if (!bImportStatus)
+				{
+					FFormatNamedArguments FilenameText
+					{
+						{ TEXT("Filename"), FText::FromString(Filename) }
+					};
+					UInterchangeResultError_Generic* Message = AddMessage<UInterchangeResultError_Generic>();
+					Message->Text = FText::Format(LOCTEXT("CannotOpenFBXFile", "Cannot open FBX file '{Filename}'."), FilenameText);
+					return false;
+				}
+
+				bool bStatus = SDKImporter->Import(SDKScene);
+
+				const FbxGlobalSettings& GlobalSettings = SDKScene->GetGlobalSettings();
+				FbxTime::EMode TimeMode = GlobalSettings.GetTimeMode();
+				//Set the original framerate from the current fbx file
+				float FbxFramerate = FbxTime::GetFrameRate(TimeMode);
+
+				//We always convert scene to UE axis and units
+				FFbxConvert::ConvertScene(SDKScene);
+
+				return true;
+			}
+
+			void FFbxParser::FillContainerWithFbxScene(UInterchangeBaseNodeContainer& NodeContainer)
+			{
+				CleanupFbxData();
+
+				FFbxMaterial FbxMaterial(*this);
+				FbxMaterial.AddAllTextures(SDKScene, NodeContainer);
+				FbxMaterial.AddAllMaterials(SDKScene, NodeContainer);
+
+				FFbxMesh FbxMesh(*this);
+				FbxMesh.AddAllMeshes(SDKScene, SDKGeometryConverter, NodeContainer, PayloadContexts);
+
+				FFbxLight FbxLight(*this);
+				FbxLight.AddAllLights(SDKScene, NodeContainer);
+
+				FFbxCamera FbxCamera(*this);
+				FbxCamera.AddAllCameras(SDKScene, NodeContainer);
+
+				FFbxScene FbxScene(*this);
+				FbxScene.AddHierarchy(SDKScene, NodeContainer);
+			}
+
+			bool FFbxParser::FetchPayloadData(const FString& PayloadKey, const FString& PayloadFilepath)
+			{
+				if (!PayloadContexts.Contains(PayloadKey))
+				{
+					UInterchangeResultError_Generic* Message = AddMessage<UInterchangeResultError_Generic>();
+					Message->Text = LOCTEXT("CannotRetrievePayload", "Cannot retrieve payload; payload key doesn't have any context.");
+					return false;
+				}
+
+				TSharedPtr<FPayloadContextBase>& PayloadContext = PayloadContexts.FindChecked(PayloadKey);
+				return PayloadContext->FetchPayloadToFile(*this, PayloadFilepath);
+			}
+
+			void FFbxParser::CleanupFbxData()
+			{
+				//////////////////////////////////////////////////////////////////////////
+				// Make sure there is a valid bind pose
+
+				//Find root bones
+				const int32 Default_NbPoses = SDKScene->GetFbxManager()->GetBindPoseCount(SDKScene);
+				// If there are no BindPoses, the following will generate them.
+				SDKScene->GetFbxManager()->CreateMissingBindPoses(SDKScene);
+				//if we created missing bind poses, update the number of bind poses
+				const int32 NbPoses = SDKScene->GetFbxManager()->GetBindPoseCount(SDKScene);
+				if (NbPoses != Default_NbPoses)
+				{
+					UInterchangeResultWarning_Generic* Message = AddMessage<UInterchangeResultWarning_Generic>();
+					Message->Text = LOCTEXT("MissingBindPose", "Missing bind pose - the FBX SDK has created one.");
+				}
+
+				//////////////////////////////////////////////////////////////////////////
+				// Scene conversion
+			}
+		} //ns Private
+	} //ns Interchange
+} //ns UE
+
+#undef LOCTEXT_NAMESPACE
diff --git a/Engine/Plugins/Experimental/Interchange/Runtime/Source/Parsers/Fbx/Private/FbxMaterial.cpp b/Engine/Plugins/Experimental/Interchange/Runtime/Source/Parsers/Fbx/Private/FbxMaterial.cpp
new file mode 100644
--- /dev/null
+++ b/Engine/Plugins/Experimental/Interchange/Runtime/Source/Parsers/Fbx/Private/FbxMaterial.cpp
@@ -0,0 +1,241 @@
+// Copyright Epic Games, Inc. All Rights Reserved.
+
+#include "FbxMaterial.h"
+
+#include "CoreMinimal.h"
+#include "FbxAPI.h"
+#include "FbxConvert.h"
+#include "FbxHelper.h"
+#include "FbxInclude.h"
+#include "InterchangeMaterialNode.h"
+#include "InterchangeResultsContainer.h"
+#include "InterchangeSceneNode.h"
+#include "InterchangeTexture2DNode.h"
+#include "InterchangeTextureNode.h"
+#include "Misc/Paths.h"
+#include "Nodes/InterchangeBaseNodeContainer.h"
+#include "Misc/Paths.h"
+
+#define LOCTEXT_NAMESPACE "InterchangeFbxMaterial"
+
+namespace UE
+{
+	namespace Interchange
+	{
+		namespace Private
+		{
+			UInterchangeMaterialNode* FFbxMaterial::CreateMaterialNode(UInterchangeBaseNodeContainer& NodeContainer, const FString& NodeUid, const FString& NodeName)
+			{
+				UInterchangeMaterialNode* MaterialNode = NewObject<UInterchangeMaterialNode>(&NodeContainer, NAME_None);
+				if (!ensure(MaterialNode))
+				{
+					UInterchangeResultError_Generic* Message = Parser.AddMessage<UInterchangeResultError_Generic>();
+					Message->Text = LOCTEXT("CannotAllocateNode", "Cannot allocate a node when importing FBX.");
+					return nullptr;
+				}
+				// Creating a UMaterialInterface
+				MaterialNode->InitializeNode(NodeUid, NodeName, EInterchangeNodeContainerType::NodeContainerType_TranslatedAsset);
+				MaterialNode->SetPayLoadKey(NodeUid);
+				NodeContainer.AddNode(MaterialNode);
+				return MaterialNode;
+			}
+
+			UInterchangeTexture2DNode* FFbxMaterial::CreateTexture2DNode(UInterchangeBaseNodeContainer& NodeContainer, const FString& NodeUid, const FString& TextureFilePath)
+			{
+				FString DisplayLabel = FPaths::GetBaseFilename(TextureFilePath);
+				UInterchangeTexture2DNode* TextureNode = NewObject<UInterchangeTexture2DNode>(&NodeContainer, NAME_None);
+				if (!ensure(TextureNode))
+				{
+					UInterchangeResultError_Generic* Message = Parser.AddMessage<UInterchangeResultError_Generic>();
+					Message->Text = LOCTEXT("CannotAllocateNode", "Cannot allocate a node when importing FBX.");
+					return nullptr;
+				}
+				// Creating a UTexture2D
+				TextureNode->InitializeNode(NodeUid, DisplayLabel, EInterchangeNodeContainerType::NodeContainerType_TranslatedAsset);
+				//All texture translator expect a file has the payload key
+				FString NormalizeFilePath = TextureFilePath;
+				FPaths::NormalizeFilename(NormalizeFilePath);
+				TextureNode->SetPayLoadKey(NormalizeFilePath);
+				NodeContainer.AddNode(TextureNode);
+				return TextureNode;
+			}
+
+			UInterchangeMaterialNode* FFbxMaterial::AddNodeMaterial(FbxSurfaceMaterial* SurfaceMaterial, UInterchangeBaseNodeContainer& NodeContainer)
+			{
+				//Create a material node
+				FString MaterialName = FFbxHelper::GetFbxObjectName(SurfaceMaterial);
+				FString NodeUid = TEXT("\\Material\\") + MaterialName;
+				UInterchangeMaterialNode* MaterialNode = Cast<UInterchangeMaterialNode>(NodeContainer.GetNode(NodeUid));
+				if (!MaterialNode)
+				{
+					MaterialNode = CreateMaterialNode(NodeContainer, NodeUid, MaterialName);
+					if (MaterialNode == nullptr)
+					{
+						FFormatNamedArguments Args
+						{
+							{ TEXT("MaterialName"), FText::FromString(MaterialName) }
+						};
+						UInterchangeResultError_Generic* Message = Parser.AddMessage<UInterchangeResultError_Generic>();
+						Message->Text = FText::Format(LOCTEXT("CannotCreateFBXMaterial", "Cannot create FBX material '{MaterialName}'."), Args);
+						return nullptr;
+					}
+
+					auto SetMaterialParameter = [this, &NodeContainer, &MaterialNode, &SurfaceMaterial, &MaterialName](const char* FbxMaterialProperty, EInterchangeMaterialNodeParameterName MaterialParameterName)
+					{
+						bool bSetMaterial = false;
+						FbxProperty FbxProperty = SurfaceMaterial->FindProperty(FbxMaterialProperty);
+						if (FbxProperty.IsValid())
+						{
+							int32 UnsupportedTextureCount = FbxProperty.GetSrcObjectCount<FbxLayeredTexture>();
+							UnsupportedTextureCount += FbxProperty.GetSrcObjectCount<FbxProceduralTexture>();
+							int32 TextureCount = FbxProperty.GetSrcObjectCount<FbxFileTexture>();
+							bool bFoundValidTexture = false;
+							if (UnsupportedTextureCount > 0)
+							{
+								FFormatNamedArguments Args
+								{
+									{ TEXT("MaterialName"), FText::FromString(MaterialName) }
+								};
+								UInterchangeResultWarning_Generic* Message = Parser.AddMessage<UInterchangeResultWarning_Generic>();
+								Message->Text = FText::Format(LOCTEXT("TextureTypeNotSupported", "Layered or procedural textures are not supported (material '{MaterialName}')."), Args);
+							}
+							else if (TextureCount > 0)
+							{
+								for (int32 TextureIndex = 0; TextureIndex < TextureCount; ++TextureIndex)
+								{
+									FbxFileTexture* FbxTextureFilePath = FbxProperty.GetSrcObject<FbxFileTexture>(TextureIndex);
+									FString TextureFilename = UTF8_TO_TCHAR(FbxTextureFilePath->GetFileName());
+									//Only import texture that exist on disk
+									if (!FPaths::FileExists(TextureFilename))
+									{
+										continue;
+									}
+									//Create a texture node and make it child of the material node
+									TArray<FString> JsonErrorMessage;
+									FString NodeUid = TEXT("\\Texture\\") + TextureFilename;
+									UInterchangeTexture2DNode* TextureNode = Cast<UInterchangeTexture2DNode>(NodeContainer.GetNode(NodeUid));
+									if (!TextureNode)
+									{
+										TextureNode = CreateTexture2DNode(NodeContainer, NodeUid, TextureFilename);
+									}
+									// add/find UVSet and set it to the texture, we pass index 0 here, I think pipeline should be able to get the UVIndex channel from the name
+									// and modify the parameter to set the correct value.
+// 									FbxString UVSetName = FbxTextureFilePath->UVSet.Get();
+// 									FString LocalUVSetName = UTF8_TO_TCHAR(UVSetName.Buffer());
+									const int32 UVChannelIndex = 0;
+									float ScaleU = (float)FbxTextureFilePath->GetScaleU();
+									float ScaleV = (float)FbxTextureFilePath->GetScaleV();
+									MaterialNode->AddTextureParameterData(MaterialParameterName, TextureNode->GetUniqueID(), UVChannelIndex, ScaleU, ScaleV);
+									MaterialNode->SetTextureDependencyUid(TextureNode->GetUniqueID());
+									bSetMaterial = true;
+									bFoundValidTexture = true;
+								}
+							}
+							
+							if (!bFoundValidTexture && MaterialParameterName == EInterchangeMaterialNodeParameterName::BaseColor)
+							{
+								//We support only the base color has a vector color
+								//TODO support all basic attributes has vector or scalar
+								FbxDouble3 FbxColor;
+								bool bFoundColor = true;
+								if (SurfaceMaterial->GetClassId().Is(FbxSurfacePhong::ClassId))
+								{
+									FbxColor = ((FbxSurfacePhong&)(*SurfaceMaterial)).Diffuse.Get();
+								}
+								else if (SurfaceMaterial->GetClassId().Is(FbxSurfaceLambert::ClassId))
+								{
+									FbxColor = ((FbxSurfaceLambert&)(*SurfaceMaterial)).Diffuse.Get();
+								}
+								else
+								{
+									bFoundColor = false;
+								}
+
+								FVector ColorData;
+								if (bFoundColor)
+								{
+									ColorData[0] = (float)(FbxColor[0]);
+									ColorData[1] = (float)(FbxColor[1]);
+									ColorData[2] = (float)(FbxColor[2]);
+								}
+								else
+								{
+									// use random color because there may be multiple materials, then they can be different 
+									ColorData[0] = 0.5f + (0.5f * FMath::Rand()) / static_cast<float>(RAND_MAX);
+									ColorData[1] = 0.5f + (0.5f * FMath::Rand()) / static_cast<float>(RAND_MAX);
+									ColorData[2] = 0.5f + (0.5f * FMath::Rand()) / static_cast<float>(RAND_MAX);
+								}
+								MaterialNode->AddVectorParameterData(MaterialParameterName, ColorData);
+								bSetMaterial = true;
+							}
+						}
+						return bSetMaterial;
+					};
+
+					SetMaterialParameter(FbxSurfaceMaterial::sDiffuse, EInterchangeMaterialNodeParameterName::BaseColor);
+					SetMaterialParameter(FbxSurfaceMaterial::sEmissive, EInterchangeMaterialNodeParameterName::EmissiveColor);
+					SetMaterialParameter(FbxSurfaceMaterial::sSpecular, EInterchangeMaterialNodeParameterName::Specular);
+					SetMaterialParameter(FbxSurfaceMaterial::sSpecularFactor, EInterchangeMaterialNodeParameterName::Roughness);
+					SetMaterialParameter(FbxSurfaceMaterial::sShininess, EInterchangeMaterialNodeParameterName::Metallic);
+					if (!SetMaterialParameter(FbxSurfaceMaterial::sNormalMap, EInterchangeMaterialNodeParameterName::Normal))
+					{
+						SetMaterialParameter(FbxSurfaceMaterial::sBump, EInterchangeMaterialNodeParameterName::Normal);
+					}
+					if (SetMaterialParameter(FbxSurfaceMaterial::sTransparentColor, EInterchangeMaterialNodeParameterName::Opacity))
+					{
+						SetMaterialParameter(FbxSurfaceMaterial::sTransparencyFactor, EInterchangeMaterialNodeParameterName::OpacityMask);
+					}
+				}
+
+				return MaterialNode;
+			}
+
+			void FFbxMaterial::AddAllTextures(FbxScene* SDKScene, UInterchangeBaseNodeContainer& NodeContainer)
+			{
+				int32 TextureCount = SDKScene->GetSrcObjectCount<FbxFileTexture>();
+				for (int32 TextureIndex = 0; TextureIndex < TextureCount; ++TextureIndex)
+				{
+					FbxFileTexture* Texture = SDKScene->GetSrcObject<FbxFileTexture>(TextureIndex);
+					FString TextureFilename = UTF8_TO_TCHAR(Texture->GetFileName());
+					//Only import texture that exist on disk
+					if (!FPaths::FileExists(TextureFilename))
+					{
+						continue;
+					}
+					//Create a texture node and make it child of the material node
+					TArray<FString> JsonErrorMessage;
+					FString NodeUid = TEXT("\\Texture\\") + TextureFilename;
+					UInterchangeTexture2DNode* TextureNode = Cast<UInterchangeTexture2DNode>(NodeContainer.GetNode(NodeUid));
+					if (!TextureNode)
+					{
+						TextureNode = CreateTexture2DNode(NodeContainer, NodeUid, TextureFilename);
+					}
+				}
+			}
+			
+			void FFbxMaterial::AddAllNodeMaterials(UInterchangeSceneNode* SceneNode, FbxNode* ParentFbxNode, UInterchangeBaseNodeContainer& NodeContainer)
+			{
+				int32 MaterialCount = ParentFbxNode->GetMaterialCount();
+				for (int32 MaterialIndex = 0; MaterialIndex < MaterialCount; ++MaterialIndex)
+				{
+					FbxSurfaceMaterial* SurfaceMaterial = ParentFbxNode->GetMaterial(MaterialIndex);
+					UInterchangeMaterialNode* MaterialNode = AddNodeMaterial(SurfaceMaterial, NodeContainer);
+					//The dependencies order is important because mesh will use index in that order to determine material use by a face
+					SceneNode->AddMaterialDependencyUid(MaterialNode->GetUniqueID());
+				}
+			}
+
+			void FFbxMaterial::AddAllMaterials(FbxScene* SDKScene, UInterchangeBaseNodeContainer& NodeContainer)
+			{
+				int32 MaterialCount = SDKScene->GetMaterialCount();
+				for (int32 MaterialIndex = 0; MaterialIndex < MaterialCount; ++MaterialIndex)
+				{
+					FbxSurfaceMaterial* SurfaceMaterial = SDKScene->GetMaterial(MaterialIndex);
+					AddNodeMaterial(SurfaceMaterial, NodeContainer);
+				}
+			}
+		} //ns Private
+	} //ns Interchange
+}//ns UE
+
+#undef LOCTEXT_NAMESPACE
diff --git a/Engine/Plugins/Experimental/Interchange/Runtime/Source/Parsers/Fbx/Private/FbxMaterial.h b/Engine/Plugins/Experimental/Interchange/Runtime/Source/Parsers/Fbx/Private/FbxMaterial.h
new file mode 100644
--- /dev/null
+++ b/Engine/Plugins/Experimental/Interchange/Runtime/Source/Parsers/Fbx/Private/FbxMaterial.h
@@ -0,0 +1,63 @@
+// Copyright Epic Games, Inc. All Rights Reserved.
+
+#pragma once
+
+#include "CoreMinimal.h"
+#include "FbxInclude.h"
+
+/** Forward declarations */
+class UInterchangeBaseNode;
+class UInterchangeBaseNodeContainer;
+class UInterchangeMaterialNode;
+class UInterchangeTexture2DNode;
+class UInterchangeSceneNode;
+
+namespace UE
+{
+	namespace Interchange
+	{
+		namespace Private
+		{
+			class FFbxParser;
+
+			class FFbxMaterial
+			{
+			public:
+				explicit FFbxMaterial(FFbxParser& InParser)
+					: Parser(InParser)
+				{}
+
+				/**
+				 * Create a UInterchangeFextureNode and add it to the NodeContainer for all texture of type FbxFileTexture the fbx file contain.
+				 *
+				 * @note - Any node that already exist in the NodeContainer will not be created or modify.
+				 */
+				void AddAllTextures(FbxScene* SDKScene, UInterchangeBaseNodeContainer& NodeContainer);
+				
+				/**
+				 * Create a UInterchangeMaterialNode and add it to the NodeContainer for all material of type FbxSurfaceMaterial the fbx file contain.
+				 * 
+				 * @note - Any node that already exist in the NodeContainer will not be created or modify.
+				 */
+				void AddAllMaterials(FbxScene* SDKScene, UInterchangeBaseNodeContainer& NodeContainer);
+
+				/**
+				 * Create a UInterchangeMaterialNode and add it to the NodeContainer for all material of type FbxSurfaceMaterial the fbx ParentFbxNode contain.
+				 * It also set the dependencies of the node materials on the interchange ParentNode.
+				 * 
+				 * @note - Any material node that already exist in the NodeContainer will simply be add has a dependency.
+				 */
+				void AddAllNodeMaterials(UInterchangeSceneNode* SceneNode, FbxNode* ParentFbxNode, UInterchangeBaseNodeContainer& NodeContainer);
+
+			protected:
+				UInterchangeMaterialNode* CreateMaterialNode(UInterchangeBaseNodeContainer& NodeContainer, const FString& NodeUID, const FString& NodeName);
+				UInterchangeTexture2DNode* CreateTexture2DNode(UInterchangeBaseNodeContainer& NodeContainer, const FString& NodeUID, const FString& TextureFilePath);
+
+			private:
+				UInterchangeMaterialNode* AddNodeMaterial(FbxSurfaceMaterial* SurfaceMaterial, UInterchangeBaseNodeContainer& NodeContainer);
+
+				FFbxParser& Parser;
+			};
+		}//ns Private
+	}//ns Interchange
+}//ns UE
diff --git a/Engine/Plugins/Experimental/Interchange/Runtime/Source/Parsers/Fbx/Private/FbxMesh.h b/Engine/Plugins/Experimental/Interchange/Runtime/Source/Parsers/Fbx/Private/FbxMesh.h
new file mode 100644
--- /dev/null
+++ b/Engine/Plugins/Experimental/Interchange/Runtime/Source/Parsers/Fbx/Private/FbxMesh.h
@@ -0,0 +1,133 @@
+// Copyright Epic Games, Inc. All Rights Reserved.
+
+#pragma once
+
+#include "CoreMinimal.h"
+#include "FbxAPI.h"
+#include "FbxHelper.h"
+#include "FbxInclude.h"
+
+/** Forward declarations */
+struct FMeshDescription;
+class UInterchangeBaseNodeContainer;
+class UInterchangeMeshNode;
+
+namespace UE
+{
+	namespace Interchange
+	{
+		namespace Private
+		{
+			class FFbxMesh;
+
+			class FMeshDescriptionImporter
+			{
+			public:
+				FMeshDescriptionImporter(FFbxParser& InParser, FMeshDescription* InMeshDescription, FbxScene* InSDKScene, FbxGeometryConverter* InSDKGeometryConverter);
+				
+				/*
+				 * Fill the mesh description using the Mesh parameter.
+				 */
+				bool FillStaticMeshDescriptionFromFbxMesh(FbxMesh* Mesh);
+				
+				/*
+				 * Fill the mesh description using the Mesh parameter and also fill the OutJointNodeUniqueIDs so the MeshDescription bone Index can be map to the correct interchange joint scene node.
+				 */
+				bool FillSkinnedMeshDescriptionFromFbxMesh(FbxMesh* Mesh, TArray<FString>& OutJointUniqueNames);
+
+				/*
+				 * Fill the mesh description using the Shape parameter.
+				 */
+				bool FillMeshDescriptionFromFbxShape(FbxShape* Shape);
+
+				/**
+				 * Add messages to the message log
+				 */
+				template <typename T>
+				T* AddMessage(FbxGeometryBase* FbxNode) const
+				{
+					T* Item = Parser.AddMessage<T>();
+					Item->MeshName = FFbxHelper::GetMeshName(FbxNode);
+					Item->InterchangeKey = FFbxHelper::GetMeshUniqueID(FbxNode);
+					return Item;
+				}
+
+			private:
+				
+				enum class EMeshType : uint8
+				{
+					None = 0, //No mesh type to import
+					Static = 1, //static mesh
+					Skinned = 2, //skinned mesh with joints
+				};
+
+				bool FillMeshDescriptionFromFbxMesh(FbxMesh* Mesh, TArray<FString>& OutJointUniqueNames, EMeshType MeshType);
+				bool IsOddNegativeScale(FbxAMatrix& TotalMatrix);
+				
+				//TODO move the real function from RenderCore to FVector, so we do not have to add render core to compute such a simple thing
+				float FbxGetBasisDeterminantSign(const FVector& XAxis, const FVector& YAxis, const FVector& ZAxis)
+				{
+					FMatrix Basis(
+						FPlane(XAxis, 0),
+						FPlane(YAxis, 0),
+						FPlane(ZAxis, 0),
+						FPlane(0, 0, 0, 1)
+					);
+					return (Basis.Determinant() < 0) ? -1.0f : +1.0f;
+				}
+
+				FFbxParser& Parser;
+				FMeshDescription* MeshDescription;
+				FbxScene* SDKScene;
+				FbxGeometryConverter* SDKGeometryConverter;
+				bool bInitialized = false;
+			};
+
+			class FMeshPayloadContext : public FPayloadContextBase
+			{
+			public:
+				virtual ~FMeshPayloadContext() {}
+				virtual FString GetPayloadType() const override { return TEXT("Mesh-PayloadContext"); }
+				virtual bool FetchPayloadToFile(FFbxParser& Parser, const FString& PayloadFilepath) override;
+				FbxMesh* Mesh = nullptr;
+				FbxScene* SDKScene = nullptr;
+				FbxGeometryConverter* SDKGeometryConverter = nullptr;
+			};
+
+			class FShapePayloadContext : public FPayloadContextBase
+			{
+			public:
+				virtual ~FShapePayloadContext() {}
+				virtual FString GetPayloadType() const override { return TEXT("Shape-PayloadContext"); }
+				virtual bool FetchPayloadToFile(FFbxParser& Parser, const FString& PayloadFilepath) override;
+				FbxShape* Shape = nullptr;
+				FbxScene* SDKScene = nullptr;
+				FbxGeometryConverter* SDKGeometryConverter = nullptr;
+			};
+
+			class FFbxMesh
+			{
+			public:
+				explicit FFbxMesh(FFbxParser& InParser)
+					: Parser(InParser)
+				{}
+
+#if 0
+				static FString GetMeshName(FbxGeometryBase* Mesh);
+				static FString GetMeshUniqueID(FbxGeometryBase* Mesh);
+#endif
+				void ExtractSkinnedMeshNodeJoints(FbxScene* SDKScene, FbxMesh* Mesh, UInterchangeMeshNode* MeshNode);
+				void AddAllMeshes(FbxScene* SDKScene, FbxGeometryConverter* SDKGeometryConverter, UInterchangeBaseNodeContainer& NodeContainer, TMap<FString, TSharedPtr<FPayloadContextBase>>& PayloadContexts);
+			
+			protected:
+				UInterchangeMeshNode* CreateMeshNode(UInterchangeBaseNodeContainer& NodeContainer, const FString& NodeName, const FString& NodeUniqueID);
+#if 0
+				static FString GetUniqueIDString(const uint64 UniqueID);
+#endif
+
+			private:
+				FFbxParser& Parser;
+			};
+		}//ns Private
+	}//ns Interchange
+}//ns UE
diff --git a/Engine/Plugins/Experimental/Interchange/Runtime/Source/Parsers/Fbx/Private/FbxScene.h b/Engine/Plugins/Experimental/Interchange/Runtime/Source/Parsers/Fbx/Private/FbxScene.h
new file mode 100644
--- /dev/null
+++ b/Engine/Plugins/Experimental/Interchange/Runtime/Source/Parsers/Fbx/Private/FbxScene.h
@@ -0,0 +1,41 @@
+// Copyright Epic Games, Inc. All Rights Reserved.
+
+#pragma once
+
+#include "CoreMinimal.h"
+#include "FbxInclude.h"
+
+/** Forward declarations */
+class UInterchangeBaseNodeContainer;
+class UInterchangeSceneNode;
+
+namespace UE
+{
+	namespace Interchange
+	{
+		namespace Private
+		{
+			class FFbxParser;
+
+			class FFbxScene
+			{
+			public:
+				explicit FFbxScene(FFbxParser& InParser)
+					: Parser(InParser)
+				{}
+
+				void AddHierarchy(FbxScene* SDKScene, UInterchangeBaseNodeContainer& NodeContainer);
+				UInterchangeSceneNode* CreateTransformNode(UInterchangeBaseNodeContainer& NodeContainer, const FString& NodeName, const FString& NodeUniqueID);
+
+			protected:
+				void CreateMeshNodeReference(UInterchangeSceneNode* UnrealSceneNode, FbxNodeAttribute* NodeAttribute, UInterchangeBaseNodeContainer& NodeContainer);
+				void CreateCameraNodeReference(UInterchangeSceneNode* UnrealSceneNode, FbxNodeAttribute* NodeAttribute, UInterchangeBaseNodeContainer& NodeContainer);
+				void CreateLightNodeReference(UInterchangeSceneNode* UnrealSceneNode, FbxNodeAttribute* NodeAttribute, UInterchangeBaseNodeContainer& NodeContainer);
+				UInterchangeSceneNode* AddHierarchyRecursively(FbxNode* Node, FbxScene* SDKScene, UInterchangeBaseNodeContainer& NodeContainer);
+
+			private:
+				FFbxParser& Parser;
+			};
+		}//ns Private
+	}//ns Interchange
+}//ns UE
diff --git a/Engine/Plugins/Experimental/Interchange/Runtime/Source/Parsers/Fbx/Private/FbxSkelMesh.cpp b/Engine/Plugins/Experimental/Interchange/Runtime/Source/Parsers/Fbx/Private/FbxSkelMesh.cpp
new file mode 100644
--- /dev/null
+++ b/Engine/Plugins/Experimental/Interchange/Runtime/Source/Parsers/Fbx/Private/FbxSkelMesh.cpp
@@ -0,0 +1,931 @@
+// Copyright Epic Games, Inc. All Rights Reserved.
+
+#include "FbxSkelMesh.h"
+
+#include "CoreMinimal.h"
+#include "FbxConvert.h"
+
+#define GeneratedLODNameSuffix "_GeneratedLOD_"
+
+
+namespace UE
+{
+	namespace Interchange
+	{
+		namespace Private
+		{
+			void FFbxSkelMesh::FindSkeletalMeshes(FbxScene* SDKScene, TArray< TArray<FbxNode*> >& outSkelMeshArray, bool bCombineSkeletalMesh, bool bForceFindRigid)
+			{
+				TArray<FbxNode*> SkeletonArray;
+				FbxNode* RootNode = SDKScene->GetRootNode();
+
+				// a) find skeletal meshes
+
+				RecursiveFindFbxSkelMesh(SDKScene, RootNode, outSkelMeshArray, SkeletonArray);
+				// for skeletal mesh, we convert the skeleton system to skeleton
+				// in less we recognize bone mesh as rigid mesh if they are textured
+				for (int32 SkelIndex = 0; SkelIndex < SkeletonArray.Num(); SkelIndex++)
+				{
+					const bool bImportMeshesInBoneHierarchy = true;
+					RecursiveFixSkeleton(SDKScene, SkeletonArray[SkelIndex], outSkelMeshArray[SkelIndex], bImportMeshesInBoneHierarchy);
+				}
+
+
+
+				// b) find rigid mesh
+
+				// If we are attempting to import a skeletal mesh but we have no hierarchy attempt to find a rigid mesh.
+				if (bForceFindRigid || outSkelMeshArray.Num() == 0)
+				{
+					RecursiveFindRigidMesh(SDKScene, RootNode, outSkelMeshArray, SkeletonArray);
+					if (bForceFindRigid)
+					{
+						//Cleanup the rigid mesh, We want to remove any real static mesh from the outSkelMeshArray
+						//Any non skinned mesh that contain no animation should be part of this array.
+						int32 AnimStackCount = SDKScene->GetSrcObjectCount<FbxAnimStack>();
+						TArray<int32> SkeletalMeshArrayToRemove;
+						for (int32 i = 0; i < outSkelMeshArray.Num(); i++)
+						{
+							bool bIsValidSkeletal = false;
+							TArray<FbxNode*>& NodeArray = outSkelMeshArray[i];
+							for (FbxNode* InspectedNode : NodeArray)
+							{
+								FbxMesh* Mesh = InspectedNode->GetMesh();
+
+								FbxLODGroup* LodGroup = InspectedNode->GetLodGroup();
+								if (LodGroup != nullptr)
+								{
+									FbxNode* SkelMeshNode = FindLODGroupNode(InspectedNode, 0, nullptr);
+									if (SkelMeshNode != nullptr)
+									{
+										Mesh = SkelMeshNode->GetMesh();
+									}
+								}
+
+								if (Mesh == nullptr)
+								{
+									continue;
+								}
+								if (Mesh->GetDeformerCount(FbxDeformer::eSkin) > 0)
+								{
+									bIsValidSkeletal = true;
+									break;
+								}
+								//If there is some anim object we count this as a valid skeletal mesh imported as rigid mesh
+								for (int32 AnimStackIndex = 0; AnimStackIndex < AnimStackCount; AnimStackIndex++)
+								{
+									FbxAnimStack* CurAnimStack = SDKScene->GetSrcObject<FbxAnimStack>(AnimStackIndex);
+									// set current anim stack
+
+									SDKScene->SetCurrentAnimationStack(CurAnimStack);
+
+									FbxTimeSpan AnimTimeSpan(FBXSDK_TIME_INFINITE, FBXSDK_TIME_MINUS_INFINITE);
+									InspectedNode->GetAnimationInterval(AnimTimeSpan, CurAnimStack);
+
+									if (AnimTimeSpan.GetDuration() > 0)
+									{
+										bIsValidSkeletal = true;
+										break;
+									}
+								}
+								if (bIsValidSkeletal)
+								{
+									break;
+								}
+							}
+							if (!bIsValidSkeletal)
+							{
+								SkeletalMeshArrayToRemove.Add(i);
+							}
+						}
+						for (int32 i = SkeletalMeshArrayToRemove.Num() - 1; i >= 0; --i)
+						{
+							if (!SkeletalMeshArrayToRemove.IsValidIndex(i) || !outSkelMeshArray.IsValidIndex(SkeletalMeshArrayToRemove[i]))
+								continue;
+							int32 IndexToRemove = SkeletalMeshArrayToRemove[i];
+							outSkelMeshArray[IndexToRemove].Empty();
+							outSkelMeshArray.RemoveAt(IndexToRemove);
+						}
+					}
+				}
+				//Empty the skeleton array
+				SkeletonArray.Empty();
+
+
+				if (bCombineSkeletalMesh)
+				{
+					//Merge all the skeletal mesh arrays into one combine mesh
+					TArray<FbxNode*> CombineNodes;
+					for (TArray<FbxNode*>& Parts : outSkelMeshArray)
+					{
+						CombineNodes.Append(Parts);
+					}
+					outSkelMeshArray.Empty(1);
+					outSkelMeshArray.Add(CombineNodes);
+				}
+			}
+
+			void FFbxSkelMesh::RecursiveGetAllMeshNode(TArray<FbxNode*> & OutAllNode, FbxNode * Node)
+			{
+				if (Node == nullptr)
+				{
+					return;
+				}
+
+				if (Node->GetMesh() != nullptr)
+				{
+					OutAllNode.Add(Node);
+					return;
+				}
+
+				//Look if its a generated LOD
+				FString FbxGeneratedNodeName = UTF8_TO_TCHAR(Node->GetName());
+				if (FbxGeneratedNodeName.Contains(TEXT(GeneratedLODNameSuffix)))
+				{
+					FString SuffixSearch = TEXT(GeneratedLODNameSuffix);
+					int32 SuffixIndex = FbxGeneratedNodeName.Find(SuffixSearch, ESearchCase::CaseSensitive, ESearchDir::FromEnd);
+					SuffixIndex += SuffixSearch.Len();
+					FString LODXNumber = FbxGeneratedNodeName.RightChop(SuffixIndex).Left(1);
+					if (LODXNumber.IsNumeric())
+					{
+						OutAllNode.Add(Node);
+						return;
+					}
+				}
+
+				for (int32 ChildIndex = 0; ChildIndex < Node->GetChildCount(); ++ChildIndex)
+				{
+					RecursiveGetAllMeshNode(OutAllNode, Node->GetChild(ChildIndex));
+				}
+			}
+
+			void FFbxSkelMesh::FindAllLODGroupNode(TArray<FbxNode*> & OutNodeInLod, FbxNode * NodeLodGroup, int32 LodIndex)
+			{
+				check(NodeLodGroup->GetChildCount() >= LodIndex);
+				FbxNode* ChildNode = NodeLodGroup->GetChild(LodIndex);
+				if (ChildNode == nullptr)
+				{
+					return;
+				}
+				RecursiveGetAllMeshNode(OutNodeInLod, ChildNode);
+			}
+
+			bool FFbxSkelMesh::IsUnrealBone(FbxNode * Link)
+			{
+				FbxNodeAttribute* Attr = Link->GetNodeAttribute();
+				if (Attr)
+				{
+					FbxNodeAttribute::EType AttrType = Attr->GetAttributeType();
+					if (AttrType == FbxNodeAttribute::eSkeleton ||
+						AttrType == FbxNodeAttribute::eMesh ||
+						AttrType == FbxNodeAttribute::eNull)
+					{
+						return true;
+					}
+				}
+
+				return false;
+			}
+
+			void FFbxSkelMesh::RecursiveBuildSkeleton(FbxNode * Link, TArray<FbxNode*> & OutSortedLinks)
+			{
+				if (IsUnrealBone(Link))
+				{
+					OutSortedLinks.Add(Link);
+					int32 ChildIndex;
+					for (ChildIndex = 0; ChildIndex < Link->GetChildCount(); ChildIndex++)
+					{
+						RecursiveBuildSkeleton(Link->GetChild(ChildIndex), OutSortedLinks);
+					}
+				}
+			}
+
+			bool FFbxSkelMesh::RetrievePoseFromBindPose(FbxScene * SDKScene, const TArray<FbxNode*> & NodeArray, FbxArray<FbxPose*> & PoseArray)
+			{
+				const int32 PoseCount = SDKScene->GetPoseCount();
+				for (int32 PoseIndex = 0; PoseIndex < PoseCount; PoseIndex++)
+				{
+					FbxPose* CurrentPose = SDKScene->GetPose(PoseIndex);
+
+					// current pose is bind pose, 
+					if (CurrentPose && CurrentPose->IsBindPose())
+					{
+						// IsValidBindPose doesn't work reliably
+						// It checks all the parent chain(regardless root given), and if the parent doesn't have correct bind pose, it fails
+						// It causes more false positive issues than the real issue we have to worry about
+						// If you'd like to try this, set CHECK_VALID_BIND_POSE to 1, and try the error message
+						// when Autodesk fixes this bug, then we might be able to re-open this
+						FString PoseName = CurrentPose->GetName();
+						// all error report status
+						FbxStatus Status;
+
+						// it does not make any difference of checking with different node
+						// it is possible pose 0 -> node array 2, but isValidBindPose function returns true even with node array 0
+						for (auto Current : NodeArray)
+						{
+							FString CurrentName = Current->GetName();
+							NodeList pMissingAncestors, pMissingDeformers, pMissingDeformersAncestors, pWrongMatrices;
+
+							if (CurrentPose->IsValidBindPoseVerbose(Current, pMissingAncestors, pMissingDeformers, pMissingDeformersAncestors, pWrongMatrices, 0.0001, &Status))
+							{
+								PoseArray.Add(CurrentPose);
+								//UE_LOG(LogFbx, Log, TEXT("Valid bind pose for Pose (%s) - %s"), *PoseName, *FString(Current->GetName()));
+								break;
+							}
+							else
+							{
+								// first try to fix up
+								// add missing ancestors
+								for (int i = 0; i < pMissingAncestors.GetCount(); i++)
+								{
+									FbxAMatrix mat = pMissingAncestors.GetAt(i)->EvaluateGlobalTransform(FBXSDK_TIME_ZERO);
+									CurrentPose->Add(pMissingAncestors.GetAt(i), mat);
+								}
+
+								pMissingAncestors.Clear();
+								pMissingDeformers.Clear();
+								pMissingDeformersAncestors.Clear();
+								pWrongMatrices.Clear();
+
+								// check it again
+								if (CurrentPose->IsValidBindPose(Current))
+								{
+									PoseArray.Add(CurrentPose);
+									//UE_LOG(LogFbx, Log, TEXT("Valid bind pose for Pose (%s) - %s"), *PoseName, *FString(Current->GetName()));
+									break;
+								}
+								else
+								{
+									// first try to find parent who is null group and see if you can try test it again
+									FbxNode* ParentNode = Current->GetParent();
+									while (ParentNode)
+									{
+										FbxNodeAttribute* Attr = ParentNode->GetNodeAttribute();
+										if (Attr && Attr->GetAttributeType() == FbxNodeAttribute::eNull)
+										{
+											// found it 
+											break;
+										}
+
+										// find next parent
+										ParentNode = ParentNode->GetParent();
+									}
+
+									if (ParentNode && CurrentPose->IsValidBindPose(ParentNode))
+									{
+										PoseArray.Add(CurrentPose);
+										//UE_LOG(LogFbx, Log, TEXT("Valid bind pose for Pose (%s) - %s"), *PoseName, *FString(Current->GetName()));
+										break;
+									}
+									else
+									{
+										FString ErrorString = Status.GetErrorString();
+										// 											if (!GIsAutomationTesting)
+										// 												UE_LOG(LogFbx, Warning, TEXT("Not valid bind pose for Pose (%s) - Node %s : %s"), *PoseName, *FString(Current->GetName()), *ErrorString);
+									}
+								}
+							}
+						}
+					}
+				}
+
+				return (PoseArray.Size() > 0);
+			}
+
+			FbxNode* FFbxSkelMesh::GetRootSkeleton(FbxScene * SDKScene, FbxNode * Link)
+			{
+				FbxNode* RootBone = Link;
+
+				// get Unreal skeleton root
+				// mesh and dummy are used as bone if they are in the skeleton hierarchy
+				while (RootBone && RootBone->GetParent())
+				{
+					bool bIsBlenderArmatureBone = false;
+					//TODO put back the creator code so we know if its blender
+					//if (FbxCreator == EFbxCreator::Blender)
+					{
+						//Hack to support armature dummy node from blender
+						//Users do not want the null attribute node named armature which is the parent of the real root bone in blender fbx file
+						//This is a hack since if a rigid mesh group root node is named "armature" it will be skip
+						const FString RootBoneParentName(RootBone->GetParent()->GetName());
+						FbxNode* GrandFather = RootBone->GetParent()->GetParent();
+						bIsBlenderArmatureBone = (GrandFather == nullptr || GrandFather == SDKScene->GetRootNode()) && (RootBoneParentName.Compare(TEXT("armature"), ESearchCase::IgnoreCase) == 0);
+					}
+
+					FbxNodeAttribute* Attr = RootBone->GetParent()->GetNodeAttribute();
+					if (Attr &&
+						(Attr->GetAttributeType() == FbxNodeAttribute::eMesh ||
+							(Attr->GetAttributeType() == FbxNodeAttribute::eNull && !bIsBlenderArmatureBone) ||
+							Attr->GetAttributeType() == FbxNodeAttribute::eSkeleton) &&
+						RootBone->GetParent() != SDKScene->GetRootNode())
+					{
+						// in some case, skeletal mesh can be ancestor of bones
+						// this avoids this situation
+						if (Attr->GetAttributeType() == FbxNodeAttribute::eMesh)
+						{
+							FbxMesh* Mesh = (FbxMesh*)Attr;
+							if (Mesh->GetDeformerCount(FbxDeformer::eSkin) > 0)
+							{
+								break;
+							}
+						}
+
+						RootBone = RootBone->GetParent();
+					}
+					else
+					{
+						break;
+					}
+				}
+
+				return RootBone;
+			}
+
+			void FFbxSkelMesh::BuildSkeletonSystem(FbxScene * SDKScene, TArray<FbxCluster*> & ClusterArray, TArray<FbxNode*> & OutSortedLinks)
+			{
+				FbxNode* Link;
+				TArray<FbxNode*> RootLinks;
+				int32 ClusterIndex;
+				for (ClusterIndex = 0; ClusterIndex < ClusterArray.Num(); ClusterIndex++)
+				{
+					Link = ClusterArray[ClusterIndex]->GetLink();
+					if (Link)
+					{
+						Link = GetRootSkeleton(SDKScene, Link);
+						int32 LinkIndex;
+						for (LinkIndex = 0; LinkIndex < RootLinks.Num(); LinkIndex++)
+						{
+							if (Link == RootLinks[LinkIndex])
+							{
+								break;
+							}
+						}
+
+						// this link is a new root, add it
+						if (LinkIndex == RootLinks.Num())
+						{
+							RootLinks.Add(Link);
+						}
+					}
+				}
+
+				for (int32 LinkIndex = 0; LinkIndex < RootLinks.Num(); LinkIndex++)
+				{
+					RecursiveBuildSkeleton(RootLinks[LinkIndex], OutSortedLinks);
+				}
+			}
+
+			bool FFbxSkelMesh::FindSkeletonJoints(FbxScene * SDKScene, TArray<FbxNode*> & NodeArray, TArray<FbxNode*> & SortedLinks, FbxArray<FbxAMatrix> & LocalsPerLink)
+			{
+				bool bIsARigidMesh = false;
+				FbxNode* Link = NULL;
+				FbxArray<FbxPose*> PoseArray;
+				TArray<FbxCluster*> ClusterArray;
+
+				if (NodeArray[0]->GetMesh()->GetDeformerCount(FbxDeformer::eSkin) == 0)
+				{
+					bIsARigidMesh = true;
+					Link = NodeArray[0];
+					RecursiveBuildSkeleton(GetRootSkeleton(SDKScene, Link), SortedLinks);
+				}
+				else
+				{
+					// get bindpose and clusters from FBX skeleton
+
+					// let's put the elements to their bind pose! (and we restore them after
+					// we have built the ClusterInformation.
+					int32 Default_NbPoses = SDKScene->GetFbxManager()->GetBindPoseCount(SDKScene);
+					// If there are no BindPoses, the following will generate them.
+					SDKScene->GetFbxManager()->CreateMissingBindPoses(SDKScene);
+
+					//if we created missing bind poses, update the number of bind poses
+					int32 NbPoses = SDKScene->GetFbxManager()->GetBindPoseCount(SDKScene);
+
+					if (NbPoses != Default_NbPoses)
+					{
+						//TODO output an error message to let the user know there is a missing bind pose when translating this fbx skeletalmesh
+						//Missing bind pose
+					}
+
+					//
+					// create the bones / skinning
+					//
+
+					for (int32 i = 0; i < NodeArray.Num(); i++)
+					{
+						FbxMesh* FbxMesh = NodeArray[i]->GetMesh();
+						const int32 SkinDeformerCount = FbxMesh->GetDeformerCount(FbxDeformer::eSkin);
+						for (int32 DeformerIndex = 0; DeformerIndex < SkinDeformerCount; DeformerIndex++)
+						{
+							FbxSkin* Skin = (FbxSkin*)FbxMesh->GetDeformer(DeformerIndex, FbxDeformer::eSkin);
+							for (int32 ClusterIndex = 0; ClusterIndex < Skin->GetClusterCount(); ClusterIndex++)
+							{
+								ClusterArray.Add(Skin->GetCluster(ClusterIndex));
+							}
+						}
+					}
+
+					if (ClusterArray.Num() == 0)
+					{
+						//TODO warn user we found a mesh with a deformer that is not bind to any skin cluster
+						return false;
+					}
+
+					// get bind pose
+					if (RetrievePoseFromBindPose(SDKScene, NodeArray, PoseArray) == false)
+					{
+						// 							if (!GIsAutomationTesting)
+						// 								UE_LOG(LogFbx, Warning, TEXT("Getting valid bind pose failed. Try to recreate bind pose"));
+						// if failed, delete bind pose, and retry.
+						const int32 PoseCount = SDKScene->GetPoseCount();
+						for (int32 PoseIndex = PoseCount - 1; PoseIndex >= 0; --PoseIndex)
+						{
+							FbxPose* CurrentPose = SDKScene->GetPose(PoseIndex);
+
+							// current pose is bind pose, 
+							if (CurrentPose && CurrentPose->IsBindPose())
+							{
+								SDKScene->RemovePose(PoseIndex);
+								CurrentPose->Destroy();
+							}
+						}
+
+						SDKScene->GetFbxManager()->CreateMissingBindPoses(SDKScene);
+						if (RetrievePoseFromBindPose(SDKScene, NodeArray, PoseArray) == false)
+						{
+							// 								if (!GIsAutomationTesting)
+							// 									UE_LOG(LogFbx, Warning, TEXT("Recreating bind pose failed."));
+						}
+						else
+						{
+							// 								if (!GIsAutomationTesting)
+							// 									UE_LOG(LogFbx, Warning, TEXT("Recreating bind pose succeeded."));
+						}
+					}
+
+					// recurse through skeleton and build ordered table
+					BuildSkeletonSystem(SDKScene, ClusterArray, SortedLinks);
+				}
+
+				// error check
+				// if no bond is found
+				if (SortedLinks.Num() == 0)
+				{
+					//						{
+					//							AddTokenizedErrorMessage(FTokenizedMessage::Create(EMessageSeverity::Warning, FText::Format(LOCTEXT("FbxSkeletaLMeshimport_NoBone", "'{0}' has no bones"), FText::FromString(NodeArray[0]->GetName()))), FFbxErrors::SkeletalMesh_NoBoneFound);
+					//						}
+					return false;
+				}
+
+				// if no bind pose is found
+				if (PoseArray.GetCount() == 0)
+				{
+					//TODO error to the user
+					//AddTokenizedErrorMessage(FTokenizedMessage::Create(EMessageSeverity::Warning, LOCTEXT("FbxSkeletaLMeshimport_MissingBindPose", "Could not find the bind pose.  It will use time 0 as bind pose.")), FFbxErrors::SkeletalMesh_InvalidBindPose);
+				}
+
+				int32 LinkIndex;
+
+				// Check for duplicate bone names and issue a warning if found
+				for (LinkIndex = 0; LinkIndex < SortedLinks.Num(); ++LinkIndex)
+				{
+					Link = SortedLinks[LinkIndex];
+
+					for (int32 AltLinkIndex = LinkIndex + 1; AltLinkIndex < SortedLinks.Num(); ++AltLinkIndex)
+					{
+						FbxNode* AltLink = SortedLinks[AltLinkIndex];
+
+						if (FCStringAnsi::Strcmp(Link->GetName(), AltLink->GetName()) == 0)
+						{
+							//								FString RawBoneName = UTF8_TO_TCHAR(Link->GetName());
+							//TODO notify the user that there is a duplicate bone name
+							// 									AddTokenizedErrorMessage(FTokenizedMessage::Create(EMessageSeverity::Error, FText::Format(LOCTEXT("Error_DuplicateBoneName", "Error, Could not import {0}.\nDuplicate bone name found ('{1}'). Each bone must have a unique name."),
+							// 															 FText::FromString(NodeArray[0]->GetName()), FText::FromString(RawBoneName))), FFbxErrors::SkeletalMesh_DuplicateBones);
+							return false;
+						}
+					}
+				}
+
+				FbxArray<FbxAMatrix> GlobalsPerLink;
+				GlobalsPerLink.Grow(SortedLinks.Num());
+				GlobalsPerLink[0].SetIdentity();
+				LocalsPerLink.Grow(SortedLinks.Num());
+				LocalsPerLink[0].SetIdentity();
+
+				bool GlobalLinkFoundFlag;
+
+				bool bAnyLinksNotInBindPose = false;
+				FString LinksWithoutBindPoses;
+
+				int32 RootIdx = INDEX_NONE;
+
+				for (LinkIndex = 0; LinkIndex < SortedLinks.Num(); LinkIndex++)
+				{
+					// Add a bone for each FBX Link
+					Link = SortedLinks[LinkIndex];
+					int32 ParentIndex = INDEX_NONE; // base value for root if no parent found
+
+					if (LinkIndex == 0)
+					{
+						RootIdx = LinkIndex;
+					}
+					else
+					{
+						const FbxNode* LinkParent = Link->GetParent();
+						// get the link parent index.
+						for (int32 ll = 0; ll < LinkIndex; ++ll) // <LinkIndex because parent is guaranteed to be before child in sortedLink
+						{
+							FbxNode* Otherlink = SortedLinks[ll];
+							if (Otherlink == LinkParent)
+							{
+								ParentIndex = ll;
+								break;
+							}
+						}
+
+						if (ParentIndex == INDEX_NONE)//We found another root inside the hierarchy, this is not supported
+						{
+							//AddTokenizedErrorMessage(FTokenizedMessage::Create(EMessageSeverity::Error, LOCTEXT("MultipleRootsFound", "Multiple roots are found in the bone hierarchy. We only support single root bone.")), FFbxErrors::SkeletalMesh_MultipleRoots);
+							return false;
+						}
+					}
+
+					GlobalLinkFoundFlag = false;
+					if (!bIsARigidMesh) //skeletal mesh
+					{
+						// there are some links, they have no cluster, but in bindpose
+						if (PoseArray.GetCount())
+						{
+							for (int32 PoseIndex = 0; PoseIndex < PoseArray.GetCount(); PoseIndex++)
+							{
+								int32 PoseLinkIndex = PoseArray[PoseIndex]->Find(Link);
+								if (PoseLinkIndex >= 0)
+								{
+									FbxMatrix NoneAffineMatrix = PoseArray[PoseIndex]->GetMatrix(PoseLinkIndex);
+									FbxAMatrix Matrix = *(FbxAMatrix*)(double*)&NoneAffineMatrix;
+									GlobalsPerLink[LinkIndex] = Matrix;
+									GlobalLinkFoundFlag = true;
+									break;
+								}
+							}
+						}
+
+						if (!GlobalLinkFoundFlag)
+						{
+							// since now we set use time 0 as ref pose this won't unlikely happen
+							// but leaving it just in case it still has case where it's missing partial bind pose
+							bAnyLinksNotInBindPose = true;
+							LinksWithoutBindPoses += UTF8_TO_TCHAR(Link->GetName());
+							LinksWithoutBindPoses += TEXT("  \n");
+
+							for (int32 ClusterIndex = 0; ClusterIndex < ClusterArray.Num(); ClusterIndex++)
+							{
+								FbxCluster* Cluster = ClusterArray[ClusterIndex];
+								if (Link == Cluster->GetLink())
+								{
+									Cluster->GetTransformLinkMatrix(GlobalsPerLink[LinkIndex]);
+									GlobalLinkFoundFlag = true;
+									break;
+								}
+							}
+						}
+					}
+
+					if (!GlobalLinkFoundFlag)
+					{
+						GlobalsPerLink[LinkIndex] = Link->EvaluateGlobalTransform();
+					}
+
+
+					if (LinkIndex)
+					{
+						LocalsPerLink[LinkIndex] = GlobalsPerLink[ParentIndex].Inverse() * GlobalsPerLink[LinkIndex];
+					}
+					else	// skeleton root
+					{
+						LocalsPerLink[LinkIndex] = GlobalsPerLink[LinkIndex];
+					}
+
+					//Check for nan and for zero scale
+					{
+						bool bFoundNan = false;
+						bool bFoundZeroScale = false;
+						for (int32 i = 0; i < 4; ++i)
+						{
+							if (i < 3)
+							{
+								if (FMath::IsNaN(LocalsPerLink[LinkIndex].GetT()[i]) || FMath::IsNaN(LocalsPerLink[LinkIndex].GetS()[i]))
+								{
+									bFoundNan = true;
+								}
+								if (FMath::IsNearlyZero(LocalsPerLink[LinkIndex].GetS()[i]))
+								{
+									bFoundZeroScale = true;
+								}
+							}
+							if (FMath::IsNaN(LocalsPerLink[LinkIndex].GetQ()[i]))
+							{
+								bFoundNan = true;
+							}
+						}
+
+						//TODO add a user error message
+						ensure(!bFoundNan);
+						ensure(!bFoundZeroScale);
+					}
+				}
+
+
+
+				if (bAnyLinksNotInBindPose)
+				{
+					//TODO error message
+					//
+					//AddTokenizedErrorMessage(FTokenizedMessage::Create(EMessageSeverity::Warning, FText::Format(LOCTEXT("FbxSkeletaLMeshimport_BonesAreMissingFromBindPose", "The following bones are missing from the bind pose:\n{0}\nThis can happen for bones that are not vert weighted. If they are not in the correct orientation after importing,\nplease set the \"Use T0 as ref pose\" option or add them to the bind pose and reimport the skeletal mesh."), FText::FromString(LinksWithoutBindPoses))), FFbxErrors::SkeletalMesh_BonesAreMissingFromBindPose);
+				}
+				return true;
+			}
+
+			FbxNode* FFbxSkelMesh::RecursiveGetFirstMeshNode(FbxNode* Node, FbxNode* NodeToFind)
+			{
+				if (Node == nullptr)
+				{
+					return nullptr;
+				}
+				if (Node->GetMesh() != nullptr)
+					return Node;
+				for (int32 ChildIndex = 0; ChildIndex < Node->GetChildCount(); ++ChildIndex)
+				{
+					FbxNode* MeshNode = RecursiveGetFirstMeshNode(Node->GetChild(ChildIndex), NodeToFind);
+					if (NodeToFind == nullptr)
+					{
+						if (MeshNode != nullptr)
+						{
+							return MeshNode;
+						}
+					}
+					else if (MeshNode == NodeToFind)
+					{
+						return MeshNode;
+					}
+				}
+				return nullptr;
+			}
+
+			FbxNode* FFbxSkelMesh::FindLODGroupNode(FbxNode* NodeLodGroup, int32 LodIndex, FbxNode* NodeToFind)
+			{
+				check(NodeLodGroup->GetChildCount() >= LodIndex);
+				FbxNode* ChildNode = NodeLodGroup->GetChild(LodIndex);
+				if (ChildNode == nullptr)
+				{
+					return nullptr;
+				}
+				return RecursiveGetFirstMeshNode(ChildNode, NodeToFind);
+			}
+
+			void FFbxSkelMesh::RecursiveFindFbxSkelMesh(FbxScene * SDKScene, FbxNode * Node, TArray< TArray<FbxNode*> > & outSkelMeshArray, TArray<FbxNode*> & SkeletonArray)
+			{
+				FbxNode* SkelMeshNode = nullptr;
+				FbxNode* NodeToAdd = Node;
+
+				if (Node->GetMesh() && Node->GetMesh()->GetDeformerCount(FbxDeformer::eSkin) > 0)
+				{
+					SkelMeshNode = Node;
+				}
+				else if (Node->GetNodeAttribute() && Node->GetNodeAttribute()->GetAttributeType() == FbxNodeAttribute::eLODGroup)
+				{
+					// for LODgroup, add the LODgroup to OutSkelMeshArray according to the skeleton that the first child bind to
+					SkelMeshNode = FindLODGroupNode(Node, 0, nullptr);
+					// check if the first child is skeletal mesh
+					if (SkelMeshNode != nullptr && !(SkelMeshNode->GetMesh() && SkelMeshNode->GetMesh()->GetDeformerCount(FbxDeformer::eSkin) > 0))
+					{
+						SkelMeshNode = nullptr;
+					}
+					// else NodeToAdd = Node;
+				}
+
+				if (SkelMeshNode)
+				{
+					// find root skeleton
+
+					check(SkelMeshNode->GetMesh() != nullptr);
+					const int32 fbxDeformerCount = SkelMeshNode->GetMesh()->GetDeformerCount();
+					FbxSkin* Deformer = static_cast<FbxSkin*>(SkelMeshNode->GetMesh()->GetDeformer(0, FbxDeformer::eSkin));
+
+					if (Deformer != NULL)
+					{
+						int32 ClusterCount = Deformer->GetClusterCount();
+						bool bFoundCorrectLink = false;
+						for (int32 ClusterId = 0; ClusterId < ClusterCount; ++ClusterId)
+						{
+							FbxNode* RootBoneLink = Deformer->GetCluster(ClusterId)->GetLink(); //Get the bone influences by this first cluster
+							RootBoneLink = GetRootSkeleton(SDKScene, RootBoneLink); // Get the skeleton root itself
+
+							if (RootBoneLink)
+							{
+								bool bAddedToExistingSkeleton = false;
+								for (int32 SkeletonIndex = 0; SkeletonIndex < SkeletonArray.Num(); ++SkeletonIndex)
+								{
+									if (RootBoneLink == SkeletonArray[SkeletonIndex])
+									{
+										// append to existed outSkelMeshArray element
+										TArray<FbxNode*>& TempArray = outSkelMeshArray[SkeletonIndex];
+										TempArray.Add(NodeToAdd);
+										bAddedToExistingSkeleton = true;
+										break;
+									}
+								}
+
+								// if there is no outSkelMeshArray element that is bind to this skeleton
+								// create new element for outSkelMeshArray
+								if (!bAddedToExistingSkeleton)
+								{
+									TArray<FbxNode*>& TempArray = outSkelMeshArray.AddDefaulted_GetRef();
+									TempArray.Add(NodeToAdd);
+									SkeletonArray.Add(RootBoneLink);
+								}
+
+								bFoundCorrectLink = true;
+								break;
+							}
+						}
+
+						// we didn't find the correct link
+						if (!bFoundCorrectLink)
+						{
+							// 								AddTokenizedErrorMessage(
+							// 									FTokenizedMessage::Create(
+							// 									EMessageSeverity::Warning,
+							// 									FText::Format(LOCTEXT("FBX_NoWeightsOnDeformer", "Ignoring mesh {0} because it but no weights."), FText::FromString(UTF8_TO_TCHAR(SkelMeshNode->GetName())))
+							// 								),
+							// 									FFbxErrors::SkeletalMesh_NoWeightsOnDeformer
+							// 								);
+						}
+					}
+				}
+
+				//Skeletalmesh node can have child so let's always iterate trough child
+				{
+					TArray<FbxNode*> ChildScaled;
+					//Sort the node to have the one with no scaling first so we have more chance
+					//to have a root skeletal mesh with no scale. Because scene import do not support
+					//root skeletal mesh containing scale
+					for (int32 ChildIndex = 0; ChildIndex < Node->GetChildCount(); ++ChildIndex)
+					{
+						FbxNode* ChildNode = Node->GetChild(ChildIndex);
+
+						if (!Node->GetNodeAttribute() || Node->GetNodeAttribute()->GetAttributeType() != FbxNodeAttribute::eLODGroup)
+						{
+							FbxVector4 NoScale(1.0, 1.0, 1.0);
+
+							if (ChildNode->EvaluateLocalScaling() == NoScale)
+							{
+								RecursiveFindFbxSkelMesh(SDKScene, ChildNode, outSkelMeshArray, SkeletonArray);
+							}
+							else
+							{
+								ChildScaled.Add(ChildNode);
+							}
+						}
+					}
+
+					for (FbxNode* ChildNode : ChildScaled)
+					{
+						RecursiveFindFbxSkelMesh(SDKScene, ChildNode, outSkelMeshArray, SkeletonArray);
+					}
+				}
+			}
+
+			void FFbxSkelMesh::RecursiveFindRigidMesh(FbxScene * SDKScene, FbxNode * Node, TArray< TArray<FbxNode*> > & outSkelMeshArray, TArray<FbxNode*> & SkeletonArray)
+			{
+				bool bRigidNodeFound = false;
+				FbxNode* RigidMeshNode = nullptr;
+
+				if (Node->GetMesh())
+				{
+					// ignore skeletal mesh
+					if (Node->GetMesh()->GetDeformerCount(FbxDeformer::eSkin) == 0)
+					{
+						RigidMeshNode = Node;
+						bRigidNodeFound = true;
+					}
+				}
+				else if (Node->GetNodeAttribute() && Node->GetNodeAttribute()->GetAttributeType() == FbxNodeAttribute::eLODGroup)
+				{
+					// for LODgroup, add the LODgroup to OutSkelMeshArray according to the skeleton that the first child bind to
+					FbxNode* FirstLOD = FindLODGroupNode(Node, 0, nullptr);
+					// check if the first child is skeletal mesh
+					if (FirstLOD != nullptr && FirstLOD->GetMesh())
+					{
+						if (FirstLOD->GetMesh()->GetDeformerCount(FbxDeformer::eSkin) == 0)
+						{
+							bRigidNodeFound = true;
+						}
+					}
+
+					if (bRigidNodeFound)
+					{
+						RigidMeshNode = Node;
+					}
+				}
+
+				if (bRigidNodeFound)
+				{
+					// find root skeleton
+					FbxNode* Link = GetRootSkeleton(SDKScene, RigidMeshNode);
+
+					int32 i;
+					for (i = 0; i < SkeletonArray.Num(); i++)
+					{
+						if (Link == SkeletonArray[i])
+						{
+							// append to existed outSkelMeshArray element
+							TArray<FbxNode*>& TempArray = outSkelMeshArray[i];
+							TempArray.Add(RigidMeshNode);
+							break;
+						}
+					}
+
+					// if there is no outSkelMeshArray element that is bind to this skeleton
+					// create new element for outSkelMeshArray
+					if (i == SkeletonArray.Num())
+					{
+						TArray<FbxNode*>& TempArray = outSkelMeshArray.AddDefaulted_GetRef();
+						TempArray.Add(RigidMeshNode);
+						SkeletonArray.Add(Link);
+					}
+				}
+
+				// for LODGroup, we will not deep in.
+				if (!(Node->GetNodeAttribute() && Node->GetNodeAttribute()->GetAttributeType() == FbxNodeAttribute::eLODGroup))
+				{
+					int32 ChildIndex;
+					for (ChildIndex = 0; ChildIndex < Node->GetChildCount(); ++ChildIndex)
+					{
+						RecursiveFindRigidMesh(SDKScene, Node->GetChild(ChildIndex), outSkelMeshArray, SkeletonArray);
+					}
+				}
+			}
+
+			void FFbxSkelMesh::RecursiveFixSkeleton(FbxScene * SDKScene, FbxNode * Node, TArray<FbxNode*> & SkelMeshes, bool bImportNestedMeshes)
+			{
+				FbxNodeAttribute* Attr = Node->GetNodeAttribute();
+				bool NodeIsLodGroup = (Attr && (Attr->GetAttributeType() == FbxNodeAttribute::eLODGroup));
+				if (!NodeIsLodGroup)
+				{
+					for (int32 i = 0; i < Node->GetChildCount(); i++)
+					{
+						RecursiveFixSkeleton(SDKScene, Node->GetChild(i), SkelMeshes, bImportNestedMeshes);
+					}
+				}
+
+				if (Attr && (Attr->GetAttributeType() == FbxNodeAttribute::eMesh || Attr->GetAttributeType() == FbxNodeAttribute::eNull))
+				{
+					if (bImportNestedMeshes && Attr->GetAttributeType() == FbxNodeAttribute::eMesh)
+					{
+						// for leaf mesh, keep them as mesh
+						int32 ChildCount = Node->GetChildCount();
+						int32 ChildIndex;
+						for (ChildIndex = 0; ChildIndex < ChildCount; ChildIndex++)
+						{
+							FbxNode* Child = Node->GetChild(ChildIndex);
+							if (Child->GetMesh() == NULL)
+							{
+								break;
+							}
+						}
+
+						if (ChildIndex != ChildCount)
+						{
+							// Remove from the mesh list it is no longer a mesh
+							SkelMeshes.Remove(Node);
+
+							//replace with skeleton
+							FbxSkeleton* lSkeleton = FbxSkeleton::Create(SDKScene->GetFbxManager(), "");
+							Node->SetNodeAttribute(lSkeleton);
+							lSkeleton->SetSkeletonType(FbxSkeleton::eLimbNode);
+						}
+						else // this mesh may be not in skeleton mesh list. If not, add it.
+						{
+							if (!SkelMeshes.Contains(Node))
+							{
+								SkelMeshes.Add(Node);
+							}
+						}
+					}
+					else
+					{
+						// Remove from the mesh list it is no longer a mesh
+						SkelMeshes.Remove(Node);
+
+						//replace with skeleton
+						FbxSkeleton* lSkeleton = FbxSkeleton::Create(SDKScene->GetFbxManager(), "");
+						Node->SetNodeAttribute(lSkeleton);
+						lSkeleton->SetSkeletonType(FbxSkeleton::eLimbNode);
+					}
+				}
+			}
+
+		}//ns Private
+	}//ns Interchange
+}//ns UE
diff --git a/Engine/Plugins/Experimental/Interchange/Runtime/Source/Parsers/Fbx/Private/FbxSkelMesh.h b/Engine/Plugins/Experimental/Interchange/Runtime/Source/Parsers/Fbx/Private/FbxSkelMesh.h
new file mode 100644
--- /dev/null
+++ b/Engine/Plugins/Experimental/Interchange/Runtime/Source/Parsers/Fbx/Private/FbxSkelMesh.h
@@ -0,0 +1,55 @@
+// Copyright Epic Games, Inc. All Rights Reserved.
+
+#pragma once
+
+#include "CoreMinimal.h"
+#include "FbxAPI.h"
+#include "FbxHelper.h"
+
+namespace UE
+{
+	namespace Interchange
+	{
+		namespace Private
+		{
+			struct FFbxSkelMesh
+			{
+			public:
+				explicit FFbxSkelMesh(FFbxParser& InParser)
+					: Parser(InParser)
+				{}
+
+				void FindSkeletalMeshes(FbxScene* SDKScene, TArray<TArray<FbxNode*>>& outSkelMeshArray, bool bCombineSkeletalMesh, bool bForceFindRigid);
+				void FindAllLODGroupNode(TArray<FbxNode*>& OutNodeInLod, FbxNode* NodeLodGroup, int32 LodIndex);
+				bool FindSkeletonJoints(FbxScene* SDKScene, TArray<FbxNode*>& NodeArray, TArray<FbxNode*>& SortedLinks, FbxArray<FbxAMatrix>& LocalsPerLink);
+				FbxNode* FindLODGroupNode(FbxNode* NodeLodGroup, int32 LodIndex, FbxNode* NodeToFind);
+
+				/**
+				 * Add messages to the message log
+				 */
+				template <typename T>
+				T* AddMessage(FbxGeometryBase* FbxNode) const
+				{
+					T* Item = Parser.AddMessage<T>();
+					Item->MeshName = FFbxHelper::GetMeshName(FbxNode);
+					Item->InterchangeKey = FFbxHelper::GetMeshUniqueID(FbxNode);
+					return Item;
+				}
+
+			private:
+				void RecursiveGetAllMeshNode(TArray<FbxNode*>& OutAllNode, FbxNode* Node);
+				bool IsUnrealBone(FbxNode* Link);
+				void RecursiveBuildSkeleton(FbxNode* Link, TArray<FbxNode*>& OutSortedLinks);
+				bool RetrievePoseFromBindPose(FbxScene* SDKScene, const TArray<FbxNode*>& NodeArray, FbxArray<FbxPose*>& PoseArray);
+				FbxNode* GetRootSkeleton(FbxScene* SDKScene, FbxNode* Link);
+				void BuildSkeletonSystem(FbxScene* SDKScene, TArray<FbxCluster*>& ClusterArray, TArray<FbxNode*>& OutSortedLinks);
+				FbxNode* RecursiveGetFirstMeshNode(FbxNode* Node, FbxNode* NodeToFind);
+				void RecursiveFindFbxSkelMesh(FbxScene* SDKScene, FbxNode* Node, TArray< TArray<FbxNode*> >& outSkelMeshArray, TArray<FbxNode*>& SkeletonArray);
+				void RecursiveFindRigidMesh(FbxScene* SDKScene, FbxNode* Node, TArray< TArray<FbxNode*> >& outSkelMeshArray, TArray<FbxNode*>& SkeletonArray);
+				void RecursiveFixSkeleton(FbxScene* SDKScene, FbxNode* Node, TArray<FbxNode*>& SkelMeshes, bool bImportNestedMeshes);
+
+				FFbxParser& Parser;
+			};
+		}//ns Private
+	}//ns Interchange
+}//ns UE
diff --git a/Engine/Plugins/Experimental/Interchange/Runtime/Source/Pipelines/Private/InterchangeGenericAssetsPipeline.cpp b/Engine/Plugins/Experimental/Interchange/Runtime/Source/Pipelines/Private/InterchangeGenericAssetsPipeline.cpp
new file mode 100644
--- /dev/null
+++ b/Engine/Plugins/Experimental/Interchange/Runtime/Source/Pipelines/Private/InterchangeGenericAssetsPipeline.cpp
@@ -0,0 +1,520 @@
+// Copyright Epic Games, Inc. All Rights Reserved. 
+
+#include "InterchangeGenericAssetsPipeline.h"
+
+#include "Animation/Skeleton.h"
+#include "CoreMinimal.h"
+#include "InterchangeMaterialFactoryNode.h"
+#include "InterchangeMaterialNode.h"
+#include "InterchangeMeshNode.h"
+#include "InterchangePipelineLog.h"
+#include "InterchangePipelineMeshesUtilities.h"
+#include "InterchangeSceneNode.h"
+#include "InterchangeSkeletalMeshFactoryNode.h"
+#include "InterchangeSkeletalMeshLodDataNode.h"
+#include "InterchangeSkeletonFactoryNode.h"
+#include "InterchangeSourceData.h"
+#include "InterchangeStaticMeshFactoryNode.h"
+#include "InterchangeStaticMeshLodDataNode.h"
+#include "InterchangeTextureFactoryNode.h"
+#include "InterchangeTextureNode.h"
+#include "Misc/Paths.h"
+#include "Nodes/InterchangeBaseNode.h"
+#include "Nodes/InterchangeBaseNodeContainer.h"
+#include "Templates/SubclassOf.h"
+#include "UObject/Object.h"
+#include "UObject/ObjectMacros.h"
+
+void UInterchangeGenericAssetsPipeline::PreDialogCleanup(const FName PipelineStackName)
+{
+	//We always clean the pipeline skeleton when showing the dialog
+	Skeleton = nullptr;
+	SaveSettings(PipelineStackName);
+}
+
+bool UInterchangeGenericAssetsPipeline::ExecutePreImportPipeline(UInterchangeBaseNodeContainer* InBaseNodeContainer, const TArray<UInterchangeSourceData*>& InSourceDatas)
+{
+	if (!InBaseNodeContainer)
+	{
+		UE_LOG(LogInterchangePipeline, Warning, TEXT("UInterchangeGenericAssetsPipeline: Cannot execute pre-import pipeline because InBaseNodeContrainer is null"));
+		return false;
+	}
+
+	BaseNodeContainer = InBaseNodeContainer;
+	SourceDatas.Empty(InSourceDatas.Num());
+	for (const UInterchangeSourceData* SourceData : InSourceDatas)
+	{
+		SourceDatas.Add(SourceData);
+	}
+
+	PipelineMeshesUtilities = UInterchangePipelineMeshesUtilities::CreateInterchangePipelineMeshesUtilities(BaseNodeContainer);
+	
+	//Find all translated node we need for this pipeline
+	BaseNodeContainer->IterateNodes([this](const FString& NodeUid, UInterchangeBaseNode* Node)
+	{
+		switch(Node->GetNodeContainerType())
+		{
+			case EInterchangeNodeContainerType::NodeContainerType_TranslatedAsset:
+			{
+				if (UInterchangeTextureNode* TextureNode = Cast<UInterchangeTextureNode>(Node))
+				{
+					TextureNodes.Add(TextureNode);
+				}
+				else if (UInterchangeMaterialNode* MaterialNode = Cast<UInterchangeMaterialNode>(Node))
+				{
+					MaterialNodes.Add(MaterialNode);
+				}
+			}
+			break;
+		}
+	});
+	
+	if (bImportTextures)
+	{
+		//import textures
+		for (const UInterchangeTextureNode* TextureNode : TextureNodes)
+		{
+			HandleCreationOfTextureFactoryNode(TextureNode);
+		}
+	}
+
+	if (bImportMaterials)
+	{
+		//import materials
+		for (const UInterchangeMaterialNode* MaterialNode : MaterialNodes)
+		{
+			CreateMaterialFactoryNode(MaterialNode);
+		}
+	}
+
+	//Create skeletalmesh factory nodes
+	ExecutePreImportPipelineSkeletalMesh();
+
+	if (bImportStaticMeshes)
+	{
+		if (bCombineStaticMeshes)
+		{
+			// Combine all the static meshes
+
+			bool bFoundMeshes = false;
+			if (bBakeMeshes)
+			{
+				// If baking transforms, get all the static mesh instance nodes, and group them by LOD
+				TArray<FString> MeshUids;
+				PipelineMeshesUtilities->GetAllStaticMeshInstance(MeshUids);
+
+				TMap<int32, TArray<FString>> MeshUidsPerLodIndex;
+
+				for (const FString& MeshUid : MeshUids)
+				{
+					const FInterchangeMeshInstance& MeshInstance = PipelineMeshesUtilities->GetMeshInstanceByUid(MeshUid);
+					for (const auto& LodIndexAndSceneNodeContainer : MeshInstance.SceneNodePerLodIndex)
+					{
+						const int32 LodIndex = LodIndexAndSceneNodeContainer.Key;
+						const FInterchangeLodSceneNodeContainer& SceneNodeContainer = LodIndexAndSceneNodeContainer.Value;
+
+						TArray<FString>& TranslatedNodes = MeshUidsPerLodIndex.FindOrAdd(LodIndex);
+						for (const UInterchangeSceneNode* SceneNode : SceneNodeContainer.SceneNodes)
+						{
+							TranslatedNodes.Add(SceneNode->GetUniqueID());
+						}
+					}
+				}
+
+				// If we got some instances, create a static mesh factory node
+				if (MeshUidsPerLodIndex.Num() > 0)
+				{
+					UInterchangeStaticMeshFactoryNode* StaticMeshFactoryNode = CreateStaticMeshFactoryNode(MeshUidsPerLodIndex);
+					StaticMeshFactoryNodes.Add(StaticMeshFactoryNode);
+					bFoundMeshes = true;
+				}
+			}
+
+			if (!bFoundMeshes)
+			{
+				// If we haven't yet managed to build a factory node, look at static mesh geometry directly.
+				TArray<FString> MeshUids;
+				PipelineMeshesUtilities->GetAllStaticMeshGeometry(MeshUids);
+
+				TMap<int32, TArray<FString>> MeshUidsPerLodIndex;
+
+				for (const FString& MeshUid : MeshUids)
+				{
+					// MeshGeometry cannot have Lod since LODs are defined in the scene node
+					const FInterchangeMeshGeometry& MeshGeometry = PipelineMeshesUtilities->GetMeshGeometryByUid(MeshUid);
+					const int32 LodIndex = 0;
+					TArray<FString>& TranslatedNodes = MeshUidsPerLodIndex.FindOrAdd(LodIndex);
+					TranslatedNodes.Add(MeshGeometry.MeshUid);
+				}
+
+				if (MeshUidsPerLodIndex.Num() > 0)
+				{
+					UInterchangeStaticMeshFactoryNode* StaticMeshFactoryNode = CreateStaticMeshFactoryNode(MeshUidsPerLodIndex);
+					StaticMeshFactoryNodes.Add(StaticMeshFactoryNode);
+				}
+			}
+
+		}
+		else
+		{
+			// Do not combine static meshes
+
+			bool bFoundMeshes = false;
+			if (bBakeMeshes)
+			{
+				TArray<FString> MeshUids;
+				PipelineMeshesUtilities->GetAllStaticMeshInstance(MeshUids);
+				
+				for (const FString& MeshUid : MeshUids)
+				{
+					// @TODO: check this is correct. What happens when we explicitly don't combine scenes with a LOD group node and multiple transform nodes?
+					TMap<int32, TArray<FString>> MeshUidsPerLodIndex;
+
+					const FInterchangeMeshInstance& MeshInstance = PipelineMeshesUtilities->GetMeshInstanceByUid(MeshUid);
+					for (const auto& LodIndexAndSceneNodeContainer : MeshInstance.SceneNodePerLodIndex)
+					{
+						const int32 LodIndex = LodIndexAndSceneNodeContainer.Key;
+						const FInterchangeLodSceneNodeContainer& SceneNodeContainer = LodIndexAndSceneNodeContainer.Value;
+
+						TArray<FString>& TranslatedNodes = MeshUidsPerLodIndex.FindOrAdd(LodIndex);
+						for (const UInterchangeSceneNode* SceneNode : SceneNodeContainer.SceneNodes)
+						{
+							TranslatedNodes.Add(SceneNode->GetUniqueID());
+						}
+					}
+
+					if (MeshUidsPerLodIndex.Num() > 0)
+					{
+						UInterchangeStaticMeshFactoryNode* StaticMeshFactoryNode = CreateStaticMeshFactoryNode(MeshUidsPerLodIndex);
+						StaticMeshFactoryNodes.Add(StaticMeshFactoryNode);
+						bFoundMeshes = true;
+					}
+				}
+			}
+
+			if (!bFoundMeshes)
+			{
+				TArray<FString> MeshUids;
+				PipelineMeshesUtilities->GetAllStaticMeshGeometry(MeshUids);
+
+				for (const FString& MeshUid : MeshUids)
+				{
+					TMap<int32, TArray<FString>> MeshUidsPerLodIndex;
+
+					const FInterchangeMeshGeometry& MeshGeometry = PipelineMeshesUtilities->GetMeshGeometryByUid(MeshUid);
+					const int32 LodIndex = 0;
+					TArray<FString>& TranslatedNodes = MeshUidsPerLodIndex.FindOrAdd(LodIndex);
+					TranslatedNodes.Add(MeshGeometry.MeshUid);
+
+					if (MeshUidsPerLodIndex.Num() > 0)
+					{
+						UInterchangeStaticMeshFactoryNode* StaticMeshFactoryNode = CreateStaticMeshFactoryNode(MeshUidsPerLodIndex);
+						StaticMeshFactoryNodes.Add(StaticMeshFactoryNode);
+					}
+				}
+			}
+		}
+	}
+
+	ImplementUseSourceNameForAssetOption();
+
+	return true;
+}
+
+bool UInterchangeGenericAssetsPipeline::ExecutePostImportPipeline(const UInterchangeBaseNodeContainer* InBaseNodeContainer, const FString& NodeKey, UObject* CreatedAsset, bool bIsAReimport)
+{
+	//We do not use the provided base container since ExecutePreImportPipeline cache it
+	//We just make sure the same one is pass in parameter
+	if (!InBaseNodeContainer || !ensure(BaseNodeContainer == InBaseNodeContainer) || !CreatedAsset)
+	{
+		return false;
+	}
+
+	UInterchangeBaseNode* Node = BaseNodeContainer->GetNode(NodeKey);
+	if (!Node)
+	{
+		return false;
+	}
+
+	//Finish the physics asset import, it need the skeletal mesh render data to create the physics collision geometry
+	PostImportPhysicsAssetImport(CreatedAsset, Node);
+
+	PostImportTextureAssetImport(CreatedAsset, bIsAReimport);
+
+	return true;
+}
+
+UInterchangeMaterialFactoryNode* UInterchangeGenericAssetsPipeline::CreateMaterialFactoryNode(const UInterchangeMaterialNode* MaterialNode)
+{
+	FString DisplayLabel = MaterialNode->GetDisplayLabel();
+	FString NodeUid = UInterchangeMaterialFactoryNode::GetMaterialFactoryNodeUidFromMaterialNodeUid(MaterialNode->GetUniqueID());
+	UInterchangeMaterialFactoryNode* MaterialFactoryNode = nullptr;
+	if (BaseNodeContainer->IsNodeUidValid(NodeUid))
+	{
+		//The node already exist, just return it
+		MaterialFactoryNode = Cast<UInterchangeMaterialFactoryNode>(BaseNodeContainer->GetNode(NodeUid));
+		if (!ensure(MaterialFactoryNode))
+		{
+			//Log an error
+		}
+	}
+	else
+	{
+		MaterialFactoryNode = NewObject<UInterchangeMaterialFactoryNode>(BaseNodeContainer, NAME_None);
+		if (!ensure(MaterialFactoryNode))
+		{
+			return nullptr;
+		}
+		//Creating a Material
+		MaterialFactoryNode->InitializeMaterialNode(NodeUid, DisplayLabel, UMaterial::StaticClass()->GetName());
+		MaterialFactoryNode->SetCustomTranslatedMaterialNodeUid(MaterialNode->GetUniqueID());
+		TArray<FString> TranslatedDependencies;
+		MaterialNode->GetTextureDependencies(TranslatedDependencies);
+		for (const FString& DepUid : TranslatedDependencies)
+		{
+			FString FactoryDepUid = UInterchangeTextureFactoryNode::GetTextureFactoryNodeUidFromTextureNodeUid(DepUid);
+			if (UInterchangeTextureFactoryNode* TextureDep = Cast<UInterchangeTextureFactoryNode>(BaseNodeContainer->GetNode(FactoryDepUid)))
+			{
+				//We found a factory dependency, because we need to be able to retrieve the UObject create by the factory
+				MaterialFactoryNode->SetTextureDependencyUid(FactoryDepUid);
+				//We also add it to the base dependencies so the texture factory will be called before this material factory
+				MaterialFactoryNode->AddFactoryDependencyUid(FactoryDepUid);
+			}
+		}
+		BaseNodeContainer->AddNode(MaterialFactoryNode);
+		MaterialFactoryNodes.Add(MaterialFactoryNode);
+		MaterialFactoryNode->AddTargetNodeUid(MaterialNode->GetUniqueID());
+		MaterialNode->AddTargetNodeUid(MaterialFactoryNode->GetUniqueID());
+	}
+	return MaterialFactoryNode;
+}
+
+
+void UInterchangeGenericAssetsPipeline::ImplementUseSourceNameForAssetOption()
+{
+	const UClass* SkeletalMeshFactoryNodeClass = UInterchangeSkeletalMeshFactoryNode::StaticClass();
+	TArray<FString> SkeletalMeshNodeUids;
+	BaseNodeContainer->GetNodes(SkeletalMeshFactoryNodeClass, SkeletalMeshNodeUids);
+
+	const UClass* StaticMeshFactoryNodeClass = UInterchangeStaticMeshFactoryNode::StaticClass();
+	TArray<FString> StaticMeshNodeUids;
+	BaseNodeContainer->GetNodes(StaticMeshFactoryNodeClass, StaticMeshNodeUids);
+
+	//TODO count also the imported animations
+
+	//If we import only one asset, and bUseSourceNameForAsset is true, we want to rename the asset using the file name.
+	const int32 MeshesAndAnimsImportedNodeCount = SkeletalMeshNodeUids.Num() + StaticMeshNodeUids.Num();
+
+	//Skeletalmesh must always rename created skeleton and physics asset to (SKNAME_Skeleton of SKNAME_PhysicsAsset). So the pipeline option
+	// bUseSourceNameForAsset will change or not the SKNAME.
+	ImplementUseSourceNameForAssetOptionSkeletalMesh(MeshesAndAnimsImportedNodeCount);
+
+	if (bUseSourceNameForAsset && MeshesAndAnimsImportedNodeCount == 1)
+	{
+		if (StaticMeshNodeUids.Num() > 0)
+		{
+			UInterchangeStaticMeshFactoryNode* StaticMeshNode = Cast<UInterchangeStaticMeshFactoryNode>(BaseNodeContainer->GetNode(StaticMeshNodeUids[0]));
+			const FString DisplayLabelName = FPaths::GetBaseFilename(SourceDatas[0]->GetFilename());
+			StaticMeshNode->SetDisplayLabel(DisplayLabelName);
+		}
+	}
+
+}
+bool UInterchangeGenericAssetsPipeline::MakeMeshFactoryNodeUidAndDisplayLabel(const TMap<int32, TArray<FString>>& MeshUidsPerLodIndex, int32 LodIndex, FString& NewNodeUid, FString& DisplayLabel)
+{
+	int32 SceneNodeCount = 0;
+
+	if (!ensure(LodIndex >= 0 && MeshUidsPerLodIndex.Num() > LodIndex))
+	{
+		return false;
+	}
+
+	for (const TPair<int32, TArray<FString>>& LodIndexAndMeshUids : MeshUidsPerLodIndex)
+	{
+		if (LodIndex == LodIndexAndMeshUids.Key)
+		{
+			const TArray<FString>& Uids = LodIndexAndMeshUids.Value;
+			if (Uids.Num() > 0)
+			{
+				const FString& Uid = Uids[0];
+				UInterchangeBaseNode* Node = BaseNodeContainer->GetNode(Uid);
+
+				if (UInterchangeMeshNode* MeshNode = Cast<UInterchangeMeshNode>(Node))
+				{
+					DisplayLabel = Node->GetDisplayLabel();
+					NewNodeUid = Uid;
+					return true;
+				}
+
+				if (UInterchangeSceneNode* SceneNode = Cast<UInterchangeSceneNode>(Node))
+				{
+					FString RefMeshUid;
+					if (SceneNode->GetCustomAssetInstanceUid(RefMeshUid))
+					{
+						UInterchangeBaseNode* MeshNode = BaseNodeContainer->GetNode(RefMeshUid);
+						if (MeshNode)
+						{
+							DisplayLabel = MeshNode->GetDisplayLabel();
+							if (Uids.Num() == 1)
+							{
+								// If we are instancing one scene node, we want to use it to name the mesh
+								DisplayLabel = SceneNode->GetDisplayLabel();
+							}
+
+							// Use the first scene node uid this LOD references, adding backslash since this uid is not asset typed (\\Mesh\\) like a mesh node
+							// @TODO: change this so that scene nodes get a type prefix of their own?
+							NewNodeUid = TEXT("\\") + Uid;
+							return true;
+						}
+					}
+				}
+			}
+
+			// We found the lod but there is no valid Mesh node to return the Uid
+			break;
+		}
+	}
+
+	return false;
+}
+
+UInterchangeStaticMeshFactoryNode* UInterchangeGenericAssetsPipeline::CreateStaticMeshFactoryNode(const TMap<int32, TArray<FString>>& MeshUidsPerLodIndex)
+{
+	if (MeshUidsPerLodIndex.Num() == 0)
+	{
+		return nullptr;
+	}
+
+	// Create the static mesh factory node, name it according to the first mesh node compositing the meshes
+	FString StaticMeshUid_MeshNamePart;
+	FString DisplayLabel;
+	const int32 BaseLodIndex = 0;
+	if (!MakeMeshFactoryNodeUidAndDisplayLabel(MeshUidsPerLodIndex, BaseLodIndex, StaticMeshUid_MeshNamePart, DisplayLabel))
+	{
+		// Log an error
+		return nullptr;
+	}
+
+	const FString StaticMeshUid = TEXT("\\StaticMesh") + StaticMeshUid_MeshNamePart;
+	UInterchangeStaticMeshFactoryNode* StaticMeshFactoryNode = NewObject<UInterchangeStaticMeshFactoryNode>(BaseNodeContainer, NAME_None);
+	if (!ensure(StaticMeshFactoryNode))
+	{
+		return nullptr;
+	}
+
+	StaticMeshFactoryNode->InitializeStaticMeshNode(StaticMeshUid, DisplayLabel, UStaticMesh::StaticClass()->GetName());
+	BaseNodeContainer->AddNode(StaticMeshFactoryNode);
+
+	AddLodDataToStaticMesh(StaticMeshFactoryNode, MeshUidsPerLodIndex);
+
+	switch (VertexColorImportOption)
+	{
+		case EInterchangeVertexColorImportOption::IVCIO_Replace:
+		{
+			StaticMeshFactoryNode->SetCustomVertexColorReplace(true);
+		}
+		break;
+		case EInterchangeVertexColorImportOption::IVCIO_Ignore:
+		{
+			StaticMeshFactoryNode->SetCustomVertexColorIgnore(true);
+		}
+		break;
+		case EInterchangeVertexColorImportOption::IVCIO_Override:
+		{
+			StaticMeshFactoryNode->SetCustomVertexColorOverride(VertexOverrideColor);
+		}
+		break;
+	}
+	return StaticMeshFactoryNode;
+}
+
+
+UInterchangeStaticMeshLodDataNode* UInterchangeGenericAssetsPipeline::CreateStaticMeshLodDataNode(const FString& NodeName, const FString& NodeUniqueID)
+{
+	FString DisplayLabel(NodeName);
+	FString NodeUID(NodeUniqueID);
+	UInterchangeStaticMeshLodDataNode* StaticMeshLodDataNode = NewObject<UInterchangeStaticMeshLodDataNode>(BaseNodeContainer, NAME_None);
+	if (!ensure(StaticMeshLodDataNode))
+	{
+		// @TODO: log error
+		return nullptr;
+	}
+
+	StaticMeshLodDataNode->InitializeNode(NodeUID, DisplayLabel, EInterchangeNodeContainerType::NodeContainerType_FactoryData);
+	BaseNodeContainer->AddNode(StaticMeshLodDataNode);
+	return StaticMeshLodDataNode;
+}
+
+
+void UInterchangeGenericAssetsPipeline::AddLodDataToStaticMesh(UInterchangeStaticMeshFactoryNode* StaticMeshFactoryNode, const TMap<int32, TArray<FString>>& NodeUidsPerLodIndex)
+{
+	const FString StaticMeshFactoryUid = StaticMeshFactoryNode->GetUniqueID();
+
+	for (const TPair<int32, TArray<FString>>& LodIndexAndNodeUids : NodeUidsPerLodIndex)
+	{
+		const int32 LodIndex = LodIndexAndNodeUids.Key;
+		if (!bImportLods && LodIndex > 0)
+		{
+			// If the pipeline should not import lods, skip any lod over base lod
+			continue;
+		}
+
+		const TArray<FString>& NodeUids = LodIndexAndNodeUids.Value;
+
+		// Create a lod data node with all the meshes for this LOD
+		const FString StaticMeshLodDataName = TEXT("LodData") + FString::FromInt(LodIndex);
+		const FString LODDataPrefix = TEXT("\\LodData") + (LodIndex > 0 ? FString::FromInt(LodIndex) : TEXT(""));
+		const FString StaticMeshLodDataUniqueID = LODDataPrefix + StaticMeshFactoryUid;
+
+		// Create LodData node if it doesn't already exist
+		UInterchangeStaticMeshLodDataNode* LodDataNode = Cast<UInterchangeStaticMeshLodDataNode>(BaseNodeContainer->GetNode(StaticMeshLodDataUniqueID));
+		if (!LodDataNode)
+		{
+			// Add the data for the LOD (all the mesh node fbx path, so we can find them when we will create the payload data)
+			LodDataNode = CreateStaticMeshLodDataNode(StaticMeshLodDataName, StaticMeshLodDataUniqueID);
+			LodDataNode->SetParentUid(StaticMeshFactoryUid);
+			StaticMeshFactoryNode->AddLodDataUniqueId(StaticMeshLodDataUniqueID);
+		}
+
+		for (const FString& NodeUid : NodeUids)
+		{
+			TArray<FString> MaterialDependencies;
+			if (const UInterchangeSceneNode* SceneNode = Cast<UInterchangeSceneNode>(BaseNodeContainer->GetNode(NodeUid)))
+			{
+				FString MeshDependency;
+				SceneNode->GetCustomAssetInstanceUid(MeshDependency);
+				if (BaseNodeContainer->IsNodeUidValid(MeshDependency))
+				{
+					StaticMeshFactoryNode->AddTargetNodeUid(MeshDependency);
+					BaseNodeContainer->GetNode(MeshDependency)->AddTargetNodeUid(StaticMeshFactoryNode->GetUniqueID());
+				}
+				
+				SceneNode->GetMaterialDependencyUids(MaterialDependencies);
+			}
+			else if (const UInterchangeMeshNode* MeshNode = Cast<UInterchangeMeshNode>(BaseNodeContainer->GetNode(NodeUid)))
+			{
+				StaticMeshFactoryNode->AddTargetNodeUid(NodeUid);
+				MeshNode->AddTargetNodeUid(StaticMeshFactoryNode->GetUniqueID());
+				MeshNode->GetMaterialDependencies(MaterialDependencies);
+			}
+
+			const int32 MaterialCount = MaterialDependencies.Num();
+			for (int32 MaterialIndex = 0; MaterialIndex < MaterialCount; ++MaterialIndex)
+			{
+				const FString MaterialFactoryNodeUid = UInterchangeMaterialFactoryNode::GetMaterialFactoryNodeUidFromMaterialNodeUid(MaterialDependencies[MaterialIndex]);
+				if (BaseNodeContainer->IsNodeUidValid(MaterialFactoryNodeUid))
+				{
+					// Create a factory dependency so Material asset are imported before the skeletal mesh asset
+					TArray<FString> FactoryDependencies;
+					StaticMeshFactoryNode->GetFactoryDependencies(FactoryDependencies);
+					if (!FactoryDependencies.Contains(MaterialFactoryNodeUid))
+					{
+						StaticMeshFactoryNode->AddFactoryDependencyUid(MaterialFactoryNodeUid);
+					}
+				}
+			}
+
+			LodDataNode->AddMeshUid(NodeUid);
+		}
+	}
+}
