Fix CIS missing include header
Fix CIS DoesPackageExist API change

#rnx

diff --git a/Engine/Plugins/Experimental/Interchange/Runtime/Source/Nodes/Public/InterchangeCineCameraFactoryNode.h b/Engine/Plugins/Experimental/Interchange/Runtime/Source/Nodes/Public/InterchangeCineCameraFactoryNode.h
new file mode 100644
--- /dev/null
+++ b/Engine/Plugins/Experimental/Interchange/Runtime/Source/Nodes/Public/InterchangeCineCameraFactoryNode.h
@@ -0,0 +1,65 @@
+// Copyright Epic Games, Inc. All Rights Reserved.
+
+#pragma once
+
+#include "CoreMinimal.h"
+#include "InterchangeActorFactoryNode.h"
+
+#if WITH_ENGINE
+	#include "CineCameraComponent.h"
+#endif
+
+#include "InterchangeCineCameraFactoryNode.generated.h"
+
+UCLASS(BlueprintType)
+class INTERCHANGENODES_API UInterchangeCineCameraFactoryNode : public UInterchangeActorFactoryNode
+{
+	GENERATED_BODY()
+
+public:
+	UFUNCTION(BlueprintCallable, Category = "Interchange | Node | CameraFactory")
+	bool GetCustomFocalLength(float& AttributeValue) const
+	{
+		IMPLEMENT_NODE_ATTRIBUTE_GETTER(FocalLength, float);
+	}
+
+	UFUNCTION(BlueprintCallable, Category = "Interchange | Node | CameraFactory")
+	bool SetCustomFocalLength(const float& AttributeValue, bool bAddApplyDelegate = true)
+	{
+		IMPLEMENT_NODE_ATTRIBUTE_SETTER(UInterchangeCineCameraFactoryNode, FocalLength, float, UCineCameraComponent);
+	}
+
+	UFUNCTION(BlueprintCallable, Category = "Interchange | Node | CameraFactory")
+	bool GetCustomSensorWidth(float& AttributeValue) const
+	{
+		IMPLEMENT_NODE_ATTRIBUTE_GETTER(SensorWidth, float);
+	}
+
+	UFUNCTION(BlueprintCallable, Category = "Interchange | Node | CameraFactory")
+	bool SetCustomSensorWidth(const float& AttributeValue, bool bAddApplyDelegate = true)
+	{
+		IMPLEMENT_NODE_ATTRIBUTE_SETTER(UInterchangeCineCameraFactoryNode, SensorWidth, float, UCineCameraComponent);
+	}
+
+	UFUNCTION(BlueprintCallable, Category = "Interchange | Node | CameraFactory")
+	bool GetCustomSensorHeight(float& AttributeValue) const
+	{
+		IMPLEMENT_NODE_ATTRIBUTE_GETTER(SensorHeight, float);
+	}
+
+	UFUNCTION(BlueprintCallable, Category = "Interchange | Node | CameraFactory")
+	bool SetCustomSensorHeight(const float& AttributeValue, bool bAddApplyDelegate = true)
+	{
+		IMPLEMENT_NODE_ATTRIBUTE_SETTER(UInterchangeCineCameraFactoryNode, SensorHeight, float, UCineCameraComponent);
+	}
+
+private:
+	const UE::Interchange::FAttributeKey Macro_CustomFocalLengthKey = UE::Interchange::FAttributeKey(TEXT("FocalLength"));
+	const UE::Interchange::FAttributeKey Macro_CustomSensorWidthKey = UE::Interchange::FAttributeKey(TEXT("SensorWidth"));
+	const UE::Interchange::FAttributeKey Macro_CustomSensorHeightKey = UE::Interchange::FAttributeKey(TEXT("SensorHeight"));
+
+private:
+	IMPLEMENT_NODE_ATTRIBUTE_APPLY_UOBJECT_BYNAME(FocalLength, float, UCineCameraComponent, TEXT("CurrentFocalLength"));
+	IMPLEMENT_NODE_ATTRIBUTE_APPLY_UOBJECT_BYNAME(SensorWidth, float, UCineCameraComponent, TEXT("Filmback.SensorWidth"));
+	IMPLEMENT_NODE_ATTRIBUTE_APPLY_UOBJECT_BYNAME(SensorHeight, float, UCineCameraComponent, TEXT("Filmback.SensorHeight"));
+};
diff --git a/Engine/Plugins/Runtime/Interchange/Source/Nodes/Public/InterchangeCineCameraFactoryNode.h b/Engine/Plugins/Runtime/Interchange/Source/Nodes/Public/InterchangeCineCameraFactoryNode.h
deleted file mode 100644
diff --git a/Engine/Source/Runtime/Experimental/Interchange/Core/Public/PackageUtils/PackageUtils.h b/Engine/Source/Runtime/Experimental/Interchange/Core/Public/PackageUtils/PackageUtils.h
--- a/Engine/Source/Runtime/Experimental/Interchange/Core/Public/PackageUtils/PackageUtils.h
+++ b/Engine/Source/Runtime/Experimental/Interchange/Core/Public/PackageUtils/PackageUtils.h
@@ -26,7 +26,7 @@ namespace UE
 				if (PackageName.Len() > 0)
 				{
 					FString PackagePath;
-					if (FPackageName::DoesPackageExist(PackageName, &PackagePath))
+					if (FPackageName::DoesPackageExist(PackageName, nullptr, &PackagePath))
 					{
 						const FString FileExtension = FPaths::GetExtension(PackagePath, true);
 						if (FileExtension == FPackageName::GetMapPackageExtension())
