Remove ifdef out code

diff --git a/Engine/Plugins/Experimental/Interchange/Runtime/Source/Parsers/Fbx/Private/FbxHelper.cpp b/Engine/Plugins/Experimental/Interchange/Runtime/Source/Parsers/Fbx/Private/FbxHelper.cpp
--- a/Engine/Plugins/Experimental/Interchange/Runtime/Source/Parsers/Fbx/Private/FbxHelper.cpp
+++ b/Engine/Plugins/Experimental/Interchange/Runtime/Source/Parsers/Fbx/Private/FbxHelper.cpp
@@ -140,919 +140,6 @@ namespace UE
 				FormatArguments.Add(TEXT("UniqueID"), UniqueID);
 				return FString::Format(TEXT("{UniqueID}"), FormatArguments);
 			}
-#if 0
-			void FFbxHelper::FindSkeletalMeshes(FbxScene* SDKScene, TArray< TArray<FbxNode*> >& outSkelMeshArray, bool bCombineSkeletalMesh, bool bForceFindRigid)
-			{
-				TArray<FbxNode*> SkeletonArray;
-				FbxNode* RootNode = SDKScene->GetRootNode();
-
-				// a) find skeletal meshes
-
-				RecursiveFindFbxSkelMesh(SDKScene, RootNode, outSkelMeshArray, SkeletonArray);
-				// for skeletal mesh, we convert the skeleton system to skeleton
-				// in less we recognize bone mesh as rigid mesh if they are textured
-				for (int32 SkelIndex = 0; SkelIndex < SkeletonArray.Num(); SkelIndex++)
-				{
-					const bool bImportMeshesInBoneHierarchy = true;
-					RecursiveFixSkeleton(SDKScene, SkeletonArray[SkelIndex], outSkelMeshArray[SkelIndex], bImportMeshesInBoneHierarchy);
-				}
-
-
-
-				// b) find rigid mesh
-
-				// If we are attempting to import a skeletal mesh but we have no hierarchy attempt to find a rigid mesh.
-				if (bForceFindRigid || outSkelMeshArray.Num() == 0)
-				{
-					RecursiveFindRigidMesh(SDKScene, RootNode, outSkelMeshArray, SkeletonArray);
-					if (bForceFindRigid)
-					{
-						//Cleanup the rigid mesh, We want to remove any real static mesh from the outSkelMeshArray
-						//Any non skinned mesh that contain no animation should be part of this array.
-						int32 AnimStackCount = SDKScene->GetSrcObjectCount<FbxAnimStack>();
-						TArray<int32> SkeletalMeshArrayToRemove;
-						for (int32 i = 0; i < outSkelMeshArray.Num(); i++)
-						{
-							bool bIsValidSkeletal = false;
-							TArray<FbxNode*>& NodeArray = outSkelMeshArray[i];
-							for (FbxNode* InspectedNode : NodeArray)
-							{
-								FbxMesh* Mesh = InspectedNode->GetMesh();
-
-								FbxLODGroup* LodGroup = InspectedNode->GetLodGroup();
-								if (LodGroup != nullptr)
-								{
-									FbxNode* SkelMeshNode = FFbxHelper::FindLODGroupNode(InspectedNode, 0, nullptr);
-									if (SkelMeshNode != nullptr)
-									{
-										Mesh = SkelMeshNode->GetMesh();
-									}
-								}
-
-								if (Mesh == nullptr)
-								{
-									continue;
-								}
-								if (Mesh->GetDeformerCount(FbxDeformer::eSkin) > 0)
-								{
-									bIsValidSkeletal = true;
-									break;
-								}
-								//If there is some anim object we count this as a valid skeletal mesh imported as rigid mesh
-								for (int32 AnimStackIndex = 0; AnimStackIndex < AnimStackCount; AnimStackIndex++)
-								{
-									FbxAnimStack* CurAnimStack = SDKScene->GetSrcObject<FbxAnimStack>(AnimStackIndex);
-									// set current anim stack
-
-									SDKScene->SetCurrentAnimationStack(CurAnimStack);
-
-									FbxTimeSpan AnimTimeSpan(FBXSDK_TIME_INFINITE, FBXSDK_TIME_MINUS_INFINITE);
-									InspectedNode->GetAnimationInterval(AnimTimeSpan, CurAnimStack);
-
-									if (AnimTimeSpan.GetDuration() > 0)
-									{
-										bIsValidSkeletal = true;
-										break;
-									}
-								}
-								if (bIsValidSkeletal)
-								{
-									break;
-								}
-							}
-							if (!bIsValidSkeletal)
-							{
-								SkeletalMeshArrayToRemove.Add(i);
-							}
-						}
-						for (int32 i = SkeletalMeshArrayToRemove.Num() - 1; i >= 0; --i)
-						{
-							if (!SkeletalMeshArrayToRemove.IsValidIndex(i) || !outSkelMeshArray.IsValidIndex(SkeletalMeshArrayToRemove[i]))
-								continue;
-							int32 IndexToRemove = SkeletalMeshArrayToRemove[i];
-							outSkelMeshArray[IndexToRemove].Empty();
-							outSkelMeshArray.RemoveAt(IndexToRemove);
-						}
-					}
-				}
-				//Empty the skeleton array
-				SkeletonArray.Empty();
-
-
-				if (bCombineSkeletalMesh)
-				{
-					//Merge all the skeletal mesh arrays into one combine mesh
-					TArray<FbxNode*> CombineNodes;
-					for (TArray<FbxNode*>& Parts : outSkelMeshArray)
-					{
-						CombineNodes.Append(Parts);
-					}
-					outSkelMeshArray.Empty(1);
-					outSkelMeshArray.Add(CombineNodes);
-				}
-			}
-
-			void FFbxHelper::RecursiveGetAllMeshNode(TArray<FbxNode*> & OutAllNode, FbxNode * Node)
-			{
-				if (Node == nullptr)
-				{
-					return;
-				}
-
-				if (Node->GetMesh() != nullptr)
-				{
-					OutAllNode.Add(Node);
-					return;
-				}
-
-				//Look if its a generated LOD
-				FString FbxGeneratedNodeName = UTF8_TO_TCHAR(Node->GetName());
-				if (FbxGeneratedNodeName.Contains(TEXT(GeneratedLODNameSuffix)))
-				{
-					FString SuffixSearch = TEXT(GeneratedLODNameSuffix);
-					int32 SuffixIndex = FbxGeneratedNodeName.Find(SuffixSearch, ESearchCase::CaseSensitive, ESearchDir::FromEnd);
-					SuffixIndex += SuffixSearch.Len();
-					FString LODXNumber = FbxGeneratedNodeName.RightChop(SuffixIndex).Left(1);
-					if (LODXNumber.IsNumeric())
-					{
-						OutAllNode.Add(Node);
-						return;
-					}
-				}
-
-				for (int32 ChildIndex = 0; ChildIndex < Node->GetChildCount(); ++ChildIndex)
-				{
-					RecursiveGetAllMeshNode(OutAllNode, Node->GetChild(ChildIndex));
-				}
-			}
-
-			void FFbxHelper::FindAllLODGroupNode(TArray<FbxNode*> & OutNodeInLod, FbxNode * NodeLodGroup, int32 LodIndex)
-			{
-				check(NodeLodGroup->GetChildCount() >= LodIndex);
-				FbxNode* ChildNode = NodeLodGroup->GetChild(LodIndex);
-				if (ChildNode == nullptr)
-				{
-					return;
-				}
-				RecursiveGetAllMeshNode(OutNodeInLod, ChildNode);
-			}
-
-			bool FFbxHelper::IsUnrealBone(FbxNode * Link)
-			{
-				FbxNodeAttribute* Attr = Link->GetNodeAttribute();
-				if (Attr)
-				{
-					FbxNodeAttribute::EType AttrType = Attr->GetAttributeType();
-					if (AttrType == FbxNodeAttribute::eSkeleton ||
-						AttrType == FbxNodeAttribute::eMesh ||
-						AttrType == FbxNodeAttribute::eNull)
-					{
-						return true;
-					}
-				}
-
-				return false;
-			}
-
-			void FFbxHelper::RecursiveBuildSkeleton(FbxNode * Link, TArray<FbxNode*> & OutSortedLinks)
-			{
-				if (IsUnrealBone(Link))
-				{
-					OutSortedLinks.Add(Link);
-					int32 ChildIndex;
-					for (ChildIndex = 0; ChildIndex < Link->GetChildCount(); ChildIndex++)
-					{
-						RecursiveBuildSkeleton(Link->GetChild(ChildIndex), OutSortedLinks);
-					}
-				}
-			}
-
-			bool FFbxHelper::RetrievePoseFromBindPose(FbxScene * SDKScene, const TArray<FbxNode*> & NodeArray, FbxArray<FbxPose*> & PoseArray)
-			{
-				const int32 PoseCount = SDKScene->GetPoseCount();
-				for (int32 PoseIndex = 0; PoseIndex < PoseCount; PoseIndex++)
-				{
-					FbxPose* CurrentPose = SDKScene->GetPose(PoseIndex);
-
-					// current pose is bind pose, 
-					if (CurrentPose && CurrentPose->IsBindPose())
-					{
-						// IsValidBindPose doesn't work reliably
-						// It checks all the parent chain(regardless root given), and if the parent doesn't have correct bind pose, it fails
-						// It causes more false positive issues than the real issue we have to worry about
-						// If you'd like to try this, set CHECK_VALID_BIND_POSE to 1, and try the error message
-						// when Autodesk fixes this bug, then we might be able to re-open this
-						FString PoseName = CurrentPose->GetName();
-						// all error report status
-						FbxStatus Status;
-
-						// it does not make any difference of checking with different node
-						// it is possible pose 0 -> node array 2, but isValidBindPose function returns true even with node array 0
-						for (auto Current : NodeArray)
-						{
-							FString CurrentName = Current->GetName();
-							NodeList pMissingAncestors, pMissingDeformers, pMissingDeformersAncestors, pWrongMatrices;
-
-							if (CurrentPose->IsValidBindPoseVerbose(Current, pMissingAncestors, pMissingDeformers, pMissingDeformersAncestors, pWrongMatrices, 0.0001, &Status))
-							{
-								PoseArray.Add(CurrentPose);
-								//UE_LOG(LogFbx, Log, TEXT("Valid bind pose for Pose (%s) - %s"), *PoseName, *FString(Current->GetName()));
-								break;
-							}
-							else
-							{
-								// first try to fix up
-								// add missing ancestors
-								for (int i = 0; i < pMissingAncestors.GetCount(); i++)
-								{
-									FbxAMatrix mat = pMissingAncestors.GetAt(i)->EvaluateGlobalTransform(FBXSDK_TIME_ZERO);
-									CurrentPose->Add(pMissingAncestors.GetAt(i), mat);
-								}
-
-								pMissingAncestors.Clear();
-								pMissingDeformers.Clear();
-								pMissingDeformersAncestors.Clear();
-								pWrongMatrices.Clear();
-
-								// check it again
-								if (CurrentPose->IsValidBindPose(Current))
-								{
-									PoseArray.Add(CurrentPose);
-									//UE_LOG(LogFbx, Log, TEXT("Valid bind pose for Pose (%s) - %s"), *PoseName, *FString(Current->GetName()));
-									break;
-								}
-								else
-								{
-									// first try to find parent who is null group and see if you can try test it again
-									FbxNode* ParentNode = Current->GetParent();
-									while (ParentNode)
-									{
-										FbxNodeAttribute* Attr = ParentNode->GetNodeAttribute();
-										if (Attr && Attr->GetAttributeType() == FbxNodeAttribute::eNull)
-										{
-											// found it 
-											break;
-										}
-
-										// find next parent
-										ParentNode = ParentNode->GetParent();
-									}
-
-									if (ParentNode && CurrentPose->IsValidBindPose(ParentNode))
-									{
-										PoseArray.Add(CurrentPose);
-										//UE_LOG(LogFbx, Log, TEXT("Valid bind pose for Pose (%s) - %s"), *PoseName, *FString(Current->GetName()));
-										break;
-									}
-									else
-									{
-										FString ErrorString = Status.GetErrorString();
-										// 											if (!GIsAutomationTesting)
-										// 												UE_LOG(LogFbx, Warning, TEXT("Not valid bind pose for Pose (%s) - Node %s : %s"), *PoseName, *FString(Current->GetName()), *ErrorString);
-									}
-								}
-							}
-						}
-					}
-				}
-
-				return (PoseArray.Size() > 0);
-			}
-
-			FbxNode* FFbxHelper::GetRootSkeleton(FbxScene * SDKScene, FbxNode * Link)
-			{
-				FbxNode* RootBone = Link;
-
-				// get Unreal skeleton root
-				// mesh and dummy are used as bone if they are in the skeleton hierarchy
-				while (RootBone && RootBone->GetParent())
-				{
-					bool bIsBlenderArmatureBone = false;
-					//TODO put back the creator code so we know if its blender
-					//if (FbxCreator == EFbxCreator::Blender)
-					{
-						//Hack to support armature dummy node from blender
-						//Users do not want the null attribute node named armature which is the parent of the real root bone in blender fbx file
-						//This is a hack since if a rigid mesh group root node is named "armature" it will be skip
-						const FString RootBoneParentName(RootBone->GetParent()->GetName());
-						FbxNode* GrandFather = RootBone->GetParent()->GetParent();
-						bIsBlenderArmatureBone = (GrandFather == nullptr || GrandFather == SDKScene->GetRootNode()) && (RootBoneParentName.Compare(TEXT("armature"), ESearchCase::IgnoreCase) == 0);
-					}
-
-					FbxNodeAttribute* Attr = RootBone->GetParent()->GetNodeAttribute();
-					if (Attr &&
-						(Attr->GetAttributeType() == FbxNodeAttribute::eMesh ||
-						(Attr->GetAttributeType() == FbxNodeAttribute::eNull && !bIsBlenderArmatureBone) ||
-						Attr->GetAttributeType() == FbxNodeAttribute::eSkeleton) &&
-						RootBone->GetParent() != SDKScene->GetRootNode())
-					{
-						// in some case, skeletal mesh can be ancestor of bones
-						// this avoids this situation
-						if (Attr->GetAttributeType() == FbxNodeAttribute::eMesh)
-						{
-							FbxMesh* Mesh = (FbxMesh*)Attr;
-							if (Mesh->GetDeformerCount(FbxDeformer::eSkin) > 0)
-							{
-								break;
-							}
-						}
-
-						RootBone = RootBone->GetParent();
-					}
-					else
-					{
-						break;
-					}
-				}
-
-				return RootBone;
-			}
-
-			void FFbxHelper::BuildSkeletonSystem(FbxScene * SDKScene, TArray<FbxCluster*> & ClusterArray, TArray<FbxNode*> & OutSortedLinks)
-			{
-				FbxNode* Link;
-				TArray<FbxNode*> RootLinks;
-				int32 ClusterIndex;
-				for (ClusterIndex = 0; ClusterIndex < ClusterArray.Num(); ClusterIndex++)
-				{
-					Link = ClusterArray[ClusterIndex]->GetLink();
-					if (Link)
-					{
-						Link = GetRootSkeleton(SDKScene, Link);
-						int32 LinkIndex;
-						for (LinkIndex = 0; LinkIndex < RootLinks.Num(); LinkIndex++)
-						{
-							if (Link == RootLinks[LinkIndex])
-							{
-								break;
-							}
-						}
-
-						// this link is a new root, add it
-						if (LinkIndex == RootLinks.Num())
-						{
-							RootLinks.Add(Link);
-						}
-					}
-				}
-
-				for (int32 LinkIndex = 0; LinkIndex < RootLinks.Num(); LinkIndex++)
-				{
-					RecursiveBuildSkeleton(RootLinks[LinkIndex], OutSortedLinks);
-				}
-			}
-
-			bool FFbxHelper::FindSkeletonJoints(FbxScene * SDKScene, TArray<FbxNode*> & NodeArray, TArray<FbxNode*> & SortedLinks, FbxArray<FbxAMatrix> & LocalsPerLink)
-			{
-				bool bIsARigidMesh = false;
-				FbxNode* Link = NULL;
-				FbxArray<FbxPose*> PoseArray;
-				TArray<FbxCluster*> ClusterArray;
-
-				if (NodeArray[0]->GetMesh()->GetDeformerCount(FbxDeformer::eSkin) == 0)
-				{
-					bIsARigidMesh = true;
-					Link = NodeArray[0];
-					RecursiveBuildSkeleton(GetRootSkeleton(SDKScene, Link), SortedLinks);
-				}
-				else
-				{
-					// get bindpose and clusters from FBX skeleton
-
-					// let's put the elements to their bind pose! (and we restore them after
-					// we have built the ClusterInformation.
-					int32 Default_NbPoses = SDKScene->GetFbxManager()->GetBindPoseCount(SDKScene);
-					// If there are no BindPoses, the following will generate them.
-					SDKScene->GetFbxManager()->CreateMissingBindPoses(SDKScene);
-
-					//if we created missing bind poses, update the number of bind poses
-					int32 NbPoses = SDKScene->GetFbxManager()->GetBindPoseCount(SDKScene);
-
-					if (NbPoses != Default_NbPoses)
-					{
-						//TODO output an error message to let the user know there is a missing bind pose when translating this fbx skeletalmesh
-						//Missing bind pose
-					}
-
-					//
-					// create the bones / skinning
-					//
-
-					for (int32 i = 0; i < NodeArray.Num(); i++)
-					{
-						FbxMesh* FbxMesh = NodeArray[i]->GetMesh();
-						const int32 SkinDeformerCount = FbxMesh->GetDeformerCount(FbxDeformer::eSkin);
-						for (int32 DeformerIndex = 0; DeformerIndex < SkinDeformerCount; DeformerIndex++)
-						{
-							FbxSkin* Skin = (FbxSkin*)FbxMesh->GetDeformer(DeformerIndex, FbxDeformer::eSkin);
-							for (int32 ClusterIndex = 0; ClusterIndex < Skin->GetClusterCount(); ClusterIndex++)
-							{
-								ClusterArray.Add(Skin->GetCluster(ClusterIndex));
-							}
-						}
-					}
-
-					if (ClusterArray.Num() == 0)
-					{
-						//TODO warn user we found a mesh with a deformer that is not bind to any skin cluster
-						return false;
-					}
-
-					// get bind pose
-					if (RetrievePoseFromBindPose(SDKScene, NodeArray, PoseArray) == false)
-					{
-						// 							if (!GIsAutomationTesting)
-						// 								UE_LOG(LogFbx, Warning, TEXT("Getting valid bind pose failed. Try to recreate bind pose"));
-													// if failed, delete bind pose, and retry.
-						const int32 PoseCount = SDKScene->GetPoseCount();
-						for (int32 PoseIndex = PoseCount - 1; PoseIndex >= 0; --PoseIndex)
-						{
-							FbxPose* CurrentPose = SDKScene->GetPose(PoseIndex);
-
-							// current pose is bind pose, 
-							if (CurrentPose && CurrentPose->IsBindPose())
-							{
-								SDKScene->RemovePose(PoseIndex);
-								CurrentPose->Destroy();
-							}
-						}
-
-						SDKScene->GetFbxManager()->CreateMissingBindPoses(SDKScene);
-						if (RetrievePoseFromBindPose(SDKScene, NodeArray, PoseArray) == false)
-						{
-							// 								if (!GIsAutomationTesting)
-							// 									UE_LOG(LogFbx, Warning, TEXT("Recreating bind pose failed."));
-						}
-						else
-						{
-							// 								if (!GIsAutomationTesting)
-							// 									UE_LOG(LogFbx, Warning, TEXT("Recreating bind pose succeeded."));
-						}
-					}
-
-					// recurse through skeleton and build ordered table
-					BuildSkeletonSystem(SDKScene, ClusterArray, SortedLinks);
-				}
-
-				// error check
-				// if no bond is found
-				if (SortedLinks.Num() == 0)
-				{
-					//						{
-					//							AddTokenizedErrorMessage(FTokenizedMessage::Create(EMessageSeverity::Warning, FText::Format(LOCTEXT("FbxSkeletaLMeshimport_NoBone", "'{0}' has no bones"), FText::FromString(NodeArray[0]->GetName()))), FFbxErrors::SkeletalMesh_NoBoneFound);
-					//						}
-					return false;
-				}
-
-				// if no bind pose is found
-				if (PoseArray.GetCount() == 0)
-				{
-					//TODO error to the user
-					//AddTokenizedErrorMessage(FTokenizedMessage::Create(EMessageSeverity::Warning, LOCTEXT("FbxSkeletaLMeshimport_MissingBindPose", "Could not find the bind pose.  It will use time 0 as bind pose.")), FFbxErrors::SkeletalMesh_InvalidBindPose);
-				}
-
-				int32 LinkIndex;
-
-				// Check for duplicate bone names and issue a warning if found
-				for (LinkIndex = 0; LinkIndex < SortedLinks.Num(); ++LinkIndex)
-				{
-					Link = SortedLinks[LinkIndex];
-
-					for (int32 AltLinkIndex = LinkIndex + 1; AltLinkIndex < SortedLinks.Num(); ++AltLinkIndex)
-					{
-						FbxNode* AltLink = SortedLinks[AltLinkIndex];
-
-						if (FCStringAnsi::Strcmp(Link->GetName(), AltLink->GetName()) == 0)
-						{
-							//								FString RawBoneName = UTF8_TO_TCHAR(Link->GetName());
-															//TODO notify the user that there is a duplicate bone name
-							// 									AddTokenizedErrorMessage(FTokenizedMessage::Create(EMessageSeverity::Error, FText::Format(LOCTEXT("Error_DuplicateBoneName", "Error, Could not import {0}.\nDuplicate bone name found ('{1}'). Each bone must have a unique name."),
-							// 															 FText::FromString(NodeArray[0]->GetName()), FText::FromString(RawBoneName))), FFbxErrors::SkeletalMesh_DuplicateBones);
-							return false;
-						}
-					}
-				}
-
-				FbxArray<FbxAMatrix> GlobalsPerLink;
-				GlobalsPerLink.Grow(SortedLinks.Num());
-				GlobalsPerLink[0].SetIdentity();
-				LocalsPerLink.Grow(SortedLinks.Num());
-				LocalsPerLink[0].SetIdentity();
-
-				bool GlobalLinkFoundFlag;
-
-				bool bAnyLinksNotInBindPose = false;
-				FString LinksWithoutBindPoses;
-
-				int32 RootIdx = INDEX_NONE;
-
-				for (LinkIndex = 0; LinkIndex < SortedLinks.Num(); LinkIndex++)
-				{
-					// Add a bone for each FBX Link
-					Link = SortedLinks[LinkIndex];
-					int32 ParentIndex = INDEX_NONE; // base value for root if no parent found
-
-					if (LinkIndex == 0)
-					{
-						RootIdx = LinkIndex;
-					}
-					else
-					{
-						const FbxNode* LinkParent = Link->GetParent();
-						// get the link parent index.
-						for (int32 ll = 0; ll < LinkIndex; ++ll) // <LinkIndex because parent is guaranteed to be before child in sortedLink
-						{
-							FbxNode* Otherlink = SortedLinks[ll];
-							if (Otherlink == LinkParent)
-							{
-								ParentIndex = ll;
-								break;
-							}
-						}
-
-						if (ParentIndex == INDEX_NONE)//We found another root inside the hierarchy, this is not supported
-						{
-							//AddTokenizedErrorMessage(FTokenizedMessage::Create(EMessageSeverity::Error, LOCTEXT("MultipleRootsFound", "Multiple roots are found in the bone hierarchy. We only support single root bone.")), FFbxErrors::SkeletalMesh_MultipleRoots);
-							return false;
-						}
-					}
-
-					GlobalLinkFoundFlag = false;
-					if (!bIsARigidMesh) //skeletal mesh
-					{
-						// there are some links, they have no cluster, but in bindpose
-						if (PoseArray.GetCount())
-						{
-							for (int32 PoseIndex = 0; PoseIndex < PoseArray.GetCount(); PoseIndex++)
-							{
-								int32 PoseLinkIndex = PoseArray[PoseIndex]->Find(Link);
-								if (PoseLinkIndex >= 0)
-								{
-									FbxMatrix NoneAffineMatrix = PoseArray[PoseIndex]->GetMatrix(PoseLinkIndex);
-									FbxAMatrix Matrix = *(FbxAMatrix*)(double*)&NoneAffineMatrix;
-									GlobalsPerLink[LinkIndex] = Matrix;
-									GlobalLinkFoundFlag = true;
-									break;
-								}
-							}
-						}
-
-						if (!GlobalLinkFoundFlag)
-						{
-							// since now we set use time 0 as ref pose this won't unlikely happen
-							// but leaving it just in case it still has case where it's missing partial bind pose
-							bAnyLinksNotInBindPose = true;
-							LinksWithoutBindPoses += UTF8_TO_TCHAR(Link->GetName());
-							LinksWithoutBindPoses += TEXT("  \n");
-
-							for (int32 ClusterIndex = 0; ClusterIndex < ClusterArray.Num(); ClusterIndex++)
-							{
-								FbxCluster* Cluster = ClusterArray[ClusterIndex];
-								if (Link == Cluster->GetLink())
-								{
-									Cluster->GetTransformLinkMatrix(GlobalsPerLink[LinkIndex]);
-									GlobalLinkFoundFlag = true;
-									break;
-								}
-							}
-						}
-					}
-
-					if (!GlobalLinkFoundFlag)
-					{
-						GlobalsPerLink[LinkIndex] = Link->EvaluateGlobalTransform();
-					}
-
-
-					if (LinkIndex)
-					{
-						LocalsPerLink[LinkIndex] = GlobalsPerLink[ParentIndex].Inverse() * GlobalsPerLink[LinkIndex];
-					}
-					else	// skeleton root
-					{
-						LocalsPerLink[LinkIndex] = GlobalsPerLink[LinkIndex];
-					}
-
-					//Check for nan and for zero scale
-					{
-						bool bFoundNan = false;
-						bool bFoundZeroScale = false;
-						for (int32 i = 0; i < 4; ++i)
-						{
-							if (i < 3)
-							{
-								if (FMath::IsNaN(LocalsPerLink[LinkIndex].GetT()[i]) || FMath::IsNaN(LocalsPerLink[LinkIndex].GetS()[i]))
-								{
-									bFoundNan = true;
-								}
-								if (FMath::IsNearlyZero(LocalsPerLink[LinkIndex].GetS()[i]))
-								{
-									bFoundZeroScale = true;
-								}
-							}
-							if (FMath::IsNaN(LocalsPerLink[LinkIndex].GetQ()[i]))
-							{
-								bFoundNan = true;
-							}
-						}
-
-						//TODO add a user error message
-						ensure(!bFoundNan);
-						ensure(!bFoundZeroScale);
-					}
-				}
-
-
-
-				if (bAnyLinksNotInBindPose)
-				{
-					//TODO error message
-					//
-					//AddTokenizedErrorMessage(FTokenizedMessage::Create(EMessageSeverity::Warning, FText::Format(LOCTEXT("FbxSkeletaLMeshimport_BonesAreMissingFromBindPose", "The following bones are missing from the bind pose:\n{0}\nThis can happen for bones that are not vert weighted. If they are not in the correct orientation after importing,\nplease set the \"Use T0 as ref pose\" option or add them to the bind pose and reimport the skeletal mesh."), FText::FromString(LinksWithoutBindPoses))), FFbxErrors::SkeletalMesh_BonesAreMissingFromBindPose);
-				}
-				return true;
-			}
-
-			FbxNode* FFbxHelper::RecursiveGetFirstMeshNode(FbxNode* Node, FbxNode* NodeToFind)
-			{
-				if (Node == nullptr)
-				{
-					return nullptr;
-				}
-				if (Node->GetMesh() != nullptr)
-					return Node;
-				for (int32 ChildIndex = 0; ChildIndex < Node->GetChildCount(); ++ChildIndex)
-				{
-					FbxNode* MeshNode = RecursiveGetFirstMeshNode(Node->GetChild(ChildIndex), NodeToFind);
-					if (NodeToFind == nullptr)
-					{
-						if (MeshNode != nullptr)
-						{
-							return MeshNode;
-						}
-					}
-					else if (MeshNode == NodeToFind)
-					{
-						return MeshNode;
-					}
-				}
-				return nullptr;
-			}
-
-			FbxNode* FFbxHelper::FindLODGroupNode(FbxNode* NodeLodGroup, int32 LodIndex, FbxNode* NodeToFind)
-			{
-				check(NodeLodGroup->GetChildCount() >= LodIndex);
-				FbxNode* ChildNode = NodeLodGroup->GetChild(LodIndex);
-				if (ChildNode == nullptr)
-				{
-					return nullptr;
-				}
-				return RecursiveGetFirstMeshNode(ChildNode, NodeToFind);
-			}
-
-			void FFbxHelper::RecursiveFindFbxSkelMesh(FbxScene * SDKScene, FbxNode * Node, TArray< TArray<FbxNode*> > & outSkelMeshArray, TArray<FbxNode*> & SkeletonArray)
-			{
-				FbxNode* SkelMeshNode = nullptr;
-				FbxNode* NodeToAdd = Node;
-
-				if (Node->GetMesh() && Node->GetMesh()->GetDeformerCount(FbxDeformer::eSkin) > 0)
-				{
-					SkelMeshNode = Node;
-				}
-				else if (Node->GetNodeAttribute() && Node->GetNodeAttribute()->GetAttributeType() == FbxNodeAttribute::eLODGroup)
-				{
-					// for LODgroup, add the LODgroup to OutSkelMeshArray according to the skeleton that the first child bind to
-					SkelMeshNode = FFbxHelper::FindLODGroupNode(Node, 0, nullptr);
-					// check if the first child is skeletal mesh
-					if (SkelMeshNode != nullptr && !(SkelMeshNode->GetMesh() && SkelMeshNode->GetMesh()->GetDeformerCount(FbxDeformer::eSkin) > 0))
-					{
-						SkelMeshNode = nullptr;
-					}
-					// else NodeToAdd = Node;
-				}
-
-				if (SkelMeshNode)
-				{
-					// find root skeleton
-
-					check(SkelMeshNode->GetMesh() != nullptr);
-					const int32 fbxDeformerCount = SkelMeshNode->GetMesh()->GetDeformerCount();
-					FbxSkin* Deformer = static_cast<FbxSkin*>(SkelMeshNode->GetMesh()->GetDeformer(0, FbxDeformer::eSkin));
-
-					if (Deformer != NULL)
-					{
-						int32 ClusterCount = Deformer->GetClusterCount();
-						bool bFoundCorrectLink = false;
-						for (int32 ClusterId = 0; ClusterId < ClusterCount; ++ClusterId)
-						{
-							FbxNode* RootBoneLink = Deformer->GetCluster(ClusterId)->GetLink(); //Get the bone influences by this first cluster
-							RootBoneLink = GetRootSkeleton(SDKScene, RootBoneLink); // Get the skeleton root itself
-
-							if (RootBoneLink)
-							{
-								bool bAddedToExistingSkeleton = false;
-								for (int32 SkeletonIndex = 0; SkeletonIndex < SkeletonArray.Num(); ++SkeletonIndex)
-								{
-									if (RootBoneLink == SkeletonArray[SkeletonIndex])
-									{
-										// append to existed outSkelMeshArray element
-										TArray<FbxNode*>& TempArray = outSkelMeshArray[SkeletonIndex];
-										TempArray.Add(NodeToAdd);
-										bAddedToExistingSkeleton = true;
-										break;
-									}
-								}
-
-								// if there is no outSkelMeshArray element that is bind to this skeleton
-								// create new element for outSkelMeshArray
-								if (!bAddedToExistingSkeleton)
-								{
-									TArray<FbxNode*>& TempArray = outSkelMeshArray.AddDefaulted_GetRef();
-									TempArray.Add(NodeToAdd);
-									SkeletonArray.Add(RootBoneLink);
-								}
-
-								bFoundCorrectLink = true;
-								break;
-							}
-						}
-
-						// we didn't find the correct link
-						if (!bFoundCorrectLink)
-						{
-							// 								AddTokenizedErrorMessage(
-							// 									FTokenizedMessage::Create(
-							// 									EMessageSeverity::Warning,
-							// 									FText::Format(LOCTEXT("FBX_NoWeightsOnDeformer", "Ignoring mesh {0} because it but no weights."), FText::FromString(UTF8_TO_TCHAR(SkelMeshNode->GetName())))
-							// 								),
-							// 									FFbxErrors::SkeletalMesh_NoWeightsOnDeformer
-							// 								);
-						}
-					}
-				}
-
-				//Skeletalmesh node can have child so let's always iterate trough child
-				{
-					TArray<FbxNode*> ChildScaled;
-					//Sort the node to have the one with no scaling first so we have more chance
-					//to have a root skeletal mesh with no scale. Because scene import do not support
-					//root skeletal mesh containing scale
-					for (int32 ChildIndex = 0; ChildIndex < Node->GetChildCount(); ++ChildIndex)
-					{
-						FbxNode* ChildNode = Node->GetChild(ChildIndex);
-
-						if (!Node->GetNodeAttribute() || Node->GetNodeAttribute()->GetAttributeType() != FbxNodeAttribute::eLODGroup)
-						{
-							FbxVector4 NoScale(1.0, 1.0, 1.0);
-
-							if (ChildNode->EvaluateLocalScaling() == NoScale)
-							{
-								RecursiveFindFbxSkelMesh(SDKScene, ChildNode, outSkelMeshArray, SkeletonArray);
-							}
-							else
-							{
-								ChildScaled.Add(ChildNode);
-							}
-						}
-					}
-
-					for (FbxNode* ChildNode : ChildScaled)
-					{
-						RecursiveFindFbxSkelMesh(SDKScene, ChildNode, outSkelMeshArray, SkeletonArray);
-					}
-				}
-			}
-
-			void FFbxHelper::RecursiveFindRigidMesh(FbxScene * SDKScene, FbxNode * Node, TArray< TArray<FbxNode*> > & outSkelMeshArray, TArray<FbxNode*> & SkeletonArray)
-			{
-				bool bRigidNodeFound = false;
-				FbxNode* RigidMeshNode = nullptr;
-
-				if (Node->GetMesh())
-				{
-					// ignore skeletal mesh
-					if (Node->GetMesh()->GetDeformerCount(FbxDeformer::eSkin) == 0)
-					{
-						RigidMeshNode = Node;
-						bRigidNodeFound = true;
-					}
-				}
-				else if (Node->GetNodeAttribute() && Node->GetNodeAttribute()->GetAttributeType() == FbxNodeAttribute::eLODGroup)
-				{
-					// for LODgroup, add the LODgroup to OutSkelMeshArray according to the skeleton that the first child bind to
-					FbxNode* FirstLOD = FindLODGroupNode(Node, 0, nullptr);
-					// check if the first child is skeletal mesh
-					if (FirstLOD != nullptr && FirstLOD->GetMesh())
-					{
-						if (FirstLOD->GetMesh()->GetDeformerCount(FbxDeformer::eSkin) == 0)
-						{
-							bRigidNodeFound = true;
-						}
-					}
-
-					if (bRigidNodeFound)
-					{
-						RigidMeshNode = Node;
-					}
-				}
-
-				if (bRigidNodeFound)
-				{
-					// find root skeleton
-					FbxNode* Link = GetRootSkeleton(SDKScene, RigidMeshNode);
-
-					int32 i;
-					for (i = 0; i < SkeletonArray.Num(); i++)
-					{
-						if (Link == SkeletonArray[i])
-						{
-							// append to existed outSkelMeshArray element
-							TArray<FbxNode*>& TempArray = outSkelMeshArray[i];
-							TempArray.Add(RigidMeshNode);
-							break;
-						}
-					}
-
-					// if there is no outSkelMeshArray element that is bind to this skeleton
-					// create new element for outSkelMeshArray
-					if (i == SkeletonArray.Num())
-					{
-						TArray<FbxNode*>& TempArray = outSkelMeshArray.AddDefaulted_GetRef();
-						TempArray.Add(RigidMeshNode);
-						SkeletonArray.Add(Link);
-					}
-				}
-
-				// for LODGroup, we will not deep in.
-				if (!(Node->GetNodeAttribute() && Node->GetNodeAttribute()->GetAttributeType() == FbxNodeAttribute::eLODGroup))
-				{
-					int32 ChildIndex;
-					for (ChildIndex = 0; ChildIndex < Node->GetChildCount(); ++ChildIndex)
-					{
-						RecursiveFindRigidMesh(SDKScene, Node->GetChild(ChildIndex), outSkelMeshArray, SkeletonArray);
-					}
-				}
-			}
-
-			void FFbxHelper::RecursiveFixSkeleton(FbxScene * SDKScene, FbxNode * Node, TArray<FbxNode*> & SkelMeshes, bool bImportNestedMeshes)
-			{
-				FbxNodeAttribute* Attr = Node->GetNodeAttribute();
-				bool NodeIsLodGroup = (Attr && (Attr->GetAttributeType() == FbxNodeAttribute::eLODGroup));
-				if (!NodeIsLodGroup)
-				{
-					for (int32 i = 0; i < Node->GetChildCount(); i++)
-					{
-						RecursiveFixSkeleton(SDKScene, Node->GetChild(i), SkelMeshes, bImportNestedMeshes);
-					}
-				}
-
-				if (Attr && (Attr->GetAttributeType() == FbxNodeAttribute::eMesh || Attr->GetAttributeType() == FbxNodeAttribute::eNull))
-				{
-					if (bImportNestedMeshes && Attr->GetAttributeType() == FbxNodeAttribute::eMesh)
-					{
-						// for leaf mesh, keep them as mesh
-						int32 ChildCount = Node->GetChildCount();
-						int32 ChildIndex;
-						for (ChildIndex = 0; ChildIndex < ChildCount; ChildIndex++)
-						{
-							FbxNode* Child = Node->GetChild(ChildIndex);
-							if (Child->GetMesh() == NULL)
-							{
-								break;
-							}
-						}
-
-						if (ChildIndex != ChildCount)
-						{
-							// Remove from the mesh list it is no longer a mesh
-							SkelMeshes.Remove(Node);
-
-							//replace with skeleton
-							FbxSkeleton* lSkeleton = FbxSkeleton::Create(SDKScene->GetFbxManager(), "");
-							Node->SetNodeAttribute(lSkeleton);
-							lSkeleton->SetSkeletonType(FbxSkeleton::eLimbNode);
-						}
-						else // this mesh may be not in skeleton mesh list. If not, add it.
-						{
-							if (!SkelMeshes.Contains(Node))
-							{
-								SkelMeshes.Add(Node);
-							}
-						}
-					}
-					else
-					{
-						// Remove from the mesh list it is no longer a mesh
-						SkelMeshes.Remove(Node);
-
-						//replace with skeleton
-						FbxSkeleton* lSkeleton = FbxSkeleton::Create(SDKScene->GetFbxManager(), "");
-						Node->SetNodeAttribute(lSkeleton);
-						lSkeleton->SetSkeletonType(FbxSkeleton::eLimbNode);
-					}
-				}
-			}
-#endif
 		}//ns Private
 	}//ns Interchange
 }//ns UE
