Add to interchange generic pipeline the option force all mesh type.

diff --git a/Engine/Plugins/Experimental/Interchange/Runtime/Source/Import/Private/Fbx/InterchangeFbxTranslator.cpp b/Engine/Plugins/Experimental/Interchange/Runtime/Source/Import/Private/Fbx/InterchangeFbxTranslator.cpp
--- a/Engine/Plugins/Experimental/Interchange/Runtime/Source/Import/Private/Fbx/InterchangeFbxTranslator.cpp
+++ b/Engine/Plugins/Experimental/Interchange/Runtime/Source/Import/Private/Fbx/InterchangeFbxTranslator.cpp
@@ -220,16 +220,15 @@ TFuture<TOptional<UE::Interchange::FStaticMeshPayloadData>> UInterchangeFbxTrans
 		FLargeMemoryReader Ar(FileData, FileDataSize);
 		StaticMeshPayload.MeshDescription.Serialize(Ar);
 
-		// This is a static mesh and the payload should never contain skinned data
+		// This is a static mesh payload can contain skinned data if we need to convert skeletalmesh to staticmesh
 		bool bFetchSkinnedData = false;
 		Ar << bFetchSkinnedData;
 		if (bFetchSkinnedData)
 		{
-			// TODO log an error saying the payload file is the wrong type for a static mesh
-			Promise->SetValue(TOptional<UE::Interchange::FStaticMeshPayloadData>());
-			return;
+			TArray<FString> JointNames;
+			//Read into a dummy structure, we will not use the influence channel of this skinned meshdescription
+			Ar << JointNames;
 		}
-
 		Promise->SetValue(MoveTemp(StaticMeshPayload));
 	}));
 
diff --git a/Engine/Plugins/Experimental/Interchange/Runtime/Source/Import/Private/Mesh/InterchangeSkeletalMeshFactory.cpp b/Engine/Plugins/Experimental/Interchange/Runtime/Source/Import/Private/Mesh/InterchangeSkeletalMeshFactory.cpp
--- a/Engine/Plugins/Experimental/Interchange/Runtime/Source/Import/Private/Mesh/InterchangeSkeletalMeshFactory.cpp
+++ b/Engine/Plugins/Experimental/Interchange/Runtime/Source/Import/Private/Mesh/InterchangeSkeletalMeshFactory.cpp
@@ -58,7 +58,7 @@ namespace UE
 			void RecursiveAddBones(const UInterchangeBaseNodeContainer* NodeContainer, const FString& JointNodeId, TArray <FJointInfo>& JointInfos, int32 ParentIndex, TArray<SkeletalMeshImportData::FBone>& RefBonesBinary)
 			{
 				const UInterchangeSceneNode* JointNode = Cast<UInterchangeSceneNode>(NodeContainer->GetNode(JointNodeId));
-				if (!JointNode || !JointNode->IsSpecializedTypeContains(FSceneNodeStaticData::GetJointSpecializeTypeString()))
+				if (!JointNode)
 				{
 					UE_LOG(LogInterchangeImport, Warning, TEXT("Invalid Skeleton Joint"));
 					return;
@@ -386,11 +386,12 @@ namespace UE
 					SkeletalMeshAppendSettings.SourceVertexIDOffset = VertexOffset;
 					FElementIDRemappings ElementIDRemappings;
 					LodMeshPayload->LodMeshDescription.Compact(ElementIDRemappings);
+					const int32 RefBoneCount = RefBonesBinary.Num();
 					//Remap the influence vertex index to point on the correct index
 					if (LodMeshPayload->JointNames.Num() > 0)
 					{
 						const int32 LocalJointCount = LodMeshPayload->JointNames.Num();
-						const int32 RefBoneCount = RefBonesBinary.Num();
+						
 						SkeletalMeshAppendSettings.SourceRemapBoneIndex.AddZeroed(LocalJointCount);
 						for (int32 LocalJointIndex = 0; LocalJointIndex < LocalJointCount; ++LocalJointIndex)
 						{
@@ -418,7 +419,10 @@ namespace UE
 						AppendSettings.MeshTransform.Reset();
 					}
 					FStaticMeshOperations::AppendMeshDescription(LodMeshPayload->LodMeshDescription, LodMeshDescription, AppendSettings);
-					FSkeletalMeshOperations::AppendSkinWeight(LodMeshPayload->LodMeshDescription, LodMeshDescription, SkeletalMeshAppendSettings);
+					if (MeshNodeContext.MeshNode->IsSkinnedMesh())
+					{
+						FSkeletalMeshOperations::AppendSkinWeight(LodMeshPayload->LodMeshDescription, LodMeshDescription, SkeletalMeshAppendSettings);
+					}
 					if (bImportMorphTarget)
 					{
 						FillBlendShapeMeshDescriptionsPerBlendShapeName(MeshNodeContext
diff --git a/Engine/Plugins/Experimental/Interchange/Runtime/Source/Pipelines/Private/InterchangeGenericAssetsPipeline.cpp b/Engine/Plugins/Experimental/Interchange/Runtime/Source/Pipelines/Private/InterchangeGenericAssetsPipeline.cpp
--- a/Engine/Plugins/Experimental/Interchange/Runtime/Source/Pipelines/Private/InterchangeGenericAssetsPipeline.cpp
+++ b/Engine/Plugins/Experimental/Interchange/Runtime/Source/Pipelines/Private/InterchangeGenericAssetsPipeline.cpp
@@ -90,8 +90,9 @@ bool UInterchangeGenericAssetsPipeline::ExecutePreImportPipeline(UInterchangeBas
 	//Create skeletalmesh factory nodes
 	ExecutePreImportPipelineSkeletalMesh();
 
-	if (bImportStaticMeshes)
+	if (bImportStaticMeshes && (ForceAllMeshHasType == EInterchangeForceMeshType::IFMT_None || ForceAllMeshHasType == EInterchangeForceMeshType::IFMT_StaticMesh))
 	{
+		const bool bConvertSkeletalMeshToStaticMesh = (ForceAllMeshHasType == EInterchangeForceMeshType::IFMT_StaticMesh);
 		if (bCombineStaticMeshes)
 		{
 			// Combine all the static meshes
@@ -101,7 +102,7 @@ bool UInterchangeGenericAssetsPipeline::ExecutePreImportPipeline(UInterchangeBas
 			{
 				// If baking transforms, get all the static mesh instance nodes, and group them by LOD
 				TArray<FString> MeshUids;
-				PipelineMeshesUtilities->GetAllStaticMeshInstance(MeshUids);
+				PipelineMeshesUtilities->GetAllStaticMeshInstance(MeshUids, bConvertSkeletalMeshToStaticMesh);
 
 				TMap<int32, TArray<FString>> MeshUidsPerLodIndex;
 
@@ -134,7 +135,7 @@ bool UInterchangeGenericAssetsPipeline::ExecutePreImportPipeline(UInterchangeBas
 			{
 				// If we haven't yet managed to build a factory node, look at static mesh geometry directly.
 				TArray<FString> MeshUids;
-				PipelineMeshesUtilities->GetAllStaticMeshGeometry(MeshUids);
+				PipelineMeshesUtilities->GetAllStaticMeshGeometry(MeshUids, bConvertSkeletalMeshToStaticMesh);
 
 				TMap<int32, TArray<FString>> MeshUidsPerLodIndex;
 
@@ -163,7 +164,7 @@ bool UInterchangeGenericAssetsPipeline::ExecutePreImportPipeline(UInterchangeBas
 			if (bBakeMeshes)
 			{
 				TArray<FString> MeshUids;
-				PipelineMeshesUtilities->GetAllStaticMeshInstance(MeshUids);
+				PipelineMeshesUtilities->GetAllStaticMeshInstance(MeshUids, bConvertSkeletalMeshToStaticMesh);
 				
 				for (const FString& MeshUid : MeshUids)
 				{
@@ -195,7 +196,7 @@ bool UInterchangeGenericAssetsPipeline::ExecutePreImportPipeline(UInterchangeBas
 			if (!bFoundMeshes)
 			{
 				TArray<FString> MeshUids;
-				PipelineMeshesUtilities->GetAllStaticMeshGeometry(MeshUids);
+				PipelineMeshesUtilities->GetAllStaticMeshGeometry(MeshUids, bConvertSkeletalMeshToStaticMesh);
 
 				for (const FString& MeshUid : MeshUids)
 				{
diff --git a/Engine/Plugins/Experimental/Interchange/Runtime/Source/Pipelines/Private/InterchangeGenericSkeletalMeshPipeline.cpp b/Engine/Plugins/Experimental/Interchange/Runtime/Source/Pipelines/Private/InterchangeGenericSkeletalMeshPipeline.cpp
--- a/Engine/Plugins/Experimental/Interchange/Runtime/Source/Pipelines/Private/InterchangeGenericSkeletalMeshPipeline.cpp
+++ b/Engine/Plugins/Experimental/Interchange/Runtime/Source/Pipelines/Private/InterchangeGenericSkeletalMeshPipeline.cpp
@@ -211,146 +211,51 @@ namespace UE::Interchange::SkeletalMeshGenericPipeline
 
 bool UInterchangeGenericAssetsPipeline::ExecutePreImportPipelineSkeletalMesh()
 {
-	
-	//PipelineMeshesUtilities;
-	TArray<UInterchangeMeshNode*> SkinnedMeshNodes;
-	TArray<UInterchangeMeshNode*> StaticMeshNodes;
-	TMap<FString, TArray<FString>> SkeletalMeshFactoryDependencyOrderPerSkeletonRootNodeUid;
-	
-	//Find all translated node we need for this pipeline
-	BaseNodeContainer->IterateNodes([this, &SkinnedMeshNodes, &StaticMeshNodes](const FString& NodeUid, UInterchangeBaseNode* Node)
+	if (!bImportSkeletalMeshes)
 	{
-		switch(Node->GetNodeContainerType())
-		{
-			case EInterchangeNodeContainerType::NodeContainerType_TranslatedAsset:
-			{
-				if (UInterchangeMeshNode* MeshNode = Cast<UInterchangeMeshNode>(Node))
-				{
-					if (MeshNode->IsSkinnedMesh())
-					{
-						SkinnedMeshNodes.Add(MeshNode);
-					}
-					else
-					{
-						StaticMeshNodes.Add(MeshNode);
-					}
-
-				}
-			}
-			break;
-		}
-	});
+		//Nothing to import
+		return true;
+	}
 
-	if (bImportSkeletalMeshes && SkinnedMeshNodes.Num() > 0)
+	if (ForceAllMeshHasType != EInterchangeForceMeshType::IFMT_None && ForceAllMeshHasType != EInterchangeForceMeshType::IFMT_SkeletalMesh)
 	{
-		auto SetSkeletalMeshDependencies = [&SkeletalMeshFactoryDependencyOrderPerSkeletonRootNodeUid](const FString& JointNodeUid, UInterchangeSkeletalMeshFactoryNode* SkeletalMeshFactoryNode)
-		{
-			TArray<FString>& SkeletalMeshFactoryDependencyOrder = SkeletalMeshFactoryDependencyOrderPerSkeletonRootNodeUid.FindOrAdd(JointNodeUid);
-			//Updating the skeleton is not multi thread safe, so we add dependency between skeletalmesh altering the same skeleton
-			//TODO make the skeletalMesh ReferenceSkeleton thread safe to allow multiple parallel skeletalmesh factory on the same skeleton asset.
-			int32 DependencyIndex = SkeletalMeshFactoryDependencyOrder.AddUnique(SkeletalMeshFactoryNode->GetUniqueID());
-			if (DependencyIndex > 0)
-			{
-				const FString SkeletalMeshFactoryNodeDependencyUid = SkeletalMeshFactoryDependencyOrder[DependencyIndex - 1];
-				SkeletalMeshFactoryNode->AddFactoryDependencyUid(SkeletalMeshFactoryNodeDependencyUid);
-			}
-		};
+		//Nothing to import
+		return true;
+	}
+	const bool bConvertStaticMeshToSkeletalMesh = (ForceAllMeshHasType == EInterchangeForceMeshType::IFMT_SkeletalMesh);
+	TMap<FString, TArray<FString>> SkeletalMeshFactoryDependencyOrderPerSkeletonRootNodeUid;
 
-		if (bCombineSkeletalMeshes)
+	auto SetSkeletalMeshDependencies = [&SkeletalMeshFactoryDependencyOrderPerSkeletonRootNodeUid](const FString& JointNodeUid, UInterchangeSkeletalMeshFactoryNode* SkeletalMeshFactoryNode)
+	{
+		TArray<FString>& SkeletalMeshFactoryDependencyOrder = SkeletalMeshFactoryDependencyOrderPerSkeletonRootNodeUid.FindOrAdd(JointNodeUid);
+		//Updating the skeleton is not multi thread safe, so we add dependency between skeletalmesh altering the same skeleton
+		//TODO make the skeletalMesh ReferenceSkeleton thread safe to allow multiple parallel skeletalmesh factory on the same skeleton asset.
+		int32 DependencyIndex = SkeletalMeshFactoryDependencyOrder.AddUnique(SkeletalMeshFactoryNode->GetUniqueID());
+		if (DependencyIndex > 0)
 		{
-			//////////////////////////////////////////////////////////////////////////
-			//Combined everything we can
-			TMap<FString, TArray<FString>> MeshUidsPerSkeletonRootUid;
-			auto CreatePerSkeletonRootUidCombinedSkinnedMesh = [this, &MeshUidsPerSkeletonRootUid, &SetSkeletalMeshDependencies](const bool bUseInstanceMesh)
-			{
-				bool bFoundInstances = false;
-				for (const TPair<FString, TArray<FString>>& SkeletonRootUidAndMeshUids : MeshUidsPerSkeletonRootUid)
-				{
-					const FString& SkeletonRootUid = SkeletonRootUidAndMeshUids.Key;
-					//Every iteration is a skeletalmesh asset that combine all MeshInstances sharing the same skeleton root node
-					UInterchangeSkeletonFactoryNode* SkeletonFactoryNode = CreateSkeletonFactoryNode(SkeletonRootUid);
-					//The MeshUids can represent a SceneNode pointing on a MeshNode or directly a MeshNode;
-					TMap<int32, TArray<FString>> MeshUidsPerLodIndex;
-					const TArray<FString>& MeshUids = SkeletonRootUidAndMeshUids.Value;
-					for (const FString& MeshUid : MeshUids)
-					{
-						if (bUseInstanceMesh)
-						{
-							const FInterchangeMeshInstance& MeshInstance = PipelineMeshesUtilities->GetMeshInstanceByUid(MeshUid);
-							for (const TPair<int32, FInterchangeLodSceneNodeContainer>& LodIndexAndSceneNodeContainer : MeshInstance.SceneNodePerLodIndex)
-							{
-								const int32 LodIndex = LodIndexAndSceneNodeContainer.Key;
-								const FInterchangeLodSceneNodeContainer& SceneNodeContainer = LodIndexAndSceneNodeContainer.Value;
-								TArray<FString>& TranslatedNodes = MeshUidsPerLodIndex.FindOrAdd(LodIndex);
-								for (const UInterchangeSceneNode* SceneNode : SceneNodeContainer.SceneNodes)
-								{
-									TranslatedNodes.Add(SceneNode->GetUniqueID());
-								}
-							}
-						}
-						else
-						{
-							//MeshGeometry cannot have Lod since LODs are define in the scene node
-							const FInterchangeMeshGeometry& MeshGeometry = PipelineMeshesUtilities->GetMeshGeometryByUid(MeshUid);
-							const int32 LodIndex = 0;
-							TArray<FString>& TranslatedNodes = MeshUidsPerLodIndex.FindOrAdd(LodIndex);
-							TranslatedNodes.Add(MeshGeometry.MeshUid);
-						}
-					}
-
-					if (MeshUidsPerLodIndex.Num() > 0)
-					{
-						UInterchangeSkeletalMeshFactoryNode* SkeletalMeshFactoryNode = CreateSkeletalMeshFactoryNode(SkeletonRootUid, MeshUidsPerLodIndex);
-						SetSkeletalMeshDependencies(SkeletonRootUid, SkeletalMeshFactoryNode);
-						SkeletonFactoryNodes.Add(SkeletonFactoryNode);
-						SkeletalMeshFactoryNodes.Add(SkeletalMeshFactoryNode);
-						bFoundInstances = true;
-					}
-				}
-				return bFoundInstances;
-			};
-
-			bool bFoundMeshes = false;
-			if (bBakeMeshes)
-			{
-				PipelineMeshesUtilities->GetCombinedSkinnedMeshInstances(MeshUidsPerSkeletonRootUid);
-				const bool bUseMeshInstance = true;
-				bFoundMeshes = CreatePerSkeletonRootUidCombinedSkinnedMesh(bUseMeshInstance);
-			}
-
-			if (!bFoundMeshes)
-			{
-				MeshUidsPerSkeletonRootUid.Empty();
-				PipelineMeshesUtilities->GetCombinedSkinnedMeshGeometries(MeshUidsPerSkeletonRootUid);
-				const bool bUseMeshInstance = false;
-				CreatePerSkeletonRootUidCombinedSkinnedMesh(bUseMeshInstance);
-			}
+			const FString SkeletalMeshFactoryNodeDependencyUid = SkeletalMeshFactoryDependencyOrder[DependencyIndex - 1];
+			SkeletalMeshFactoryNode->AddFactoryDependencyUid(SkeletalMeshFactoryNodeDependencyUid);
 		}
-		else
+	};
+
+	if (bCombineSkeletalMeshes)
+	{
+		//////////////////////////////////////////////////////////////////////////
+		//Combined everything we can
+		TMap<FString, TArray<FString>> MeshUidsPerSkeletonRootUid;
+		auto CreatePerSkeletonRootUidCombinedSkinnedMesh = [this, &MeshUidsPerSkeletonRootUid, &SetSkeletalMeshDependencies](const bool bUseInstanceMesh)
 		{
-			//////////////////////////////////////////////////////////////////////////
-			//Do not combined meshes
-			TArray<FString> MeshUids;
-			auto CreatePerSkeletonRootUidSkinnedMesh = [this, &MeshUids, &SetSkeletalMeshDependencies](const bool bUseInstanceMesh)
+			bool bFoundInstances = false;
+			for (const TPair<FString, TArray<FString>>& SkeletonRootUidAndMeshUids : MeshUidsPerSkeletonRootUid)
 			{
-				bool bFoundInstances = false;
+				const FString& SkeletonRootUid = SkeletonRootUidAndMeshUids.Key;
+				//Every iteration is a skeletalmesh asset that combine all MeshInstances sharing the same skeleton root node
+				UInterchangeSkeletonFactoryNode* SkeletonFactoryNode = CreateSkeletonFactoryNode(SkeletonRootUid);
+				//The MeshUids can represent a SceneNode pointing on a MeshNode or directly a MeshNode;
+				TMap<int32, TArray<FString>> MeshUidsPerLodIndex;
+				const TArray<FString>& MeshUids = SkeletonRootUidAndMeshUids.Value;
 				for (const FString& MeshUid : MeshUids)
 				{
-					//Every iteration is a skeletalmesh asset that combine all MeshInstances sharing the same skeleton root node
-					//The MeshUids can represent a SceneNode pointing on a MeshNode or directly a MeshNode;
-					TMap<int32, TArray<FString>> MeshUidsPerLodIndex;
-					FString SkeletonRootUid;
-					if (!(bUseInstanceMesh ? PipelineMeshesUtilities->IsValidMeshInstanceUid(MeshUid) : PipelineMeshesUtilities->IsValidMeshGeometryUid(MeshUid)))
-					{
-						continue;
-					}
-					SkeletonRootUid = (bUseInstanceMesh ? PipelineMeshesUtilities->GetMeshInstanceSkeletonRootUid(MeshUid) : PipelineMeshesUtilities->GetMeshGeometrySkeletonRootUid(MeshUid));
-					if (SkeletonRootUid.IsEmpty())
-					{
-						//Log an error
-						continue;
-					}
-					UInterchangeSkeletonFactoryNode* SkeletonFactoryNode = CreateSkeletonFactoryNode(SkeletonRootUid);
 					if (bUseInstanceMesh)
 					{
 						const FInterchangeMeshInstance& MeshInstance = PipelineMeshesUtilities->GetMeshInstanceByUid(MeshUid);
@@ -367,58 +272,116 @@ bool UInterchangeGenericAssetsPipeline::ExecutePreImportPipelineSkeletalMesh()
 					}
 					else
 					{
+						//MeshGeometry cannot have Lod since LODs are define in the scene node
 						const FInterchangeMeshGeometry& MeshGeometry = PipelineMeshesUtilities->GetMeshGeometryByUid(MeshUid);
 						const int32 LodIndex = 0;
 						TArray<FString>& TranslatedNodes = MeshUidsPerLodIndex.FindOrAdd(LodIndex);
 						TranslatedNodes.Add(MeshGeometry.MeshUid);
 					}
-					if (MeshUidsPerLodIndex.Num() > 0)
-					{
-						UInterchangeSkeletalMeshFactoryNode* SkeletalMeshFactoryNode = CreateSkeletalMeshFactoryNode(SkeletonRootUid, MeshUidsPerLodIndex);
-						SetSkeletalMeshDependencies(SkeletonRootUid, SkeletalMeshFactoryNode);
-						SkeletonFactoryNodes.Add(SkeletonFactoryNode);
-						SkeletalMeshFactoryNodes.Add(SkeletalMeshFactoryNode);
-						bFoundInstances = true;
-					}
 				}
-				return bFoundInstances;
-			};
-			
-			bool bFoundMeshes = false;
-			if (bBakeMeshes)
-			{
-				PipelineMeshesUtilities->GetAllSkinnedMeshInstance(MeshUids);
-				const bool bUseMeshInstance = true;
-				bFoundMeshes = CreatePerSkeletonRootUidSkinnedMesh(bUseMeshInstance);
+
+				if (MeshUidsPerLodIndex.Num() > 0)
+				{
+					UInterchangeSkeletalMeshFactoryNode* SkeletalMeshFactoryNode = CreateSkeletalMeshFactoryNode(SkeletonRootUid, MeshUidsPerLodIndex);
+					SetSkeletalMeshDependencies(SkeletonRootUid, SkeletalMeshFactoryNode);
+					SkeletonFactoryNodes.Add(SkeletonFactoryNode);
+					SkeletalMeshFactoryNodes.Add(SkeletalMeshFactoryNode);
+					bFoundInstances = true;
+				}
 			}
+			return bFoundInstances;
+		};
 
-			if (!bFoundMeshes)
+		bool bFoundMeshes = false;
+		if (bBakeMeshes)
+		{
+			PipelineMeshesUtilities->GetCombinedSkinnedMeshInstances(MeshUidsPerSkeletonRootUid, bConvertStaticMeshToSkeletalMesh);
+			const bool bUseMeshInstance = true;
+			bFoundMeshes = CreatePerSkeletonRootUidCombinedSkinnedMesh(bUseMeshInstance);
+		}
+
+		if (!bFoundMeshes)
+		{
+			MeshUidsPerSkeletonRootUid.Empty();
+			PipelineMeshesUtilities->GetCombinedSkinnedMeshGeometries(MeshUidsPerSkeletonRootUid);
+			const bool bUseMeshInstance = false;
+			CreatePerSkeletonRootUidCombinedSkinnedMesh(bUseMeshInstance);
+		}
+	}
+	else
+	{
+		//////////////////////////////////////////////////////////////////////////
+		//Do not combined meshes
+		TArray<FString> MeshUids;
+		auto CreatePerSkeletonRootUidSkinnedMesh = [this, &MeshUids, &SetSkeletalMeshDependencies](const bool bUseInstanceMesh)
+		{
+			bool bFoundInstances = false;
+			for (const FString& MeshUid : MeshUids)
 			{
-				MeshUids.Empty();
-				PipelineMeshesUtilities->GetAllSkinnedMeshGeometry(MeshUids);
-				const bool bUseMeshInstance = false;
-				CreatePerSkeletonRootUidSkinnedMesh(bUseMeshInstance);
+				//Every iteration is a skeletalmesh asset that combine all MeshInstances sharing the same skeleton root node
+				//The MeshUids can represent a SceneNode pointing on a MeshNode or directly a MeshNode;
+				TMap<int32, TArray<FString>> MeshUidsPerLodIndex;
+				FString SkeletonRootUid;
+				if (!(bUseInstanceMesh ? PipelineMeshesUtilities->IsValidMeshInstanceUid(MeshUid) : PipelineMeshesUtilities->IsValidMeshGeometryUid(MeshUid)))
+				{
+					continue;
+				}
+				SkeletonRootUid = (bUseInstanceMesh ? PipelineMeshesUtilities->GetMeshInstanceSkeletonRootUid(MeshUid) : PipelineMeshesUtilities->GetMeshGeometrySkeletonRootUid(MeshUid));
+				if (SkeletonRootUid.IsEmpty())
+				{
+					//Log an error
+					continue;
+				}
+				UInterchangeSkeletonFactoryNode* SkeletonFactoryNode = CreateSkeletonFactoryNode(SkeletonRootUid);
+				if (bUseInstanceMesh)
+				{
+					const FInterchangeMeshInstance& MeshInstance = PipelineMeshesUtilities->GetMeshInstanceByUid(MeshUid);
+					for (const TPair<int32, FInterchangeLodSceneNodeContainer>& LodIndexAndSceneNodeContainer : MeshInstance.SceneNodePerLodIndex)
+					{
+						const int32 LodIndex = LodIndexAndSceneNodeContainer.Key;
+						const FInterchangeLodSceneNodeContainer& SceneNodeContainer = LodIndexAndSceneNodeContainer.Value;
+						TArray<FString>& TranslatedNodes = MeshUidsPerLodIndex.FindOrAdd(LodIndex);
+						for (const UInterchangeSceneNode* SceneNode : SceneNodeContainer.SceneNodes)
+						{
+							TranslatedNodes.Add(SceneNode->GetUniqueID());
+						}
+					}
+				}
+				else
+				{
+					const FInterchangeMeshGeometry& MeshGeometry = PipelineMeshesUtilities->GetMeshGeometryByUid(MeshUid);
+					const int32 LodIndex = 0;
+					TArray<FString>& TranslatedNodes = MeshUidsPerLodIndex.FindOrAdd(LodIndex);
+					TranslatedNodes.Add(MeshGeometry.MeshUid);
+				}
+				if (MeshUidsPerLodIndex.Num() > 0)
+				{
+					UInterchangeSkeletalMeshFactoryNode* SkeletalMeshFactoryNode = CreateSkeletalMeshFactoryNode(SkeletonRootUid, MeshUidsPerLodIndex);
+					SetSkeletalMeshDependencies(SkeletonRootUid, SkeletalMeshFactoryNode);
+					SkeletonFactoryNodes.Add(SkeletonFactoryNode);
+					SkeletalMeshFactoryNodes.Add(SkeletalMeshFactoryNode);
+					bFoundInstances = true;
+				}
 			}
+			return bFoundInstances;
+		};
+			
+		bool bFoundMeshes = false;
+		if (bBakeMeshes)
+		{
+			PipelineMeshesUtilities->GetAllSkinnedMeshInstance(MeshUids, bConvertStaticMeshToSkeletalMesh);
+			const bool bUseMeshInstance = true;
+			bFoundMeshes = CreatePerSkeletonRootUidSkinnedMesh(bUseMeshInstance);
 		}
-	}
-
-
-	const UClass* SkeletalMeshFactoryNodeClass = UInterchangeSkeletalMeshFactoryNode::StaticClass();
-	TArray<FString> SkeletalMeshNodeUids;
-	BaseNodeContainer->GetNodes(SkeletalMeshFactoryNodeClass, SkeletalMeshNodeUids);
 
-	//If we import only one asset, and bUseSourceNameForAsset is true, we want to rename the asset using the file name.
-	const int32 MeshesAndAnimsImportedNodeCount = SkeletalMeshNodeUids.Num(); // + StaticMeshNodeUids.Num();
-	if (bUseSourceNameForAsset && MeshesAndAnimsImportedNodeCount == 1)
-	{
-		if (SkeletalMeshNodeUids.Num() > 0)
+		if (!bFoundMeshes)
 		{
-			UInterchangeSkeletalMeshFactoryNode* SkeletalMeshNode = Cast<UInterchangeSkeletalMeshFactoryNode>(BaseNodeContainer->GetNode(SkeletalMeshNodeUids[0]));
-			const FString DisplayLabelName = FPaths::GetBaseFilename(SourceDatas[0]->GetFilename());
-			SkeletalMeshNode->SetDisplayLabel(DisplayLabelName);
+			MeshUids.Empty();
+			PipelineMeshesUtilities->GetAllSkinnedMeshGeometry(MeshUids);
+			const bool bUseMeshInstance = false;
+			CreatePerSkeletonRootUidSkinnedMesh(bUseMeshInstance);
 		}
 	}
-
 	return true;
 }
 
diff --git a/Engine/Plugins/Experimental/Interchange/Runtime/Source/Pipelines/Private/InterchangePipelineMeshesUtilities.cpp b/Engine/Plugins/Experimental/Interchange/Runtime/Source/Pipelines/Private/InterchangePipelineMeshesUtilities.cpp
--- a/Engine/Plugins/Experimental/Interchange/Runtime/Source/Pipelines/Private/InterchangePipelineMeshesUtilities.cpp
+++ b/Engine/Plugins/Experimental/Interchange/Runtime/Source/Pipelines/Private/InterchangePipelineMeshesUtilities.cpp
@@ -170,50 +170,50 @@ void UInterchangePipelineMeshesUtilities::IterateAllMeshInstance(TFunctionRef<vo
 	}
 }
 
-void UInterchangePipelineMeshesUtilities::GetAllSkinnedMeshInstance(TArray<FString>& MeshInstanceUids) const
+void UInterchangePipelineMeshesUtilities::GetAllSkinnedMeshInstance(TArray<FString>& MeshInstanceUids, const bool bConvertStaticMeshToSkeletalMesh) const
 {
 	MeshInstanceUids.Empty(MeshInstancesPerMeshInstanceUid.Num());
 	for (const TPair<FString, FInterchangeMeshInstance>& MeshInstanceUidAndMeshInstance : MeshInstancesPerMeshInstanceUid)
 	{
 		const FInterchangeMeshInstance& MeshInstance = MeshInstanceUidAndMeshInstance.Value;
-		if(MeshInstance.bReferenceSkinnedMesh)
+		if((bConvertStaticMeshToSkeletalMesh || MeshInstance.bReferenceSkinnedMesh) && !MeshInstance.bReferenceBlendShape)
 		{
 			MeshInstanceUids.Add(MeshInstance.MeshInstanceUid);
 		}
 	}
 }
 
-void UInterchangePipelineMeshesUtilities::IterateAllSkinnedMeshInstance(TFunctionRef<void(const FInterchangeMeshInstance&)> IterationLambda) const
+void UInterchangePipelineMeshesUtilities::IterateAllSkinnedMeshInstance(TFunctionRef<void(const FInterchangeMeshInstance&)> IterationLambda, const bool bConvertStaticMeshToSkeletalMesh) const
 {
 	for (const TPair<FString, FInterchangeMeshInstance>& MeshInstanceUidAndMeshInstance : MeshInstancesPerMeshInstanceUid)
 	{
 		const FInterchangeMeshInstance& MeshInstance = MeshInstanceUidAndMeshInstance.Value;
-		if (MeshInstance.bReferenceSkinnedMesh)
+		if ((bConvertStaticMeshToSkeletalMesh || MeshInstance.bReferenceSkinnedMesh) && !MeshInstance.bReferenceBlendShape)
 		{
 			IterationLambda(MeshInstance);
 		}
 	}
 }
 
-void UInterchangePipelineMeshesUtilities::GetAllStaticMeshInstance(TArray<FString>& MeshInstanceUids) const
+void UInterchangePipelineMeshesUtilities::GetAllStaticMeshInstance(TArray<FString>& MeshInstanceUids, const bool bConvertSkeletalMeshToStaticMesh) const
 {
 	MeshInstanceUids.Empty(MeshInstancesPerMeshInstanceUid.Num());
 	for (const TPair<FString, FInterchangeMeshInstance>& MeshInstanceUidAndMeshInstance : MeshInstancesPerMeshInstanceUid)
 	{
 		const FInterchangeMeshInstance& MeshInstance = MeshInstanceUidAndMeshInstance.Value;
-		if (!MeshInstance.bReferenceSkinnedMesh && !MeshInstance.bReferenceBlendShape)
+		if ((bConvertSkeletalMeshToStaticMesh || !MeshInstance.bReferenceSkinnedMesh) && !MeshInstance.bReferenceBlendShape)
 		{
 			MeshInstanceUids.Add(MeshInstance.MeshInstanceUid);
 		}
 	}
 }
 
-void UInterchangePipelineMeshesUtilities::IterateAllStaticMeshInstance(TFunctionRef<void(const FInterchangeMeshInstance&)> IterationLambda) const
+void UInterchangePipelineMeshesUtilities::IterateAllStaticMeshInstance(TFunctionRef<void(const FInterchangeMeshInstance&)> IterationLambda, const bool bConvertSkeletalMeshToStaticMesh) const
 {
 	for (const TPair<FString, FInterchangeMeshInstance>& MeshInstanceUidAndMeshInstance : MeshInstancesPerMeshInstanceUid)
 	{
 		const FInterchangeMeshInstance& MeshInstance = MeshInstanceUidAndMeshInstance.Value;
-		if (!MeshInstance.bReferenceSkinnedMesh && !MeshInstance.bReferenceBlendShape)
+		if ((bConvertSkeletalMeshToStaticMesh || !MeshInstance.bReferenceSkinnedMesh) && !MeshInstance.bReferenceBlendShape)
 		{
 			IterationLambda(MeshInstance);
 		}
@@ -259,25 +259,25 @@ void UInterchangePipelineMeshesUtilities::IterateAllSkinnedMeshGeometry(TFunctio
 	}
 }
 
-void UInterchangePipelineMeshesUtilities::GetAllStaticMeshGeometry(TArray<FString>& MeshGeometryUids) const
+void UInterchangePipelineMeshesUtilities::GetAllStaticMeshGeometry(TArray<FString>& MeshGeometryUids, const bool bConvertSkeletalMeshToStaticMesh) const
 {
 	MeshGeometryUids.Empty(MeshGeometriesPerMeshUid.Num());
 	for (const TPair<FString, FInterchangeMeshGeometry>& MeshGeometryUidAndMeshGeometry : MeshGeometriesPerMeshUid)
 	{
 		const FInterchangeMeshGeometry& MeshGeometry = MeshGeometryUidAndMeshGeometry.Value;
-		if (!MeshGeometry.MeshNode->IsSkinnedMesh() && !MeshGeometry.MeshNode->IsBlendShape())
+		if ((bConvertSkeletalMeshToStaticMesh || !MeshGeometry.MeshNode->IsSkinnedMesh()) && !MeshGeometry.MeshNode->IsBlendShape())
 		{
 			MeshGeometryUids.Add(MeshGeometry.MeshUid);
 		}
 	}
 }
 
-void UInterchangePipelineMeshesUtilities::IterateAllStaticMeshGeometry(TFunctionRef<void(const FInterchangeMeshGeometry&)> IterationLambda) const
+void UInterchangePipelineMeshesUtilities::IterateAllStaticMeshGeometry(TFunctionRef<void(const FInterchangeMeshGeometry&)> IterationLambda, const bool bConvertSkeletalMeshToStaticMesh) const
 {
 	for (const TPair<FString, FInterchangeMeshGeometry>& MeshGeometryUidAndMeshGeometry : MeshGeometriesPerMeshUid)
 	{
 		const FInterchangeMeshGeometry& MeshGeometry = MeshGeometryUidAndMeshGeometry.Value;
-		if (!MeshGeometry.MeshNode->IsSkinnedMesh() && !MeshGeometry.MeshNode->IsBlendShape())
+		if ((bConvertSkeletalMeshToStaticMesh || !MeshGeometry.MeshNode->IsSkinnedMesh()) && !MeshGeometry.MeshNode->IsBlendShape())
 		{
 			IterationLambda(MeshGeometry);
 		}
@@ -345,15 +345,16 @@ void UInterchangePipelineMeshesUtilities::IterateAllMeshInstanceUsingMeshGeometr
 	}
 }
 
-void UInterchangePipelineMeshesUtilities::GetCombinedSkinnedMeshInstances(TMap<FString, TArray<FString>>& OutMeshInstanceUidsPerSkeletonRootUid) const
+void UInterchangePipelineMeshesUtilities::GetCombinedSkinnedMeshInstances(TMap<FString, TArray<FString>>& OutMeshInstanceUidsPerSkeletonRootUid, const bool bConvertStaticMeshToSkeletalMesh) const
 {
 	for (const TPair<FString, FInterchangeMeshInstance>& MeshInstanceUidAndMeshInstance : MeshInstancesPerMeshInstanceUid)
 	{
 		const FInterchangeMeshInstance& MeshInstance = MeshInstanceUidAndMeshInstance.Value;
-		if (!MeshInstance.bReferenceSkinnedMesh)
+		if (!bConvertStaticMeshToSkeletalMesh && !MeshInstance.bReferenceSkinnedMesh)
 		{
 			continue;
 		}
+		bool bIsStaticMesh = bConvertStaticMeshToSkeletalMesh && !MeshInstance.bReferenceSkinnedMesh;
 		//Find the root skeleton for this MeshInstance
 		FString SkeletonRootUid;
 		for (const FString& MeshGeometryUid : MeshInstance.ReferencingMeshGeometryUids)
@@ -371,6 +372,11 @@ void UInterchangePipelineMeshesUtilities::GetCombinedSkinnedMeshInstances(TMap<F
 					break;
 				}
 			}
+			else if (bIsStaticMesh)
+			{
+				//Create a joint from the instance node (the scene node pointing on the mesh).
+				SkeletonRootUid = MeshInstance.MeshInstanceUid;
+			}
 			else
 			{
 				//every skinned geometry should have a skeleton root node ???
diff --git a/Engine/Plugins/Experimental/Interchange/Runtime/Source/Pipelines/Public/InterchangeGenericAssetsPipeline.h b/Engine/Plugins/Experimental/Interchange/Runtime/Source/Pipelines/Public/InterchangeGenericAssetsPipeline.h
--- a/Engine/Plugins/Experimental/Interchange/Runtime/Source/Pipelines/Public/InterchangeGenericAssetsPipeline.h
+++ b/Engine/Plugins/Experimental/Interchange/Runtime/Source/Pipelines/Public/InterchangeGenericAssetsPipeline.h
@@ -93,8 +93,8 @@ public:
 
 
 	/** Allow to convert mesh to a particular type */
-// 	UPROPERTY(EditAnywhere, BlueprintReadWrite, Category = COMMON_MESHES_CATEGORY)
-// 	TEnumAsByte<enum EForceMeshType> ForceAllMeshHasType = EForceMeshType::FMT_None;
+ 	UPROPERTY(EditAnywhere, BlueprintReadWrite, Category = COMMON_MESHES_CATEGORY)
+ 	TEnumAsByte<enum EInterchangeForceMeshType> ForceAllMeshHasType = EInterchangeForceMeshType::IFMT_None;
 
 	/** If enable, meshes LODs will be imported. Note that it required the advanced bBakeMesh property to be enabled. */
 	UPROPERTY(EditAnywhere, BlueprintReadWrite, Category = COMMON_MESHES_CATEGORY, meta = (editcondition = "bBakeMeshes"))
@@ -106,7 +106,7 @@ public:
 
 	/** Specify how vertex colors should be imported */
 	UPROPERTY(EditAnywhere, BlueprintReadWrite, Category = COMMON_MESHES_CATEGORY)
-	TEnumAsByte<enum EInterchangeVertexColorImportOption> VertexColorImportOption;
+	TEnumAsByte<enum EInterchangeVertexColorImportOption> VertexColorImportOption = EInterchangeVertexColorImportOption::IVCIO_Replace;
 
 	/** Specify override color in the case that VertexColorImportOption is set to Override */
 	UPROPERTY(EditAnywhere, BlueprintReadWrite, Category = COMMON_MESHES_CATEGORY)
diff --git a/Engine/Plugins/Experimental/Interchange/Runtime/Source/Pipelines/Public/InterchangePipelineMeshesUtilities.h b/Engine/Plugins/Experimental/Interchange/Runtime/Source/Pipelines/Public/InterchangePipelineMeshesUtilities.h
--- a/Engine/Plugins/Experimental/Interchange/Runtime/Source/Pipelines/Public/InterchangePipelineMeshesUtilities.h
+++ b/Engine/Plugins/Experimental/Interchange/Runtime/Source/Pipelines/Public/InterchangePipelineMeshesUtilities.h
@@ -134,23 +134,23 @@ public:
 	* Get all skinned mesh instance unique ids.
 	*/
 	UFUNCTION(BlueprintCallable, Category = "Interchange | Pipeline | Meshes")
-	void GetAllSkinnedMeshInstance(TArray<FString>& MeshInstanceUids) const;
+	void GetAllSkinnedMeshInstance(TArray<FString>& MeshInstanceUids, const bool bConvertStaticMeshToSkeletalMesh) const;
 
 	/**
 	* Iterate all skinned mesh instance.
 	*/
-	void IterateAllSkinnedMeshInstance(TFunctionRef<void(const FInterchangeMeshInstance&)> IterationLambda) const;
+	void IterateAllSkinnedMeshInstance(TFunctionRef<void(const FInterchangeMeshInstance&)> IterationLambda, const bool bConvertStaticMeshToSkeletalMesh) const;
 
 	/**
 	* Get all static mesh instance unique ids.
 	*/
 	UFUNCTION(BlueprintCallable, Category = "Interchange | Pipeline | Meshes")
-	void GetAllStaticMeshInstance(TArray<FString>& MeshInstanceUids) const;
+	void GetAllStaticMeshInstance(TArray<FString>& MeshInstanceUids, const bool bConvertSkeletalMeshToStaticMesh) const;
 
 	/**
 	* Iterate all static mesh instance.
 	*/
-	void IterateAllStaticMeshInstance(TFunctionRef<void(const FInterchangeMeshInstance&)> IterationLambda) const;
+	void IterateAllStaticMeshInstance(TFunctionRef<void(const FInterchangeMeshInstance&)> IterationLambda, const bool bConvertSkeletalMeshToStaticMesh) const;
 
 	/**
 	* Get all mesh geometry unique ids.
@@ -178,12 +178,12 @@ public:
 	* Get all static mesh geometry unique ids.
 	*/
 	UFUNCTION(BlueprintCallable, Category = "Interchange | Pipeline | Meshes")
-	void GetAllStaticMeshGeometry(TArray<FString>& MeshGeometryUids) const;
+	void GetAllStaticMeshGeometry(TArray<FString>& MeshGeometryUids, const bool bConvertSkeletalMeshToStaticMesh) const;
 
 	/**
 	* Iterate all static mesh geometry.
 	*/
-	void IterateAllStaticMeshGeometry(TFunctionRef<void(const FInterchangeMeshGeometry&)> IterationLambda) const;
+	void IterateAllStaticMeshGeometry(TFunctionRef<void(const FInterchangeMeshGeometry&)> IterationLambda, const bool bConvertSkeletalMeshToStaticMesh) const;
 
 	/**
 	* Get all not instanced mesh geometry unique ids.
@@ -235,7 +235,7 @@ public:
 	* Return a list of skinned FInterchangeMeshInstance uid that can be combined together.
 	* We cannot create a skinned mesh with multiple skeleton root node, This function return combined MeshInstance per skeleton roots
 	*/
-	void GetCombinedSkinnedMeshInstances(TMap<FString, TArray<FString>>& OutMeshInstanceUidsPerSkeletonRootUid) const;
+	void GetCombinedSkinnedMeshInstances(TMap<FString, TArray<FString>>& OutMeshInstanceUidsPerSkeletonRootUid, const bool bConvertStaticMeshToSkeletalMesh) const;
 	
 	/**
 	* Return a list of skinned FInterchangeMeshGeometry uid that can be combined together.
