Add to the generic pipeline the option UpdateSkeletonReferencePose, which if enable will update the bone transform of the referenced skeleton in the generic option with the imported bone transform.

#rnx

diff --git a/Engine/Plugins/Experimental/Interchange/Runtime/Source/Pipelines/Private/InterchangeGenericAssetsPipeline.cpp b/Engine/Plugins/Experimental/Interchange/Runtime/Source/Pipelines/Private/InterchangeGenericAssetsPipeline.cpp
--- a/Engine/Plugins/Experimental/Interchange/Runtime/Source/Pipelines/Private/InterchangeGenericAssetsPipeline.cpp
+++ b/Engine/Plugins/Experimental/Interchange/Runtime/Source/Pipelines/Private/InterchangeGenericAssetsPipeline.cpp
@@ -237,6 +237,8 @@ bool UInterchangeGenericAssetsPipeline::ExecutePostImportPipeline(const UInterch
 		return false;
 	}
 
+	PostImportSkeletalMesh(CreatedAsset, Node);
+
 	//Finish the physics asset import, it need the skeletal mesh render data to create the physics collision geometry
 	PostImportPhysicsAssetImport(CreatedAsset, Node);
 
diff --git a/Engine/Plugins/Experimental/Interchange/Runtime/Source/Pipelines/Private/InterchangeGenericSkeletalMeshPipeline.cpp b/Engine/Plugins/Experimental/Interchange/Runtime/Source/Pipelines/Private/InterchangeGenericSkeletalMeshPipeline.cpp
--- a/Engine/Plugins/Experimental/Interchange/Runtime/Source/Pipelines/Private/InterchangeGenericSkeletalMeshPipeline.cpp
+++ b/Engine/Plugins/Experimental/Interchange/Runtime/Source/Pipelines/Private/InterchangeGenericSkeletalMeshPipeline.cpp
@@ -678,6 +678,26 @@ void UInterchangeGenericAssetsPipeline::AddLodDataToSkeletalMesh(const UIntercha
 	}
 }
 
+void UInterchangeGenericAssetsPipeline::PostImportSkeletalMesh(UObject* CreatedAsset, UInterchangeBaseNode* Node)
+{
+	if (!BaseNodeContainer)
+	{
+		return;
+	}
+
+	USkeletalMesh* SkeletalMesh = Cast<USkeletalMesh>(CreatedAsset);
+	if (!SkeletalMesh)
+	{
+		return;
+	}
+
+	if (bUpdateSkeletonReferencePose && !Skeleton.IsNull() && SkeletalMesh->GetSkeleton() == Skeleton.Get())
+	{
+		SkeletalMesh->GetSkeleton()->UpdateReferencePoseFromMesh(SkeletalMesh);
+		//TODO: notify editor the skeleton has change
+	}
+}
+
 void UInterchangeGenericAssetsPipeline::PostImportPhysicsAssetImport(UObject* CreatedAsset, UInterchangeBaseNode* Node)
 {
 #if WITH_EDITOR
diff --git a/Engine/Plugins/Experimental/Interchange/Runtime/Source/Pipelines/Public/InterchangeGenericAssetsPipeline.h b/Engine/Plugins/Experimental/Interchange/Runtime/Source/Pipelines/Public/InterchangeGenericAssetsPipeline.h
--- a/Engine/Plugins/Experimental/Interchange/Runtime/Source/Pipelines/Public/InterchangeGenericAssetsPipeline.h
+++ b/Engine/Plugins/Experimental/Interchange/Runtime/Source/Pipelines/Public/InterchangeGenericAssetsPipeline.h
@@ -147,9 +147,13 @@ public:
 	UPROPERTY(EditAnywhere, BlueprintReadWrite, Category = SKELETAL_MESHES_CATEGORY)
 	bool bImportMorphTargets = true;
 
+	/** Enable this option to update Skeleton (of the mesh)'s reference pose. Mesh's reference pose is always updated.  */
+	UPROPERTY(EditAnywhere, BlueprintReadWrite, Category = SKELETAL_MESHES_CATEGORY)
+	bool bUpdateSkeletonReferencePose = false;
+
 	/** If checked, create new PhysicsAsset if it doesn't have it */
 	UPROPERTY(EditAnywhere, BlueprintReadWrite, Category = SKELETAL_MESHES_CATEGORY)
-	uint32 bCreatePhysicsAsset : 1;
+	bool bCreatePhysicsAsset;
 
 	/** If this is set, use this specified PhysicsAsset. If its not set and bCreatePhysicsAsset is false, the importer will not generate or set any physic asset. */
 	UPROPERTY(EditAnywhere, BlueprintReadWrite, Category = SKELETAL_MESHES_CATEGORY, meta = (editcondition = "!bCreatePhysicsAsset"))
@@ -267,6 +271,11 @@ private:
 	 */
 	bool ExecutePreImportPipelineSkeletalMesh();
 
+	/**
+	 * This function will finish creating the skeletalmesh asset
+	 */
+	void PostImportSkeletalMesh(UObject* CreatedAsset, UInterchangeBaseNode* Node);
+
 	/**
 	 * This function will finish creating the physics asset with the skeletalmesh render data
 	 */
