Add the "Bind Pose" transform to the joint node so we can import correctly the skeletalmesh internal skeleton(refskeleton of the skeletalmesh).
Remove some ifdef out code

#rnx

diff --git a/Engine/Plugins/Experimental/Interchange/Runtime/Source/Parsers/Fbx/Private/FbxConvert.cpp b/Engine/Plugins/Experimental/Interchange/Runtime/Source/Parsers/Fbx/Private/FbxConvert.cpp
--- a/Engine/Plugins/Experimental/Interchange/Runtime/Source/Parsers/Fbx/Private/FbxConvert.cpp
+++ b/Engine/Plugins/Experimental/Interchange/Runtime/Source/Parsers/Fbx/Private/FbxConvert.cpp
@@ -169,37 +169,25 @@ namespace UE
 				return FVector(Vector[0], -Vector[1], Vector[2]);
 			}
 
-			ANSICHAR* FFbxConvert::MakeName(const ANSICHAR* Name)
+			FString FFbxConvert::MakeName(const ANSICHAR* Name)
 			{
-				const int SpecialChars[] = { '.', ',', '/', '`', '%' };
+				const TCHAR SpecialChars[] = { TEXT('.'), TEXT(','), TEXT('/'), TEXT('`'), TEXT('%') };
 
-				const int len = FCStringAnsi::Strlen(Name);
-				ANSICHAR* TmpName = new ANSICHAR[len + 1];
-
-				FCStringAnsi::Strcpy(TmpName, len + 1, Name);
-
-				for (int32 i = 0; i < UE_ARRAY_COUNT(SpecialChars); i++)
-				{
-					ANSICHAR* CharPtr = TmpName;
-					while ((CharPtr = FCStringAnsi::Strchr(CharPtr, SpecialChars[i])) != NULL)
-					{
-						CharPtr[0] = '_';
-					}
-				}
+				FString TmpName = MakeString(Name);
 
 				// Remove namespaces
-				ANSICHAR* NewName;
-				NewName = FCStringAnsi::Strchr(TmpName, ':');
-
-				// there may be multiple namespace, so find the last ':'
-				while (NewName && FCStringAnsi::Strchr(NewName + 1, ':'))
+				int32 LastNamespaceTokenIndex = INDEX_NONE;
+				if (TmpName.FindLastChar(TEXT(':'), LastNamespaceTokenIndex))
 				{
-					NewName = FCStringAnsi::Strchr(NewName + 1, ':');
+					const bool bAllowShrinking = true;
+					//+1 to remove the ':' character we found
+					TmpName.RightChopInline(LastNamespaceTokenIndex + 1, bAllowShrinking);
 				}
 
-				if (NewName)
+				//Remove the special chars
+				for (int32 i = 0; i < UE_ARRAY_COUNT(SpecialChars); i++)
 				{
-					return NewName + 1;
+					TmpName.ReplaceCharInline(SpecialChars[i], TEXT('_'), ESearchCase::CaseSensitive);
 				}
 
 				return TmpName;
diff --git a/Engine/Plugins/Experimental/Interchange/Runtime/Source/Parsers/Fbx/Private/FbxConvert.h b/Engine/Plugins/Experimental/Interchange/Runtime/Source/Parsers/Fbx/Private/FbxConvert.h
--- a/Engine/Plugins/Experimental/Interchange/Runtime/Source/Parsers/Fbx/Private/FbxConvert.h
+++ b/Engine/Plugins/Experimental/Interchange/Runtime/Source/Parsers/Fbx/Private/FbxConvert.h
@@ -56,7 +56,7 @@ namespace UE
 				 * Replace all special characters with '_', then remove all namespace
 				 * Special characters are . , / ` %
 				 */
-				static ANSICHAR* MakeName(const ANSICHAR* Name);
+				static FString MakeName(const ANSICHAR* Name);
 
 				/**
 				 * Convert ANSI char to a FString using ANSI_TO_TCHAR macro
diff --git a/Engine/Plugins/Experimental/Interchange/Runtime/Source/Parsers/Fbx/Private/FbxHelper.cpp b/Engine/Plugins/Experimental/Interchange/Runtime/Source/Parsers/Fbx/Private/FbxHelper.cpp
--- a/Engine/Plugins/Experimental/Interchange/Runtime/Source/Parsers/Fbx/Private/FbxHelper.cpp
+++ b/Engine/Plugins/Experimental/Interchange/Runtime/Source/Parsers/Fbx/Private/FbxHelper.cpp
@@ -100,7 +100,7 @@ namespace UE
 				{
 					return FString();
 				}
-				FString ObjName = UTF8_TO_TCHAR(FFbxConvert::MakeName(Object->GetName()));
+				FString ObjName = FFbxConvert::MakeName(Object->GetName());
 				if (ObjName.Equals(TEXT("none"), ESearchCase::IgnoreCase))
 				{
 					//Replace None by Null because None clash with NAME_None and the create asset will instead call the object ClassName_X
diff --git a/Engine/Plugins/Experimental/Interchange/Runtime/Source/Parsers/Fbx/Private/FbxHelper.h b/Engine/Plugins/Experimental/Interchange/Runtime/Source/Parsers/Fbx/Private/FbxHelper.h
--- a/Engine/Plugins/Experimental/Interchange/Runtime/Source/Parsers/Fbx/Private/FbxHelper.h
+++ b/Engine/Plugins/Experimental/Interchange/Runtime/Source/Parsers/Fbx/Private/FbxHelper.h
@@ -45,33 +45,6 @@ namespace UE
 
 			protected:
 				static FString GetUniqueIDString(const uint64 UniqueID);
-
-#if 0
-				// This stuff is moved to FbxSkelMesh.h
-				// 
-				//Skeletalmesh helper
-				static void FindSkeletalMeshes(FbxScene* SDKScene, TArray< TArray<FbxNode*> >& outSkelMeshArray, bool bCombineSkeletalMesh, bool bForceFindRigid);
-				static void FindAllLODGroupNode(TArray<FbxNode*>& OutNodeInLod, FbxNode* NodeLodGroup, int32 LodIndex);
-				static bool FindSkeletonJoints(FbxScene* SDKScene, TArray<FbxNode*>& NodeArray, TArray<FbxNode*>& SortedLinks, FbxArray<FbxAMatrix>& LocalsPerLink);
-				static FbxNode* FindLODGroupNode(FbxNode* NodeLodGroup, int32 LodIndex, FbxNode* NodeToFind);
-
-			private:
-
-				/* Skletalmesh private helper begin */
-
-				static void RecursiveGetAllMeshNode(TArray<FbxNode*>& OutAllNode, FbxNode* Node);
-				static bool IsUnrealBone(FbxNode* Link);
-				static void RecursiveBuildSkeleton(FbxNode* Link, TArray<FbxNode*>& OutSortedLinks);
-				static bool RetrievePoseFromBindPose(FbxScene* SDKScene, const TArray<FbxNode*>& NodeArray, FbxArray<FbxPose*>& PoseArray);
-				static FbxNode* GetRootSkeleton(FbxScene* SDKScene, FbxNode* Link);
-				static void BuildSkeletonSystem(FbxScene* SDKScene, TArray<FbxCluster*>& ClusterArray, TArray<FbxNode*>& OutSortedLinks);
-				static FbxNode* RecursiveGetFirstMeshNode(FbxNode* Node, FbxNode* NodeToFind);
-				static void RecursiveFindFbxSkelMesh(FbxScene* SDKScene, FbxNode* Node, TArray< TArray<FbxNode*> >& outSkelMeshArray, TArray<FbxNode*>& SkeletonArray);
-				static void RecursiveFindRigidMesh(FbxScene* SDKScene, FbxNode* Node, TArray< TArray<FbxNode*> >& outSkelMeshArray, TArray<FbxNode*>& SkeletonArray);
-				static void RecursiveFixSkeleton(FbxScene* SDKScene, FbxNode* Node, TArray<FbxNode*>& SkelMeshes, bool bImportNestedMeshes);
-
-				/* Skletalmesh private helper end */
-#endif
 			};
 		}//ns Private
 	}//ns Interchange
diff --git a/Engine/Plugins/Experimental/Interchange/Runtime/Source/Parsers/Fbx/Private/FbxMesh.cpp b/Engine/Plugins/Experimental/Interchange/Runtime/Source/Parsers/Fbx/Private/FbxMesh.cpp
--- a/Engine/Plugins/Experimental/Interchange/Runtime/Source/Parsers/Fbx/Private/FbxMesh.cpp
+++ b/Engine/Plugins/Experimental/Interchange/Runtime/Source/Parsers/Fbx/Private/FbxMesh.cpp
@@ -1232,105 +1232,6 @@ namespace UE
 
 			//////////////////////////////////////////////////////////////////////////
 			/// FFbxMesh implementation
-#if 0
-			FString FFbxMesh::GetUniqueIDString(const uint64 UniqueID)
-			{
-				FStringFormatNamedArguments FormatArguments;
-				FormatArguments.Add(TEXT("UniqueID"), UniqueID);
-				return FString::Format(TEXT("{UniqueID}"), FormatArguments);
-			}
-
-			FString FFbxMesh::GetMeshName(FbxGeometryBase* Mesh)
-			{
-				FString MeshName = FFbxHelper::GetFbxObjectName(Mesh);
-				if (MeshName.IsEmpty())
-				{
-					if (Mesh->GetNodeCount() > 0)
-					{
-						if (Mesh->GetAttributeType() == FbxNodeAttribute::eMesh)
-						{
-							MeshName = TEXT("Mesh_");
-						}
-						else if (Mesh->GetAttributeType() == FbxNodeAttribute::eShape)
-						{
-							MeshName = TEXT("Shape_");
-						}
-						MeshName += FFbxHelper::GetFbxObjectName(Mesh->GetNode(0));
-					}
-					else
-					{
-						uint64 UniqueFbxObjectID = Mesh->GetUniqueID();
-						MeshName += GetUniqueIDString(UniqueFbxObjectID);
-					}
-				}
-				return MeshName;
-			}
-			
-			FString FFbxMesh::GetMeshUniqueID(FbxGeometryBase* Mesh)
-			{
-				FString MeshUniqueID;
-				if (Mesh->GetAttributeType() == FbxNodeAttribute::eMesh)
-				{
-					MeshUniqueID = TEXT("\\Mesh\\");
-				}
-				else if (Mesh->GetAttributeType() == FbxNodeAttribute::eShape)
-				{
-					MeshUniqueID = TEXT("\\Shape\\");
-				}
-				FString MeshName = FFbxHelper::GetFbxObjectName(Mesh);
-				if (MeshName.IsEmpty())
-				{
-					if (Mesh->GetNodeCount() > 0)
-					{
-						MeshUniqueID += FFbxHelper::GetFbxNodeHierarchyName(Mesh->GetNode(0));
-					}
-					else
-					{
-						MeshUniqueID += GetMeshName(Mesh);
-					}
-				}
-				else
-				{
-					MeshUniqueID += MeshName;
-				}
-				return MeshUniqueID;
-			}
-#endif
-
-			void FFbxMesh::ExtractSkinnedMeshNodeJoints(FbxScene* SDKScene, FbxMesh* Mesh, UInterchangeMeshNode* MeshNode)
-			{
-				TArray<FString> JointNodeUniqueIDs;
-				const int32 SkinDeformerCount = Mesh->GetDeformerCount(FbxDeformer::eSkin);
-				for (int32 DeformerIndex = 0; DeformerIndex < SkinDeformerCount; DeformerIndex++)
-				{
-					FbxSkin* Skin = (FbxSkin*)Mesh->GetDeformer(DeformerIndex, FbxDeformer::eSkin);
-					if (!ensure(Skin))
-					{
-						continue;
-					}
-					const int32 ClusterCount = Skin->GetClusterCount();
-					JointNodeUniqueIDs.Reserve(ClusterCount);
-					for (int32 ClusterIndex = 0; ClusterIndex < ClusterCount; ClusterIndex++)
-					{
-						FbxCluster* Cluster = Skin->GetCluster(ClusterIndex);
-						// When Maya plug-in exports rigid binding, it will generate "CompensationCluster" for each ancestor links.
-						// FBX writes these "CompensationCluster" out. The CompensationCluster also has weight 1 for vertices.
-						// Unreal importer should skip these clusters.
-						if (!Cluster || (FCStringAnsi::Strcmp(Cluster->GetUserDataID(), "Maya_ClusterHint") == 0 && FCStringAnsi::Strcmp(Cluster->GetUserData(), "CompensationCluster") == 0))
-						{
-							continue;
-						}
-						FbxNode* Link = Cluster->GetLink();
-						FString JointNodeUniqueID = FFbxHelper::GetFbxNodeHierarchyName(Link);
-						// find the bone index
-						if (!JointNodeUniqueIDs.Contains(JointNodeUniqueID))
-						{
-							JointNodeUniqueIDs.Add(JointNodeUniqueID);
-							MeshNode->SetSkeletonDependencyUid(JointNodeUniqueID);
-						}
-					}
-				}
-			}
 
 			void FFbxMesh::AddAllMeshes(FbxScene* SDKScene, FbxGeometryConverter* SDKGeometryConverter, UInterchangeBaseNodeContainer& NodeContainer, TMap<FString, TSharedPtr<FPayloadContextBase>>& PayloadContexts)
 			{
@@ -1357,7 +1258,7 @@ namespace UE
 						{
 							//Set the skinned mesh attribute
 							MeshNode->SetSkinnedMesh(true);
-							ExtractSkinnedMeshNodeJoints(SDKScene, Mesh, MeshNode);
+							ExtractSkinnedMeshNodeJoints(SDKScene, NodeContainer, Mesh, MeshNode);
 						}
 						Mesh->ComputeBBox();
 						const int32 MeshVertexCount = Mesh->GetControlPointsCount();
@@ -1480,6 +1381,117 @@ namespace UE
 				}
 			}
 
+			bool FFbxMesh::GetGlobalJointBindPoseTransform(FbxScene* SDKScene, FbxNode* Joint, FbxAMatrix& GlobalBindPoseJointMatrix)
+			{
+				const FString LinkName = FFbxHelper::GetFbxObjectName(Joint);
+
+				const int32 PoseCount = SDKScene->GetPoseCount();
+				for (int32 PoseIndex = 0; PoseIndex < PoseCount; PoseIndex++)
+				{
+					FbxPose* CurrentPose = SDKScene->GetPose(PoseIndex);
+					if (CurrentPose && CurrentPose->IsBindPose())
+					{
+						FString PoseName = FFbxHelper::GetFbxObjectName(CurrentPose);
+						int32 PoseLinkIndex = CurrentPose->Find(Joint);
+						if (PoseLinkIndex >= 0)
+						{
+							FbxMatrix NoneAffineMatrix = CurrentPose->GetMatrix(PoseLinkIndex);
+							GlobalBindPoseJointMatrix = *(FbxAMatrix*)(double*)&NoneAffineMatrix;
+							return true;
+						}
+					}
+				}
+
+				//Search all skeletalmesh(FbxGeometry with valid deformer) using this joint and see if there is a valid FbxCluster
+				//containing a TransformLinkMatrix for this joint
+				const int32 GeometryCount = SDKScene->GetGeometryCount();
+				for (int32 GeometryIndex = 0; GeometryIndex < GeometryCount; ++GeometryIndex)
+				{
+					const FbxGeometry* Geometry = SDKScene->GetGeometry(GeometryIndex);
+					if (!ensure(Geometry))
+					{
+						continue;
+					}
+					const int32 GeometryDeformerCount = Geometry->GetDeformerCount(FbxDeformer::eSkin);
+					for (int32 GeometryDeformerIndex = 0; GeometryDeformerIndex < GeometryDeformerCount; ++GeometryDeformerIndex)
+					{
+						const FbxSkin* Skin = (FbxSkin*)Geometry->GetDeformer(GeometryDeformerIndex, FbxDeformer::eSkin);
+						if (!ensure(Skin))
+						{
+							continue;
+						}
+						const int32 ClusterCount = Skin->GetClusterCount();
+						for (int32 ClusterIndex = 0; ClusterIndex < ClusterCount; ClusterIndex++)
+						{
+							const FbxCluster* Cluster = Skin->GetCluster(ClusterIndex);
+							// When Maya plug-in exports rigid binding, it will generate "CompensationCluster" for each ancestor links.
+							// FBX writes these "CompensationCluster" out. The CompensationCluster also has weight 1 for vertices.
+							// Unreal importer should skip these clusters.
+							if (!Cluster || (FCStringAnsi::Strcmp(Cluster->GetUserDataID(), "Maya_ClusterHint") == 0 && FCStringAnsi::Strcmp(Cluster->GetUserData(), "CompensationCluster") == 0))
+							{
+								continue;
+							}
+							if (Joint == Cluster->GetLink())
+							{
+								Cluster->GetTransformLinkMatrix(GlobalBindPoseJointMatrix);
+								return true;
+							}
+						}
+					}
+				}
+				return false;
+			}
+
+			void FFbxMesh::ExtractSkinnedMeshNodeJoints(FbxScene* SDKScene, UInterchangeBaseNodeContainer& NodeContainer, FbxMesh* Mesh, UInterchangeMeshNode* MeshNode)
+			{
+				TArray<FString> JointNodeUniqueIDs;
+
+				//Get the bind pose from fbx
+				const int32 PoseCount = SDKScene->GetPoseCount();
+				TArray<FbxPose*> BindPoses;
+				for (int32 PoseIndex = 0; PoseIndex < PoseCount; PoseIndex++)
+				{
+					FbxPose* CurrentPose = SDKScene->GetPose(PoseIndex);
+
+					// current pose is bind pose, 
+					if (CurrentPose && CurrentPose->IsBindPose())
+					{
+						BindPoses.Add(CurrentPose);
+					}
+				}
+
+				const int32 SkinDeformerCount = Mesh->GetDeformerCount(FbxDeformer::eSkin);
+				for (int32 DeformerIndex = 0; DeformerIndex < SkinDeformerCount; DeformerIndex++)
+				{
+					FbxSkin* Skin = (FbxSkin*)Mesh->GetDeformer(DeformerIndex, FbxDeformer::eSkin);
+					if (!ensure(Skin))
+					{
+						continue;
+					}
+					const int32 ClusterCount = Skin->GetClusterCount();
+					JointNodeUniqueIDs.Reserve(JointNodeUniqueIDs.Num() + ClusterCount);
+					for (int32 ClusterIndex = 0; ClusterIndex < ClusterCount; ClusterIndex++)
+					{
+						FbxCluster* Cluster = Skin->GetCluster(ClusterIndex);
+						// When Maya plug-in exports rigid binding, it will generate "CompensationCluster" for each ancestor links.
+						// FBX writes these "CompensationCluster" out. The CompensationCluster also has weight 1 for vertices.
+						// Unreal importer should skip these clusters.
+						if (!Cluster || (FCStringAnsi::Strcmp(Cluster->GetUserDataID(), "Maya_ClusterHint") == 0 && FCStringAnsi::Strcmp(Cluster->GetUserData(), "CompensationCluster") == 0))
+						{
+							continue;
+						}
+						FbxNode* Link = Cluster->GetLink();
+						FString JointNodeUniqueID = FFbxHelper::GetFbxNodeHierarchyName(Link);
+						// find the bone index
+						if (!JointNodeUniqueIDs.Contains(JointNodeUniqueID))
+						{
+							JointNodeUniqueIDs.Add(JointNodeUniqueID);
+							MeshNode->SetSkeletonDependencyUid(JointNodeUniqueID);
+						}
+					}
+				}
+			}
+
 			UInterchangeMeshNode* FFbxMesh::CreateMeshNode(UInterchangeBaseNodeContainer& NodeContainer, const FString& NodeName, const FString& NodeUniqueID)
 			{
 				FString DisplayLabel(NodeName);
diff --git a/Engine/Plugins/Experimental/Interchange/Runtime/Source/Parsers/Fbx/Private/FbxMesh.h b/Engine/Plugins/Experimental/Interchange/Runtime/Source/Parsers/Fbx/Private/FbxMesh.h
--- a/Engine/Plugins/Experimental/Interchange/Runtime/Source/Parsers/Fbx/Private/FbxMesh.h
+++ b/Engine/Plugins/Experimental/Interchange/Runtime/Source/Parsers/Fbx/Private/FbxMesh.h
@@ -112,18 +112,11 @@ namespace UE
 					: Parser(InParser)
 				{}
 
-#if 0
-				static FString GetMeshName(FbxGeometryBase* Mesh);
-				static FString GetMeshUniqueID(FbxGeometryBase* Mesh);
-#endif
-				void ExtractSkinnedMeshNodeJoints(FbxScene* SDKScene, FbxMesh* Mesh, UInterchangeMeshNode* MeshNode);
 				void AddAllMeshes(FbxScene* SDKScene, FbxGeometryConverter* SDKGeometryConverter, UInterchangeBaseNodeContainer& NodeContainer, TMap<FString, TSharedPtr<FPayloadContextBase>>& PayloadContexts);
-			
+				static bool GetGlobalJointBindPoseTransform(FbxScene* SDKScene, FbxNode* Joint, FbxAMatrix& GlobalBindPoseJointMatrix);
 			protected:
+				void ExtractSkinnedMeshNodeJoints(FbxScene* SDKScene, UInterchangeBaseNodeContainer& NodeContainer, FbxMesh* Mesh, UInterchangeMeshNode* MeshNode);
 				UInterchangeMeshNode* CreateMeshNode(UInterchangeBaseNodeContainer& NodeContainer, const FString& NodeName, const FString& NodeUniqueID);
-#if 0
-				static FString GetUniqueIDString(const uint64 UniqueID);
-#endif
 
 			private:
 				FFbxParser& Parser;
diff --git a/Engine/Plugins/Experimental/Interchange/Runtime/Source/Parsers/Fbx/Private/FbxScene.cpp b/Engine/Plugins/Experimental/Interchange/Runtime/Source/Parsers/Fbx/Private/FbxScene.cpp
--- a/Engine/Plugins/Experimental/Interchange/Runtime/Source/Parsers/Fbx/Private/FbxScene.cpp
+++ b/Engine/Plugins/Experimental/Interchange/Runtime/Source/Parsers/Fbx/Private/FbxScene.cpp
@@ -69,13 +69,17 @@ namespace UE
 				CreateAssetNodeReference(UnrealSceneNode, NodeAttribute, NodeContainer, UInterchangeLightNode::StaticAssetTypeName());
 			}
 
-			UInterchangeSceneNode* FFbxScene::AddHierarchyRecursively(FbxNode* Node, FbxScene* SDKScene, UInterchangeBaseNodeContainer& NodeContainer)
+			void FFbxScene::AddHierarchyRecursively(UInterchangeSceneNode* UnrealParentNode, FbxNode* Node, FbxScene* SDKScene, UInterchangeBaseNodeContainer& NodeContainer)
 			{
 				FString NodeName = FFbxHelper::GetFbxObjectName(Node);
 				FString NodeUniqueID = FFbxHelper::GetFbxNodeHierarchyName(Node);
 
 				UInterchangeSceneNode* UnrealNode = CreateTransformNode(NodeContainer, NodeName, NodeUniqueID);
 				check(UnrealNode);
+				if (UnrealParentNode)
+				{
+					UnrealNode->SetParentUid(UnrealParentNode->GetUniqueID());
+				}
 				
 				auto GetConvertedTransform = [Node](FbxAMatrix& NewFbxMatrix)
 				{
@@ -125,7 +129,6 @@ namespace UE
 						case FbxNodeAttribute::eOpticalReference:
 						case FbxNodeAttribute::eOpticalMarker:
 						case FbxNodeAttribute::eCachedEffect:
-						case FbxNodeAttribute::eNull:
 						case FbxNodeAttribute::eMarker:
 						case FbxNodeAttribute::eCameraStereo:
 						case FbxNodeAttribute::eCameraSwitcher:
@@ -139,14 +142,40 @@ namespace UE
 						case FbxNodeAttribute::eLine:
 							//Unsupported attribute
 							break;
+
 						case FbxNodeAttribute::eShape: //We do not add a dependency for shape on the scene node since shapes are a MeshNode dependency.
 							break;
+
+						case FbxNodeAttribute::eNull:
 						case FbxNodeAttribute::eSkeleton:
+							//eNull node will be set has a skeleton node
 						{
 							//Add the joint specialized type
 							UnrealNode->AddSpecializedType(FSceneNodeStaticData::GetJointSpecializeTypeString());
+							//Get the bind pose for this joint
+							FbxAMatrix GlobalBindPoseJointMatrix;
+							if (FFbxMesh::GetGlobalJointBindPoseTransform(SDKScene, Node, GlobalBindPoseJointMatrix))
+							{
+								FTransform GlobalBindPoseJointTransform = GetConvertedTransform(GlobalBindPoseJointMatrix);
+								UnrealNode->SetCustomGlobalTransform(GlobalBindPoseJointTransform);
+								//We grab the fbx parent node to compute the local transform
+								if (FbxNode* ParentNode = Node->GetParent())
+								{
+									FbxAMatrix GlobalFbxParentMatrix = ParentNode->EvaluateGlobalTransform();
+									FFbxMesh::GetGlobalJointBindPoseTransform(SDKScene, ParentNode, GlobalFbxParentMatrix);
+									FbxAMatrix	LocalFbxMatrix = GlobalFbxParentMatrix.Inverse() * GlobalBindPoseJointMatrix;
+									FTransform LocalBindPoseJointTransform = GetConvertedTransform(LocalFbxMatrix);
+									UnrealNode->SetCustomLocalTransform(LocalBindPoseJointTransform);
+								}
+								else
+								{
+									//No parent, set the same matrix has the global
+									UnrealNode->SetCustomLocalTransform(GlobalBindPoseJointTransform);
+								}
+							}
 							break;
 						}
+
 						case FbxNodeAttribute::eMesh:
 						{
 							//For Mesh attribute we add the fbx nodes materials
@@ -179,13 +208,8 @@ namespace UE
 				for (int32 ChildIndex = 0; ChildIndex < ChildCount; ++ChildIndex)
 				{
 					FbxNode* ChildNode = Node->GetChild(ChildIndex);
-					UInterchangeSceneNode* UnrealChildNode = AddHierarchyRecursively(ChildNode, SDKScene, NodeContainer);
-					if (UnrealChildNode)
-					{
-						UnrealChildNode->SetParentUid(UnrealNode->GetUniqueID());
-					}
+					AddHierarchyRecursively(UnrealNode, ChildNode, SDKScene, NodeContainer);
 				}
-				return UnrealNode;
 			}
 
 			UInterchangeSceneNode* FFbxScene::CreateTransformNode(UInterchangeBaseNodeContainer& NodeContainer, const FString& NodeName, const FString& NodeUniqueID)
@@ -208,7 +232,7 @@ namespace UE
 			void FFbxScene::AddHierarchy(FbxScene* SDKScene, UInterchangeBaseNodeContainer& NodeContainer)
 			{
 				 FbxNode* RootNode = SDKScene->GetRootNode();
-				 AddHierarchyRecursively(RootNode, SDKScene, NodeContainer);
+				 AddHierarchyRecursively(nullptr, RootNode, SDKScene, NodeContainer);
 			}
 		} //ns Private
 	} //ns Interchange
diff --git a/Engine/Plugins/Experimental/Interchange/Runtime/Source/Parsers/Fbx/Private/FbxScene.h b/Engine/Plugins/Experimental/Interchange/Runtime/Source/Parsers/Fbx/Private/FbxScene.h
--- a/Engine/Plugins/Experimental/Interchange/Runtime/Source/Parsers/Fbx/Private/FbxScene.h
+++ b/Engine/Plugins/Experimental/Interchange/Runtime/Source/Parsers/Fbx/Private/FbxScene.h
@@ -31,7 +31,7 @@ namespace UE
 				void CreateMeshNodeReference(UInterchangeSceneNode* UnrealSceneNode, FbxNodeAttribute* NodeAttribute, UInterchangeBaseNodeContainer& NodeContainer);
 				void CreateCameraNodeReference(UInterchangeSceneNode* UnrealSceneNode, FbxNodeAttribute* NodeAttribute, UInterchangeBaseNodeContainer& NodeContainer);
 				void CreateLightNodeReference(UInterchangeSceneNode* UnrealSceneNode, FbxNodeAttribute* NodeAttribute, UInterchangeBaseNodeContainer& NodeContainer);
-				UInterchangeSceneNode* AddHierarchyRecursively(FbxNode* Node, FbxScene* SDKScene, UInterchangeBaseNodeContainer& NodeContainer);
+				void AddHierarchyRecursively(UInterchangeSceneNode* UnrealParentNode, FbxNode* Node, FbxScene* SDKScene, UInterchangeBaseNodeContainer& NodeContainer);
 
 			private:
 				FFbxParser& Parser;
diff --git a/Engine/Plugins/Experimental/Interchange/Runtime/Source/Parsers/Fbx/Private/FbxSkelMesh.cpp b/Engine/Plugins/Experimental/Interchange/Runtime/Source/Parsers/Fbx/Private/FbxSkelMesh.cpp
deleted file mode 100644
diff --git a/Engine/Plugins/Experimental/Interchange/Runtime/Source/Parsers/Fbx/Private/FbxSkelMesh.h b/Engine/Plugins/Experimental/Interchange/Runtime/Source/Parsers/Fbx/Private/FbxSkelMesh.h
deleted file mode 100644
