Implement interchange option bImportMeshesInBoneHierarchy in the generic pipeline.
#rnx

diff --git a/Engine/Plugins/Experimental/Interchange/Runtime/Source/Pipelines/Private/InterchangeGenericSkeletalMeshPipeline.cpp b/Engine/Plugins/Experimental/Interchange/Runtime/Source/Pipelines/Private/InterchangeGenericSkeletalMeshPipeline.cpp
--- a/Engine/Plugins/Experimental/Interchange/Runtime/Source/Pipelines/Private/InterchangeGenericSkeletalMeshPipeline.cpp
+++ b/Engine/Plugins/Experimental/Interchange/Runtime/Source/Pipelines/Private/InterchangeGenericSkeletalMeshPipeline.cpp
@@ -207,6 +207,42 @@ namespace UE::Interchange::SkeletalMeshGenericPipeline
 		// if the hierarchy matches, and if it's more then 1 bone, we allow
 		return (NumOfBoneMatches > 0);
 	}
+
+	bool RecursiveFindChildUid(const UInterchangeBaseNodeContainer* BaseNodeContainer, const FString& ParentUid, const FString& SearchUid)
+	{
+		if (ParentUid == SearchUid)
+		{
+			return true;
+		}
+		const int32 ChildCount = BaseNodeContainer->GetNodeChildrenCount(ParentUid);
+		TArray<FString> Childrens = BaseNodeContainer->GetNodeChildrenUids(ParentUid);
+		for (int32 ChildIndex = 0; ChildIndex < ChildCount; ++ChildIndex)
+		{
+			if (RecursiveFindChildUid(BaseNodeContainer, Childrens[ChildIndex], SearchUid))
+			{
+				return true;
+			}
+		}
+		return false;
+	}
+
+	void RemoveNestedMeshNodes(const UInterchangeBaseNodeContainer* BaseNodeContainer, const UInterchangeSkeletonFactoryNode* SkeletonFactoryNode, TArray<FString>& NodeUids)
+	{
+		if (!SkeletonFactoryNode)
+		{
+			return;
+		}
+		FString SkeletonRootJointUid;
+		SkeletonFactoryNode->GetCustomRootJointUid(SkeletonRootJointUid);
+		for (int32 NodeIndex = NodeUids.Num()-1; NodeIndex >= 0; NodeIndex--)
+		{
+			const FString& NodeUid = NodeUids[NodeIndex];
+			if (RecursiveFindChildUid(BaseNodeContainer, SkeletonRootJointUid, NodeUid))
+			{
+				NodeUids.RemoveAt(NodeIndex);
+			}
+		}
+	}
 }
 
 bool UInterchangeGenericAssetsPipeline::ExecutePreImportPipelineSkeletalMesh()
@@ -621,7 +657,13 @@ void UInterchangeGenericAssetsPipeline::AddLodDataToSkeletalMesh(const UIntercha
 			//If the pipeline should not import lods, skip any lod over base lod
 			continue;
 		}
-		const TArray<FString>& NodeUids = LodIndexAndNodeUids.Value;
+
+		//Copy the nodes unique id because we need to remove nested mesh if the option is to not import them
+		TArray<FString> NodeUids = LodIndexAndNodeUids.Value;
+		if (!bImportMeshesInBoneHierarchy)
+		{
+			UE::Interchange::SkeletalMeshGenericPipeline::RemoveNestedMeshNodes(BaseNodeContainer, SkeletonFactoryNode, NodeUids);
+		}
 
 		//Create a lod data node with all the meshes for this LOD
 		const FString SkeletalMeshLodDataName = TEXT("LodData") + FString::FromInt(LodIndex);
diff --git a/Engine/Plugins/Experimental/Interchange/Runtime/Source/Pipelines/Public/InterchangeGenericAssetsPipeline.h b/Engine/Plugins/Experimental/Interchange/Runtime/Source/Pipelines/Public/InterchangeGenericAssetsPipeline.h
--- a/Engine/Plugins/Experimental/Interchange/Runtime/Source/Pipelines/Public/InterchangeGenericAssetsPipeline.h
+++ b/Engine/Plugins/Experimental/Interchange/Runtime/Source/Pipelines/Public/InterchangeGenericAssetsPipeline.h
@@ -151,6 +151,10 @@ public:
 	UPROPERTY(EditAnywhere, BlueprintReadWrite, Category = SKELETAL_MESHES_CATEGORY)
 	bool bUpdateSkeletonReferencePose = false;
 
+	/** If checked, meshes nested in bone hierarchies will be imported instead of being converted to bones. */
+	UPROPERTY(EditAnywhere, BlueprintReadWrite, Category = SKELETAL_MESHES_CATEGORY)
+	bool bImportMeshesInBoneHierarchy = true;
+
 	/** If checked, create new PhysicsAsset if it doesn't have it */
 	UPROPERTY(EditAnywhere, BlueprintReadWrite, Category = SKELETAL_MESHES_CATEGORY)
 	bool bCreatePhysicsAsset;
