Fix bug that prevent blendshape import (morph target) when using interchange.
#rnx

diff --git a/Engine/Plugins/Experimental/Interchange/Runtime/Source/Import/Private/Mesh/InterchangeSkeletalMeshFactory.cpp b/Engine/Plugins/Experimental/Interchange/Runtime/Source/Import/Private/Mesh/InterchangeSkeletalMeshFactory.cpp
--- a/Engine/Plugins/Experimental/Interchange/Runtime/Source/Import/Private/Mesh/InterchangeSkeletalMeshFactory.cpp
+++ b/Engine/Plugins/Experimental/Interchange/Runtime/Source/Import/Private/Mesh/InterchangeSkeletalMeshFactory.cpp
@@ -235,7 +235,7 @@ namespace UE
 					const int32 VertexOffset = BlendShapePayloadData->VertexOffset;
 					const int32 SourceMeshVertexCount = SourceMeshDescription.Vertices().Num();
 					const int32 DestinationVertexIndexMax = VertexOffset + SourceMeshVertexCount;
-					if (!DestinationSkeletalMeshImportData.Points.IsValidIndex(DestinationVertexIndexMax))
+					if (!DestinationSkeletalMeshImportData.Points.IsValidIndex(DestinationVertexIndexMax-1))
 					{
 						UE_LOG(LogInterchangeImport, Error, TEXT("Corrupted blend shape optional payload data [%s]"), *BlendShapeName);
 						continue;
@@ -291,7 +291,7 @@ namespace UE
 						const int32 VertexOffset = BlendShapePayloadData->VertexOffset;
 						const int32 SourceMeshVertexCount = SourceMeshDescription.Vertices().Num();
 						const int32 DestinationVertexIndexMax = VertexOffset + SourceMeshVertexCount;
-						if (!ensure(DestinationSkeletalMeshImportData.Points.IsValidIndex(DestinationVertexIndexMax)))
+						if (!ensure(DestinationSkeletalMeshImportData.Points.IsValidIndex(DestinationVertexIndexMax-1)))
 						{
 							//This error was suppose to be catch in the pre parallel for loop
 							break;
