Add bUseT0AsReference generic pipeline option
#rnx

diff --git a/Engine/Plugins/Experimental/Interchange/Runtime/Source/Import/Private/Mesh/InterchangeSkeletalMeshFactory.cpp b/Engine/Plugins/Experimental/Interchange/Runtime/Source/Import/Private/Mesh/InterchangeSkeletalMeshFactory.cpp
--- a/Engine/Plugins/Experimental/Interchange/Runtime/Source/Import/Private/Mesh/InterchangeSkeletalMeshFactory.cpp
+++ b/Engine/Plugins/Experimental/Interchange/Runtime/Source/Import/Private/Mesh/InterchangeSkeletalMeshFactory.cpp
@@ -55,7 +55,7 @@
 				FTransform	LocalTransform; // local transform
 			};
 
-			void RecursiveAddBones(const UInterchangeBaseNodeContainer* NodeContainer, const FString& JointNodeId, TArray <FJointInfo>& JointInfos, int32 ParentIndex, TArray<SkeletalMeshImportData::FBone>& RefBonesBinary)
+			void RecursiveAddBones(const UInterchangeBaseNodeContainer* NodeContainer, const FString& JointNodeId, TArray <FJointInfo>& JointInfos, int32 ParentIndex, TArray<SkeletalMeshImportData::FBone>& RefBonesBinary, const bool bUseTimeZeroAsBindPose, bool & bOutDiffPose)
 			{
 				const UInterchangeSceneNode* JointNode = Cast<UInterchangeSceneNode>(NodeContainer->GetNode(JointNodeId));
 				if (!JointNode)
@@ -67,7 +67,33 @@
 				int32 JointInfoIndex = JointInfos.Num();
 				FJointInfo& Info = JointInfos.AddZeroed_GetRef();
 				Info.Name = JointNode->GetDisplayLabel();
-				ensure(JointNode->GetCustomLocalTransform(Info.LocalTransform));
+
+				FTransform LocalTransform;
+				ensure(JointNode->GetCustomLocalTransform(LocalTransform));
+
+				FTransform TimeZeroLocalTransform;
+				const bool bHasTimeZeroTransform = JointNode->GetCustomTimeZeroLocalTransform(TimeZeroLocalTransform);
+				FTransform BindPoseLocalTransform;
+				const bool bHasBindPoseTransform = JointNode->GetCustomBindPoseLocalTransform(BindPoseLocalTransform);
+
+				Info.LocalTransform = bHasBindPoseTransform ? BindPoseLocalTransform : LocalTransform;
+				//If user want to bind the mesh at time zero try to get the time zero transform
+				if (bUseTimeZeroAsBindPose && bHasTimeZeroTransform)
+				{
+					if (bHasBindPoseTransform)
+					{
+						if(!TimeZeroLocalTransform.Equals(Info.LocalTransform))
+						{
+							bOutDiffPose = true;
+						}
+					}
+					Info.LocalTransform = TimeZeroLocalTransform;
+				}
+				else if (bHasBindPoseTransform)
+				{
+					Info.LocalTransform = BindPoseLocalTransform;
+				}
+
 				Info.ParentIndex = ParentIndex;
 
 				SkeletalMeshImportData::FBone& Bone = RefBonesBinary.AddZeroed_GetRef();
@@ -84,11 +110,11 @@
 				Bone.NumChildren = ChildrenIds.Num();
 				for (int32 ChildIndex = 0; ChildIndex < ChildrenIds.Num(); ++ChildIndex)
 				{
-					RecursiveAddBones(NodeContainer, ChildrenIds[ChildIndex], JointInfos, JointInfoIndex, RefBonesBinary);
+					RecursiveAddBones(NodeContainer, ChildrenIds[ChildIndex], JointInfos, JointInfoIndex, RefBonesBinary, bUseTimeZeroAsBindPose, bOutDiffPose);
 				}
 			}
 
-			bool ProcessImportMeshSkeleton(const USkeleton* SkeletonAsset, FReferenceSkeleton& RefSkeleton, int32& SkeletalDepth, const UInterchangeBaseNodeContainer* NodeContainer, const FString& RootJointNodeId, TArray<SkeletalMeshImportData::FBone>& RefBonesBinary)
+			bool ProcessImportMeshSkeleton(const USkeleton* SkeletonAsset, FReferenceSkeleton& RefSkeleton, int32& SkeletalDepth, const UInterchangeBaseNodeContainer* NodeContainer, const FString& RootJointNodeId, TArray<SkeletalMeshImportData::FBone>& RefBonesBinary, const bool bUseTimeZeroAsBindPose, bool& bOutDiffPose)
 			{
 				auto FixupBoneName = [](FString BoneName)
 				{
@@ -103,7 +129,12 @@
 
 				FReferenceSkeletonModifier RefSkelModifier(RefSkeleton, SkeletonAsset);
 				TArray <FJointInfo> JointInfos;
-				RecursiveAddBones(NodeContainer, RootJointNodeId, JointInfos, INDEX_NONE, RefBonesBinary);
+				RecursiveAddBones(NodeContainer, RootJointNodeId, JointInfos, INDEX_NONE, RefBonesBinary, bUseTimeZeroAsBindPose, bOutDiffPose);
+				if (bOutDiffPose)
+				{
+					//bOutDiffPose can only be true if the user ask to bind with time zero transform.
+					ensure(bUseTimeZeroAsBindPose);
+				}
 				// Digest bones to the serializable format.
 				for (int32 b = 0; b < JointInfos.Num(); b++)
 				{
@@ -333,12 +364,156 @@
 				return;
 			}
 
+			const UInterchangeSceneNode* RecursiveFindJointByName(const UInterchangeBaseNodeContainer* NodeContainer, const FString& ParentJointNodeId, const FString& JointName)
+			{
+				if (const UInterchangeSceneNode* JointNode = Cast<UInterchangeSceneNode>(NodeContainer->GetNode(ParentJointNodeId)))
+				{
+					if (JointNode->GetDisplayLabel().Equals(JointName))
+					{
+						return JointNode;
+					}
+				}
+				TArray<FString> NodeChildrenUids = NodeContainer->GetNodeChildrenUids(ParentJointNodeId);
+				for (int32 ChildIndex = 0; ChildIndex < NodeChildrenUids.Num(); ++ChildIndex)
+				{
+					if (const UInterchangeSceneNode* JointNode = RecursiveFindJointByName(NodeContainer, NodeChildrenUids[ChildIndex], JointName))
+					{
+						return JointNode;
+					}
+				}
+
+				return nullptr;
+			}
+
+			void SkinVertexPositionToTimeZero(UE::Interchange::FSkeletalMeshLodPayloadData& LodMeshPayload
+				, const UInterchangeBaseNodeContainer* NodeContainer
+				, const FString& RootJointNodeId
+				, const FTransform& MeshGlobalTransform)
+			{
+				FMeshDescription& MeshDescription = LodMeshPayload.LodMeshDescription;
+				const int32 VertexCount = MeshDescription.Vertices().Num();
+				const TArray<FString>& JointNames = LodMeshPayload.JointNames;
+				// Create a copy of the vertex array to receive vertex deformations.
+				TArray<FVector> DestinationVertexPositions;
+				DestinationVertexPositions.AddZeroed(VertexCount);
+
+				FSkeletalMeshAttributes Attributes(MeshDescription);
+				TVertexAttributesRef<FVector3f> VertexPositions = Attributes.GetVertexPositions();
+				FSkinWeightsVertexAttributesRef VertexSkinWeights = Attributes.GetVertexSkinWeights();
+
+				for (FVertexID VertexID : MeshDescription.Vertices().GetElementIDs())
+				{
+					//We can use GetValue because the Meshdescription was compacted before the copy
+					DestinationVertexPositions[VertexID.GetValue()] = VertexPositions[VertexID];
+				}
+
+				// Deform the vertex array with the links contained in the mesh.
+				TArray<FMatrix> SkinDeformations;
+				SkinDeformations.AddZeroed(VertexCount);
+
+				TArray<double> SkinWeights;
+				SkinWeights.AddZeroed(VertexCount);
+
+				//We assume normalize weight method in this bind pose conversion
+
+				const FTransform MeshGlobalTransformInverse = MeshGlobalTransform.Inverse();
+				const int32 JointCount = JointNames.Num();
+				for (int32 JointIndex = 0; JointIndex < JointCount; ++JointIndex)
+				{
+					const FString& JointName = JointNames[JointIndex];
+
+					const UInterchangeSceneNode* JointNode = RecursiveFindJointByName(NodeContainer, RootJointNodeId, JointName);
+					if (!ensure(JointNode))
+					{
+						continue;
+					}
+					
+					FTransform JointBindPoseGlobalTransform;
+					if (!JointNode->GetCustomBindPoseGlobalTransform(JointBindPoseGlobalTransform))
+					{
+						//If there is no bind pose we will fall back on the CustomGlobalTransform of the link.
+						//We ensure here because any scenenode should have a valid CustomGlobalTransform.
+						if (!ensure(JointNode->GetCustomGlobalTransform(JointBindPoseGlobalTransform)))
+						{
+							//No value to convert from, skip this joint.
+							continue;
+						}
+					}
+
+					FTransform JointTimeZeroGlobalTransform;
+					if (!JointNode->GetCustomTimeZeroGlobalTransform(JointTimeZeroGlobalTransform))
+					{
+						//If there is no time zero global transform we cannot set the bind pose to time zero.
+						//We must skip this joint.
+						continue;
+					}
+
+					//Get the mesh transform in local relative to the bind pose transform
+ 					const FTransform MeshTransformRelativeToBindPoseTransform = MeshGlobalTransform * JointBindPoseGlobalTransform.Inverse();
+					//Get the time zero pose transform in local relative to the mesh transform
+ 					const FTransform TimeZeroTransformRelativeToMeshTransform = JointTimeZeroGlobalTransform * MeshGlobalTransformInverse;
+					//Multiply both transform to get a matrix that will transform the mesh vertices from the bind pose skinning to the time zero skinning
+ 					const FMatrix VertexTransformMatrix = (MeshTransformRelativeToBindPoseTransform * TimeZeroTransformRelativeToMeshTransform).ToMatrixWithScale();
+
+					//Iterate all bone vertices
+					for (FVertexID VertexID : MeshDescription.Vertices().GetElementIDs())
+					{
+						const int32 VertexIndex = VertexID.GetValue();
+						const FVertexBoneWeights BoneWeights = VertexSkinWeights.Get(VertexID);
+						const int32 InfluenceCount = BoneWeights.Num();
+						float Weight = 0.0f;
+						for (int32 InfluenceIndex = 0; InfluenceIndex < InfluenceCount; ++InfluenceIndex)
+						{
+							FBoneIndexType BoneIndex = BoneWeights[InfluenceIndex].GetBoneIndex();
+							if (JointIndex == BoneIndex)
+							{
+								Weight = BoneWeights[InfluenceIndex].GetWeight();
+								break;
+							}
+						}
+						if (FMath::IsNearlyZero(Weight))
+						{
+							continue;
+						}
+
+						//The weight multiply the vertex transform matrix so we can have multiple joint affecting this vertex.
+						const FMatrix Influence = VertexTransformMatrix * Weight;
+						//Add the weighted result
+						SkinDeformations[VertexIndex] += Influence;
+						//Add the total weight so we can normalize the result in case the accumulated weight is different then 1
+						SkinWeights[VertexIndex] += Weight;
+					}
+				}
+
+				for (FVertexID VertexID : MeshDescription.Vertices().GetElementIDs())
+				{
+					const int32 VertexIndex = VertexID.GetValue();
+					const FVector3f lSrcVertex = DestinationVertexPositions[VertexIndex];
+					FVector& lDstVertex = DestinationVertexPositions[VertexIndex];
+					double Weight = SkinWeights[VertexIndex];
+
+					// Deform the vertex if there was at least a link with an influence on the vertex,
+					if (!FMath::IsNearlyZero(Weight))
+					{
+						//Apply skinning of all joints
+						lDstVertex = SkinDeformations[VertexIndex].TransformPosition(lSrcVertex);
+						//Normalized, in case the weight is different then 1
+						lDstVertex /= Weight;
+						//Set the new vertex position in the mesh description
+						VertexPositions[VertexID] = lDstVertex;
+					}
+				}
+			}
+
 			void RetrieveAllSkeletalMeshPayloadsAndFillImportData(const UInterchangeSkeletalMeshFactoryNode* SkeletalMeshFactoryNode
 																  , FSkeletalMeshImportData& DestinationImportData
 																  , TArray<FMeshNodeContext>& MeshReferences
 																  , TArray<SkeletalMeshImportData::FBone>& RefBonesBinary
 																  , const UInterchangeSkeletalMeshFactory::FCreateAssetParams& Arguments
-																  , const IInterchangeSkeletalMeshPayloadInterface* SkeletalMeshTranslatorPayloadInterface)
+																  , const IInterchangeSkeletalMeshPayloadInterface* SkeletalMeshTranslatorPayloadInterface
+																  , const bool bSkinControlPointToTimeZero
+																  , const UInterchangeBaseNodeContainer* NodeContainer
+																  , const FString& RootJointNodeId)
 			{
 				if (!SkeletalMeshTranslatorPayloadInterface)
 				{
@@ -386,6 +561,20 @@
 					SkeletalMeshAppendSettings.SourceVertexIDOffset = VertexOffset;
 					FElementIDRemappings ElementIDRemappings;
 					LodMeshPayload->LodMeshDescription.Compact(ElementIDRemappings);
+
+					if (bSkinControlPointToTimeZero)
+					{
+						//We need to rebind the mesh at time 0. Skeleton joint have the time zero transform, so we need to apply the skinning to the mesh
+						//With the skeleton transform at time zero
+						FTransform MeshGlobalTransform;
+						MeshGlobalTransform.SetIdentity();
+						if (MeshNodeContext.SceneGlobalTransform.IsSet())
+						{
+							MeshGlobalTransform = MeshNodeContext.SceneGlobalTransform.GetValue();
+						}
+						SkinVertexPositionToTimeZero(LodMeshPayload.GetValue(), NodeContainer, RootJointNodeId, MeshGlobalTransform);
+					}
+
 					const int32 RefBoneCount = RefBonesBinary.Num();
 					//Remap the influence vertex index to point on the correct index
 					if (LodMeshPayload->JointNames.Num() > 0)
@@ -408,7 +597,6 @@
 							}
 						}
 					}
-
 					//Bake the payload, with the provide transform
 					if (MeshNodeContext.SceneGlobalTransform.IsSet())
 					{
@@ -811,55 +999,6 @@
 					continue;
 				}
 
-				
-				TArray<UE::Interchange::Private::FMeshNodeContext> MeshReferences;
-				//Scope to query the mesh node
-				{
-					TArray<FString> MeshUids;
-					LodDataNode->GetMeshUids(MeshUids);
-					MeshReferences.Reserve(MeshUids.Num());
-					for (const FString& MeshUid : MeshUids)
-					{
-						UE::Interchange::Private::FMeshNodeContext MeshReference;
-						MeshReference.MeshNode = Cast<UInterchangeMeshNode>(Arguments.NodeContainer->GetNode(MeshUid));
-						if (!MeshReference.MeshNode)
-						{
-							//The reference is a scene node and we need to bake the geometry
-							MeshReference.SceneNode = Cast<UInterchangeSceneNode>(Arguments.NodeContainer->GetNode(MeshUid));
-							if (!ensure(MeshReference.SceneNode != nullptr))
-							{
-								UE_LOG(LogInterchangeImport, Warning, TEXT("Invalid LOD mesh reference when importing SkeletalMesh asset %s"), *Arguments.AssetName);
-								continue;
-							}
-							FString MeshDependencyUid;
-							MeshReference.SceneNode->GetCustomAssetInstanceUid(MeshDependencyUid);
-							MeshReference.MeshNode = Cast<UInterchangeMeshNode>(Arguments.NodeContainer->GetNode(MeshDependencyUid));
-							//Cache the scene node global matrix, we will use this matrix to bake the vertices
-							FTransform SceneNodeGlobalTransform;
-							if (MeshReference.SceneNode->GetCustomGlobalTransform(SceneNodeGlobalTransform))
-							{
-								MeshReference.SceneGlobalTransform = SceneNodeGlobalTransform;
-							}
-						}
-						if (!ensure(MeshReference.MeshNode != nullptr))
-						{
-							UE_LOG(LogInterchangeImport, Warning, TEXT("Invalid LOD mesh reference when importing SkeletalMesh asset %s"), *Arguments.AssetName);
-							continue;
-						}
-						TOptional<FString> MeshPayloadKey = MeshReference.MeshNode->GetPayLoadKey();
-						if (MeshPayloadKey.IsSet())
-						{
-							MeshReference.TranslatorPayloadKey = MeshPayloadKey.GetValue();
-						}
-						else
-						{
-							UE_LOG(LogInterchangeImport, Warning, TEXT("Empty LOD mesh reference payload when importing SkeletalMesh asset %s"), *Arguments.AssetName);
-							continue;
-						}
-						MeshReferences.Add(MeshReference);
-					}
-				}
-
 				FString SkeletonNodeUid;
 				if (!LodDataNode->GetCustomSkeletonUid(SkeletonNodeUid))
 				{
@@ -872,14 +1011,14 @@
 					UE_LOG(LogInterchangeImport, Warning, TEXT("Invalid Skeleton LOD when importing SkeletalMesh asset %s"), *Arguments.AssetName);
 					continue;
 				}
-				
+
 				FSoftObjectPath SpecifiedSkeleton;
 				SkeletalMeshFactoryNode->GetCustomSkeletonSoftObjectPath(SpecifiedSkeleton);
 				bool bSpecifiedSkeleton = SpecifiedSkeleton.IsValid();
-				if(SkeletonReference == nullptr)
+				if (SkeletonReference == nullptr)
 				{
 					UObject* SkeletonObject = nullptr;
-					
+
 					if (SpecifiedSkeleton.IsValid())
 					{
 						SkeletonObject = SpecifiedSkeleton.TryLoad();
@@ -892,7 +1031,7 @@
 					if (SkeletonObject)
 					{
 						SkeletonReference = Cast<USkeleton>(SkeletonObject);
-						
+
 					}
 					if (!ensure(SkeletonReference))
 					{
@@ -907,24 +1046,85 @@
 					UE_LOG(LogInterchangeImport, Warning, TEXT("Invalid Skeleton LOD Root Joint when importing SkeletalMesh asset %s"), *Arguments.AssetName);
 					continue;
 				}
-				
+
 				int32 SkeletonDepth = 0;
 				TArray<SkeletalMeshImportData::FBone> RefBonesBinary;
-				UE::Interchange::Private::ProcessImportMeshSkeleton(SkeletonReference, SkeletalMesh->GetRefSkeleton(), SkeletonDepth, Arguments.NodeContainer, RootJointNodeId, RefBonesBinary);
+				bool bUseTimeZeroAsBindPose = false;
+				SkeletonNode->GetCustomUseTimeZeroForBindPose(bUseTimeZeroAsBindPose);
+				bool bDiffPose = false;
+				UE::Interchange::Private::ProcessImportMeshSkeleton(SkeletonReference, SkeletalMesh->GetRefSkeleton(), SkeletonDepth, Arguments.NodeContainer, RootJointNodeId, RefBonesBinary, bUseTimeZeroAsBindPose, bDiffPose);
 				if (bSpecifiedSkeleton && !SkeletonReference->IsCompatibleMesh(SkeletalMesh))
 				{
-					UE_LOG(LogInterchangeImport, Warning, TEXT("The skeleton %s is incompatible with the imported skeletalmesh asset %s"), *SkeletonReference->GetName() , *Arguments.AssetName);
+					UE_LOG(LogInterchangeImport, Warning, TEXT("The skeleton %s is incompatible with the imported skeletalmesh asset %s"), *SkeletonReference->GetName(), *Arguments.AssetName);
 				}
+				
+				TArray<UE::Interchange::Private::FMeshNodeContext> MeshReferences;
+				//Scope to query the mesh node
+				{
+					TArray<FString> MeshUids;
+					LodDataNode->GetMeshUids(MeshUids);
+					MeshReferences.Reserve(MeshUids.Num());
+					for (const FString& MeshUid : MeshUids)
+					{
+						UE::Interchange::Private::FMeshNodeContext MeshReference;
+						MeshReference.MeshNode = Cast<UInterchangeMeshNode>(Arguments.NodeContainer->GetNode(MeshUid));
+						if (!MeshReference.MeshNode)
+						{
+							//The reference is a scene node and we need to bake the geometry
+							MeshReference.SceneNode = Cast<UInterchangeSceneNode>(Arguments.NodeContainer->GetNode(MeshUid));
+							if (!ensure(MeshReference.SceneNode != nullptr))
+							{
+								UE_LOG(LogInterchangeImport, Warning, TEXT("Invalid LOD mesh reference when importing SkeletalMesh asset %s"), *Arguments.AssetName);
+								continue;
+							}
+							FString MeshDependencyUid;
+							MeshReference.SceneNode->GetCustomAssetInstanceUid(MeshDependencyUid);
+							MeshReference.MeshNode = Cast<UInterchangeMeshNode>(Arguments.NodeContainer->GetNode(MeshDependencyUid));
+							//Cache the scene node global matrix, we will use this matrix to bake the vertices, add the node geometric mesh offset to this matrix to bake it properly
+							FTransform SceneNodeGlobalTransform;
+							if (!bUseTimeZeroAsBindPose || !MeshReference.SceneNode->GetCustomTimeZeroGlobalTransform(SceneNodeGlobalTransform))
+							{
+								ensure(MeshReference.SceneNode->GetCustomGlobalTransform(SceneNodeGlobalTransform));
+							}
+							FTransform SceneNodeGeometricTransform;
+							if(MeshReference.SceneNode->GetCustomGeometricTransform(SceneNodeGeometricTransform))
+							{
+								SceneNodeGlobalTransform *= SceneNodeGeometricTransform;
+							}
+							MeshReference.SceneGlobalTransform = SceneNodeGlobalTransform;
+						}
+						if (!ensure(MeshReference.MeshNode != nullptr))
+						{
+							UE_LOG(LogInterchangeImport, Warning, TEXT("Invalid LOD mesh reference when importing SkeletalMesh asset %s"), *Arguments.AssetName);
+							continue;
+						}
+						TOptional<FString> MeshPayloadKey = MeshReference.MeshNode->GetPayLoadKey();
+						if (MeshPayloadKey.IsSet())
+						{
+							MeshReference.TranslatorPayloadKey = MeshPayloadKey.GetValue();
+						}
+						else
+						{
+							UE_LOG(LogInterchangeImport, Warning, TEXT("Empty LOD mesh reference payload when importing SkeletalMesh asset %s"), *Arguments.AssetName);
+							continue;
+						}
+						MeshReferences.Add(MeshReference);
+					}
+				}
+
 				//Add the lod mesh data to the skeletalmesh
 				FSkeletalMeshImportData SkeletalMeshImportData;
-				
+				const bool bSkinControlPointToTimeZero = bUseTimeZeroAsBindPose && bDiffPose;
 				//Get all meshes and blend shapes payload and fill the SkeletalMeshImportData structure
 				UE::Interchange::Private::RetrieveAllSkeletalMeshPayloadsAndFillImportData(SkeletalMeshFactoryNode
 																						   , SkeletalMeshImportData
 																						   , MeshReferences
 																						   , RefBonesBinary
 																						   , Arguments
-																						   , SkeletalMeshTranslatorPayloadInterface);
+																						   , SkeletalMeshTranslatorPayloadInterface
+																						   , bSkinControlPointToTimeZero
+																						   , Arguments.NodeContainer
+																						   , RootJointNodeId);
 				//////////////////////////////////////////////////////////////////////////
 				//Manage vertex color, we want to use the translated source data
 				//Replace -> do nothing
@@ -1040,6 +1240,10 @@
 				}
 			}
 		}
+		else //Re-Import
+		{
+			//We want to swap the source data and set properly any skeleton change
+		}
 		
 		//Getting the file Hash will cache it into the source data
 		Arguments.SourceData->GetFileContentHash();
diff --git a/Engine/Plugins/Experimental/Interchange/Runtime/Source/Nodes/Public/InterchangeSceneNode.h b/Engine/Plugins/Experimental/Interchange/Runtime/Source/Nodes/Public/InterchangeSceneNode.h
--- a/Engine/Plugins/Experimental/Interchange/Runtime/Source/Nodes/Public/InterchangeSceneNode.h
+++ b/Engine/Plugins/Experimental/Interchange/Runtime/Source/Nodes/Public/InterchangeSceneNode.h
@@ -220,32 +220,112 @@ public:
 		return MaterialDependencyUids.RemoveName(MaterialDependencyUid);
 	}
 
-	/** Return false if the Attribute was not set previously.*/
+	//Default transform is the transform we have in the node (no bind pose, no time evaluation).
+
+	/** Return the default scene node local transform. */
 	UFUNCTION(BlueprintCallable, Category = "Interchange | Node | Scene")
 	bool GetCustomLocalTransform(FTransform& AttributeValue) const
 	{
 		IMPLEMENT_NODE_ATTRIBUTE_GETTER(LocalTransform, FTransform);
 	}
 
+	/** Store the default scene node local transform. */
 	UFUNCTION(BlueprintCallable, Category = "Interchange | Node | Scene")
 	bool SetCustomLocalTransform(const FTransform& AttributeValue)
 	{
 		IMPLEMENT_NODE_ATTRIBUTE_SETTER_NODELEGATE(LocalTransform, FTransform);
 	}
 
-	/** Return false if the Attribute was not set previously.*/
+	/** Return the default scene node global transform. */
 	UFUNCTION(BlueprintCallable, Category = "Interchange | Node | Scene")
 	bool GetCustomGlobalTransform(FTransform& AttributeValue) const
 	{
 		IMPLEMENT_NODE_ATTRIBUTE_GETTER(GlobalTransform, FTransform);
 	}
 
+	/** Store the default scene node global transform. */
 	UFUNCTION(BlueprintCallable, Category = "Interchange | Node | Scene")
 	bool SetCustomGlobalTransform(const FTransform& AttributeValue)
 	{
 		IMPLEMENT_NODE_ATTRIBUTE_SETTER_NODELEGATE(GlobalTransform, FTransform);
 	}
 
+	//Bind pose transform is the transform of the joint when the binding with the mesh was done.
+	//This attribute should be set only if we have a joint.
+
+	/** Return the bind pose scene node local transform. */
+	UFUNCTION(BlueprintCallable, Category = "Interchange | Node | Joint")
+	bool GetCustomBindPoseLocalTransform(FTransform& AttributeValue) const
+	{
+		IMPLEMENT_NODE_ATTRIBUTE_GETTER(BindPoseLocalTransform, FTransform);
+	}
+
+	/** Store the bind pose scene node local transform. */
+	UFUNCTION(BlueprintCallable, Category = "Interchange | Node | Joint")
+	bool SetCustomBindPoseLocalTransform(const FTransform& AttributeValue)
+	{
+		IMPLEMENT_NODE_ATTRIBUTE_SETTER_NODELEGATE(BindPoseLocalTransform, FTransform);
+	}
+
+	/** Return the bind pose scene node global transform. */
+	UFUNCTION(BlueprintCallable, Category = "Interchange | Node | Joint")
+	bool GetCustomBindPoseGlobalTransform(FTransform& AttributeValue) const
+	{
+		IMPLEMENT_NODE_ATTRIBUTE_GETTER(BindPoseGlobalTransform, FTransform);
+	}
+
+	/** Store the bind pose scene node global transform. */
+	UFUNCTION(BlueprintCallable, Category = "Interchange | Node | Joint")
+	bool SetCustomBindPoseGlobalTransform(const FTransform& AttributeValue)
+	{
+		IMPLEMENT_NODE_ATTRIBUTE_SETTER_NODELEGATE(BindPoseGlobalTransform, FTransform);
+	}
+
+	//Time zero transform is the transform of the node at time zero.
+	//This is useful when there is no bind pose or when we import rigid mesh.
+
+	/** Return the time zero scene node local transform. */
+	UFUNCTION(BlueprintCallable, Category = "Interchange | Node | Joint")
+	bool GetCustomTimeZeroLocalTransform(FTransform& AttributeValue) const
+	{
+		IMPLEMENT_NODE_ATTRIBUTE_GETTER(TimeZeroLocalTransform, FTransform);
+	}
+
+	/** Store the time zero scene node local transform. */
+	UFUNCTION(BlueprintCallable, Category = "Interchange | Node | Joint")
+	bool SetCustomTimeZeroLocalTransform(const FTransform& AttributeValue)
+	{
+		IMPLEMENT_NODE_ATTRIBUTE_SETTER_NODELEGATE(TimeZeroLocalTransform, FTransform);
+	}
+
+	/** Return the time zero scene node global transform. */
+	UFUNCTION(BlueprintCallable, Category = "Interchange | Node | Joint")
+	bool GetCustomTimeZeroGlobalTransform(FTransform& AttributeValue) const
+	{
+		IMPLEMENT_NODE_ATTRIBUTE_GETTER(TimeZeroGlobalTransform, FTransform);
+	}
+
+	/** Store the time zero scene node global transform. */
+	UFUNCTION(BlueprintCallable, Category = "Interchange | Node | Joint")
+	bool SetCustomTimeZeroGlobalTransform(const FTransform& AttributeValue)
+	{
+		IMPLEMENT_NODE_ATTRIBUTE_SETTER_NODELEGATE(TimeZeroGlobalTransform, FTransform);
+	}
+
+	/** Return the geometric offset. Any mesh attach to this scene node will be offset using this transform. */
+	UFUNCTION(BlueprintCallable, Category = "Interchange | Node | Joint")
+	bool GetCustomGeometricTransform(FTransform& AttributeValue) const
+	{
+		IMPLEMENT_NODE_ATTRIBUTE_GETTER(GeometricTransform, FTransform);
+	}
+
+	/** Store the geometric offset. Any mesh attach to this scene node will be offset using this transform. */
+	UFUNCTION(BlueprintCallable, Category = "Interchange | Node | Joint")
+	bool SetCustomGeometricTransform(const FTransform& AttributeValue)
+	{
+		IMPLEMENT_NODE_ATTRIBUTE_SETTER_NODELEGATE(GeometricTransform, FTransform);
+	}
+
 	/** Return false if the Attribute was not set previously.*/
 	UFUNCTION(BlueprintCallable, Category = "Interchange | Node | Scene")
 	bool GetCustomAssetInstanceUid(FString& AttributeValue) const
@@ -264,6 +344,11 @@ private:
 	//Scene Attribute Keys
 	const UE::Interchange::FAttributeKey Macro_CustomLocalTransformKey = UE::Interchange::FAttributeKey(TEXT("LocalTransform"));
 	const UE::Interchange::FAttributeKey Macro_CustomGlobalTransformKey = UE::Interchange::FAttributeKey(TEXT("GlobalTransform"));
+	const UE::Interchange::FAttributeKey Macro_CustomBindPoseLocalTransformKey = UE::Interchange::FAttributeKey(TEXT("BindPoseLocalTransform"));
+	const UE::Interchange::FAttributeKey Macro_CustomBindPoseGlobalTransformKey = UE::Interchange::FAttributeKey(TEXT("BindPoseGlobalTransform"));
+	const UE::Interchange::FAttributeKey Macro_CustomTimeZeroLocalTransformKey = UE::Interchange::FAttributeKey(TEXT("TimeZeroLocalTransform"));
+	const UE::Interchange::FAttributeKey Macro_CustomTimeZeroGlobalTransformKey = UE::Interchange::FAttributeKey(TEXT("TimeZeroGlobalTransform"));
+	const UE::Interchange::FAttributeKey Macro_CustomGeometricTransformKey = UE::Interchange::FAttributeKey(TEXT("GeometricTransform"));
 	const UE::Interchange::FAttributeKey Macro_CustomAssetInstanceUidKey = UE::Interchange::FAttributeKey(TEXT("AssetInstanceUid"));
 
 	UE::Interchange::FNameAttributeArrayHelper NodeSpecializeTypes;
diff --git a/Engine/Plugins/Experimental/Interchange/Runtime/Source/Nodes/Public/InterchangeSkeletonFactoryNode.h b/Engine/Plugins/Experimental/Interchange/Runtime/Source/Nodes/Public/InterchangeSkeletonFactoryNode.h
--- a/Engine/Plugins/Experimental/Interchange/Runtime/Source/Nodes/Public/InterchangeSkeletonFactoryNode.h
+++ b/Engine/Plugins/Experimental/Interchange/Runtime/Source/Nodes/Public/InterchangeSkeletonFactoryNode.h
@@ -94,6 +94,20 @@ public:
 		IMPLEMENT_NODE_ATTRIBUTE_SETTER_NODELEGATE(RootJointID, FString)
 	}
 
+	/** Query if this skeleton should replace joint transform with time zero evaluation instead of bind pose.*/
+	UFUNCTION(BlueprintCallable, Category = "Interchange | Node | Skeleton")
+	bool GetCustomUseTimeZeroForBindPose(bool& AttributeValue) const
+	{
+		IMPLEMENT_NODE_ATTRIBUTE_GETTER(UseTimeZeroForBindPose, bool);
+	}
+
+	/** If AttributeValue is true, force this skeleton bind pose to use time zero instead of bind pose.*/
+	UFUNCTION(BlueprintCallable, Category = "Interchange | Node | Skeleton")
+	bool SetCustomUseTimeZeroForBindPose(const bool& AttributeValue)
+	{
+		IMPLEMENT_NODE_ATTRIBUTE_SETTER_NODELEGATE(UseTimeZeroForBindPose, bool)
+	}
+
 private:
 
 	void FillAssetClassFromAttribute()
@@ -132,6 +146,7 @@ private:
 
 	//Skeleton
 	const UE::Interchange::FAttributeKey Macro_CustomRootJointIDKey = UE::Interchange::FAttributeKey(TEXT("RootJointID"));
+	const UE::Interchange::FAttributeKey Macro_CustomUseTimeZeroForBindPoseKey = UE::Interchange::FAttributeKey(TEXT("UseTimeZeroForBindPose"));
 	
 protected:
 #if WITH_ENGINE
diff --git a/Engine/Plugins/Experimental/Interchange/Runtime/Source/Parsers/Fbx/Private/FbxScene.cpp b/Engine/Plugins/Experimental/Interchange/Runtime/Source/Parsers/Fbx/Private/FbxScene.cpp
--- a/Engine/Plugins/Experimental/Interchange/Runtime/Source/Parsers/Fbx/Private/FbxScene.cpp
+++ b/Engine/Plugins/Experimental/Interchange/Runtime/Source/Parsers/Fbx/Private/FbxScene.cpp
@@ -25,7 +25,7 @@ namespace UE
 		namespace Private
 		{
 
-			void FFbxScene::CreateMeshNodeReference(UInterchangeSceneNode* UnrealSceneNode, FbxNodeAttribute* NodeAttribute, UInterchangeBaseNodeContainer& NodeContainer)
+			void FFbxScene::CreateMeshNodeReference(UInterchangeSceneNode* UnrealSceneNode, FbxNodeAttribute* NodeAttribute, UInterchangeBaseNodeContainer& NodeContainer, const FTransform& GeometricTransform)
 			{
 				UInterchangeMeshNode* MeshNode = nullptr;
 				if (NodeAttribute->GetAttributeType() == FbxNodeAttribute::eMesh)
@@ -46,6 +46,7 @@ namespace UE
 				{
 					UnrealSceneNode->SetCustomAssetInstanceUid(MeshNode->GetUniqueID());
 					MeshNode->SetSceneInstanceUid(UnrealSceneNode->GetUniqueID());
+					UnrealSceneNode->SetCustomGeometricTransform(GeometricTransform);
 				}
 			}
 
@@ -107,16 +108,23 @@ namespace UE
 					return Transform;
 				};
 
+				//Set the node default transform
 				{
-					//Set the global node transform
 					FbxAMatrix GlobalFbxMatrix = Node->EvaluateGlobalTransform();
 					FTransform GlobalTransform = GetConvertedTransform(GlobalFbxMatrix);
 					UnrealNode->SetCustomGlobalTransform(GlobalTransform);
-
-					//Set the local node transform
-					FbxAMatrix LocalFbxMatrix = Node->EvaluateLocalTransform();
-					FTransform LocalTransform = GetConvertedTransform(LocalFbxMatrix);
-					UnrealNode->SetCustomLocalTransform(LocalTransform);
+					if (FbxNode* ParentNode = Node->GetParent())
+					{
+						FbxAMatrix GlobalFbxParentMatrix = ParentNode->EvaluateGlobalTransform();
+						FbxAMatrix	LocalFbxMatrix = GlobalFbxParentMatrix.Inverse() * GlobalFbxMatrix;
+						FTransform LocalTransform = GetConvertedTransform(LocalFbxMatrix);
+						UnrealNode->SetCustomLocalTransform(LocalTransform);
+					}
+					else
+					{
+						//No parent, set the same matrix has the global
+						UnrealNode->SetCustomLocalTransform(GlobalTransform);
+					}
 				}
 
 				int32 AttributeCount = Node->GetNodeAttributeCount();
@@ -152,12 +160,12 @@ namespace UE
 						{
 							//Add the joint specialized type
 							UnrealNode->AddSpecializedType(FSceneNodeStaticData::GetJointSpecializeTypeString());
-							//Get the bind pose for this joint
+							//Get the bind pose transform for this joint
 							FbxAMatrix GlobalBindPoseJointMatrix;
 							if (FFbxMesh::GetGlobalJointBindPoseTransform(SDKScene, Node, GlobalBindPoseJointMatrix))
 							{
 								FTransform GlobalBindPoseJointTransform = GetConvertedTransform(GlobalBindPoseJointMatrix);
-								UnrealNode->SetCustomGlobalTransform(GlobalBindPoseJointTransform);
+								UnrealNode->SetCustomBindPoseGlobalTransform(GlobalBindPoseJointTransform);
 								//We grab the fbx parent node to compute the local transform
 								if (FbxNode* ParentNode = Node->GetParent())
 								{
@@ -165,14 +173,35 @@ namespace UE
 									FFbxMesh::GetGlobalJointBindPoseTransform(SDKScene, ParentNode, GlobalFbxParentMatrix);
 									FbxAMatrix	LocalFbxMatrix = GlobalFbxParentMatrix.Inverse() * GlobalBindPoseJointMatrix;
 									FTransform LocalBindPoseJointTransform = GetConvertedTransform(LocalFbxMatrix);
-									UnrealNode->SetCustomLocalTransform(LocalBindPoseJointTransform);
+									UnrealNode->SetCustomBindPoseLocalTransform(LocalBindPoseJointTransform);
+								}
+								else
+								{
+									//No parent, set the same matrix has the global
+									UnrealNode->SetCustomBindPoseLocalTransform(GlobalBindPoseJointTransform);
+								}
+							}
+
+							//Get time Zero transform for this joint
+							{
+								//Set the global node transform
+								FbxAMatrix GlobalFbxMatrix = Node->EvaluateGlobalTransform(FBXSDK_TIME_ZERO);
+								FTransform GlobalTransform = GetConvertedTransform(GlobalFbxMatrix);
+								UnrealNode->SetCustomTimeZeroGlobalTransform(GlobalTransform);
+								if (FbxNode* ParentNode = Node->GetParent())
+								{
+									FbxAMatrix GlobalFbxParentMatrix = ParentNode->EvaluateGlobalTransform(FBXSDK_TIME_ZERO);
+									FbxAMatrix	LocalFbxMatrix = GlobalFbxParentMatrix.Inverse() * GlobalFbxMatrix;
+									FTransform LocalTransform = GetConvertedTransform(LocalFbxMatrix);
+									UnrealNode->SetCustomTimeZeroLocalTransform(LocalTransform);
 								}
 								else
 								{
 									//No parent, set the same matrix has the global
-									UnrealNode->SetCustomLocalTransform(GlobalBindPoseJointTransform);
+									UnrealNode->SetCustomTimeZeroLocalTransform(GlobalTransform);
 								}
 							}
+
 							break;
 						}
 
@@ -181,7 +210,18 @@ namespace UE
 							//For Mesh attribute we add the fbx nodes materials
 							FFbxMaterial FbxMaterial(Parser);
 							FbxMaterial.AddAllNodeMaterials(UnrealNode, Node, NodeContainer);
-							CreateMeshNodeReference(UnrealNode, NodeAttribute, NodeContainer);
+							//Get the Geometric offset transform and set it in the mesh node
+							//The geometric offset is not part of the hierarchy transform, it is not inherited
+							FbxAMatrix Geometry;
+							FbxVector4 Translation, Rotation, Scaling;
+							Translation = Node->GetGeometricTranslation(FbxNode::eSourcePivot);
+							Rotation = Node->GetGeometricRotation(FbxNode::eSourcePivot);
+							Scaling = Node->GetGeometricScaling(FbxNode::eSourcePivot);
+							Geometry.SetT(Translation);
+							Geometry.SetR(Rotation);
+							Geometry.SetS(Scaling);
+							FTransform GeometricTransform = GetConvertedTransform(Geometry);
+							CreateMeshNodeReference(UnrealNode, NodeAttribute, NodeContainer, GeometricTransform);
 							break;
 						}
 						case FbxNodeAttribute::eLODGroup:
diff --git a/Engine/Plugins/Experimental/Interchange/Runtime/Source/Parsers/Fbx/Private/FbxScene.h b/Engine/Plugins/Experimental/Interchange/Runtime/Source/Parsers/Fbx/Private/FbxScene.h
--- a/Engine/Plugins/Experimental/Interchange/Runtime/Source/Parsers/Fbx/Private/FbxScene.h
+++ b/Engine/Plugins/Experimental/Interchange/Runtime/Source/Parsers/Fbx/Private/FbxScene.h
@@ -28,7 +28,7 @@ namespace UE
 				UInterchangeSceneNode* CreateTransformNode(UInterchangeBaseNodeContainer& NodeContainer, const FString& NodeName, const FString& NodeUniqueID);
 
 			protected:
-				void CreateMeshNodeReference(UInterchangeSceneNode* UnrealSceneNode, FbxNodeAttribute* NodeAttribute, UInterchangeBaseNodeContainer& NodeContainer);
+				void CreateMeshNodeReference(UInterchangeSceneNode* UnrealSceneNode, FbxNodeAttribute* NodeAttribute, UInterchangeBaseNodeContainer& NodeContainer, const FTransform& GeometricTransform);
 				void CreateCameraNodeReference(UInterchangeSceneNode* UnrealSceneNode, FbxNodeAttribute* NodeAttribute, UInterchangeBaseNodeContainer& NodeContainer);
 				void CreateLightNodeReference(UInterchangeSceneNode* UnrealSceneNode, FbxNodeAttribute* NodeAttribute, UInterchangeBaseNodeContainer& NodeContainer);
 				void AddHierarchyRecursively(UInterchangeSceneNode* UnrealParentNode, FbxNode* Node, FbxScene* SDKScene, UInterchangeBaseNodeContainer& NodeContainer);
diff --git a/Engine/Plugins/Experimental/Interchange/Runtime/Source/Pipelines/Private/InterchangeGenericSkeletalMeshPipeline.cpp b/Engine/Plugins/Experimental/Interchange/Runtime/Source/Pipelines/Private/InterchangeGenericSkeletalMeshPipeline.cpp
--- a/Engine/Plugins/Experimental/Interchange/Runtime/Source/Pipelines/Private/InterchangeGenericSkeletalMeshPipeline.cpp
+++ b/Engine/Plugins/Experimental/Interchange/Runtime/Source/Pipelines/Private/InterchangeGenericSkeletalMeshPipeline.cpp
@@ -102,7 +102,7 @@ namespace UE::Interchange::SkeletalMeshGenericPipeline
 		return true;
 	}
 
-	void RecursiveBuildSkeletalSkeleton(const FString JoinToAddUid, const int32 ParentIndex, const UInterchangeBaseNodeContainer* BaseNodeContainer, TArray<FMeshBoneInfo>& SkeletalLodRawInfos, TArray<FTransform>& SkeletalLodRawTransforms)
+	void RecursiveBuildSkeletalSkeleton(const FString JoinToAddUid, const int32 ParentIndex, const UInterchangeBaseNodeContainer* BaseNodeContainer, TArray<FMeshBoneInfo>& SkeletalLodRawInfos)
 	{
 		const UInterchangeSceneNode* SceneNode = Cast<UInterchangeSceneNode>(BaseNodeContainer->GetNode(JoinToAddUid));
 		if (!SceneNode || !SceneNode->IsSpecializedTypeContains(FSceneNodeStaticData::GetJointSpecializeTypeString()))
@@ -117,14 +117,11 @@ namespace UE::Interchange::SkeletalMeshGenericPipeline
 #if WITH_EDITORONLY_DATA
 		Info.ExportName = Info.Name.ToString();
 #endif
-		FTransform& JoinTransform = SkeletalLodRawTransforms.AddZeroed_GetRef();
-		SceneNode->GetCustomLocalTransform(JoinTransform);
-
 		//Iterate childrens
 		const TArray<FString> ChildrenIds = BaseNodeContainer->GetNodeChildrenUids(JoinToAddUid);
 		for (int32 ChildIndex = 0; ChildIndex < ChildrenIds.Num(); ++ChildIndex)
 		{
-			RecursiveBuildSkeletalSkeleton(ChildrenIds[ChildIndex], JoinIndex, BaseNodeContainer, SkeletalLodRawInfos, SkeletalLodRawTransforms);
+			RecursiveBuildSkeletalSkeleton(ChildrenIds[ChildIndex], JoinIndex, BaseNodeContainer, SkeletalLodRawInfos);
 		}
 	}
 
@@ -138,9 +135,7 @@ namespace UE::Interchange::SkeletalMeshGenericPipeline
 
 		TArray<FMeshBoneInfo> SkeletalLodRawInfos;
 		SkeletalLodRawInfos.Reserve(SkeletonBoneCount);
-		TArray<FTransform> SkeletalLodRawTransforms;
-		SkeletalLodRawTransforms.Reserve(SkeletonBoneCount);
-		RecursiveBuildSkeletalSkeleton(RootJoinUid, INDEX_NONE, BaseNodeContainer, SkeletalLodRawInfos, SkeletalLodRawTransforms);
+		RecursiveBuildSkeletalSkeleton(RootJoinUid, INDEX_NONE, BaseNodeContainer, SkeletalLodRawInfos);
 		const int32 SkeletalLodBoneCount = SkeletalLodRawInfos.Num();
 
 		// first ensure the parent exists for each bone
@@ -442,6 +437,13 @@ UInterchangeSkeletonFactoryNode* UInterchangeGenericAssetsPipeline::CreateSkelet
 			//Log an error
 			return nullptr;
 		}
+		FString ExistingSkeletonRootJointUid;
+		SkeletonFactoryNode->GetCustomRootJointUid(ExistingSkeletonRootJointUid);
+		if (!ensure(ExistingSkeletonRootJointUid.Equals(RootJointUid)))
+		{
+			//Log an error
+			return nullptr;
+		}
 	}
 	else
 	{
@@ -452,6 +454,7 @@ UInterchangeSkeletonFactoryNode* UInterchangeGenericAssetsPipeline::CreateSkelet
 		}
 		SkeletonFactoryNode->InitializeSkeletonNode(SkeletonUid, DisplayLabel, USkeleton::StaticClass()->GetName());
 		SkeletonFactoryNode->SetCustomRootJointUid(RootJointNode->GetUniqueID());
+		SkeletonFactoryNode->SetCustomUseTimeZeroForBindPose(bUseT0AsRefPose);
 		BaseNodeContainer->AddNode(SkeletonFactoryNode);
 	}
 
diff --git a/Engine/Plugins/Experimental/Interchange/Runtime/Source/Pipelines/Public/InterchangeGenericAssetsPipeline.h b/Engine/Plugins/Experimental/Interchange/Runtime/Source/Pipelines/Public/InterchangeGenericAssetsPipeline.h
--- a/Engine/Plugins/Experimental/Interchange/Runtime/Source/Pipelines/Public/InterchangeGenericAssetsPipeline.h
+++ b/Engine/Plugins/Experimental/Interchange/Runtime/Source/Pipelines/Public/InterchangeGenericAssetsPipeline.h
@@ -155,9 +155,13 @@ public:
 	UPROPERTY(EditAnywhere, BlueprintReadWrite, Category = SKELETAL_MESHES_CATEGORY)
 	bool bImportMeshesInBoneHierarchy = true;
 
+	/** Enable this option to use frame 0 as reference pose */
+	UPROPERTY(EditAnywhere, BlueprintReadWrite, Category = SKELETAL_MESHES_CATEGORY)
+	bool bUseT0AsRefPose = false;
+
 	/** If checked, create new PhysicsAsset if it doesn't have it */
 	UPROPERTY(EditAnywhere, BlueprintReadWrite, Category = SKELETAL_MESHES_CATEGORY)
-	bool bCreatePhysicsAsset;
+	bool bCreatePhysicsAsset = true;
 
 	/** If this is set, use this specified PhysicsAsset. If its not set and bCreatePhysicsAsset is false, the importer will not generate or set any physic asset. */
 	UPROPERTY(EditAnywhere, BlueprintReadWrite, Category = SKELETAL_MESHES_CATEGORY, meta = (editcondition = "!bCreatePhysicsAsset"))
