Fix interchange skeletal mesh material index issue when importing several LODs in the same file
#rnx

diff --git a/Engine/Plugins/Experimental/Interchange/Runtime/Source/Import/Private/Mesh/InterchangeSkeletalMeshFactory.cpp b/Engine/Plugins/Experimental/Interchange/Runtime/Source/Import/Private/Mesh/InterchangeSkeletalMeshFactory.cpp
--- a/Engine/Plugins/Experimental/Interchange/Runtime/Source/Import/Private/Mesh/InterchangeSkeletalMeshFactory.cpp
+++ b/Engine/Plugins/Experimental/Interchange/Runtime/Source/Import/Private/Mesh/InterchangeSkeletalMeshFactory.cpp
@@ -1198,6 +1198,36 @@ UObject* UInterchangeSkeletalMeshFactory::CreateAsset(const FCreateAssetParams&
 				NewLODInfo.LODHysteresis = 0.02f;
 				NewLODInfo.bImportWithBaseMesh = true;
 
+				//Make sure each LOD material index point on the correct skeletalMesh materials list index.
+				TArray<FSkeletalMaterial>& Materials = SkeletalMesh->GetMaterials();
+				const TArray<SkeletalMeshImportData::FMaterial>& ImportedMaterials = SkeletalMeshImportData.Materials;
+				if (FSkeletalMeshLODInfo* LodInfo = SkeletalMesh->GetLODInfo(CurrentLodIndex))
+				{
+					LodInfo->LODMaterialMap.Empty();
+					// Now set up the material mapping array.
+					for (int32 ImportedMaterialIndex = 0; ImportedMaterialIndex < ImportedMaterials.Num(); ImportedMaterialIndex++)
+					{
+						FName ImportedMaterialName = *(ImportedMaterials[ImportedMaterialIndex].MaterialImportName);
+						//Match by name
+						int32 LODMatIndex = INDEX_NONE;
+						for (int32 MaterialIndex = 0; MaterialIndex < Materials.Num(); ++MaterialIndex)
+						{
+							const FSkeletalMaterial& SkeletalMaterial = Materials[MaterialIndex];
+							if (SkeletalMaterial.ImportedMaterialSlotName != NAME_None && SkeletalMaterial.ImportedMaterialSlotName == ImportedMaterialName)
+							{
+								LODMatIndex = MaterialIndex;
+								break;
+							}
+						}
+						// If we dont have a match, add a new entry to the material list.
+						if (LODMatIndex == INDEX_NONE)
+						{
+							LODMatIndex = Materials.Add(FSkeletalMaterial(ImportedMaterials[ImportedMaterialIndex].Material.Get(), true, false, ImportedMaterialName, ImportedMaterialName));
+						}
+						LodInfo->LODMaterialMap.Add(LODMatIndex);
+					}
+				}
+
 				//Add the bound to the skeletal mesh
 				if (SkeletalMesh->GetImportedBounds().BoxExtent.IsNearlyZero())
 				{
diff --git a/Engine/Source/Runtime/Engine/Private/SkeletalMesh.cpp b/Engine/Source/Runtime/Engine/Private/SkeletalMesh.cpp
--- a/Engine/Source/Runtime/Engine/Private/SkeletalMesh.cpp
+++ b/Engine/Source/Runtime/Engine/Private/SkeletalMesh.cpp
@@ -4627,7 +4627,7 @@ namespace InternalSkeletalMeshHelper
 				continue;
 			}
 			const TArray<int16>* BackupSectionsPtr = BackupSectionsPerLOD.Find(LODIndex);
-			if (!BackupSectionsPtr)
+			if (!BackupSectionsPtr || BackupSectionsPtr->IsEmpty())
 			{
 				continue;
 			}
