Make sure the material imported by interchange are used by a mesh we import, do not import unused material when we use the generic asset pipeline.

#rnx

diff --git a/Engine/Plugins/Experimental/Interchange/Runtime/Source/Import/Private/Mesh/InterchangeSkeletalMeshFactory.cpp b/Engine/Plugins/Experimental/Interchange/Runtime/Source/Import/Private/Mesh/InterchangeSkeletalMeshFactory.cpp
--- a/Engine/Plugins/Experimental/Interchange/Runtime/Source/Import/Private/Mesh/InterchangeSkeletalMeshFactory.cpp
+++ b/Engine/Plugins/Experimental/Interchange/Runtime/Source/Import/Private/Mesh/InterchangeSkeletalMeshFactory.cpp
@@ -1177,6 +1177,10 @@ UObject* UInterchangeSkeletalMeshFactory::CreateAsset(const FCreateAssetParams&
 					{
 						continue;
 					}
+					if (!MaterialFactoryNode->IsEnabled())
+					{
+						continue;
+					}
 					UMaterialInterface* MaterialInterface = Cast<UMaterialInterface>(MaterialFactoryNode->ReferenceObject.ResolveObject());
 					if (!MaterialInterface)
 					{
diff --git a/Engine/Plugins/Experimental/Interchange/Runtime/Source/Import/Private/Mesh/InterchangeStaticMeshFactory.cpp b/Engine/Plugins/Experimental/Interchange/Runtime/Source/Import/Private/Mesh/InterchangeStaticMeshFactory.cpp
--- a/Engine/Plugins/Experimental/Interchange/Runtime/Source/Import/Private/Mesh/InterchangeStaticMeshFactory.cpp
+++ b/Engine/Plugins/Experimental/Interchange/Runtime/Source/Import/Private/Mesh/InterchangeStaticMeshFactory.cpp
@@ -318,6 +318,10 @@ UObject* UInterchangeStaticMeshFactory::CreateAsset(const FCreateAssetParams& Ar
 					{
 						continue;
 					}
+					if (!MaterialFactoryNode->IsEnabled())
+					{
+						continue;
+					}
 
 					FName MaterialSlotName = *MaterialFactoryNode->GetDisplayLabel();
 					UMaterialInterface* MaterialInterface = Cast<UMaterialInterface>(MaterialFactoryNode->ReferenceObject.ResolveObject());
diff --git a/Engine/Plugins/Experimental/Interchange/Runtime/Source/Parsers/Fbx/Private/FbxMesh.cpp b/Engine/Plugins/Experimental/Interchange/Runtime/Source/Parsers/Fbx/Private/FbxMesh.cpp
--- a/Engine/Plugins/Experimental/Interchange/Runtime/Source/Parsers/Fbx/Private/FbxMesh.cpp
+++ b/Engine/Plugins/Experimental/Interchange/Runtime/Source/Parsers/Fbx/Private/FbxMesh.cpp
@@ -203,7 +203,7 @@ namespace UE
 			
 			bool FMeshDescriptionImporter::FillStaticMeshDescriptionFromFbxMesh(FbxMesh* Mesh)
 			{
-				if (!ensure(bInitialized) || !ensure(MeshDescription) || !ensure(Mesh) || !ensure(Mesh->GetDeformerCount() == 0))
+				if (!ensure(bInitialized) || !ensure(MeshDescription) || !ensure(Mesh))
 				{
 					return false;
 				}
@@ -1233,6 +1233,35 @@ namespace UE
 			//////////////////////////////////////////////////////////////////////////
 			/// FFbxMesh implementation
 
+			void GetMaterialIndex(FbxMesh* Mesh, TArray<int32>& MaterialIndexes)
+			{
+				int32 PolygonCount = Mesh->GetPolygonCount();
+				if (FbxGeometryElementMaterial* GeometryElementMaterial = Mesh->GetElementMaterial())
+				{
+					FbxLayerElementArrayTemplate<int32>& IndexArray = GeometryElementMaterial->GetIndexArray();
+					MaterialIndexes.Reset();
+					switch (GeometryElementMaterial->GetMappingMode())
+					{
+						case FbxGeometryElement::eByPolygon:
+						{
+							if (IndexArray.GetCount() == PolygonCount)
+							{
+								for (int i = 0; i < PolygonCount; ++i)
+								{
+									MaterialIndexes.AddUnique(IndexArray.GetAt(i));
+								}
+							}
+						}
+						break;
+						case FbxGeometryElement::eAllSame:
+						{
+							MaterialIndexes.AddUnique(IndexArray.GetAt(0));
+						}
+						break;
+					}
+				}
+			}
+
 			void FFbxMesh::AddAllMeshes(FbxScene* SDKScene, FbxGeometryConverter* SDKGeometryConverter, UInterchangeBaseNodeContainer& NodeContainer, TMap<FString, TSharedPtr<FPayloadContextBase>>& PayloadContexts)
 			{
 				int32 GeometryCount = SDKScene->GetGeometryCount();
@@ -1280,16 +1309,52 @@ namespace UE
 						const int32 MeshUVCount = Mesh->GetElementUVCount();
 						MeshNode->SetCustomUVCount(MeshUVCount);
 						
-						//Add Material dependencies, we use always the first fbx node that instance the fbx geometry to grab the materials.
-						if (FbxNode* FbxMeshNode = Mesh->GetNode())
+						//Add Material dependencies, we use always the first fbx node that instance the fbx geometry to grab the fbx surface materials.
 						{
-							const int32 MaterialCount = FbxMeshNode->GetMaterialCount();
-							for (int32 MaterialIndex = 0; MaterialIndex < MaterialCount; ++MaterialIndex)
+							//Grab all Material indexes use by the mesh
+							TArray<int32> MaterialIndexes;
+							int32 PolygonCount = Mesh->GetPolygonCount();
+							if (FbxGeometryElementMaterial* GeometryElementMaterial = Mesh->GetElementMaterial())
 							{
-								FbxSurfaceMaterial* FbxMaterial = FbxMeshNode->GetMaterial(MaterialIndex);
-								FString MaterialName = FFbxHelper::GetFbxObjectName(FbxMaterial);
-								FString MaterialUid = TEXT("\\Material\\") + MaterialName;
-								MeshNode->SetMaterialDependencyUid(MaterialUid);
+								FbxLayerElementArrayTemplate<int32>& IndexArray = GeometryElementMaterial->GetIndexArray();
+								switch (GeometryElementMaterial->GetMappingMode())
+								{
+									case FbxGeometryElement::eByPolygon:
+									{
+										if (IndexArray.GetCount() == PolygonCount)
+										{
+											for (int i = 0; i < PolygonCount; ++i)
+											{
+												MaterialIndexes.AddUnique(IndexArray.GetAt(i));
+											}
+										}
+									}
+									break;
+
+									case FbxGeometryElement::eAllSame:
+									{
+										if (IndexArray.GetCount() > 0)
+										{
+											MaterialIndexes.AddUnique(IndexArray.GetAt(0));
+										}
+									}
+									break;
+								}
+							}
+							if (FbxNode* FbxMeshNode = Mesh->GetNode())
+							{
+								bool bAddAllNodeMaterials = (MaterialIndexes.Num() == 0);
+								const int32 MaterialCount = FbxMeshNode->GetMaterialCount();
+								for (int32 MaterialIndex = 0; MaterialIndex < MaterialCount; ++MaterialIndex)
+								{
+									FbxSurfaceMaterial* FbxMaterial = FbxMeshNode->GetMaterial(MaterialIndex);
+									FString MaterialName = FFbxHelper::GetFbxObjectName(FbxMaterial);
+									FString MaterialUid = TEXT("\\Material\\") + MaterialName;
+									if (bAddAllNodeMaterials || MaterialIndexes.Contains(MaterialIndex))
+									{
+										MeshNode->SetMaterialDependencyUid(MaterialUid);
+									}
+								}
 							}
 						}
 					}
diff --git a/Engine/Plugins/Experimental/Interchange/Runtime/Source/Pipelines/Private/InterchangeGenericAssetsPipeline.cpp b/Engine/Plugins/Experimental/Interchange/Runtime/Source/Pipelines/Private/InterchangeGenericAssetsPipeline.cpp
--- a/Engine/Plugins/Experimental/Interchange/Runtime/Source/Pipelines/Private/InterchangeGenericAssetsPipeline.cpp
+++ b/Engine/Plugins/Experimental/Interchange/Runtime/Source/Pipelines/Private/InterchangeGenericAssetsPipeline.cpp
@@ -83,7 +83,11 @@ bool UInterchangeGenericAssetsPipeline::ExecutePreImportPipeline(UInterchangeBas
 		//import materials
 		for (const UInterchangeMaterialNode* MaterialNode : MaterialNodes)
 		{
-			CreateMaterialFactoryNode(MaterialNode);
+			if (UInterchangeMaterialFactoryNode* MaterialFactoryNode = CreateMaterialFactoryNode(MaterialNode))
+			{
+				//By default we do not create the materials, every node with mesh attribute can enable them. So we wont create unused materials.
+				MaterialFactoryNode->SetEnabled(false);
+			}
 		}
 	}
 
@@ -488,11 +492,15 @@ void UInterchangeGenericAssetsPipeline::AddLodDataToStaticMesh(UInterchangeStati
 				SceneNode->GetCustomAssetInstanceUid(MeshDependency);
 				if (BaseNodeContainer->IsNodeUidValid(MeshDependency))
 				{
+					const UInterchangeMeshNode* MeshDependencyNode = Cast<UInterchangeMeshNode>(BaseNodeContainer->GetNode(MeshDependency));
 					StaticMeshFactoryNode->AddTargetNodeUid(MeshDependency);
-					BaseNodeContainer->GetNode(MeshDependency)->AddTargetNodeUid(StaticMeshFactoryNode->GetUniqueID());
+					MeshDependencyNode->AddTargetNodeUid(StaticMeshFactoryNode->GetUniqueID());
+					MeshDependencyNode->GetMaterialDependencies(MaterialDependencies);
+				}
+				else
+				{
+					SceneNode->GetMaterialDependencyUids(MaterialDependencies);
 				}
-				
-				SceneNode->GetMaterialDependencyUids(MaterialDependencies);
 			}
 			else if (const UInterchangeMeshNode* MeshNode = Cast<UInterchangeMeshNode>(BaseNodeContainer->GetNode(NodeUid)))
 			{
@@ -507,6 +515,7 @@ void UInterchangeGenericAssetsPipeline::AddLodDataToStaticMesh(UInterchangeStati
 				const FString MaterialFactoryNodeUid = UInterchangeMaterialFactoryNode::GetMaterialFactoryNodeUidFromMaterialNodeUid(MaterialDependencies[MaterialIndex]);
 				if (BaseNodeContainer->IsNodeUidValid(MaterialFactoryNodeUid))
 				{
+					BaseNodeContainer->GetNode(MaterialFactoryNodeUid)->SetEnabled(true);
 					// Create a factory dependency so Material asset are imported before the skeletal mesh asset
 					TArray<FString> FactoryDependencies;
 					StaticMeshFactoryNode->GetFactoryDependencies(FactoryDependencies);
diff --git a/Engine/Plugins/Experimental/Interchange/Runtime/Source/Pipelines/Private/InterchangeGenericSkeletalMeshPipeline.cpp b/Engine/Plugins/Experimental/Interchange/Runtime/Source/Pipelines/Private/InterchangeGenericSkeletalMeshPipeline.cpp
--- a/Engine/Plugins/Experimental/Interchange/Runtime/Source/Pipelines/Private/InterchangeGenericSkeletalMeshPipeline.cpp
+++ b/Engine/Plugins/Experimental/Interchange/Runtime/Source/Pipelines/Private/InterchangeGenericSkeletalMeshPipeline.cpp
@@ -691,11 +691,15 @@ void UInterchangeGenericAssetsPipeline::AddLodDataToSkeletalMesh(const UIntercha
 				SceneNode->GetCustomAssetInstanceUid(MeshDependency);
 				if (BaseNodeContainer->IsNodeUidValid(MeshDependency))
 				{
+					const UInterchangeMeshNode* MeshDependencyNode = Cast<UInterchangeMeshNode>(BaseNodeContainer->GetNode(MeshDependency));
 					SkeletalMeshFactoryNode->AddTargetNodeUid(MeshDependency);
-					BaseNodeContainer->GetNode(MeshDependency)->AddTargetNodeUid(SkeletalMeshFactoryNode->GetUniqueID());
+					MeshDependencyNode->AddTargetNodeUid(SkeletalMeshFactoryNode->GetUniqueID());
+					MeshDependencyNode->GetMaterialDependencies(MaterialDependencies);
+				}
+				else
+				{
+					SceneNode->GetMaterialDependencyUids(MaterialDependencies);
 				}
-
-				SceneNode->GetMaterialDependencyUids(MaterialDependencies);
 			}
 			else if (const UInterchangeMeshNode* MeshNode = Cast<UInterchangeMeshNode>(BaseNodeContainer->GetNode(NodeUid)))
 			{
@@ -709,6 +713,7 @@ void UInterchangeGenericAssetsPipeline::AddLodDataToSkeletalMesh(const UIntercha
 				const FString MaterialFactoryNodeUid = UInterchangeMaterialFactoryNode::GetMaterialFactoryNodeUidFromMaterialNodeUid(MaterialDependencies[MaterialIndex]);
 				if (BaseNodeContainer->IsNodeUidValid(MaterialFactoryNodeUid))
 				{
+					BaseNodeContainer->GetNode(MaterialFactoryNodeUid)->SetEnabled(true);
 					//Create a factory dependency so Material asset are import before the skeletal mesh asset
 					TArray<FString> FactoryDependencies;
 					SkeletalMeshFactoryNode->GetFactoryDependencies(FactoryDependencies);
