Fix CIS missing include
#rnx

diff --git a/Engine/Plugins/Experimental/Interchange/Runtime/Source/Import/Private/Mesh/InterchangeStaticMeshFactory.cpp b/Engine/Plugins/Experimental/Interchange/Runtime/Source/Import/Private/Mesh/InterchangeStaticMeshFactory.cpp
--- a/Engine/Plugins/Experimental/Interchange/Runtime/Source/Import/Private/Mesh/InterchangeStaticMeshFactory.cpp
+++ b/Engine/Plugins/Experimental/Interchange/Runtime/Source/Import/Private/Mesh/InterchangeStaticMeshFactory.cpp
@@ -13,6 +13,9 @@
 #include "InterchangeStaticMeshFactoryNode.h"
 #include "InterchangeSourceData.h"
 #include "InterchangeTranslatorBase.h"
+#include "Materials/Material.h"
+#include "Materials/MaterialInstance.h"
+#include "Materials/MaterialInterface.h"
 #include "Mesh/InterchangeStaticMeshPayload.h"
 #include "Mesh/InterchangeStaticMeshPayloadInterface.h"
 #include "Nodes/InterchangeBaseNode.h"
diff --git a/Engine/Plugins/Experimental/Interchange/Runtime/Source/Nodes/Public/InterchangeMaterialFactoryNode.h b/Engine/Plugins/Experimental/Interchange/Runtime/Source/Nodes/Public/InterchangeMaterialFactoryNode.h
--- a/Engine/Plugins/Experimental/Interchange/Runtime/Source/Nodes/Public/InterchangeMaterialFactoryNode.h
+++ b/Engine/Plugins/Experimental/Interchange/Runtime/Source/Nodes/Public/InterchangeMaterialFactoryNode.h
@@ -6,6 +6,10 @@
 #include "Nodes/InterchangeBaseNode.h"
 #include "UObject/UObjectGlobals.h"
 
+#if WITH_ENGINE
+#include "Materials/MaterialInterface.h"
+#endif
+
 #include "InterchangeMaterialFactoryNode.generated.h"
 
 UCLASS(BlueprintType, Experimental)
@@ -132,7 +136,7 @@ private:
 
 protected:
 #if WITH_ENGINE
-	TSubclassOf<UTexture> AssetClass = nullptr;
+	TSubclassOf<UMaterialInterface> AssetClass = nullptr;
 #endif
 	bool bIsMaterialNodeClassInitialized = false;
 
diff --git a/Engine/Plugins/Experimental/Interchange/Runtime/Source/Nodes/Public/InterchangeSkeletalMeshFactoryNode.h b/Engine/Plugins/Experimental/Interchange/Runtime/Source/Nodes/Public/InterchangeSkeletalMeshFactoryNode.h
--- a/Engine/Plugins/Experimental/Interchange/Runtime/Source/Nodes/Public/InterchangeSkeletalMeshFactoryNode.h
+++ b/Engine/Plugins/Experimental/Interchange/Runtime/Source/Nodes/Public/InterchangeSkeletalMeshFactoryNode.h
@@ -5,6 +5,10 @@
 #include "CoreMinimal.h"
 #include "Nodes/InterchangeBaseNode.h"
 
+#if WITH_ENGINE
+#include "Engine/SkeletalMesh.h"
+#endif
+
 #include "InterchangeSkeletalMeshFactoryNode.generated.h"
 
 //Interchange namespace
@@ -140,7 +144,7 @@ private:
 	UE::Interchange::TArrayAttributeHelper<FString> LodDependencies;
 protected:
 #if WITH_ENGINE
-	TSubclassOf<class USkeletalMesh> AssetClass = nullptr;
+	TSubclassOf<USkeletalMesh> AssetClass = nullptr;
 #endif
 	bool bIsNodeClassInitialized = false;
 };
diff --git a/Engine/Plugins/Experimental/Interchange/Runtime/Source/Nodes/Public/InterchangeStaticMeshFactoryNode.h b/Engine/Plugins/Experimental/Interchange/Runtime/Source/Nodes/Public/InterchangeStaticMeshFactoryNode.h
--- a/Engine/Plugins/Experimental/Interchange/Runtime/Source/Nodes/Public/InterchangeStaticMeshFactoryNode.h
+++ b/Engine/Plugins/Experimental/Interchange/Runtime/Source/Nodes/Public/InterchangeStaticMeshFactoryNode.h
@@ -5,6 +5,10 @@
 #include "CoreMinimal.h"
 #include "Nodes/InterchangeBaseNode.h"
 
+#if WITH_ENGINE
+#include "Engine/StaticMesh.h"
+#endif
+
 #include "InterchangeStaticMeshFactoryNode.generated.h"
 
 
@@ -102,7 +106,7 @@ private:
 	UE::Interchange::TArrayAttributeHelper<FString> LodDependencies;
 protected:
 #if WITH_ENGINE
-	TSubclassOf<class UStaticMesh> AssetClass = nullptr;
+	TSubclassOf<UStaticMesh> AssetClass = nullptr;
 #endif
 	bool bIsNodeClassInitialized = false;
 };
