Interchange material graphs support

Interchange Core:
- Removed prefix from EInterchangeNodeContainerType since its an enum class.
- Specialized our delegate apply and fill for FStrings when the property if an object property.
- Fixed an out of bound memory access in InterchangePrivateNodeBase::FindPropertyByPathChecked

Interchange Materials:
- Added ShaderGraph and ShaderNodes interchange nodes.
- Replaced the MaterialNode and added a MaterialExpression node.
- Added support for Lambert, Phong and UE PBR shading models.
- Supports texture sample and lerp shader nodes.

Known limitations:
- FBX transparency goes into UE opaque channel but it should be flipped.
- FBX Phong shininess goes from 0 to infinite and it currently goes into UE's specular channel. It needs to be scaled between 0 and 1.
- Material paramters are numbered based on their position in the material expressions list to avoid having 2 parameters with the same name.
- The Lerp shader node should have a flexible set of inputs. Its currently limited to linear colors.

diff --git a/Engine/Plugins/Experimental/Interchange/Editor/Source/Pipelines/Private/InterchangeEditorPipelineDetails.cpp b/Engine/Plugins/Experimental/Interchange/Editor/Source/Pipelines/Private/InterchangeEditorPipelineDetails.cpp
--- a/Engine/Plugins/Experimental/Interchange/Editor/Source/Pipelines/Private/InterchangeEditorPipelineDetails.cpp
+++ b/Engine/Plugins/Experimental/Interchange/Editor/Source/Pipelines/Private/InterchangeEditorPipelineDetails.cpp
@@ -543,7 +543,7 @@ void FInterchangeBaseNodeDetailsCustomization::BuildLinearColorValueContent(IDet
 	UE::Interchange::EAttributeTypes AttributeType = InterchangeBaseNode->GetAttributeType(AttributeKey);
 	if (AttributeType == UE::Interchange::EAttributeTypes::LinearColor)
 	{
-		const UE::Interchange::FAttributeStorage::TAttributeHandle<FColor> AttributeHandle = InterchangeBaseNode->GetAttributeHandle<FColor>(AttributeKey);
+		const UE::Interchange::FAttributeStorage::TAttributeHandle<FLinearColor> AttributeHandle = InterchangeBaseNode->GetAttributeHandle<FLinearColor>(AttributeKey);
 		if (!AttributeHandle.IsValid())
 		{
 			CreateInvalidHandleRow(AttributeCategory, AttributeKey);
diff --git a/Engine/Plugins/Experimental/Interchange/Runtime/Source/Import/InterchangeImport.Build.cs b/Engine/Plugins/Experimental/Interchange/Runtime/Source/Import/InterchangeImport.Build.cs
--- a/Engine/Plugins/Experimental/Interchange/Runtime/Source/Import/InterchangeImport.Build.cs
+++ b/Engine/Plugins/Experimental/Interchange/Runtime/Source/Import/InterchangeImport.Build.cs
@@ -41,6 +41,7 @@ namespace UnrealBuildTool.Rules
 				PrivateDependencyModuleNames.AddRange(
 					new string[]
 					{
+						"MaterialEditor",
 						"UnrealEd",
 					}
 				);
diff --git a/Engine/Plugins/Experimental/Interchange/Runtime/Source/Import/Private/Material/InterchangeMaterialFactory.cpp b/Engine/Plugins/Experimental/Interchange/Runtime/Source/Import/Private/Material/InterchangeMaterialFactory.cpp
--- a/Engine/Plugins/Experimental/Interchange/Runtime/Source/Import/Private/Material/InterchangeMaterialFactory.cpp
+++ b/Engine/Plugins/Experimental/Interchange/Runtime/Source/Import/Private/Material/InterchangeMaterialFactory.cpp
@@ -4,26 +4,153 @@
 #include "InterchangeImportCommon.h"
 #include "InterchangeImportLog.h"
 #include "InterchangeMaterialFactoryNode.h"
-#include "InterchangeMaterialNode.h"
+#include "InterchangeTextureNode.h"
 #include "InterchangeTextureFactoryNode.h"
+#include "InterchangeShaderGraphNode.h"
 #include "InterchangeSourceData.h"
 #include "Materials/Material.h"
+#include "Materials/MaterialExpression.h"
 #include "Materials/MaterialExpressionScalarParameter.h"
 #include "Materials/MaterialExpressionTextureCoordinate.h"
 #include "Materials/MaterialExpressionTextureSample.h"
 #include "Materials/MaterialExpressionVectorParameter.h"
-#include "Materials/MaterialInstance.h"
 #include "Materials/MaterialInterface.h"
 #include "Nodes/InterchangeBaseNode.h"
 #include "Nodes/InterchangeBaseNodeContainer.h"
 
 
 #if WITH_EDITORONLY_DATA
-
 #include "EditorFramework/AssetImportData.h"
-
 #endif //WITH_EDITORONLY_DATA
 
+#if WITH_EDITOR
+#include "MaterialEditingLibrary.h"
+#endif //WITH_EDITOR
+
+#define LOCTEXT_NAMESPACE "InterchangeMaterialFactory"
+
+namespace UE
+{
+	namespace Interchange
+	{
+		namespace MaterialFactory
+		{
+			namespace Internal
+			{
+#if WITH_EDITOR
+				/**
+				 * Finds a UMaterialExpression class by name.
+				 * @param ClassName		The name of the class to look for (ie:UClass*->GetName()).
+				 * @return				A sub class of UMaterialExpression or nullptr.
+				 */
+				TSubclassOf<UMaterialExpression> FindExpressionClass(const TCHAR* ClassName)
+				{
+					check(ClassName);
+
+					UObject* const ClassPackage = ANY_PACKAGE;
+					UClass* MaterialExpressionClass = FindObject<UClass>(ClassPackage, ClassName);
+
+					if (!MaterialExpressionClass)
+					{
+						if (UObjectRedirector* RenamedClassRedirector = FindObject<UObjectRedirector>(ClassPackage, ClassName))
+						{
+							MaterialExpressionClass = CastChecked<UClass>(RenamedClassRedirector->DestinationObject);
+						}
+					}
+
+					if (MaterialExpressionClass->IsChildOf<UMaterialExpression>())
+					{
+						return MaterialExpressionClass;
+					}
+					else
+					{
+						return nullptr;
+					}
+				}
+
+				int32 GetInputIndex(UMaterialExpression& MaterialExpression, const FString& InputName)
+				{
+					int32 ExpressionInputIndex = 0;
+
+					for (const FExpressionInput* ExpressionInput : MaterialExpression.GetInputs())
+					{
+						// MaterialFuncCall appends the type to the input name when calling GetInputName
+						// and the InputName in FExpressionInput is optional so we'll check both here to be safe
+						if (MaterialExpression.GetInputName(ExpressionInputIndex) == *InputName ||
+							(ExpressionInput && ExpressionInput->InputName == *InputName))
+						{
+							return ExpressionInputIndex;
+						}
+
+						++ExpressionInputIndex;
+					}
+
+					return INDEX_NONE;
+				}
+
+				int32 GetOutputIndex(UMaterialExpression& MaterialExpression, const FString& OutputName)
+				{
+					int32 ExpressionOutputIndex = 0;
+
+					for (const FExpressionOutput& ExpressionOutput : MaterialExpression.GetOutputs())
+					{
+						if (ExpressionOutput.OutputName == *OutputName)
+						{
+							return ExpressionOutputIndex;
+						}
+
+						++ExpressionOutputIndex;
+					}
+
+					return 0; // Consider 0 as the default output to connect to since most expressions have a single output
+				}
+
+				void SetupFunctionCallExpression(UMaterial* Material, const UInterchangeFactoryBase::FCreateAssetParams& Arguments, UMaterialExpressionMaterialFunctionCall* FunctionCallExpression)
+				{
+					FunctionCallExpression->UpdateFromFunctionResource();
+				}
+
+				void SetupTextureExpression(const UInterchangeFactoryBase::FCreateAssetParams& Arguments, const UInterchangeMaterialExpressionFactoryNode* ExpressionNode, UMaterialExpressionTextureBase* TextureExpression)
+				{
+					const UInterchangeTextureNode* TextureNode = nullptr;
+
+					FString TextureDependency;
+
+					TArray<FString> TargetNodes;
+					ExpressionNode->GetTargetNodeUids(TargetNodes);
+					if (!TargetNodes.IsEmpty())
+					{
+						if (const UInterchangeShaderNode* ShaderNode = Cast<UInterchangeShaderNode>(Arguments.NodeContainer->GetNode(TargetNodes[0])))
+						{
+							FString TextureUid;
+							ShaderNode->GetStringAttribute(UInterchangeShaderPortsAPI::MakeInputValueKey(TEXT("TextureUid")), TextureUid);
+
+							TextureNode = Cast<UInterchangeTextureNode>(Arguments.NodeContainer->GetNode(TextureUid));
+						}
+					}
+
+					if (TextureNode)
+					{
+						TArray<FString> TextureNodeTargetNodes;
+						TextureNode->GetTargetNodeUids(TextureNodeTargetNodes);
+
+						if (const UInterchangeTextureFactoryNode* TextureFactoryNode = Cast<UInterchangeTextureFactoryNode>(Arguments.NodeContainer->GetNode(TextureNodeTargetNodes[0])))
+						{
+							if (UTexture* Texture = Cast<UTexture>(TextureFactoryNode->ReferenceObject.TryLoad()))
+							{
+								TextureExpression->Texture = Texture;
+							}
+						}
+					}
+
+					TextureExpression->AutoSetSampleType();
+				}
+#endif // #if WITH_EDITOR
+			}
+		}
+	}
+}
+
 UClass* UInterchangeMaterialFactory::GetFactoryClass() const
 {
 	return UMaterialInterface::StaticClass();
@@ -33,24 +160,17 @@ UObject* UInterchangeMaterialFactory::CreateEmptyAsset(const FCreateAssetParams&
 {
 	UObject* Material = nullptr;
 
-#if WITH_EDITORONLY_DATA
+#if WITH_EDITOR
 	if (!Arguments.AssetNode || !Arguments.AssetNode->GetObjectClass()->IsChildOf(GetFactoryClass()))
 	{
 		return nullptr;
 	}
 
-	const UInterchangeMaterialFactoryNode* MaterialFactoryNode = Cast<UInterchangeMaterialFactoryNode>(Arguments.AssetNode);
+	const UInterchangeBaseMaterialFactoryNode* MaterialFactoryNode = Cast<UInterchangeBaseMaterialFactoryNode>(Arguments.AssetNode);
 	if (MaterialFactoryNode == nullptr)
 	{
 		return nullptr;
 	}
-	FString MaterialTranslatedNodeUid;
-	MaterialFactoryNode->GetCustomTranslatedMaterialNodeUid(MaterialTranslatedNodeUid);
-	const UInterchangeMaterialNode* MaterialTranslatedNode = Cast<UInterchangeMaterialNode>(Arguments.NodeContainer->GetNode(MaterialTranslatedNodeUid));
-	if (MaterialTranslatedNode == nullptr)
-	{
-		return nullptr;
-	}
 
 	const UClass* MaterialClass = MaterialFactoryNode->GetObjectClass();
 	if (!ensure(MaterialClass && MaterialClass->IsChildOf(GetFactoryClass())))
@@ -78,13 +198,15 @@ UObject* UInterchangeMaterialFactory::CreateEmptyAsset(const FCreateAssetParams&
 		return nullptr;
 	}
 	Material->PreEditChange(nullptr);
-#endif //WITH_EDITORONLY_DATA
+
+#endif //WITH_EDITOR
+
 	return Material;
 }
 
 UObject* UInterchangeMaterialFactory::CreateAsset(const FCreateAssetParams& Arguments)
 {
-#if !WITH_EDITORONLY_DATA
+#if !WITH_EDITOR
 
 	UE_LOG(LogInterchangeImport, Error, TEXT("Cannot import Material asset in runtime, this is an editor only feature."));
 	return nullptr;
@@ -96,18 +218,11 @@ UObject* UInterchangeMaterialFactory::CreateAsset(const FCreateAssetParams& Argu
 		return nullptr;
 	}
 	
-	const UInterchangeMaterialFactoryNode* MaterialFactoryNode = Cast<UInterchangeMaterialFactoryNode>(Arguments.AssetNode);
+	const UInterchangeBaseMaterialFactoryNode* MaterialFactoryNode = Cast<UInterchangeBaseMaterialFactoryNode>(Arguments.AssetNode);
 	if (MaterialFactoryNode == nullptr)
 	{
 		return nullptr;
 	}
-	FString MaterialTranslatedNodeUid;
-	MaterialFactoryNode->GetCustomTranslatedMaterialNodeUid(MaterialTranslatedNodeUid);
-	const UInterchangeMaterialNode* MaterialTranslatedNode = Cast<UInterchangeMaterialNode>(Arguments.NodeContainer->GetNode(MaterialTranslatedNodeUid));
-	if (MaterialTranslatedNode == nullptr)
-	{
-		return nullptr;
-	}
 
 	const UClass* MaterialClass = MaterialFactoryNode->GetObjectClass();
 	check(MaterialClass && MaterialClass->IsChildOf(GetFactoryClass()));
@@ -120,7 +235,7 @@ UObject* UInterchangeMaterialFactory::CreateAsset(const FCreateAssetParams& Argu
 	if (!ExistingAsset)
 	{
 		//NewObject is not thread safe, the asset registry directory watcher tick on the main thread can trig before we finish initializing the UObject and will crash
-		//The UObject should have been create by calling CreateEmptyAsset on the main thread.
+		//The UObject should have been created by calling CreateEmptyAsset on the main thread.
 		check(IsInGameThread());
 		MaterialObject = NewObject<UObject>(Arguments.Parent, MaterialClass, *Arguments.AssetName, RF_Public | RF_Standalone);
 	}
@@ -142,148 +257,15 @@ UObject* UInterchangeMaterialFactory::CreateAsset(const FCreateAssetParams& Argu
 		//TODO design a re-import process for the material (expressions and input connections)
 		if(!Arguments.ReimportObject)
 		{
-			UMaterial* Material = Cast<UMaterial>(MaterialObject);
-			
-			if (Material)
+			if (UMaterialInterface* MaterialInterface = Cast<UMaterialInterface>(MaterialObject))
 			{
-				auto ApplyInputParameter = [&MaterialFactoryNode, &MaterialTranslatedNode, &Material, &Arguments](EInterchangeMaterialNodeParameterName ParameterName, FExpressionInput& MaterialInput, const FVector2D& Location)
+				if (UMaterial* Material = Cast<UMaterial>(MaterialObject))
 				{
-					FString OutTextureUid;
-					int32 OutUVSetIndex = 0;
-					float ScaleU = 0.0f;
-					float ScaleV = 0.0f;
-					FVector OutVectorParameter = FVector(0.0f);
-					float OutFloatParameter = 0.0f;
-					if (MaterialTranslatedNode->GetTextureParameterData(ParameterName, OutTextureUid, OutUVSetIndex, ScaleU, ScaleV))
-					{
-						OutTextureUid = UInterchangeTextureFactoryNode::GetTextureFactoryNodeUidFromTextureNodeUid(OutTextureUid);
-						UTexture* TextureReference = nullptr;
-						if (Arguments.NodeContainer)
-						{
-							TArray<FString> MaterialTextureDependencies;
-							MaterialFactoryNode->GetTextureDependencies(MaterialTextureDependencies);
-							int32 DependenciesCount = MaterialTextureDependencies.Num();
-							for (int32 DependIndex = 0; DependIndex < DependenciesCount; ++DependIndex)
-							{
-								if (OutTextureUid != MaterialTextureDependencies[DependIndex])
-								{
-									continue;
-								}
-								const UInterchangeTextureFactoryNode* DepNode = Cast<UInterchangeTextureFactoryNode>(Arguments.NodeContainer->GetNode(MaterialTextureDependencies[DependIndex]));
-								if (DepNode && DepNode->ReferenceObject.IsAsset())
-								{
-									//Use Resolve object so we just look for in memory UObject
-									UObject* TextureObject = DepNode->ReferenceObject.ResolveObject();
-									if (TextureObject)
-									{
-										TextureReference = Cast<UTexture>(TextureObject);
-										break;
-									}
-								}
-							}
-						}
-
-						if (!TextureReference)
-						{
-							//TODO: Try to load the texture by extracting an asset name from the filename
-						}
-
-
-						if (TextureReference)
-						{
-							//Set the input to use this texture
-							// and link it to the material 
-							UMaterialExpressionTextureSample* UnrealTextureExpression = NewObject<UMaterialExpressionTextureSample>(Material);
-							Material->Expressions.Add(UnrealTextureExpression);
-							MaterialInput.Expression = UnrealTextureExpression;
-							UnrealTextureExpression->Texture = TextureReference;
-							// 						UnrealTextureExpression->SamplerType = bSetupAsNormalMap ?
-							// 							(bIsVirtualTexture ? SAMPLERTYPE_VirtualNormal : SAMPLERTYPE_Normal) :
-							// 							(bIsVirtualTexture ? SAMPLERTYPE_VirtualColor : SAMPLERTYPE_Color);
-
-							UnrealTextureExpression->MaterialExpressionEditorX = FMath::TruncToInt(Location.X);
-							UnrealTextureExpression->MaterialExpressionEditorY = FMath::TruncToInt(Location.Y);
-
-							//////////////////////////////////////////////////////////////////////////
-							//TODO
-							//This UV set come from fbx we have to find a way to have more context here, the UV index has to be set when creating the node
-							//The name is not suffisant, we need an index
-
-							// add/find UVSet and set it to the texture
-
-							if ((OutUVSetIndex != 0 && OutUVSetIndex != INDEX_NONE) || ScaleU != 1.0f || ScaleV != 1.0f)
-							{
-								// Create a texture coord node for the texture sample
-								UMaterialExpressionTextureCoordinate* MyCoordExpression = NewObject<UMaterialExpressionTextureCoordinate>(Material);
-								Material->Expressions.Add(MyCoordExpression);
-								MyCoordExpression->CoordinateIndex = (OutUVSetIndex >= 0) ? OutUVSetIndex : 0;
-								MyCoordExpression->UTiling = ScaleU;
-								MyCoordExpression->VTiling = ScaleV;
-								UnrealTextureExpression->Coordinates.Expression = MyCoordExpression;
-								MyCoordExpression->MaterialExpressionEditorX = FMath::TruncToInt(Location.X - 175);
-								MyCoordExpression->MaterialExpressionEditorY = FMath::TruncToInt(Location.Y);
-
-							}
-						}
-					}
-					else if (MaterialTranslatedNode->GetVectorParameterData(ParameterName, OutVectorParameter))
-					{
-						UMaterialExpressionVectorParameter* MyColorExpression = NewObject<UMaterialExpressionVectorParameter>(Material);
-						Material->Expressions.Add( MyColorExpression );
-						MaterialInput.Expression = MyColorExpression;
-
-						MyColorExpression->DefaultValue.R = OutVectorParameter[0];
-						MyColorExpression->DefaultValue.G = OutVectorParameter[1];
-						MyColorExpression->DefaultValue.B = OutVectorParameter[2];
-
-						MyColorExpression->MaterialExpressionEditorX = FMath::TruncToInt(Location.X);
-						MyColorExpression->MaterialExpressionEditorY = FMath::TruncToInt(Location.Y);
-					}
-					else if (MaterialTranslatedNode->GetScalarParameterData(ParameterName, OutFloatParameter))
-					{
-						UMaterialExpressionScalarParameter* MyScalarExpression = NewObject<UMaterialExpressionScalarParameter>(Material);
-						Material->Expressions.Add(MyScalarExpression);
-						MaterialInput.Expression = MyScalarExpression;
-
-						MyScalarExpression->DefaultValue = OutFloatParameter;
-
-						MyScalarExpression->MaterialExpressionEditorX = FMath::TruncToInt(Location.X);
-						MyScalarExpression->MaterialExpressionEditorY = FMath::TruncToInt(Location.Y);
-					}
-					
-					bool bResult = false;
-					if (MaterialInput.Expression)
-					{
-						TArray<FExpressionOutput> Outputs = MaterialInput.Expression->GetOutputs();
-						FExpressionOutput* Output = Outputs.GetData();
-						MaterialInput.Mask = Output->Mask;
-						MaterialInput.MaskR = Output->MaskR;
-						MaterialInput.MaskG = Output->MaskG;
-						MaterialInput.MaskB = Output->MaskB;
-						MaterialInput.MaskA = Output->MaskA;
-						//If we create an expression we return true
-						bResult = true;
-					}
-					return bResult;
-				};
-
-				ApplyInputParameter(EInterchangeMaterialNodeParameterName::BaseColor, Material->BaseColor, FVector2D(-250, -100));
-				ApplyInputParameter(EInterchangeMaterialNodeParameterName::Metallic, Material->Metallic, FVector2D(-750, 0));
-				ApplyInputParameter(EInterchangeMaterialNodeParameterName::Specular, Material->Specular, FVector2D(-500, 100));
-				ApplyInputParameter(EInterchangeMaterialNodeParameterName::Roughness, Material->Roughness, FVector2D(-250, 200));
-				ApplyInputParameter(EInterchangeMaterialNodeParameterName::EmissiveColor, Material->EmissiveColor, FVector2D(-750, 300));
-				if (ApplyInputParameter(EInterchangeMaterialNodeParameterName::Opacity, Material->Opacity, FVector2D(-500, 400)))
-				{
-					Material->BlendMode = BLEND_Translucent;
+					SetupMaterial(Material, Arguments, MaterialFactoryNode);
 				}
-				ApplyInputParameter(EInterchangeMaterialNodeParameterName::OpacityMask, Material->OpacityMask, FVector2D(-250, 500));
-				ApplyInputParameter(EInterchangeMaterialNodeParameterName::Normal, Material->Normal, FVector2D(-750, 600));
-			}
 
-			/** Apply all MaterialNode custom attributes to the material asset */
-			MaterialFactoryNode->ApplyAllCustomAttributeToObject(Material);
-
-			//TODO support material instance
+				MaterialFactoryNode->ApplyAllCustomAttributeToObject(MaterialObject);
+			}
 		}
 		
 		//Getting the file Hash will cache it into the source data
@@ -291,14 +273,9 @@ UObject* UInterchangeMaterialFactory::CreateAsset(const FCreateAssetParams& Argu
 
 		//The interchange completion task (call in the GameThread after the factories pass), will call PostEditChange which will trig another asynchronous system that will build all material in parallel
 	}
-	else
-	{
-		//The material is not a UMaterialInterface
-		MaterialObject->RemoveFromRoot();
-		MaterialObject->MarkAsGarbage();
-	}
+
 	return MaterialObject;
-#endif
+#endif // WITH_EDITOR
 }
 
 /* This function is call in the completion task on the main thread, use it to call main thread post creation step for your assets*/
@@ -324,4 +301,269 @@ void UInterchangeMaterialFactory::PreImportPreCompletedCallback(const FImportPre
 		ImportedMaterial->AssetImportData = UE::Interchange::FFactoryCommon::UpdateImportAssetData(UpdateImportAssetDataParameters);
 	}
 #endif
-}
\ No newline at end of file
+}
+
+#if WITH_EDITOR
+void UInterchangeMaterialFactory::SetupMaterial(UMaterial* Material, const FCreateAssetParams& Arguments, const UInterchangeBaseMaterialFactoryNode* BaseMaterialFactoryNode)
+{
+	using namespace UE::Interchange::MaterialFactory::Internal;
+
+	const UInterchangeMaterialFactoryNode* MaterialFactoryNode = Cast<UInterchangeMaterialFactoryNode>(BaseMaterialFactoryNode);
+
+	TMap<FString, UMaterialExpression*> Expressions;
+
+	// Base Color
+	{
+		FString BaseColorUid;
+		FString OutputName;
+		MaterialFactoryNode->GetBaseColorConnection(BaseColorUid, OutputName);
+
+		const UInterchangeMaterialExpressionFactoryNode* BaseColor = Cast<UInterchangeMaterialExpressionFactoryNode>(Arguments.NodeContainer->GetNode(BaseColorUid));
+
+		if (BaseColor)
+		{
+			if (UMaterialExpression* BaseColorExpression = CreateExpressionsForNode(Material, Arguments, BaseColor, Expressions))
+			{
+				if (FExpressionInput* BaseColorInput = Material->GetExpressionInputForProperty(MP_BaseColor))
+				{
+					BaseColorExpression->ConnectExpression(BaseColorInput, GetOutputIndex(*BaseColorExpression, OutputName));
+				}
+			}
+		}
+	}
+	
+	// Metallic
+	{
+		FString MetallicUid;
+		FString OutputName;
+		MaterialFactoryNode->GetMetallicConnection(MetallicUid, OutputName);
+
+		const UInterchangeMaterialExpressionFactoryNode* MetallicNode = Cast<UInterchangeMaterialExpressionFactoryNode>(Arguments.NodeContainer->GetNode(MetallicUid));
+
+		if (MetallicNode)
+		{
+			if (UMaterialExpression* MetallicExpression = CreateExpressionsForNode(Material, Arguments, MetallicNode, Expressions))
+			{
+				if (FExpressionInput* MetallicInput = Material->GetExpressionInputForProperty(MP_Metallic))
+				{
+					MetallicExpression->ConnectExpression(MetallicInput, GetOutputIndex(*MetallicExpression, OutputName));
+				}
+			}
+		}
+	}
+
+	// Specular
+	{
+		FString SpecularUid;
+		FString OutputName;
+		MaterialFactoryNode->GetSpecularConnection(SpecularUid, OutputName);
+
+		const UInterchangeMaterialExpressionFactoryNode* SpecularNode = Cast<UInterchangeMaterialExpressionFactoryNode>(Arguments.NodeContainer->GetNode(SpecularUid));
+
+		if (SpecularNode)
+		{
+			if (UMaterialExpression* SpecularExpression = CreateExpressionsForNode(Material, Arguments, SpecularNode, Expressions))
+			{
+				if (FExpressionInput* SpecularInput = Material->GetExpressionInputForProperty(MP_Specular))
+				{
+					SpecularExpression->ConnectExpression(SpecularInput, GetOutputIndex(*SpecularExpression, OutputName));
+				}
+			}
+		}
+	}
+
+	// Roughness
+	{
+		FString RoughnessUid;
+		FString OutputName;
+		MaterialFactoryNode->GetRoughnessConnection(RoughnessUid, OutputName);
+
+		const UInterchangeMaterialExpressionFactoryNode* Roughness = Cast<UInterchangeMaterialExpressionFactoryNode>(Arguments.NodeContainer->GetNode(RoughnessUid));
+
+		if (Roughness)
+		{
+			if (UMaterialExpression* RoughnessExpression = CreateExpressionsForNode(Material, Arguments, Roughness, Expressions))
+			{
+				if (FExpressionInput* RoughnessInput = Material->GetExpressionInputForProperty(MP_Roughness))
+				{
+					RoughnessExpression->ConnectExpression(RoughnessInput, GetOutputIndex(*RoughnessExpression, OutputName));
+				}
+			}
+		}
+	}
+
+	// Emissive
+	{
+		FString EmissiveColorUid;
+		FString OutputName;
+		MaterialFactoryNode->GetEmissiveColorConnection(EmissiveColorUid, OutputName);
+
+		const UInterchangeMaterialExpressionFactoryNode* EmissiveColor = Cast<UInterchangeMaterialExpressionFactoryNode>(Arguments.NodeContainer->GetNode(EmissiveColorUid));
+
+		if (EmissiveColor)
+		{
+			if (UMaterialExpression* EmissiveExpression = CreateExpressionsForNode(Material, Arguments, EmissiveColor, Expressions))
+			{
+				if (FExpressionInput* EmissiveInput = Material->GetExpressionInputForProperty(MP_EmissiveColor))
+				{
+					EmissiveExpression->ConnectExpression(EmissiveInput, GetOutputIndex(*EmissiveExpression, OutputName));
+				}
+			}
+		}
+	}
+	
+	// Normal
+	{
+		FString NormalUid;
+		FString OutputName;
+		MaterialFactoryNode->GetNormalConnection(NormalUid, OutputName);
+
+		const UInterchangeMaterialExpressionFactoryNode* NormalNode = Cast<UInterchangeMaterialExpressionFactoryNode>(Arguments.NodeContainer->GetNode(NormalUid));
+
+		if (NormalNode)
+		{
+			if (UMaterialExpression* NormalExpression = CreateExpressionsForNode(Material, Arguments, NormalNode, Expressions))
+			{
+				if (FExpressionInput* NormalInput = Material->GetExpressionInputForProperty(MP_Normal))
+				{
+					NormalExpression->ConnectExpression(NormalInput, GetOutputIndex(*NormalExpression, OutputName));
+				}
+			}
+		}
+	}
+
+	// Opacity
+	{
+		FString OpacityUid;
+		FString OutputName;
+		MaterialFactoryNode->GetOpacityConnection(OpacityUid, OutputName);
+
+		const UInterchangeMaterialExpressionFactoryNode* OpacityNode = Cast<UInterchangeMaterialExpressionFactoryNode>(Arguments.NodeContainer->GetNode(OpacityUid));
+
+		if (OpacityNode)
+		{
+			if (UMaterialExpression* OpacityExpression = CreateExpressionsForNode(Material, Arguments, OpacityNode, Expressions))
+			{
+				if (FExpressionInput* OpacityInput = Material->GetExpressionInputForProperty(MP_Opacity))
+				{
+					OpacityExpression->ConnectExpression(OpacityInput, GetOutputIndex(*OpacityExpression, OutputName));
+				}
+			}
+		}
+	}
+
+	UMaterialEditingLibrary::LayoutMaterialExpressions(Material);
+}
+
+UMaterialExpression* UInterchangeMaterialFactory::CreateExpression(UMaterial* Material, const UInterchangeFactoryBase::FCreateAssetParams& Arguments, const UInterchangeMaterialExpressionFactoryNode& ExpressionNode)
+{
+	using namespace UE::Interchange::MaterialFactory::Internal;
+
+	FString ExpressionClassName;
+	ExpressionNode.GetCustomExpressionClassName(ExpressionClassName);
+
+	TSubclassOf<UMaterialExpression> ExpressionClass = FindExpressionClass(*ExpressionClassName);
+
+	if (!ExpressionClass.Get())
+	{
+		UInterchangeResultWarning_Generic* Message = AddMessage<UInterchangeResultWarning_Generic>();
+		Message->Text = FText::Format(LOCTEXT("ExpressionClassNotFound", "Invalid class {0} for material expression node {1}."),
+			FText::FromString(ExpressionClassName),
+			FText::FromString(ExpressionNode.GetDisplayLabel()));
+
+		return nullptr;
+	}
+
+	UMaterialExpression* MaterialExpression = UMaterialEditingLibrary::CreateMaterialExpression(Material, ExpressionClass);
+
+	if (!MaterialExpression)
+	{
+		return nullptr;
+	}
+
+	// Set the parameter name if the material expression has one (some material expressions don't inherit from UMaterialExpressionParameter, ie: UMaterialExpressionTextureSampleParameter
+	if (FNameProperty* Property = FindFProperty<FNameProperty>(MaterialExpression->GetClass(), GET_MEMBER_NAME_CHECKED(UMaterialExpressionParameter, ParameterName)))
+	{
+		*(Property->ContainerPtrToValuePtr<FName>(MaterialExpression)) = FName(*(ExpressionNode.GetDisplayLabel() + LexToString(Material->Expressions.Num())));
+	}
+
+	ExpressionNode.ApplyAllCustomAttributeToObject(MaterialExpression);
+
+	if (UMaterialExpressionTextureBase* TextureExpression = Cast<UMaterialExpressionTextureBase>(MaterialExpression))
+	{
+		SetupTextureExpression(Arguments, &ExpressionNode, TextureExpression);	
+	}
+	else if (UMaterialExpressionMaterialFunctionCall* FunctionCallExpression = Cast<UMaterialExpressionMaterialFunctionCall>(MaterialExpression))
+	{
+		SetupFunctionCallExpression(Material, Arguments, FunctionCallExpression);
+	}
+
+	return MaterialExpression;
+}
+
+UMaterialExpression* UInterchangeMaterialFactory::CreateExpressionsForNode(UMaterial* Material, const UInterchangeFactoryBase::FCreateAssetParams& Arguments,
+	const UInterchangeMaterialExpressionFactoryNode* Expression, TMap<FString, UMaterialExpression*>& Expressions)
+{
+	using namespace UE::Interchange::MaterialFactory::Internal;
+
+	UMaterialExpression* MaterialExpression = Expressions.FindRef(Expression->GetUniqueID());
+
+	if (MaterialExpression)
+	{
+		return MaterialExpression;
+	}
+
+	MaterialExpression = CreateExpression(Material, Arguments, *Expression);
+
+	if (!MaterialExpression)
+	{
+		return nullptr;
+	}
+
+	Expressions.Add(Expression->GetUniqueID()) = MaterialExpression;
+
+	TArray<FString> Inputs;
+	UInterchangeShaderPortsAPI::GatherInputs(Expression, Inputs);
+
+	for (const FString& InputName : Inputs)
+	{
+		FString ConnectedExpressionUid;
+		FString OutputName;
+		if (UInterchangeShaderPortsAPI::GetInputConnection(Expression, InputName, ConnectedExpressionUid, OutputName))
+		{
+			const UInterchangeMaterialExpressionFactoryNode* ConnectedExpressionNode = Cast<UInterchangeMaterialExpressionFactoryNode>(Arguments.NodeContainer->GetNode(ConnectedExpressionUid));
+			if (ConnectedExpressionNode)
+			{
+				UMaterialExpression* ConnectedExpression = Expressions.FindRef(ConnectedExpressionUid);
+				if (!ConnectedExpression)
+				{
+					ConnectedExpression = CreateExpressionsForNode(Material, Arguments, ConnectedExpressionNode, Expressions);
+				}
+
+				const int32 InputIndex = GetInputIndex(*MaterialExpression, InputName);
+				const int32 OutputIndex = GetOutputIndex(*MaterialExpression, OutputName);
+
+				if (InputIndex != INDEX_NONE)
+				{
+					FExpressionInput* ExpressionInput = MaterialExpression->GetInput(InputIndex);
+					if (ExpressionInput)
+					{
+						ConnectedExpression->ConnectExpression(ExpressionInput, OutputIndex);
+					}
+				}
+				else
+				{
+					UInterchangeResultWarning_Generic* Message = AddMessage<UInterchangeResultWarning_Generic>();
+					Message->Text = FText::Format(LOCTEXT("InputNotFound", "Invalid input {0} for material expression node {1}."),
+						FText::FromString(InputName),
+						FText::FromString(Expression->GetDisplayLabel()));
+				}
+			}
+		}
+	}
+
+	return MaterialExpression;
+}
+#endif // #if WITH_EDITOR
+
+#undef LOCTEXT_NAMESPACE
diff --git a/Engine/Plugins/Experimental/Interchange/Runtime/Source/Import/Private/Mesh/InterchangeSkeletalMeshFactory.cpp b/Engine/Plugins/Experimental/Interchange/Runtime/Source/Import/Private/Mesh/InterchangeSkeletalMeshFactory.cpp
--- a/Engine/Plugins/Experimental/Interchange/Runtime/Source/Import/Private/Mesh/InterchangeSkeletalMeshFactory.cpp
+++ b/Engine/Plugins/Experimental/Interchange/Runtime/Source/Import/Private/Mesh/InterchangeSkeletalMeshFactory.cpp
@@ -8,7 +8,6 @@
 #include "InterchangeImportCommon.h"
 #include "InterchangeImportLog.h"
 #include "InterchangeMaterialFactoryNode.h"
-#include "InterchangeMaterialNode.h"
 #include "InterchangeMeshNode.h"
 #include "InterchangeSceneNode.h"
 #include "InterchangeSkeletalMeshLodDataNode.h"
diff --git a/Engine/Plugins/Experimental/Interchange/Runtime/Source/Import/Private/Mesh/InterchangeStaticMeshFactory.cpp b/Engine/Plugins/Experimental/Interchange/Runtime/Source/Import/Private/Mesh/InterchangeStaticMeshFactory.cpp
--- a/Engine/Plugins/Experimental/Interchange/Runtime/Source/Import/Private/Mesh/InterchangeStaticMeshFactory.cpp
+++ b/Engine/Plugins/Experimental/Interchange/Runtime/Source/Import/Private/Mesh/InterchangeStaticMeshFactory.cpp
@@ -6,7 +6,6 @@
 #include "InterchangeImportCommon.h"
 #include "InterchangeImportLog.h"
 #include "InterchangeMaterialFactoryNode.h"
-#include "InterchangeMaterialNode.h"
 #include "InterchangeMeshNode.h"
 #include "InterchangeSceneNode.h"
 #include "InterchangeStaticMeshLodDataNode.h"
diff --git a/Engine/Plugins/Experimental/Interchange/Runtime/Source/Import/Public/Material/InterchangeMaterialFactory.h b/Engine/Plugins/Experimental/Interchange/Runtime/Source/Import/Public/Material/InterchangeMaterialFactory.h
--- a/Engine/Plugins/Experimental/Interchange/Runtime/Source/Import/Public/Material/InterchangeMaterialFactory.h
+++ b/Engine/Plugins/Experimental/Interchange/Runtime/Source/Import/Public/Material/InterchangeMaterialFactory.h
@@ -9,12 +9,17 @@
 
 #include "InterchangeMaterialFactory.generated.h"
 
+class UInterchangeBaseMaterialFactoryNode;
+class UInterchangeMaterialExpressionFactoryNode;
+class UMaterial;
+class UMaterialExpression;
+
 UCLASS(BlueprintType, Experimental)
 class INTERCHANGEIMPORT_API UInterchangeMaterialFactory : public UInterchangeFactoryBase
 {
 	GENERATED_BODY()
-public:
 
+public:
 	//////////////////////////////////////////////////////////////////////////
 	// Interchange factory base interface begin
 
@@ -25,6 +30,15 @@ public:
 
 	// Interchange factory base interface end
 	//////////////////////////////////////////////////////////////////////////
+private:
+#if WITH_EDITOR
+	void SetupMaterial(UMaterial* Material, const FCreateAssetParams& Arguments, const UInterchangeBaseMaterialFactoryNode* MaterialFactoryNode);
+
+	UMaterialExpression* CreateExpressionsForNode(UMaterial* Material, const UInterchangeFactoryBase::FCreateAssetParams& Arguments,
+	const UInterchangeMaterialExpressionFactoryNode* Expression, TMap<FString, UMaterialExpression*>& Expressions);
+
+	UMaterialExpression* CreateExpression(UMaterial* Material, const UInterchangeFactoryBase::FCreateAssetParams& Arguments, const UInterchangeMaterialExpressionFactoryNode& ExpressionNode);
+#endif // #if WITH_EDITOR
 };
 
 
diff --git a/Engine/Plugins/Experimental/Interchange/Runtime/Source/Nodes/Private/InterchangeMaterialFactoryNode.cpp b/Engine/Plugins/Experimental/Interchange/Runtime/Source/Nodes/Private/InterchangeMaterialFactoryNode.cpp
--- a/Engine/Plugins/Experimental/Interchange/Runtime/Source/Nodes/Private/InterchangeMaterialFactoryNode.cpp
+++ b/Engine/Plugins/Experimental/Interchange/Runtime/Source/Nodes/Private/InterchangeMaterialFactoryNode.cpp
@@ -10,22 +10,10 @@
 
 #endif
 
-
-UInterchangeMaterialFactoryNode::UInterchangeMaterialFactoryNode()
-{
-#if WITH_ENGINE
-	AssetClass = nullptr;
-#endif
-	TextureDependencies.Initialize(Attributes, TextureDependenciesKey.Key);
-}
-
-void UInterchangeMaterialFactoryNode::InitializeMaterialNode(const FString& UniqueID, const FString& DisplayLabel, const FString& InAssetClass)
+FString UInterchangeBaseMaterialFactoryNode::GetMaterialFactoryNodeUidFromMaterialNodeUid(const FString& TranslatedNodeUid)
 {
-	bIsMaterialNodeClassInitialized = false;
-	InitializeNode(UniqueID, DisplayLabel, EInterchangeNodeContainerType::FactoryData);
-	FString OperationName = GetTypeName() + TEXT(".SetAssetClassName");
-	InterchangePrivateNodeBase::SetCustomAttribute<FString>(*Attributes, ClassNameAttributeKey, OperationName, InAssetClass);
-	FillAssetClassFromAttribute();
+	FString NewUid = TEXT("Factory_") + TranslatedNodeUid;
+	return NewUid;
 }
 
 FString UInterchangeMaterialFactoryNode::GetTypeName() const
@@ -36,202 +24,130 @@ FString UInterchangeMaterialFactoryNode::GetTypeName() const
 
 UClass* UInterchangeMaterialFactoryNode::GetObjectClass() const
 {
-	ensure(bIsMaterialNodeClassInitialized);
 #if WITH_ENGINE
-	return AssetClass.Get() != nullptr ? AssetClass.Get() : UMaterial::StaticClass();
+	return UMaterial::StaticClass();
 #else
 	return nullptr;
 #endif
 }
 
-FString UInterchangeMaterialFactoryNode::GetKeyDisplayName(const UE::Interchange::FAttributeKey& NodeAttributeKey) const
-{
-	FString KeyDisplayName = NodeAttributeKey.ToString();
-	if (NodeAttributeKey == TextureDependenciesKey)
-	{
-		return KeyDisplayName = TEXT("Texture Dependencies count");
-	}
-	else if (NodeAttributeKey.Key.StartsWith(TextureDependenciesKey.Key))
-	{
-		KeyDisplayName = TEXT("Texture Dependencies Index ");
-		const FString IndexKey = UE::Interchange::TArrayAttributeHelper<FString>::IndexKey();
-		int32 IndexPosition = NodeAttributeKey.Key.Find(IndexKey) + IndexKey.Len();
-		if (IndexPosition < NodeAttributeKey.Key.Len())
-		{
-			KeyDisplayName += NodeAttributeKey.Key.RightChop(IndexPosition);
-		}
-		return KeyDisplayName;
-	}
-	return Super::GetKeyDisplayName(NodeAttributeKey);
-}
-
-FString UInterchangeMaterialFactoryNode::GetAttributeCategory(const UE::Interchange::FAttributeKey& NodeAttributeKey) const
-{
-	if (NodeAttributeKey == Macro_CustomTranslatedMaterialNodeUidKey
-		|| NodeAttributeKey == Macro_CustomMaterialUsageKey
-		|| NodeAttributeKey == Macro_CustomBlendModeKey)
-	{
-		return FString(TEXT("MaterialFactory"));
-	}
-	else if (NodeAttributeKey.Key.StartsWith(TextureDependenciesKey.Key))
-	{
-		return FString(TEXT("TextureDependencies"));
-	}
-	return Super::GetAttributeCategory(NodeAttributeKey);
-}
-
-FString UInterchangeMaterialFactoryNode::GetMaterialFactoryNodeUidFromMaterialNodeUid(const FString& TranslatedNodeUid)
+bool UInterchangeMaterialFactoryNode::GetBaseColorConnection(FString& ExpressionNodeUid, FString& OutputName) const
 {
-	FString NewUid = TEXT("Factory_") + TranslatedNodeUid;
-	return NewUid;
+	return UInterchangeShaderPortsAPI::GetInputConnection(this, UE::Interchange::Materials::PBR::Parameters::BaseColor.ToString(), ExpressionNodeUid, OutputName);
 }
 
-int32 UInterchangeMaterialFactoryNode::GetTextureDependeciesCount() const
+bool UInterchangeMaterialFactoryNode::ConnectToBaseColor(const FString& AttributeValue)
 {
-	return TextureDependencies.GetCount();
+	return UInterchangeShaderPortsAPI::ConnectDefaultOuputToInput(this, UE::Interchange::Materials::PBR::Parameters::BaseColor.ToString(), AttributeValue);
+}
+	
+bool UInterchangeMaterialFactoryNode::ConnectOutputToBaseColor(const FString& ExpressionNodeUid, const FString& OutputName)
+{
+	return UInterchangeShaderPortsAPI::ConnectOuputToInput(this, UE::Interchange::Materials::PBR::Parameters::BaseColor.ToString(), ExpressionNodeUid, OutputName);
 }
 
-void UInterchangeMaterialFactoryNode::GetTextureDependencies(TArray<FString>& OutDependencies) const
+bool UInterchangeMaterialFactoryNode::GetMetallicConnection(FString& ExpressionNodeUid, FString& OutputName) const
 {
-	TextureDependencies.GetItems(OutDependencies);
+	return UInterchangeShaderPortsAPI::GetInputConnection(this, UE::Interchange::Materials::PBR::Parameters::Metallic.ToString(), ExpressionNodeUid, OutputName);
 }
 
-void UInterchangeMaterialFactoryNode::GetTextureDependency(const int32 Index, FString& OutDependency) const
+bool UInterchangeMaterialFactoryNode::ConnectToMetallic(const FString& AttributeValue)
 {
-	TextureDependencies.GetItem(Index, OutDependency);
+	return UInterchangeShaderPortsAPI::ConnectDefaultOuputToInput(this, UE::Interchange::Materials::PBR::Parameters::Metallic.ToString(), AttributeValue);
 }
 
-bool UInterchangeMaterialFactoryNode::SetTextureDependencyUid(const FString& DependencyUid)
+bool UInterchangeMaterialFactoryNode::ConnectOutputToMetallic(const FString& ExpressionNodeUid, const FString& OutputName)
 {
-	return TextureDependencies.AddItem(DependencyUid);
+	return UInterchangeShaderPortsAPI::ConnectOuputToInput(this, UE::Interchange::Materials::PBR::Parameters::Metallic.ToString(), ExpressionNodeUid, OutputName);
 }
 
-bool UInterchangeMaterialFactoryNode::RemoveTextureDependencyUid(const FString& DependencyUid)
+bool UInterchangeMaterialFactoryNode::GetSpecularConnection(FString& ExpressionNodeUid, FString& OutputName) const
 {
-	return TextureDependencies.RemoveItem(DependencyUid);
+	return UInterchangeShaderPortsAPI::GetInputConnection(this, UE::Interchange::Materials::PBR::Parameters::Specular.ToString(), ExpressionNodeUid, OutputName);
 }
 
-bool UInterchangeMaterialFactoryNode::GetCustomTranslatedMaterialNodeUid(FString& AttributeValue) const
+bool UInterchangeMaterialFactoryNode::ConnectToSpecular(const FString& ExpressionNodeUid)
 {
-	IMPLEMENT_NODE_ATTRIBUTE_GETTER(TranslatedMaterialNodeUid, FString);
+	return UInterchangeShaderPortsAPI::ConnectDefaultOuputToInput(this, UE::Interchange::Materials::PBR::Parameters::Specular.ToString(), ExpressionNodeUid);
 }
 
-bool UInterchangeMaterialFactoryNode::SetCustomTranslatedMaterialNodeUid(const FString& AttributeValue)
+bool UInterchangeMaterialFactoryNode::ConnectOutputToSpecular(const FString& ExpressionNodeUid, const FString& OutputName)
 {
-	IMPLEMENT_NODE_ATTRIBUTE_SETTER_NODELEGATE(TranslatedMaterialNodeUid, FString);
+	return UInterchangeShaderPortsAPI::ConnectOuputToInput(this, UE::Interchange::Materials::PBR::Parameters::Specular.ToString(), ExpressionNodeUid, OutputName);
 }
 
-bool UInterchangeMaterialFactoryNode::GetCustomMaterialUsage(uint8& AttributeValue) const
+bool UInterchangeMaterialFactoryNode::GetRoughnessConnection(FString& ExpressionNodeUid, FString& OutputName) const
 {
-	IMPLEMENT_NODE_ATTRIBUTE_GETTER(MaterialUsage, uint8);
+	return UInterchangeShaderPortsAPI::GetInputConnection(this, UE::Interchange::Materials::PBR::Parameters::Roughness.ToString(), ExpressionNodeUid, OutputName);
 }
 
-bool UInterchangeMaterialFactoryNode::SetCustomMaterialUsage(const uint8& AttributeValue, bool bAddApplyDelegate /*= true*/)
+bool UInterchangeMaterialFactoryNode::ConnectToRoughness(const FString& ExpressionNodeUid)
 {
-	IMPLEMENT_NODE_ATTRIBUTE_SETTER(UInterchangeMaterialFactoryNode, MaterialUsage, uint8, UMaterialInterface);
+	return UInterchangeShaderPortsAPI::ConnectDefaultOuputToInput(this, UE::Interchange::Materials::PBR::Parameters::Roughness.ToString(), ExpressionNodeUid);
 }
 
-bool UInterchangeMaterialFactoryNode::GetCustomBlendMode(uint8& AttributeValue) const
+bool UInterchangeMaterialFactoryNode::ConnectOutputToRoughness(const FString& ExpressionNodeUid, const FString& OutputName)
 {
-	IMPLEMENT_NODE_ATTRIBUTE_GETTER(BlendMode, uint8);
+	return UInterchangeShaderPortsAPI::ConnectOuputToInput(this, UE::Interchange::Materials::PBR::Parameters::Roughness.ToString(), ExpressionNodeUid, OutputName);
 }
 
-bool UInterchangeMaterialFactoryNode::SetCustomBlendMode(const uint8& AttributeValue, bool bAddApplyDelegate /*= true*/)
+bool UInterchangeMaterialFactoryNode::GetEmissiveColorConnection(FString& ExpressionNodeUid, FString& OutputName) const
 {
-	IMPLEMENT_NODE_ATTRIBUTE_SETTER(UInterchangeMaterialFactoryNode, BlendMode, uint8, UMaterialInterface);
+	return UInterchangeShaderPortsAPI::GetInputConnection(this, UE::Interchange::Materials::Common::Parameters::EmissiveColor.ToString(), ExpressionNodeUid, OutputName);
 }
 
-void UInterchangeMaterialFactoryNode::Serialize(FArchive& Ar)
+bool UInterchangeMaterialFactoryNode::ConnectToEmissiveColor(const FString& ExpressionNodeUid)
 {
-	Super::Serialize(Ar);
-#if WITH_ENGINE
-	if (Ar.IsLoading())
-	{
-		//Make sure the class is properly set when we compile with engine, this will set the
-		//bIsMaterialNodeClassInitialized to true.
-		SetMaterialNodeClassFromClassAttribute();
-	}
-#endif //#if WITH_ENGINE
+	return UInterchangeShaderPortsAPI::ConnectDefaultOuputToInput(this, UE::Interchange::Materials::Common::Parameters::EmissiveColor.ToString(), ExpressionNodeUid);
 }
 
-void UInterchangeMaterialFactoryNode::FillAssetClassFromAttribute()
+bool UInterchangeMaterialFactoryNode::ConnectOutputToEmissiveColor(const FString& ExpressionNodeUid, const FString& OutputName)
 {
-#if WITH_ENGINE
-	FString OperationName = GetTypeName() + TEXT(".GetAssetClassName");
-	FString ClassName;
-	InterchangePrivateNodeBase::GetCustomAttribute<FString>(*Attributes, ClassNameAttributeKey, OperationName, ClassName);
-	if (ClassName.Equals(UMaterial::StaticClass()->GetName()))
-	{
-		AssetClass = UMaterial::StaticClass();
-		bIsMaterialNodeClassInitialized = true;
-	}
-	else if (ClassName.Equals(UMaterialInstance::StaticClass()->GetName()))
-	{
-		AssetClass = UMaterialInstance::StaticClass();
-		bIsMaterialNodeClassInitialized = true;
-	}
-#endif
+	return UInterchangeShaderPortsAPI::ConnectOuputToInput(this, UE::Interchange::Materials::PBR::Parameters::EmissiveColor.ToString(), ExpressionNodeUid, OutputName);
 }
 
-bool UInterchangeMaterialFactoryNode::SetMaterialNodeClassFromClassAttribute()
+bool UInterchangeMaterialFactoryNode::GetNormalConnection(FString& ExpressionNodeUid, FString& OutputName) const
 {
-	if (!bIsMaterialNodeClassInitialized)
-	{
-		FillAssetClassFromAttribute();
-	}
-	return bIsMaterialNodeClassInitialized;
+	return UInterchangeShaderPortsAPI::GetInputConnection(this, UE::Interchange::Materials::Common::Parameters::Normal.ToString(), ExpressionNodeUid, OutputName);
 }
 
-#if WITH_ENGINE
+bool UInterchangeMaterialFactoryNode::ConnectToNormal(const FString& ExpressionNodeUid)
+{
+	return UInterchangeShaderPortsAPI::ConnectDefaultOuputToInput(this, UE::Interchange::Materials::Common::Parameters::Normal.ToString(), ExpressionNodeUid);
+}
 
-bool UInterchangeMaterialFactoryNode::ApplyCustomMaterialUsageToAsset(UObject* Asset) const
-{
-	if (!Asset)
-	{
-		return false;
-	}
-	UMaterialInterface* TypedObject = Cast<UMaterialInterface>(Asset);
-	if (!TypedObject)
-	{
-		return false;
-	}
-	uint8 ValueData;
-	if (GetCustomMaterialUsage(ValueData))
-	{
-		TypedObject->CheckMaterialUsage_Concurrent(EMaterialUsage(ValueData));
-		return true;
-	}
-	return false;
-}
-
-bool UInterchangeMaterialFactoryNode::FillCustomMaterialUsageFromAsset(UObject* Asset)
-{
-	if (!Asset)
-	{
-		return false;
-	}
-	UMaterialInterface* TypedObject = Cast<UMaterialInterface>(Asset);
-	if (!TypedObject)
-	{
-		return false;
-	}
-	const UMaterial* Material = TypedObject->GetMaterial_Concurrent();
-	if (!Material)
-	{
-		return false;
-	}
-	for (int32 Usage = 0; Usage < MATUSAGE_MAX; Usage++)
-	{
-		const EMaterialUsage UsageEnum = (EMaterialUsage)Usage;
-		if (Material->GetUsageByFlag(UsageEnum))
-		{
-			if (SetCustomMaterialUsage((int8)UsageEnum, false))
-			{
-				return true;
-			}
-		}
-	}
-	return false;
-}
-
-#endif //WITH_ENGINE
\ No newline at end of file
+bool UInterchangeMaterialFactoryNode::ConnectOutputToNormal(const FString& ExpressionNodeUid, const FString& OutputName)
+{
+	return UInterchangeShaderPortsAPI::ConnectOuputToInput(this, UE::Interchange::Materials::PBR::Parameters::Normal.ToString(), ExpressionNodeUid, OutputName);
+}
+
+bool UInterchangeMaterialFactoryNode::GetOpacityConnection(FString& ExpressionNodeUid, FString& OutputName) const
+{
+	return UInterchangeShaderPortsAPI::GetInputConnection(this, UE::Interchange::Materials::Common::Parameters::Opacity.ToString(), ExpressionNodeUid, OutputName);
+}
+
+bool UInterchangeMaterialFactoryNode::ConnectToOpacity(const FString& AttributeValue)
+{
+	return UInterchangeShaderPortsAPI::ConnectDefaultOuputToInput(this, UE::Interchange::Materials::Common::Parameters::Opacity.ToString(), AttributeValue);
+}
+
+bool UInterchangeMaterialFactoryNode::ConnectOutputToOpacity(const FString& ExpressionNodeUid, const FString& OutputName)
+{
+	return UInterchangeShaderPortsAPI::ConnectOuputToInput(this, UE::Interchange::Materials::PBR::Parameters::Opacity.ToString(), ExpressionNodeUid, OutputName);
+}
+
+FString UInterchangeMaterialExpressionFactoryNode::GetTypeName() const
+{
+	const FString TypeName = TEXT("MaterialExpressionFactoryNode");
+	return TypeName;
+}
+
+bool UInterchangeMaterialExpressionFactoryNode::GetCustomExpressionClassName(FString& AttributeValue) const
+{
+	IMPLEMENT_NODE_ATTRIBUTE_GETTER(ExpressionClassName, FString);
+}
+
+bool UInterchangeMaterialExpressionFactoryNode::SetCustomExpressionClassName(const FString& AttributeValue)
+{
+	IMPLEMENT_NODE_ATTRIBUTE_SETTER_NODELEGATE(ExpressionClassName, FString);
+}
diff --git a/Engine/Plugins/Experimental/Interchange/Runtime/Source/Nodes/Private/InterchangeMaterialNode.cpp b/Engine/Plugins/Experimental/Interchange/Runtime/Source/Nodes/Private/InterchangeMaterialNode.cpp
deleted file mode 100644
diff --git a/Engine/Plugins/Experimental/Interchange/Runtime/Source/Nodes/Private/InterchangeShaderGraphNode.cpp b/Engine/Plugins/Experimental/Interchange/Runtime/Source/Nodes/Private/InterchangeShaderGraphNode.cpp
new file mode 100644
--- /dev/null
+++ b/Engine/Plugins/Experimental/Interchange/Runtime/Source/Nodes/Private/InterchangeShaderGraphNode.cpp
@@ -0,0 +1,135 @@
+// Copyright Epic Games, Inc. All Rights Reserved.
+
+#include "InterchangeShaderGraphNode.h"
+
+const TCHAR* UInterchangeShaderPortsAPI::InputPrefix = TEXT("Inputs");
+const TCHAR* UInterchangeShaderPortsAPI::InputSeparator = TEXT(":");
+
+FString UInterchangeShaderPortsAPI::MakeInputConnectionKey(const FString& InputName)
+{
+	TStringBuilder<128> StringBuilder;
+	StringBuilder.Append(InputPrefix);
+	StringBuilder.Append(InputSeparator);
+	StringBuilder.Append(InputName);
+	StringBuilder.Append(InputSeparator);
+	StringBuilder.Append(TEXT("Connect"));
+
+	return StringBuilder.ToString();
+}
+
+FString UInterchangeShaderPortsAPI::MakeInputValueKey(const FString& InputName)
+{
+	TStringBuilder<128> StringBuilder;
+	StringBuilder.Append(InputPrefix);
+	StringBuilder.Append(InputSeparator);
+	StringBuilder.Append(InputName);
+	StringBuilder.Append(InputSeparator);
+	StringBuilder.Append(TEXT("Value"));
+
+	return StringBuilder.ToString();
+}
+
+FString UInterchangeShaderPortsAPI::MakeInputName(const FString& InputKey)
+{
+	FString InputName;
+	FString Discard;
+
+	InputKey.Split(InputSeparator, &Discard, &InputName, ESearchCase::IgnoreCase, ESearchDir::FromStart);
+	InputName.Split(InputSeparator, &InputName, &Discard, ESearchCase::IgnoreCase, ESearchDir::FromStart);
+
+	return InputName;
+}
+
+bool UInterchangeShaderPortsAPI::IsAnInput(const FString& AttributeKey)
+{
+	TStringBuilder<128> StringBuilder;
+	StringBuilder.Append(InputPrefix);
+	StringBuilder.Append(InputSeparator);
+
+	return AttributeKey.StartsWith(StringBuilder.ToString());
+}
+
+bool UInterchangeShaderPortsAPI::HasInput(const UInterchangeBaseNode* InterchangeNode, const FName& InInputName)
+{
+	TArray<FString> InputNames;
+	GatherInputs(InterchangeNode, InputNames);
+
+	for (const FString& InputName : InputNames)
+	{
+		if (InputName.Equals(InInputName.ToString(), ESearchCase::IgnoreCase))
+		{
+			return true;
+		}
+	}
+
+	return false;
+}
+
+void UInterchangeShaderPortsAPI::GatherInputs(const UInterchangeBaseNode* InterchangeNode, TArray<FString>& OutInputNames)
+{
+	TArray< UE::Interchange::FAttributeKey > AttributeKeys;
+	InterchangeNode->GetAttributeKeys(AttributeKeys);
+
+	for (const UE::Interchange::FAttributeKey& AttributeKey : AttributeKeys)
+	{
+		if (IsAnInput(AttributeKey.ToString()))
+		{
+			OutInputNames.Add(MakeInputName(AttributeKey.ToString()));
+		}
+	}
+}
+
+bool UInterchangeShaderPortsAPI::ConnectDefaultOuputToInput(UInterchangeBaseNode* InterchangeNode, const FString& InputName, const FString& ExpressionUid)
+{
+	return InterchangeNode->AddStringAttribute(MakeInputConnectionKey(InputName), ExpressionUid);
+}
+
+bool UInterchangeShaderPortsAPI::ConnectOuputToInput(UInterchangeBaseNode* InterchangeNode, const FString& InputName, const FString& ExpressionUid, const FString& OutputName)
+{
+	if (OutputName.IsEmpty())
+	{
+		return ConnectDefaultOuputToInput(InterchangeNode, InputName, ExpressionUid);
+	}
+	else
+	{
+		return InterchangeNode->AddStringAttribute(MakeInputConnectionKey(InputName), ExpressionUid + InputSeparator + OutputName);
+	}
+}
+
+UE::Interchange::EAttributeTypes UInterchangeShaderPortsAPI::GetInputType(const UInterchangeBaseNode* InterchangeNode, const FString& InputName)
+{
+	return InterchangeNode->GetAttributeType(UE::Interchange::FAttributeKey(MakeInputValueKey(InputName)));
+}
+
+bool UInterchangeShaderPortsAPI::GetInputConnection(const UInterchangeBaseNode* InterchangeNode, const FString& InputName, FString& OutExpressionUid, FString& OutputName)
+{
+	if (InterchangeNode->GetStringAttribute(MakeInputConnectionKey(InputName), OutExpressionUid))
+	{
+		OutExpressionUid.Split(InputSeparator, &OutExpressionUid, &OutputName);
+		return true;
+	}
+
+	return false;
+}
+
+FString UInterchangeShaderNode::GetTypeName() const
+{
+	const FString TypeName = TEXT("ShaderNode");
+	return TypeName;
+}
+
+bool UInterchangeShaderNode::GetCustomShaderType(FString& AttributeValue) const
+{
+	IMPLEMENT_NODE_ATTRIBUTE_GETTER(ShaderType, FString);
+}
+
+bool UInterchangeShaderNode::SetCustomShaderType(const FString& AttributeValue)
+{
+	IMPLEMENT_NODE_ATTRIBUTE_SETTER_NODELEGATE(ShaderType, FString);
+}
+
+FString UInterchangeShaderGraphNode::GetTypeName() const
+{
+	const FString TypeName = TEXT("ShaderGraphNode");
+	return TypeName;
+}
diff --git a/Engine/Plugins/Experimental/Interchange/Runtime/Source/Nodes/Public/InterchangeMaterialDefinitions.h b/Engine/Plugins/Experimental/Interchange/Runtime/Source/Nodes/Public/InterchangeMaterialDefinitions.h
new file mode 100644
--- /dev/null
+++ b/Engine/Plugins/Experimental/Interchange/Runtime/Source/Nodes/Public/InterchangeMaterialDefinitions.h
@@ -0,0 +1,86 @@
+// Copyright Epic Games, Inc. All Rights Reserved.
+
+#pragma once
+
+namespace UE
+{
+namespace Interchange
+{
+	namespace Materials
+	{
+		namespace Standard
+		{
+			namespace Nodes
+			{
+				namespace Lerp
+				{
+					const FName Name = TEXT("Lerp");
+
+					namespace Inputs
+					{
+						const FName A = TEXT("A"); // Type: linear color
+						const FName B = TEXT("B"); // Type: linear color
+						const FName Factor = TEXT("Factor"); // Type: float
+					}
+				}
+
+				namespace TextureSample
+				{
+					const FName Name = TEXT("TextureSample");
+
+					namespace Inputs
+					{
+						const FName Texture = TEXT("TextureUid"); // Type: FString (unique id of a texture node)
+						const FName UTiling = TEXT("UTiling"); // Type: float
+						const FName VTiling = TEXT("VTiling"); // Type: float
+					}
+				}
+			}
+		}
+
+		namespace Common
+		{
+			namespace Parameters
+			{
+				const FName EmissiveColor = TEXT("EmissiveColor"); // Type: linear color
+				const FName Normal = TEXT("Normal"); // Type: vector3f
+				const FName Opacity = TEXT("Opacity"); // Type: float
+			}
+		}
+
+		namespace Lambert
+		{
+			namespace Parameters
+			{
+				using namespace Common::Parameters;
+
+				const FName DiffuseColor = TEXT("DiffuseColor"); // Type: linear color
+			}
+		}
+
+		namespace Phong
+		{
+			namespace Parameters
+			{
+				using namespace Lambert::Parameters;
+
+				const FName SpecularColor = TEXT("SpecularColor"); // Type: linear color
+				const FName Shininess = TEXT("Shininess"); // Type: float
+			}
+		}
+
+		namespace PBR
+		{
+			namespace Parameters
+			{
+				using namespace Common::Parameters;
+
+				const FName BaseColor = TEXT("BaseColor"); // Type: vector3
+				const FName Metallic = TEXT("Metallic"); // Type: float
+				const FName Specular = TEXT("Specular"); // Type: float
+				const FName Roughness = TEXT("Roughness"); // Type: float
+			}
+		}
+	}
+}
+}
diff --git a/Engine/Plugins/Experimental/Interchange/Runtime/Source/Nodes/Public/InterchangeMaterialFactoryNode.h b/Engine/Plugins/Experimental/Interchange/Runtime/Source/Nodes/Public/InterchangeMaterialFactoryNode.h
--- a/Engine/Plugins/Experimental/Interchange/Runtime/Source/Nodes/Public/InterchangeMaterialFactoryNode.h
+++ b/Engine/Plugins/Experimental/Interchange/Runtime/Source/Nodes/Public/InterchangeMaterialFactoryNode.h
@@ -4,7 +4,8 @@
 
 #include "CoreMinimal.h"
 #include "Nodes/InterchangeBaseNode.h"
-#include "UObject/UObjectGlobals.h"
+#include "InterchangeMaterialDefinitions.h"
+#include "InterchangeShaderGraphNode.h"
 
 #if WITH_ENGINE
 #include "Materials/MaterialInterface.h"
@@ -12,133 +13,108 @@
 
 #include "InterchangeMaterialFactoryNode.generated.h"
 
-UCLASS(BlueprintType, Experimental)
-class INTERCHANGENODES_API UInterchangeMaterialFactoryNode : public UInterchangeBaseNode
+UCLASS(Abstract, Experimental)
+class INTERCHANGENODES_API UInterchangeBaseMaterialFactoryNode : public UInterchangeBaseNode
 {
 	GENERATED_BODY()
 
 public:
-	UInterchangeMaterialFactoryNode();
+	static FString GetMaterialFactoryNodeUidFromMaterialNodeUid(const FString& TranslatedNodeUid);
+};
 
-	/**
-	 * Initialize node data
-	 * @param: UniqueID - The uniqueId for this node
-	 * @param DisplayLabel - The name of the node
-	 * @param InAssetClass - The class the material factory will create for this node.
-	 */
-	UFUNCTION(BlueprintCallable, Category = "Interchange | Node | MaterialFactory")
-	void InitializeMaterialNode(const FString& UniqueID, const FString& DisplayLabel, const FString& InAssetClass);
+UCLASS(BlueprintType, Experimental)
+class INTERCHANGENODES_API UInterchangeMaterialFactoryNode : public UInterchangeBaseMaterialFactoryNode
+{
+	GENERATED_BODY()
+
+public:
 
-	/**
-	 * Return the node type name of the class, we use this when reporting error
-	 */
 	virtual FString GetTypeName() const override;
 
-	/** Get the class this node want to create */
 	UFUNCTION(BlueprintCallable, Category = "Interchange | Node | MaterialFactory")
 	virtual class UClass* GetObjectClass() const override;
 
-	virtual FString GetKeyDisplayName(const UE::Interchange::FAttributeKey& NodeAttributeKey) const override;
-
-	virtual FString GetAttributeCategory(const UE::Interchange::FAttributeKey& NodeAttributeKey) const override;
-
-	static FString GetMaterialFactoryNodeUidFromMaterialNodeUid(const FString& TranslatedNodeUid);
+	UFUNCTION(BlueprintCallable, Category = "Interchange | Node | MaterialFactory")
+	bool GetBaseColorConnection(FString& ExpressionNodeUid, FString& OutputName) const;
 
-	/**
-	 * This function allow to retrieve the number of Texture dependencies for this object.
-	 *
-	 */
 	UFUNCTION(BlueprintCallable, Category = "Interchange | Node | MaterialFactory")
-	int32 GetTextureDependeciesCount() const;
+	bool ConnectToBaseColor(const FString& AttributeValue);
+	
+	UFUNCTION(BlueprintCallable, Category = "Interchange | Node | MaterialFactory")
+	bool ConnectOutputToBaseColor(const FString& ExpressionNodeUid, const FString& OutputName);
 
-	/**
-	 * This function allow to retrieve the Texture dependency for this object.
-	 *
-	 */
 	UFUNCTION(BlueprintCallable, Category = "Interchange | Node | MaterialFactory")
-	void GetTextureDependencies(TArray<FString>& OutDependencies) const;
+	bool GetMetallicConnection(FString& ExpressionNodeUid, FString& OutputName) const;
 
-	/**
-	 * This function allow to retrieve one Texture dependency for this object.
-	 *
-	 */
 	UFUNCTION(BlueprintCallable, Category = "Interchange | Node | MaterialFactory")
-	void GetTextureDependency(const int32 Index, FString& OutDependency) const;
+	bool ConnectToMetallic(const FString& AttributeValue);
 
-	/**
-	 * Add one Texture dependency to this object.
-	 */
 	UFUNCTION(BlueprintCallable, Category = "Interchange | Node | MaterialFactory")
-	bool SetTextureDependencyUid(const FString& DependencyUid);
+	bool ConnectOutputToMetallic(const FString& ExpressionNodeUid, const FString& OutputName);
 
-	/**
-	 * Remove one Texture dependency from this object.
-	 */
 	UFUNCTION(BlueprintCallable, Category = "Interchange | Node | MaterialFactory")
-	bool RemoveTextureDependencyUid(const FString& DependencyUid);
+	bool GetSpecularConnection(FString& ExpressionNodeUid, FString& OutputName) const;
 
-	/**
-	 * Get the translated texture node unique ID.
-	 */
 	UFUNCTION(BlueprintCallable, Category = "Interchange | Node | MaterialFactory")
-	bool GetCustomTranslatedMaterialNodeUid(FString& AttributeValue) const;
+	bool ConnectToSpecular(const FString& ExpressionNodeUid);
 
-	/**
-	 * Set the translated texture node unique ID. This is the reference to the node that was create by the translator and this node is needed to get the texture payload.
-	 */
 	UFUNCTION(BlueprintCallable, Category = "Interchange | Node | MaterialFactory")
-	bool SetCustomTranslatedMaterialNodeUid(const FString& AttributeValue);
+	bool ConnectOutputToSpecular(const FString& ExpressionNodeUid, const FString& OutputName);
 
-	/** Return false if the Attribute was not set previously.*/
 	UFUNCTION(BlueprintCallable, Category = "Interchange | Node | MaterialFactory")
-	bool GetCustomMaterialUsage(uint8& AttributeValue) const;
+	bool GetRoughnessConnection(FString& ExpressionNodeUid, FString& OutputName) const;
 
 	UFUNCTION(BlueprintCallable, Category = "Interchange | Node | MaterialFactory")
-	bool SetCustomMaterialUsage(const uint8& AttributeValue, bool bAddApplyDelegate = true);
+	bool ConnectToRoughness(const FString& ExpressionNodeUid);
 
-	/** Return false if the Attribute was not set previously.*/
 	UFUNCTION(BlueprintCallable, Category = "Interchange | Node | MaterialFactory")
-	bool GetCustomBlendMode(uint8& AttributeValue) const;
+	bool ConnectOutputToRoughness(const FString& ExpressionNodeUid, const FString& OutputName);
 
 	UFUNCTION(BlueprintCallable, Category = "Interchange | Node | MaterialFactory")
-	bool SetCustomBlendMode(const uint8& AttributeValue, bool bAddApplyDelegate = true);
+	bool GetEmissiveColorConnection(FString& ExpressionNodeUid, FString& OutputName) const;
 
-	virtual void Serialize(FArchive& Ar) override;
+	UFUNCTION(BlueprintCallable, Category = "Interchange | Node | MaterialFactory")
+	bool ConnectToEmissiveColor(const FString& ExpressionNodeUid);
 
-private:
+	UFUNCTION(BlueprintCallable, Category = "Interchange | Node | MaterialFactory")
+	bool ConnectOutputToEmissiveColor(const FString& ExpressionNodeUid, const FString& OutputName);
 
-	void FillAssetClassFromAttribute();
+	UFUNCTION(BlueprintCallable, Category = "Interchange | Node | MaterialFactory")
+	bool GetNormalConnection(FString& ExpressionNodeUid, FString& OutputName) const;
 
-	bool SetMaterialNodeClassFromClassAttribute();
+	UFUNCTION(BlueprintCallable, Category = "Interchange | Node | MaterialFactory")
+	bool ConnectToNormal(const FString& ExpressionNodeUid);
 
-	const UE::Interchange::FAttributeKey ClassNameAttributeKey = UE::Interchange::FBaseNodeStaticData::ClassTypeAttributeKey();
-	
-	//Translated texture node unique id
-	const UE::Interchange::FAttributeKey Macro_CustomTranslatedMaterialNodeUidKey = UE::Interchange::FAttributeKey(TEXT("TranslatedMaterialNodeUid"));
+	UFUNCTION(BlueprintCallable, Category = "Interchange | Node | MaterialFactory")
+	bool ConnectOutputToNormal(const FString& ExpressionNodeUid, const FString& OutputName);
 
-	//Material Attribute
-	const UE::Interchange::FAttributeKey Macro_CustomMaterialUsageKey = UE::Interchange::FAttributeKey(TEXT("MaterialUsage"));
-	const UE::Interchange::FAttributeKey Macro_CustomBlendModeKey = UE::Interchange::FAttributeKey(TEXT("BlendMode"));
+	UFUNCTION(BlueprintCallable, Category = "Interchange | Node | MaterialFactory")
+	bool GetOpacityConnection(FString& ExpressionNodeUid, FString& OutputName) const;
 
-	const UE::Interchange::FAttributeKey TextureDependenciesKey = UE::Interchange::FAttributeKey(TEXT("__TextureDependenciesKey__"));
+	UFUNCTION(BlueprintCallable, Category = "Interchange | Node | MaterialFactory")
+	bool ConnectToOpacity(const FString& AttributeValue);
 
+	UFUNCTION(BlueprintCallable, Category = "Interchange | Node | MaterialFactory")
+	bool ConnectOutputToOpacity(const FString& ExpressionNodeUid, const FString& OutputName);
+};
 
-#if WITH_ENGINE
+UCLASS(BlueprintType, Experimental)
+class INTERCHANGENODES_API UInterchangeMaterialExpressionFactoryNode : public UInterchangeBaseNode
+{
+	GENERATED_BODY()
 
-	//Material Ussage is not handle properly, you can have several usage check
-	//Currently I can set only one and retrieve the first, TODO make this work (maybe one key per usage...)
-	bool ApplyCustomMaterialUsageToAsset(UObject* Asset) const;
+public:
+	virtual FString GetTypeName() const override;
 
-	bool FillCustomMaterialUsageFromAsset(UObject* Asset);
-#endif
+public:
+	UFUNCTION(BlueprintCallable, Category = "Interchange | Node | Material")
+	bool GetCustomExpressionClassName(FString& AttributeValue) const;
 
-	IMPLEMENT_NODE_ATTRIBUTE_APPLY_UOBJECT(BlendMode, uint8, UMaterial);
+	UFUNCTION(BlueprintCallable, Category = "Interchange | Node | Material")
+	bool SetCustomExpressionClassName(const FString& AttributeValue);
 
-protected:
-#if WITH_ENGINE
-	TSubclassOf<UMaterialInterface> AssetClass = nullptr;
-#endif
-	bool bIsMaterialNodeClassInitialized = false;
+private:
+	const UE::Interchange::FAttributeKey Macro_CustomExpressionClassNameKey = UE::Interchange::FAttributeKey(TEXT("ExpressionClassName"));
 
-	UE::Interchange::TArrayAttributeHelper<FString> TextureDependencies;
 };
+
diff --git a/Engine/Plugins/Experimental/Interchange/Runtime/Source/Nodes/Public/InterchangeMaterialNode.h b/Engine/Plugins/Experimental/Interchange/Runtime/Source/Nodes/Public/InterchangeMaterialNode.h
deleted file mode 100644
diff --git a/Engine/Plugins/Experimental/Interchange/Runtime/Source/Nodes/Public/InterchangeShaderGraphNode.h b/Engine/Plugins/Experimental/Interchange/Runtime/Source/Nodes/Public/InterchangeShaderGraphNode.h
new file mode 100644
--- /dev/null
+++ b/Engine/Plugins/Experimental/Interchange/Runtime/Source/Nodes/Public/InterchangeShaderGraphNode.h
@@ -0,0 +1,132 @@
+// Copyright Epic Games, Inc. All Rights Reserved.
+ 
+#pragma once
+ 
+#include "Nodes/InterchangeBaseNode.h"
+#include "UObject/Object.h"
+ 
+#include "InterchangeShaderGraphNode.generated.h"
+ 
+/**
+ * The Shader Ports API manages a set of inputs and outputs attributes.
+ * This API can be used over any InterchangeBaseNode that wants to support shader ports as attributes.
+ */
+UCLASS(BlueprintType, Experimental)
+class INTERCHANGENODES_API UInterchangeShaderPortsAPI : public UObject
+{
+	GENERATED_BODY()
+ 
+public:
+	/**
+	 * Makes an attribute key to represent a node being connected to an input (ie: Inputs:InputName:Connect).
+	 */
+	UFUNCTION(BlueprintCallable, Category = "Interchange | Node | Material")
+	static FString MakeInputConnectionKey(const FString& InputName);
+ 
+	/**
+	 * Makes an attribute key to represent a value being given to an input (ie: Inputs:InputName:Value).
+	 */
+	UFUNCTION(BlueprintCallable, Category = "Interchange | Node | Material")
+	static FString MakeInputValueKey(const FString& InputName);
+ 
+	/**
+	 * From an attribute key associated with an input (ie: Inputs:InputName:Value), retrieves the input name from it.
+	 */
+	UFUNCTION(BlueprintCallable, Category = "Interchange | Node | Material")
+	static FString MakeInputName(const FString& InputKey);
+	
+	/**
+	 * Returns true if the attribute key is associated with an input (starts with "Inputs:").
+	 */
+	UFUNCTION(BlueprintCallable, Category = "Interchange | Node | Material")
+	static bool IsAnInput(const FString& AttributeKey);
+ 
+	/**
+	 * Checks if a particular input exists on a given node.
+	 */
+	UFUNCTION(BlueprintCallable, Category = "Interchange | Node | Material")
+	static bool HasInput(const UInterchangeBaseNode* InterchangeNode, const FName& InInputName);
+ 
+	/**
+	 * Retrieves the names of all the inputs for a given node.
+	 */
+	UFUNCTION(BlueprintCallable, Category = "Interchange | Node | Material")
+	static void GatherInputs(const UInterchangeBaseNode* InterchangeNode, TArray<FString>& OutInputNames);
+ 
+	/**
+	 * Adds an input connection attribute.
+	 * @param InterchangeNode	The Node to create the input on.
+	 * @param InputName			The name to give to the input.
+	 * @param ExpressionUid		The unique id of the node to connect to the input.
+	 * @return					true if the input connection was succesfully added to the node.
+	 */
+	UFUNCTION(BlueprintCallable, Category = "Interchange | Node | Material")
+	static bool ConnectDefaultOuputToInput(UInterchangeBaseNode* InterchangeNode, const FString& InputName, const FString& ExpressionUid);
+ 
+	/**
+	 * Adds an input connection attribute.
+	 * @param InterchangeNode	The Node to create the input on.
+	 * @param InputName			The name to give to the input.
+	 * @param ExpressionUid		The unique id of the node to connect to the input.
+	 * @param OutputName		The name of the ouput from ExpressionUid to connect to the input.
+	 * @return					true if the input connection was succesfully added to the node.
+	 */
+	UFUNCTION(BlueprintCallable, Category = "Interchange | Node | Material")
+	static bool ConnectOuputToInput(UInterchangeBaseNode* InterchangeNode, const FString& InputName, const FString& ExpressionUid, const FString& OutputName);
+ 
+	/**
+	 * Retrieves the node unique id and the ouputname connected to a given input, if any.
+	 */
+	UFUNCTION(BlueprintCallable, Category = "Interchange | Node | Material")
+	static bool GetInputConnection(const UInterchangeBaseNode* InterchangeNode, const FString& InputName, FString& OutExpressionUid, FString& OutputName);
+	
+	/**
+	 * For an input with a value, returns the type of the stored value.
+	 */
+	static UE::Interchange::EAttributeTypes GetInputType(const UInterchangeBaseNode* InterchangeNode, const FString& InputName);
+ 
+private:
+	static const TCHAR* InputPrefix;
+	static const TCHAR* InputSeparator;
+};
+ 
+/**
+ * A shader node is a named set of inputs and outputs. It can be connected to other shader nodes and finally to a shader graph input.
+ */
+UCLASS(BlueprintType, Experimental)
+class INTERCHANGENODES_API UInterchangeShaderNode : public UInterchangeBaseNode
+{
+	GENERATED_BODY()
+ 
+public:
+	virtual FString GetTypeName() const override;
+ 
+public:
+	UFUNCTION(BlueprintCallable, Category = "Interchange | Node | Material")
+	bool GetCustomShaderType(FString& AttributeValue) const;
+ 
+	/**
+	 * Sets which type of shader this nodes represents. Can be arbitrary or one of the predefined shader types.
+	 * The material pipeline handling the shader node should be aware of the shader type that is being set here.
+	 */
+	UFUNCTION(BlueprintCallable, Category = "Interchange | Node | Material")
+	bool SetCustomShaderType(const FString& AttributeValue);
+ 
+private:
+	const UE::Interchange::FAttributeKey Macro_CustomShaderTypeKey = UE::Interchange::FAttributeKey(TEXT("ShaderType"));
+};
+ 
+/**
+ * A shader graph has its own set of inputs on which shader nodes can be connected to.
+ */
+UCLASS(BlueprintType, Experimental)
+class INTERCHANGENODES_API UInterchangeShaderGraphNode : public UInterchangeShaderNode
+{
+	GENERATED_BODY()
+ 
+public:
+	/**
+	 * Return the node type name of the class, we use this when reporting error
+	 */
+	virtual FString GetTypeName() const override;
+};
\ No newline at end of file
diff --git a/Engine/Plugins/Experimental/Interchange/Runtime/Source/Parsers/Fbx/Private/FbxAPI.cpp b/Engine/Plugins/Experimental/Interchange/Runtime/Source/Parsers/Fbx/Private/FbxAPI.cpp
--- a/Engine/Plugins/Experimental/Interchange/Runtime/Source/Parsers/Fbx/Private/FbxAPI.cpp
+++ b/Engine/Plugins/Experimental/Interchange/Runtime/Source/Parsers/Fbx/Private/FbxAPI.cpp
@@ -11,7 +11,6 @@
 #include "FbxMaterial.h"
 #include "FbxMesh.h"
 #include "FbxScene.h"
-#include "InterchangeMaterialNode.h"
 #include "InterchangeTextureNode.h"
 #include "Nodes/InterchangeBaseNodeContainer.h"
 
diff --git a/Engine/Plugins/Experimental/Interchange/Runtime/Source/Parsers/Fbx/Private/FbxConvert.cpp b/Engine/Plugins/Experimental/Interchange/Runtime/Source/Parsers/Fbx/Private/FbxConvert.cpp
--- a/Engine/Plugins/Experimental/Interchange/Runtime/Source/Parsers/Fbx/Private/FbxConvert.cpp
+++ b/Engine/Plugins/Experimental/Interchange/Runtime/Source/Parsers/Fbx/Private/FbxConvert.cpp
@@ -169,6 +169,17 @@ namespace UE
 				return FVector(Vector[0], -Vector[1], Vector[2]);
 			}
 
+			FLinearColor FFbxConvert::ConvertColor(const FbxDouble3& Color)
+			{
+				FLinearColor LinearColor;
+				LinearColor.R =(float)Color[0];
+				LinearColor.G =(float)Color[1];
+				LinearColor.B =(float)Color[2];
+				LinearColor.A = 1.f;
+
+				return LinearColor;
+			}
+
 			FString FFbxConvert::MakeName(const ANSICHAR* Name)
 			{
 				const TCHAR SpecialChars[] = { TEXT('.'), TEXT(','), TEXT('/'), TEXT('`'), TEXT('%') };
diff --git a/Engine/Plugins/Experimental/Interchange/Runtime/Source/Parsers/Fbx/Private/FbxConvert.h b/Engine/Plugins/Experimental/Interchange/Runtime/Source/Parsers/Fbx/Private/FbxConvert.h
--- a/Engine/Plugins/Experimental/Interchange/Runtime/Source/Parsers/Fbx/Private/FbxConvert.h
+++ b/Engine/Plugins/Experimental/Interchange/Runtime/Source/Parsers/Fbx/Private/FbxConvert.h
@@ -29,6 +29,7 @@ namespace UE
 				static FRotator ConvertRotation(FbxQuaternion Quaternion);
 				static FVector ConvertPos(const FbxVector4& Vector);
 				static FVector ConvertDir(const FbxVector4& Vector);
+				static FLinearColor ConvertColor(const FbxDouble3& Color);
 
 				static FTransform AdjustCameraTransform(const FTransform& Transform);
 				static FTransform AdjustLightTransform(const FTransform& Transform);
diff --git a/Engine/Plugins/Experimental/Interchange/Runtime/Source/Parsers/Fbx/Private/FbxMaterial.cpp b/Engine/Plugins/Experimental/Interchange/Runtime/Source/Parsers/Fbx/Private/FbxMaterial.cpp
--- a/Engine/Plugins/Experimental/Interchange/Runtime/Source/Parsers/Fbx/Private/FbxMaterial.cpp
+++ b/Engine/Plugins/Experimental/Interchange/Runtime/Source/Parsers/Fbx/Private/FbxMaterial.cpp
@@ -2,19 +2,18 @@
 
 #include "FbxMaterial.h"
 
-#include "CoreMinimal.h"
 #include "FbxAPI.h"
 #include "FbxConvert.h"
 #include "FbxHelper.h"
 #include "FbxInclude.h"
-#include "InterchangeMaterialNode.h"
+#include "InterchangeMaterialDefinitions.h"
 #include "InterchangeResultsContainer.h"
 #include "InterchangeSceneNode.h"
+#include "InterchangeShaderGraphNode.h"
 #include "InterchangeTexture2DNode.h"
 #include "InterchangeTextureNode.h"
 #include "Misc/Paths.h"
 #include "Nodes/InterchangeBaseNodeContainer.h"
-#include "Misc/Paths.h"
 
 #define LOCTEXT_NAMESPACE "InterchangeFbxMaterial"
 
@@ -24,52 +23,154 @@ namespace UE
 	{
 		namespace Private
 		{
-			UInterchangeMaterialNode* FFbxMaterial::CreateMaterialNode(UInterchangeBaseNodeContainer& NodeContainer, const FString& NodeUid, const FString& NodeName)
+			UInterchangeShaderGraphNode* FFbxMaterial::CreateShaderGraphNode(UInterchangeBaseNodeContainer& NodeContainer, const FString& NodeUid, const FString& NodeName)
 			{
-				UInterchangeMaterialNode* MaterialNode = NewObject<UInterchangeMaterialNode>(&NodeContainer, NAME_None);
-				if (!ensure(MaterialNode))
-				{
-					UInterchangeResultError_Generic* Message = Parser.AddMessage<UInterchangeResultError_Generic>();
-					Message->Text = LOCTEXT("CannotAllocateNode", "Cannot allocate a node when importing FBX.");
-					return nullptr;
-				}
-				// Creating a UMaterialInterface
-				MaterialNode->InitializeNode(NodeUid, NodeName, EInterchangeNodeContainerType::TranslatedAsset);
-				MaterialNode->SetPayLoadKey(NodeUid);
-				NodeContainer.AddNode(MaterialNode);
-				return MaterialNode;
+				UInterchangeShaderGraphNode* ShaderGraphNode = NewObject<UInterchangeShaderGraphNode>(&NodeContainer);
+				ShaderGraphNode->InitializeNode(NodeUid, NodeName, EInterchangeNodeContainerType::TranslatedAsset);
+				NodeContainer.AddNode(ShaderGraphNode);
+
+				return ShaderGraphNode;
+			}
+
+			UInterchangeShaderNode* FFbxMaterial::CreateShaderNode(UInterchangeBaseNodeContainer& NodeContainer, const FString& ParentUid, const FString& NodeName)
+			{
+				UInterchangeShaderNode* ShaderNode = NewObject<UInterchangeShaderNode>(&NodeContainer);
+
+				const FString NodeUid = ParentUid + TEXT("\\") + NodeName;
+				ShaderNode->InitializeNode(NodeUid, NodeName, EInterchangeNodeContainerType::TranslatedAsset);
+				NodeContainer.AddNode(ShaderNode);
+				NodeContainer.SetNodeParentUid(NodeUid, ParentUid);
+
+				return ShaderNode;
 			}
 
 			UInterchangeTexture2DNode* FFbxMaterial::CreateTexture2DNode(UInterchangeBaseNodeContainer& NodeContainer, const FString& NodeUid, const FString& TextureFilePath)
 			{
 				FString DisplayLabel = FPaths::GetBaseFilename(TextureFilePath);
 				UInterchangeTexture2DNode* TextureNode = NewObject<UInterchangeTexture2DNode>(&NodeContainer, NAME_None);
-				if (!ensure(TextureNode))
-				{
-					UInterchangeResultError_Generic* Message = Parser.AddMessage<UInterchangeResultError_Generic>();
-					Message->Text = LOCTEXT("CannotAllocateNode", "Cannot allocate a node when importing FBX.");
-					return nullptr;
-				}
-				// Creating a UTexture2D
 				TextureNode->InitializeNode(NodeUid, DisplayLabel, EInterchangeNodeContainerType::TranslatedAsset);
+
 				//All texture translator expect a file has the payload key
 				FString NormalizeFilePath = TextureFilePath;
 				FPaths::NormalizeFilename(NormalizeFilePath);
 				TextureNode->SetPayLoadKey(NormalizeFilePath);
 				NodeContainer.AddNode(TextureNode);
+
 				return TextureNode;
 			}
 
-			UInterchangeMaterialNode* FFbxMaterial::AddNodeMaterial(FbxSurfaceMaterial* SurfaceMaterial, UInterchangeBaseNodeContainer& NodeContainer)
+			void FFbxMaterial::ConvertPropertyToShaderNode(UInterchangeBaseNodeContainer& NodeContainer, UInterchangeShaderGraphNode* ShaderGraphNode, FbxProperty& Property, float Factor, FName InputName, TVariant<FLinearColor, float> DefaultValue)
 			{
+				using namespace Materials::Standard::Nodes;
+
+				UInterchangeShaderNode* LerpNode = nullptr;
+
+				const int32 TextureCount = Property.GetSrcObjectCount<FbxFileTexture>();
+				const FbxDataType DataType = Property.GetPropertyDataType();
+
+				if (!FMath::IsNearlyEqual(Factor, 1.f))
+				{
+					FString LerpNodeName = InputName.ToString() + TEXT("Lerp");
+					LerpNode = CreateShaderNode(NodeContainer, ShaderGraphNode->GetUniqueID(), LerpNodeName);
+					LerpNode->SetCustomShaderType(Lerp::Name.ToString());
+
+					if (TextureCount == 0)
+					{
+						if (DataType.GetType() == eFbxDouble || DataType.GetType() == eFbxFloat || DataType.GetType() == eFbxInt)
+						{
+							LerpNode->AddFloatAttribute(UInterchangeShaderPortsAPI::MakeInputValueKey(Lerp::Inputs::B.ToString()), DefaultValue.Get<float>());
+						}
+						else if (DataType.GetType() == eFbxDouble3)
+						{
+							LerpNode->AddLinearColorAttribute(UInterchangeShaderPortsAPI::MakeInputValueKey(Lerp::Inputs::B.ToString()), DefaultValue.Get<FLinearColor>());
+						}
+					}
+					else
+					{
+						LerpNode->AddLinearColorAttribute(UInterchangeShaderPortsAPI::MakeInputValueKey(Lerp::Inputs::B.ToString()), DefaultValue.Get<FLinearColor>());
+					}
+
+					const float InverseFactor = 1.f - Factor; // We lerp from A to B and prefer to put the strongest input in A so we need to flip the lerp factor
+					LerpNode->AddFloatAttribute(UInterchangeShaderPortsAPI::MakeInputValueKey(Lerp::Inputs::Factor.ToString()), InverseFactor);
+
+					UInterchangeShaderPortsAPI::ConnectDefaultOuputToInput(ShaderGraphNode, InputName.ToString(), LerpNode->GetUniqueID());
+				}
+
+				// Handles max one texture per property.
+				if (TextureCount > 0)
+				{
+					FbxFileTexture* FbxTextureFilePath = Property.GetSrcObject<FbxFileTexture>(0);
+
+					FString TextureFilename = UTF8_TO_TCHAR(FbxTextureFilePath->GetFileName());
+					FString TextureName = FPaths::GetBaseFilename(TextureFilename);
+
+					UInterchangeShaderNode* TextureSampleShader = NewObject<UInterchangeShaderNode>(&NodeContainer);
+					TextureSampleShader->SetCustomShaderType(TextureSample::Name.ToString());
+
+					FString TextureSampleShaderUid = ShaderGraphNode->GetUniqueID() + TEXT("\\Textures\\") + TextureName;
+					TextureSampleShader->InitializeNode(TextureSampleShaderUid, TextureName, EInterchangeNodeContainerType::TranslatedAsset);
+
+					NodeContainer.AddNode(TextureSampleShader);
+					NodeContainer.SetNodeParentUid(TextureSampleShaderUid, ShaderGraphNode->GetUniqueID());
+
+					FString TextureNodeUid = TEXT("\\Texture\\") + TextureFilename;
+					UInterchangeTexture2DNode* TextureNode = Cast<UInterchangeTexture2DNode>(NodeContainer.GetNode(TextureNodeUid));
+					if (!TextureNode)
+					{
+						TextureNode = CreateTexture2DNode(NodeContainer, TextureNodeUid, TextureFilename);
+					}
+
+					TextureSampleShader->AddStringAttribute(UInterchangeShaderPortsAPI::MakeInputValueKey(TextureSample::Inputs::Texture.ToString()), TextureNode->GetUniqueID());
+					TextureSampleShader->AddFloatAttribute(UInterchangeShaderPortsAPI::MakeInputValueKey(TextureSample::Inputs::UTiling.ToString()), (float)FbxTextureFilePath->GetScaleU());
+					TextureSampleShader->AddFloatAttribute(UInterchangeShaderPortsAPI::MakeInputValueKey(TextureSample::Inputs::VTiling.ToString()), (float)FbxTextureFilePath->GetScaleV());
+
+					if (LerpNode)
+					{
+						UInterchangeShaderPortsAPI::ConnectDefaultOuputToInput(LerpNode, Lerp::Inputs::A.ToString(), TextureSampleShaderUid);
+					}
+					else
+					{
+						UInterchangeShaderPortsAPI::ConnectDefaultOuputToInput(ShaderGraphNode, InputName.ToString(), TextureSampleShaderUid);
+					}
+				}
+				else if (DataType.GetType() == eFbxDouble || DataType.GetType() == eFbxFloat || DataType.GetType() == eFbxInt)
+				{
+					if (LerpNode)
+					{
+						LerpNode->AddFloatAttribute(UInterchangeShaderPortsAPI::MakeInputValueKey(Lerp::Inputs::A.ToString()), Property.Get<float>());
+					}
+					else
+					{
+						ShaderGraphNode->AddFloatAttribute(UInterchangeShaderPortsAPI::MakeInputValueKey(InputName.ToString()), Property.Get<float>());
+					}
+				}
+				else if (DataType.GetType() == eFbxDouble3)
+				{
+					const FLinearColor PropertyValue = FFbxConvert::ConvertColor(Property.Get<FbxDouble3>());
+
+					if (LerpNode)
+					{
+						LerpNode->AddLinearColorAttribute(UInterchangeShaderPortsAPI::MakeInputValueKey(Lerp::Inputs::A.ToString()), PropertyValue);
+					}
+					else
+					{
+						ShaderGraphNode->AddLinearColorAttribute(UInterchangeShaderPortsAPI::MakeInputValueKey(InputName.ToString()), PropertyValue);
+					}
+				}
+			}
+
+			UInterchangeShaderGraphNode* FFbxMaterial::AddShaderGraphNode(FbxSurfaceMaterial* SurfaceMaterial, UInterchangeBaseNodeContainer& NodeContainer)
+			{
+				using namespace UE::Interchange::Materials;
+
 				//Create a material node
 				FString MaterialName = FFbxHelper::GetFbxObjectName(SurfaceMaterial);
 				FString NodeUid = TEXT("\\Material\\") + MaterialName;
-				UInterchangeMaterialNode* MaterialNode = Cast<UInterchangeMaterialNode>(NodeContainer.GetNode(NodeUid));
-				if (!MaterialNode)
+				UInterchangeShaderGraphNode* ShaderGraphNode = Cast<UInterchangeShaderGraphNode>(NodeContainer.GetNode(NodeUid));
+				if (!ShaderGraphNode)
 				{
-					MaterialNode = CreateMaterialNode(NodeContainer, NodeUid, MaterialName);
-					if (MaterialNode == nullptr)
+					ShaderGraphNode = CreateShaderGraphNode(NodeContainer, NodeUid, MaterialName);
+					if (ShaderGraphNode == nullptr)
 					{
 						FFormatNamedArguments Args
 						{
@@ -80,114 +181,78 @@ namespace UE
 						return nullptr;
 					}
 
-					auto SetMaterialParameter = [this, &NodeContainer, &MaterialNode, &SurfaceMaterial, &MaterialName](const char* FbxMaterialProperty, EInterchangeMaterialNodeParameterName MaterialParameterName)
+					// for each material property, add an input value or a shader node
+					FbxProperty CurrentProperty = SurfaceMaterial->GetFirstProperty();
+					while (CurrentProperty.IsValid())
 					{
-						bool bSetMaterial = false;
-						FbxProperty FbxProperty = SurfaceMaterial->FindProperty(FbxMaterialProperty);
-						if (FbxProperty.IsValid())
+						if (CurrentProperty.Modified())
 						{
-							int32 UnsupportedTextureCount = FbxProperty.GetSrcObjectCount<FbxLayeredTexture>();
-							UnsupportedTextureCount += FbxProperty.GetSrcObjectCount<FbxProceduralTexture>();
-							int32 TextureCount = FbxProperty.GetSrcObjectCount<FbxFileTexture>();
-							bool bFoundValidTexture = false;
-							if (UnsupportedTextureCount > 0)
+							const FName PropertyName(CurrentProperty.GetNameAsCStr());
+
+							// Convert FBX property names to standard Interchange Phong parameter names when possible
 							{
-								FFormatNamedArguments Args
+								if (PropertyName == FbxSurfaceMaterial::sDiffuse)
 								{
-									{ TEXT("MaterialName"), FText::FromString(MaterialName) }
-								};
-								UInterchangeResultWarning_Generic* Message = Parser.AddMessage<UInterchangeResultWarning_Generic>();
-								Message->Text = FText::Format(LOCTEXT("TextureTypeNotSupported", "Layered or procedural textures are not supported (material '{MaterialName}')."), Args);
-							}
-							else if (TextureCount > 0)
-							{
-								for (int32 TextureIndex = 0; TextureIndex < TextureCount; ++TextureIndex)
+									const float Factor = (float)((FbxSurfaceLambert*)SurfaceMaterial)->DiffuseFactor.Get();
+									TVariant<FLinearColor, float> DefaultValue;
+									DefaultValue.Set<FLinearColor>(FLinearColor::Black);
+									ConvertPropertyToShaderNode(NodeContainer, ShaderGraphNode, CurrentProperty, Factor, Phong::Parameters::DiffuseColor, DefaultValue);
+								}
+								else if (PropertyName == FbxSurfaceMaterial::sSpecular)
 								{
-									FbxFileTexture* FbxTextureFilePath = FbxProperty.GetSrcObject<FbxFileTexture>(TextureIndex);
-									FString TextureFilename = UTF8_TO_TCHAR(FbxTextureFilePath->GetFileName());
-									//Only import texture that exist on disk
-									if (!FPaths::FileExists(TextureFilename))
-									{
-										continue;
-									}
-									//Create a texture node and make it child of the material node
-									TArray<FString> JsonErrorMessage;
-									FString NodeUid = TEXT("\\Texture\\") + TextureFilename;
-									UInterchangeTexture2DNode* TextureNode = Cast<UInterchangeTexture2DNode>(NodeContainer.GetNode(NodeUid));
-									if (!TextureNode)
-									{
-										TextureNode = CreateTexture2DNode(NodeContainer, NodeUid, TextureFilename);
-									}
-									// add/find UVSet and set it to the texture, we pass index 0 here, I think pipeline should be able to get the UVIndex channel from the name
-									// and modify the parameter to set the correct value.
-// 									FbxString UVSetName = FbxTextureFilePath->UVSet.Get();
-// 									FString LocalUVSetName = UTF8_TO_TCHAR(UVSetName.Buffer());
-									const int32 UVChannelIndex = 0;
-									float ScaleU = (float)FbxTextureFilePath->GetScaleU();
-									float ScaleV = (float)FbxTextureFilePath->GetScaleV();
-									MaterialNode->AddTextureParameterData(MaterialParameterName, TextureNode->GetUniqueID(), UVChannelIndex, ScaleU, ScaleV);
-									MaterialNode->SetTextureDependencyUid(TextureNode->GetUniqueID());
-									bSetMaterial = true;
-									bFoundValidTexture = true;
+									const float Factor = (float)((FbxSurfacePhong*)SurfaceMaterial)->SpecularFactor.Get();
+									TVariant<FLinearColor, float> DefaultValue;
+									DefaultValue.Set<FLinearColor>(FLinearColor::Black);
+									ConvertPropertyToShaderNode(NodeContainer, ShaderGraphNode, CurrentProperty, Factor, Phong::Parameters::SpecularColor, DefaultValue);
 								}
-							}
-							
-							if (!bFoundValidTexture && MaterialParameterName == EInterchangeMaterialNodeParameterName::BaseColor)
-							{
-								//We support only the base color has a vector color
-								//TODO support all basic attributes has vector or scalar
-								FbxDouble3 FbxColor;
-								bool bFoundColor = true;
-								if (SurfaceMaterial->GetClassId().Is(FbxSurfacePhong::ClassId))
+								else if (PropertyName == FbxSurfaceMaterial::sShininess)
 								{
-									FbxColor = ((FbxSurfacePhong&)(*SurfaceMaterial)).Diffuse.Get();
+									const float Factor = 1.f;
+									TVariant<FLinearColor, float> DefaultValue;
+									DefaultValue.Set<float>(0.2f);
+									ConvertPropertyToShaderNode(NodeContainer, ShaderGraphNode, CurrentProperty, Factor, Phong::Parameters::Shininess, DefaultValue);
 								}
-								else if (SurfaceMaterial->GetClassId().Is(FbxSurfaceLambert::ClassId))
+								else if (PropertyName == FbxSurfaceMaterial::sEmissive)
 								{
-									FbxColor = ((FbxSurfaceLambert&)(*SurfaceMaterial)).Diffuse.Get();
+									const float Factor = (float)((FbxSurfaceLambert*)SurfaceMaterial)->EmissiveFactor.Get();
+									TVariant<FLinearColor, float> DefaultValue;
+									DefaultValue.Set<FLinearColor>(FLinearColor::Black);
+									ConvertPropertyToShaderNode(NodeContainer, ShaderGraphNode, CurrentProperty, Factor, Phong::Parameters::EmissiveColor, DefaultValue);
 								}
-								else
+								else if (PropertyName == FbxSurfaceMaterial::sNormalMap)
 								{
-									bFoundColor = false;
+									const float Factor = 1.f;
+									TVariant<FLinearColor, float> DefaultValue;
+									DefaultValue.Set<FLinearColor>(FLinearColor::Black);
+									ConvertPropertyToShaderNode(NodeContainer, ShaderGraphNode, CurrentProperty, Factor, Phong::Parameters::Normal, DefaultValue);
 								}
-
-								FVector ColorData;
-								if (bFoundColor)
+								else if (PropertyName == FbxSurfaceMaterial::sTransparentColor)
+								{
+									const float OpacityFactor = 1.f - (float)((FbxSurfaceLambert*)SurfaceMaterial)->TransparencyFactor.Get();
+									TVariant<FLinearColor, float> DefaultValue;
+									DefaultValue.Set<float>(1.f); // Opaque
+									ConvertPropertyToShaderNode(NodeContainer, ShaderGraphNode, CurrentProperty, OpacityFactor, Phong::Parameters::Opacity, DefaultValue);
+								}
+								else if (PropertyName == FbxSurfaceMaterial::sDiffuseFactor || PropertyName == FbxSurfaceMaterial::sSpecularFactor ||
+										 PropertyName == FbxSurfaceMaterial::sTransparencyFactor)
 								{
-									ColorData[0] = (float)(FbxColor[0]);
-									ColorData[1] = (float)(FbxColor[1]);
-									ColorData[2] = (float)(FbxColor[2]);
+									// Skip factors as they were processed with their corresponding inputs
 								}
 								else
 								{
-									// use random color because there may be multiple materials, then they can be different 
-									ColorData[0] = 0.5f + (0.5f * FMath::Rand()) / static_cast<float>(RAND_MAX);
-									ColorData[1] = 0.5f + (0.5f * FMath::Rand()) / static_cast<float>(RAND_MAX);
-									ColorData[2] = 0.5f + (0.5f * FMath::Rand()) / static_cast<float>(RAND_MAX);
+									const float DummyFactor = 1.f;
+									TVariant<FLinearColor, float> DummyDefaultValue;
+									DummyDefaultValue.Set<float>(1.f); // Will be ignore because we have a factor of 1
+									ConvertPropertyToShaderNode(NodeContainer, ShaderGraphNode, CurrentProperty, DummyFactor, PropertyName, DummyDefaultValue);
 								}
-								MaterialNode->AddVectorParameterData(MaterialParameterName, ColorData);
-								bSetMaterial = true;
 							}
 						}
-						return bSetMaterial;
-					};
-
-					SetMaterialParameter(FbxSurfaceMaterial::sDiffuse, EInterchangeMaterialNodeParameterName::BaseColor);
-					SetMaterialParameter(FbxSurfaceMaterial::sEmissive, EInterchangeMaterialNodeParameterName::EmissiveColor);
-					SetMaterialParameter(FbxSurfaceMaterial::sSpecular, EInterchangeMaterialNodeParameterName::Specular);
-					SetMaterialParameter(FbxSurfaceMaterial::sSpecularFactor, EInterchangeMaterialNodeParameterName::Roughness);
-					SetMaterialParameter(FbxSurfaceMaterial::sShininess, EInterchangeMaterialNodeParameterName::Metallic);
-					if (!SetMaterialParameter(FbxSurfaceMaterial::sNormalMap, EInterchangeMaterialNodeParameterName::Normal))
-					{
-						SetMaterialParameter(FbxSurfaceMaterial::sBump, EInterchangeMaterialNodeParameterName::Normal);
-					}
-					if (SetMaterialParameter(FbxSurfaceMaterial::sTransparentColor, EInterchangeMaterialNodeParameterName::Opacity))
-					{
-						SetMaterialParameter(FbxSurfaceMaterial::sTransparencyFactor, EInterchangeMaterialNodeParameterName::OpacityMask);
+
+						CurrentProperty = SurfaceMaterial->GetNextProperty(CurrentProperty);
 					}
 				}
 
-				return MaterialNode;
+				return ShaderGraphNode;
 			}
 
 			void FFbxMaterial::AddAllTextures(FbxScene* SDKScene, UInterchangeBaseNodeContainer& NodeContainer)
@@ -219,9 +284,9 @@ namespace UE
 				for (int32 MaterialIndex = 0; MaterialIndex < MaterialCount; ++MaterialIndex)
 				{
 					FbxSurfaceMaterial* SurfaceMaterial = ParentFbxNode->GetMaterial(MaterialIndex);
-					UInterchangeMaterialNode* MaterialNode = AddNodeMaterial(SurfaceMaterial, NodeContainer);
+					UInterchangeShaderGraphNode* ShaderGraphNode = AddShaderGraphNode(SurfaceMaterial, NodeContainer);
 					//The dependencies order is important because mesh will use index in that order to determine material use by a face
-					SceneNode->AddMaterialDependencyUid(MaterialNode->GetUniqueID());
+					SceneNode->AddMaterialDependencyUid(ShaderGraphNode->GetUniqueID());
 				}
 			}
 
@@ -231,7 +296,7 @@ namespace UE
 				for (int32 MaterialIndex = 0; MaterialIndex < MaterialCount; ++MaterialIndex)
 				{
 					FbxSurfaceMaterial* SurfaceMaterial = SDKScene->GetMaterial(MaterialIndex);
-					AddNodeMaterial(SurfaceMaterial, NodeContainer);
+					AddShaderGraphNode(SurfaceMaterial, NodeContainer);
 				}
 			}
 		} //ns Private
diff --git a/Engine/Plugins/Experimental/Interchange/Runtime/Source/Parsers/Fbx/Private/FbxMaterial.h b/Engine/Plugins/Experimental/Interchange/Runtime/Source/Parsers/Fbx/Private/FbxMaterial.h
--- a/Engine/Plugins/Experimental/Interchange/Runtime/Source/Parsers/Fbx/Private/FbxMaterial.h
+++ b/Engine/Plugins/Experimental/Interchange/Runtime/Source/Parsers/Fbx/Private/FbxMaterial.h
@@ -4,11 +4,13 @@
 
 #include "CoreMinimal.h"
 #include "FbxInclude.h"
+#include "Misc/TVariant.h"
 
 /** Forward declarations */
 class UInterchangeBaseNode;
 class UInterchangeBaseNodeContainer;
-class UInterchangeMaterialNode;
+class UInterchangeShaderGraphNode;
+class UInterchangeShaderNode;
 class UInterchangeTexture2DNode;
 class UInterchangeSceneNode;
 
@@ -30,14 +32,14 @@ namespace UE
 				/**
 				 * Create a UInterchangeFextureNode and add it to the NodeContainer for all texture of type FbxFileTexture the fbx file contain.
 				 *
-				 * @note - Any node that already exist in the NodeContainer will not be created or modify.
+				 * @note - Any node that already exist in the NodeContainer will not be created or modified.
 				 */
 				void AddAllTextures(FbxScene* SDKScene, UInterchangeBaseNodeContainer& NodeContainer);
 				
 				/**
 				 * Create a UInterchangeMaterialNode and add it to the NodeContainer for all material of type FbxSurfaceMaterial the fbx file contain.
 				 * 
-				 * @note - Any node that already exist in the NodeContainer will not be created or modify.
+				 * @note - Any node that already exist in the NodeContainer will not be created or modified.
 				 */
 				void AddAllMaterials(FbxScene* SDKScene, UInterchangeBaseNodeContainer& NodeContainer);
 
@@ -50,11 +52,14 @@ namespace UE
 				void AddAllNodeMaterials(UInterchangeSceneNode* SceneNode, FbxNode* ParentFbxNode, UInterchangeBaseNodeContainer& NodeContainer);
 
 			protected:
-				UInterchangeMaterialNode* CreateMaterialNode(UInterchangeBaseNodeContainer& NodeContainer, const FString& NodeUID, const FString& NodeName);
+				UInterchangeShaderGraphNode* CreateShaderGraphNode(UInterchangeBaseNodeContainer& NodeContainer, const FString& NodeUID, const FString& NodeName);
+				UInterchangeShaderNode* CreateShaderNode(UInterchangeBaseNodeContainer& NodeContainer, const FString& NodeUid, const FString& NodeName);
 				UInterchangeTexture2DNode* CreateTexture2DNode(UInterchangeBaseNodeContainer& NodeContainer, const FString& NodeUID, const FString& TextureFilePath);
+				void ConvertPropertyToShaderNode(UInterchangeBaseNodeContainer& NodeContainer, UInterchangeShaderGraphNode* ShaderGraphNode,
+					FbxProperty& Property, float Factor, FName PropertyName, TVariant<FLinearColor, float> DefaultValue);
 
 			private:
-				UInterchangeMaterialNode* AddNodeMaterial(FbxSurfaceMaterial* SurfaceMaterial, UInterchangeBaseNodeContainer& NodeContainer);
+				UInterchangeShaderGraphNode* AddShaderGraphNode(FbxSurfaceMaterial* SurfaceMaterial, UInterchangeBaseNodeContainer& NodeContainer);
 
 				FFbxParser& Parser;
 			};
diff --git a/Engine/Plugins/Experimental/Interchange/Runtime/Source/Pipelines/Private/InterchangeGenericAssetsPipeline.cpp b/Engine/Plugins/Experimental/Interchange/Runtime/Source/Pipelines/Private/InterchangeGenericAssetsPipeline.cpp
--- a/Engine/Plugins/Experimental/Interchange/Runtime/Source/Pipelines/Private/InterchangeGenericAssetsPipeline.cpp
+++ b/Engine/Plugins/Experimental/Interchange/Runtime/Source/Pipelines/Private/InterchangeGenericAssetsPipeline.cpp
@@ -5,7 +5,6 @@
 #include "Animation/Skeleton.h"
 #include "CoreMinimal.h"
 #include "InterchangeMaterialFactoryNode.h"
-#include "InterchangeMaterialNode.h"
 #include "InterchangeMeshNode.h"
 #include "InterchangePipelineLog.h"
 #include "InterchangePipelineMeshesUtilities.h"
@@ -25,6 +24,11 @@
 #include "UObject/Object.h"
 #include "UObject/ObjectMacros.h"
 
+UInterchangeGenericAssetsPipeline::UInterchangeGenericAssetsPipeline()
+{
+	MaterialPipeline = CreateDefaultSubobject<UInterchangeGenericMaterialPipeline>("MaterialPipeline");
+}
+
 void UInterchangeGenericAssetsPipeline::PreDialogCleanup(const FName PipelineStackName)
 {
 	//We always clean the pipeline skeleton when showing the dialog
@@ -60,10 +64,6 @@ void UInterchangeGenericAssetsPipeline::ExecutePreImportPipeline(UInterchangeBas
 				{
 					TextureNodes.Add(TextureNode);
 				}
-				else if (UInterchangeMaterialNode* MaterialNode = Cast<UInterchangeMaterialNode>(Node))
-				{
-					MaterialNodes.Add(MaterialNode);
-				}
 			}
 			break;
 		}
@@ -78,25 +78,17 @@ void UInterchangeGenericAssetsPipeline::ExecutePreImportPipeline(UInterchangeBas
 		}
 	}
 
-	if (bImportMaterials)
+	if (MaterialPipeline)
 	{
-		//import materials
-		for (const UInterchangeMaterialNode* MaterialNode : MaterialNodes)
-		{
-			if (UInterchangeMaterialFactoryNode* MaterialFactoryNode = CreateMaterialFactoryNode(MaterialNode))
-			{
-				//By default we do not create the materials, every node with mesh attribute can enable them. So we wont create unused materials.
-				MaterialFactoryNode->SetEnabled(false);
-			}
-		}
+		MaterialPipeline->ExecutePreImportPipeline(InBaseNodeContainer, InSourceDatas);
 	}
 
 	//Create skeletalmesh factory nodes
 	ExecutePreImportPipelineSkeletalMesh();
 
-	if (bImportStaticMeshes && (ForceAllMeshHasType == EInterchangeForceMeshType::IFMT_None || ForceAllMeshHasType == EInterchangeForceMeshType::IFMT_StaticMesh))
+	if (bImportStaticMeshes && (ForceAllMeshAsType == EInterchangeForceMeshType::IFMT_None || ForceAllMeshAsType == EInterchangeForceMeshType::IFMT_StaticMesh))
 	{
-		const bool bConvertSkeletalMeshToStaticMesh = (ForceAllMeshHasType == EInterchangeForceMeshType::IFMT_StaticMesh);
+		const bool bConvertSkeletalMeshToStaticMesh = (ForceAllMeshAsType == EInterchangeForceMeshType::IFMT_StaticMesh);
 		if (bCombineStaticMeshes)
 		{
 			// Combine all the static meshes
@@ -247,52 +239,6 @@ void UInterchangeGenericAssetsPipeline::ExecutePostImportPipeline(const UInterch
 	PostImportTextureAssetImport(CreatedAsset, bIsAReimport);
 }
 
-UInterchangeMaterialFactoryNode* UInterchangeGenericAssetsPipeline::CreateMaterialFactoryNode(const UInterchangeMaterialNode* MaterialNode)
-{
-	FString DisplayLabel = MaterialNode->GetDisplayLabel();
-	FString NodeUid = UInterchangeMaterialFactoryNode::GetMaterialFactoryNodeUidFromMaterialNodeUid(MaterialNode->GetUniqueID());
-	UInterchangeMaterialFactoryNode* MaterialFactoryNode = nullptr;
-	if (BaseNodeContainer->IsNodeUidValid(NodeUid))
-	{
-		//The node already exist, just return it
-		MaterialFactoryNode = Cast<UInterchangeMaterialFactoryNode>(BaseNodeContainer->GetNode(NodeUid));
-		if (!ensure(MaterialFactoryNode))
-		{
-			//Log an error
-		}
-	}
-	else
-	{
-		MaterialFactoryNode = NewObject<UInterchangeMaterialFactoryNode>(BaseNodeContainer, NAME_None);
-		if (!ensure(MaterialFactoryNode))
-		{
-			return nullptr;
-		}
-		//Creating a Material
-		MaterialFactoryNode->InitializeMaterialNode(NodeUid, DisplayLabel, UMaterial::StaticClass()->GetName());
-		MaterialFactoryNode->SetCustomTranslatedMaterialNodeUid(MaterialNode->GetUniqueID());
-		TArray<FString> TranslatedDependencies;
-		MaterialNode->GetTextureDependencies(TranslatedDependencies);
-		for (const FString& DepUid : TranslatedDependencies)
-		{
-			FString FactoryDepUid = UInterchangeTextureFactoryNode::GetTextureFactoryNodeUidFromTextureNodeUid(DepUid);
-			if (UInterchangeTextureFactoryNode* TextureDep = Cast<UInterchangeTextureFactoryNode>(BaseNodeContainer->GetNode(FactoryDepUid)))
-			{
-				//We found a factory dependency, because we need to be able to retrieve the UObject create by the factory
-				MaterialFactoryNode->SetTextureDependencyUid(FactoryDepUid);
-				//We also add it to the base dependencies so the texture factory will be called before this material factory
-				MaterialFactoryNode->AddFactoryDependencyUid(FactoryDepUid);
-			}
-		}
-		BaseNodeContainer->AddNode(MaterialFactoryNode);
-		MaterialFactoryNodes.Add(MaterialFactoryNode);
-		MaterialFactoryNode->AddTargetNodeUid(MaterialNode->GetUniqueID());
-		MaterialNode->AddTargetNodeUid(MaterialFactoryNode->GetUniqueID());
-	}
-	return MaterialFactoryNode;
-}
-
-
 void UInterchangeGenericAssetsPipeline::ImplementUseSourceNameForAssetOption()
 {
 	const UClass* SkeletalMeshFactoryNodeClass = UInterchangeSkeletalMeshFactoryNode::StaticClass();
@@ -323,6 +269,7 @@ void UInterchangeGenericAssetsPipeline::ImplementUseSourceNameForAssetOption()
 	}
 
 }
+
 bool UInterchangeGenericAssetsPipeline::MakeMeshFactoryNodeUidAndDisplayLabel(const TMap<int32, TArray<FString>>& MeshUidsPerLodIndex, int32 LodIndex, FString& NewNodeUid, FString& DisplayLabel)
 {
 	int32 SceneNodeCount = 0;
diff --git a/Engine/Plugins/Experimental/Interchange/Runtime/Source/Pipelines/Private/InterchangeGenericMaterialPipeline.cpp b/Engine/Plugins/Experimental/Interchange/Runtime/Source/Pipelines/Private/InterchangeGenericMaterialPipeline.cpp
new file mode 100644
--- /dev/null
+++ b/Engine/Plugins/Experimental/Interchange/Runtime/Source/Pipelines/Private/InterchangeGenericMaterialPipeline.cpp
@@ -0,0 +1,637 @@
+// Copyright Epic Games, Inc. All Rights Reserved. 
+
+#include "InterchangeGenericMaterialPipeline.h"
+
+#include "CoreMinimal.h"
+
+#include "InterchangeMaterialDefinitions.h"
+#include "InterchangeMaterialFactoryNode.h"
+#include "InterchangeShaderGraphNode.h"
+#include "InterchangeTexture2DNode.h"
+#include "InterchangeTexture2DArrayNode.h"
+#include "InterchangeTextureCubeNode.h"
+#include "InterchangeTextureNode.h"
+#include "InterchangePipelineLog.h"
+#include "InterchangeSourceData.h"
+
+#include "Materials/MaterialExpressionLinearInterpolate.h"
+#include "Materials/MaterialExpressionOneMinus.h"
+#include "Materials/MaterialExpressionScalarParameter.h"
+#include "Materials/MaterialExpressionTextureCoordinate.h"
+#include "Materials/MaterialExpressionTextureSample.h"
+#include "Materials/MaterialExpressionTextureSampleParameter2D.h"
+#include "Materials/MaterialExpressionTextureSampleParameter2DArray.h"
+#include "Materials/MaterialExpressionTextureSampleParameterCube.h"
+#include "Materials/MaterialExpressionVectorParameter.h"
+#include "Misc/Paths.h"
+#include "Nodes/InterchangeBaseNode.h"
+#include "Nodes/InterchangeBaseNodeContainer.h"
+#include "Templates/SubclassOf.h"
+#include "UObject/Object.h"
+#include "UObject/ObjectMacros.h"
+
+void UInterchangeGenericMaterialPipeline::ExecutePreImportPipeline(UInterchangeBaseNodeContainer* InBaseNodeContainer, const TArray<UInterchangeSourceData*>& InSourceDatas)
+{
+	if (!InBaseNodeContainer)
+	{
+		UE_LOG(LogInterchangePipeline, Warning, TEXT("UInterchangeGenericMaterialPipeline: Cannot execute pre-import pipeline because InBaseNodeContrainer is null"));
+		return;
+	}
+
+	BaseNodeContainer = InBaseNodeContainer;
+	SourceDatas.Empty(InSourceDatas.Num());
+	for (const UInterchangeSourceData* SourceData : InSourceDatas)
+	{
+		SourceDatas.Add(SourceData);
+	}
+	
+	//Find all translated node we need for this pipeline
+	BaseNodeContainer->IterateNodes([this](const FString& NodeUid, UInterchangeBaseNode* Node)
+	{
+		switch(Node->GetNodeContainerType())
+		{
+			case EInterchangeNodeContainerType::TranslatedAsset:
+			{
+				if (UInterchangeShaderGraphNode* MaterialNode = Cast<UInterchangeShaderGraphNode>(Node))
+				{
+					MaterialNodes.Add(MaterialNode);
+				}
+			}
+			break;
+		}
+	});
+
+	//import materials
+	if (MaterialImport == EInterchangeMaterialImportOption::ImportAsMaterials)
+	{
+		for (const UInterchangeShaderGraphNode* ShaderGraphNode : MaterialNodes)
+		{
+			if (UInterchangeMaterialFactoryNode* MaterialFactoryNode = CreateMaterialFactoryNode(ShaderGraphNode))
+			{
+				//By default we do not create the materials, every node with mesh attribute can enable them. So we wont create unused materials.
+				MaterialFactoryNode->SetEnabled(false);
+			}
+		}
+	}
+
+	return;
+}
+
+UInterchangeBaseMaterialFactoryNode* UInterchangeGenericMaterialPipeline::CreateBaseMaterialFactoryNode(const UInterchangeBaseNode* MaterialNode, TSubclassOf<UInterchangeBaseMaterialFactoryNode> NodeType)
+{
+	FString DisplayLabel = MaterialNode->GetDisplayLabel();
+	const FString NodeUid = UInterchangeMaterialFactoryNode::GetMaterialFactoryNodeUidFromMaterialNodeUid(MaterialNode->GetUniqueID());
+	UInterchangeBaseMaterialFactoryNode* MaterialFactoryNode = nullptr;
+	if (BaseNodeContainer->IsNodeUidValid(NodeUid))
+	{
+		//The node already exist, just return it
+		MaterialFactoryNode = Cast<UInterchangeBaseMaterialFactoryNode>(BaseNodeContainer->GetNode(NodeUid));
+		if (!ensure(MaterialFactoryNode))
+		{
+			//Log an error
+		}
+	}
+	else
+	{
+		MaterialFactoryNode = NewObject<UInterchangeBaseMaterialFactoryNode>(BaseNodeContainer, NodeType.Get(), NAME_None);
+		if (!ensure(MaterialFactoryNode))
+		{
+			return nullptr;
+		}
+		//Creating a Material
+		MaterialFactoryNode->InitializeNode(NodeUid, DisplayLabel, EInterchangeNodeContainerType::FactoryData);
+		
+		BaseNodeContainer->AddNode(MaterialFactoryNode);
+		MaterialFactoryNodes.Add(MaterialFactoryNode);
+		MaterialFactoryNode->AddTargetNodeUid(MaterialNode->GetUniqueID());
+		MaterialNode->AddTargetNodeUid(MaterialFactoryNode->GetUniqueID());
+	}
+	return MaterialFactoryNode;
+}
+
+bool UInterchangeGenericMaterialPipeline::HandlePhongModel(const UInterchangeShaderGraphNode* ShaderGraphNode, UInterchangeMaterialFactoryNode* MaterialFactoryNode)
+{
+	using namespace UE::Interchange::Materials::Phong;
+
+	const bool bHasDiffuseInput = UInterchangeShaderPortsAPI::HasInput(ShaderGraphNode, Parameters::DiffuseColor);
+	const bool bHasSpecularInput = UInterchangeShaderPortsAPI::HasInput(ShaderGraphNode, Parameters::SpecularColor);
+
+	if (bHasDiffuseInput && bHasSpecularInput)
+	{
+		// ConvertFromDiffSpec function call
+		UInterchangeMaterialExpressionFactoryNode* FunctionCallExpression = NewObject<UInterchangeMaterialExpressionFactoryNode>(BaseNodeContainer, NAME_None);
+		FunctionCallExpression->SetCustomExpressionClassName(UMaterialExpressionMaterialFunctionCall::StaticClass()->GetName());
+		FString FunctionCallExpressionUid = MaterialFactoryNode->GetUniqueID() + TEXT("\\Inputs\\BaseColor\\DiffSpecFunc");
+		FunctionCallExpression->InitializeNode(FunctionCallExpressionUid, TEXT("DiffSpecFunc"), EInterchangeNodeContainerType::FactoryData);
+
+		BaseNodeContainer->AddNode(FunctionCallExpression);
+		BaseNodeContainer->SetNodeParentUid(FunctionCallExpressionUid, MaterialFactoryNode->GetUniqueID());
+
+		const FName MaterialFunctionMemberName = GET_MEMBER_NAME_CHECKED(UMaterialExpressionMaterialFunctionCall, MaterialFunction);
+
+		FunctionCallExpression->AddStringAttribute(
+			MaterialFunctionMemberName.ToString(),
+			TEXT("MaterialFunction'/Engine/Functions/Engine_MaterialFunctions01/Shading/ConvertFromDiffSpec.ConvertFromDiffSpec'"));
+		FunctionCallExpression->AddApplyAndFillDelegates<FString>(MaterialFunctionMemberName.ToString(), UMaterialExpressionMaterialFunctionCall::StaticClass(), MaterialFunctionMemberName);
+
+		MaterialFactoryNode->ConnectOutputToBaseColor(FunctionCallExpressionUid, UE::Interchange::Materials::PBR::Parameters::BaseColor.ToString());
+		MaterialFactoryNode->ConnectOutputToMetallic(FunctionCallExpressionUid, UE::Interchange::Materials::PBR::Parameters::Metallic.ToString());
+		MaterialFactoryNode->ConnectOutputToSpecular(FunctionCallExpressionUid, UE::Interchange::Materials::PBR::Parameters::Specular.ToString());
+		
+		// Diffuse
+		{
+			TTuple<UInterchangeMaterialExpressionFactoryNode*, FString> DiffuseExpressionFactoryNode =
+				CreateMaterialExpressionForInput(MaterialFactoryNode, ShaderGraphNode, Parameters::DiffuseColor.ToString(), FunctionCallExpression->GetUniqueID());
+
+			if (DiffuseExpressionFactoryNode.Get<0>())
+			{
+				UInterchangeShaderPortsAPI::ConnectOuputToInput(FunctionCallExpression, TEXT("DiffuseColor"),
+					DiffuseExpressionFactoryNode.Get<0>()->GetUniqueID(), DiffuseExpressionFactoryNode.Get<1>());
+			}
+		}
+
+		// Specular
+		{
+			TTuple<UInterchangeMaterialExpressionFactoryNode*, FString> SpecularExpressionFactoryNode =
+				CreateMaterialExpressionForInput(MaterialFactoryNode, ShaderGraphNode, Parameters::SpecularColor.ToString(), FunctionCallExpression->GetUniqueID());
+
+			if (SpecularExpressionFactoryNode.Get<0>())
+			{
+				UInterchangeShaderPortsAPI::ConnectOuputToInput(FunctionCallExpression, TEXT("SpecularColor"),
+					SpecularExpressionFactoryNode.Get<0>()->GetUniqueID(), SpecularExpressionFactoryNode.Get<1>());
+			}
+		}
+		
+		// Shininess
+		{
+			const bool bHasShininessInput = UInterchangeShaderPortsAPI::HasInput(ShaderGraphNode, Parameters::Shininess);
+			if (bHasShininessInput)
+			{
+				TTuple<UInterchangeMaterialExpressionFactoryNode*, FString> ShininessExpressionFactoryNode =
+					CreateMaterialExpressionForInput(MaterialFactoryNode, ShaderGraphNode, Parameters::Shininess.ToString(), MaterialFactoryNode->GetUniqueID());
+
+				if (ShininessExpressionFactoryNode.Get<0>())
+				{
+					UInterchangeMaterialExpressionFactoryNode* InverseShininessNode =
+						CreateExpressionNode(TEXT("InverseShininess"), ShininessExpressionFactoryNode.Get<0>()->GetUniqueID(), UMaterialExpressionOneMinus::StaticClass());
+
+					UInterchangeShaderPortsAPI::ConnectOuputToInput(InverseShininessNode, GET_MEMBER_NAME_CHECKED(UMaterialExpressionOneMinus, Input).ToString(),
+						ShininessExpressionFactoryNode.Get<0>()->GetUniqueID(), ShininessExpressionFactoryNode.Get<1>());
+
+					MaterialFactoryNode->ConnectToRoughness(InverseShininessNode->GetUniqueID());
+				}
+			}
+		}
+
+		return true;
+	}
+
+	return false;
+}
+
+bool UInterchangeGenericMaterialPipeline::HandleLambertModel(const UInterchangeShaderGraphNode* ShaderGraphNode, UInterchangeMaterialFactoryNode* MaterialFactoryNode)
+{
+	using namespace UE::Interchange::Materials::Lambert;
+
+	// Diffuse
+	const bool bHasDiffuseInput = UInterchangeShaderPortsAPI::HasInput(ShaderGraphNode, Parameters::DiffuseColor);
+
+	if (bHasDiffuseInput)
+	{
+		TTuple<UInterchangeMaterialExpressionFactoryNode*, FString> DiffuseExpressionFactoryNode =
+			CreateMaterialExpressionForInput(MaterialFactoryNode, ShaderGraphNode, Parameters::DiffuseColor.ToString(), MaterialFactoryNode->GetUniqueID());
+
+		if (DiffuseExpressionFactoryNode.Get<0>())
+		{
+			MaterialFactoryNode->ConnectOutputToBaseColor(DiffuseExpressionFactoryNode.Get<0>()->GetUniqueID(), DiffuseExpressionFactoryNode.Get<1>());
+		}
+
+		return true;
+	}
+
+	return false;
+}
+
+bool UInterchangeGenericMaterialPipeline::HandlePBRModel(const UInterchangeShaderGraphNode* ShaderGraphNode, UInterchangeMaterialFactoryNode* MaterialFactoryNode)
+{
+	using namespace UE::Interchange::Materials::PBR;
+
+	// BaseColor
+	{
+		const bool bHasInput = UInterchangeShaderPortsAPI::HasInput(ShaderGraphNode, Parameters::BaseColor);
+
+		if (bHasInput)
+		{
+			TTuple<UInterchangeMaterialExpressionFactoryNode*, FString> ExpressionFactoryNode =
+				CreateMaterialExpressionForInput(MaterialFactoryNode, ShaderGraphNode, Parameters::BaseColor.ToString(), MaterialFactoryNode->GetUniqueID());
+
+			if (ExpressionFactoryNode.Get<0>())
+			{
+				MaterialFactoryNode->ConnectOutputToBaseColor(ExpressionFactoryNode.Get<0>()->GetUniqueID(), ExpressionFactoryNode.Get<1>());
+			}
+
+			return true;
+		}
+	}
+
+	// Metallic
+	{
+		const bool bHasInput = UInterchangeShaderPortsAPI::HasInput(ShaderGraphNode, Parameters::Metallic);
+
+		if (bHasInput)
+		{
+			TTuple<UInterchangeMaterialExpressionFactoryNode*, FString> ExpressionFactoryNode =
+				CreateMaterialExpressionForInput(MaterialFactoryNode, ShaderGraphNode, Parameters::Metallic.ToString(), MaterialFactoryNode->GetUniqueID());
+
+			if (ExpressionFactoryNode.Get<0>())
+			{
+				MaterialFactoryNode->ConnectOutputToMetallic(ExpressionFactoryNode.Get<0>()->GetUniqueID(), ExpressionFactoryNode.Get<1>());
+			}
+
+			return true;
+		}
+	}
+
+	// Specular
+	{
+		const bool bHasInput = UInterchangeShaderPortsAPI::HasInput(ShaderGraphNode, Parameters::Specular);
+
+		if (bHasInput)
+		{
+			TTuple<UInterchangeMaterialExpressionFactoryNode*, FString> ExpressionFactoryNode =
+				CreateMaterialExpressionForInput(MaterialFactoryNode, ShaderGraphNode, Parameters::Specular.ToString(), MaterialFactoryNode->GetUniqueID());
+
+			if (ExpressionFactoryNode.Get<0>())
+			{
+				MaterialFactoryNode->ConnectOutputToSpecular(ExpressionFactoryNode.Get<0>()->GetUniqueID(), ExpressionFactoryNode.Get<1>());
+			}
+
+			return true;
+		}
+	}
+
+	// Roughness
+	{
+		const bool bHasInput = UInterchangeShaderPortsAPI::HasInput(ShaderGraphNode, Parameters::Roughness);
+
+		if (bHasInput)
+		{
+			TTuple<UInterchangeMaterialExpressionFactoryNode*, FString> ExpressionFactoryNode =
+				CreateMaterialExpressionForInput(MaterialFactoryNode, ShaderGraphNode, Parameters::Roughness.ToString(), MaterialFactoryNode->GetUniqueID());
+
+			if (ExpressionFactoryNode.Get<0>())
+			{
+				MaterialFactoryNode->ConnectOutputToRoughness(ExpressionFactoryNode.Get<0>()->GetUniqueID(), ExpressionFactoryNode.Get<1>());
+			}
+
+			return true;
+		}
+	}
+
+	return false;
+}
+
+bool UInterchangeGenericMaterialPipeline::HandleCommonParameters(const UInterchangeShaderGraphNode* ShaderGraphNode, UInterchangeMaterialFactoryNode* MaterialFactoryNode)
+{
+	using namespace UE::Interchange::Materials::Common;
+
+	// Emissive
+	{
+		const bool bHasEmissiveInput = UInterchangeShaderPortsAPI::HasInput(ShaderGraphNode, Parameters::EmissiveColor);
+
+		if (bHasEmissiveInput)
+		{
+			TTuple<UInterchangeMaterialExpressionFactoryNode*, FString> EmissiveExpressionFactoryNode =
+				CreateMaterialExpressionForInput(MaterialFactoryNode, ShaderGraphNode, Parameters::EmissiveColor.ToString(), MaterialFactoryNode->GetUniqueID());
+
+			if (EmissiveExpressionFactoryNode.Get<0>())
+			{
+				MaterialFactoryNode->ConnectOutputToEmissiveColor(EmissiveExpressionFactoryNode.Get<0>()->GetUniqueID(), EmissiveExpressionFactoryNode.Get<1>());
+			}
+
+			return true;
+		}
+	}
+
+	// Normal
+	{
+		const bool bHasNormalInput = UInterchangeShaderPortsAPI::HasInput(ShaderGraphNode, Parameters::Normal);
+
+		if (bHasNormalInput)
+		{
+			TTuple<UInterchangeMaterialExpressionFactoryNode*, FString> ExpressionFactoryNode =
+				CreateMaterialExpressionForInput(MaterialFactoryNode, ShaderGraphNode, Parameters::Normal.ToString(), MaterialFactoryNode->GetUniqueID());
+
+			if (ExpressionFactoryNode.Get<0>())
+			{
+				MaterialFactoryNode->ConnectOutputToNormal(ExpressionFactoryNode.Get<0>()->GetUniqueID(), ExpressionFactoryNode.Get<1>());
+			}
+
+			return true;
+		}
+	}
+
+	// Opacity
+	{
+		const bool bHasOpacityInput = UInterchangeShaderPortsAPI::HasInput(ShaderGraphNode, Parameters::Opacity);
+
+		if (bHasOpacityInput)
+		{
+			TTuple<UInterchangeMaterialExpressionFactoryNode*, FString> OpacityExpressionFactoryNode =
+				CreateMaterialExpressionForInput(MaterialFactoryNode, ShaderGraphNode, Parameters::Opacity.ToString(), MaterialFactoryNode->GetUniqueID());
+
+			if (OpacityExpressionFactoryNode.Get<0>())
+			{
+				MaterialFactoryNode->ConnectOutputToOpacity(OpacityExpressionFactoryNode.Get<0>()->GetUniqueID(), OpacityExpressionFactoryNode.Get<1>());
+			}
+
+			return true;
+		}
+	}
+
+	return false;
+}
+
+void UInterchangeGenericMaterialPipeline::HandleTextureSampleNode(const UInterchangeShaderNode* ShaderNode, UInterchangeMaterialExpressionFactoryNode* TextureSampleFactoryNode)
+{
+	using namespace UE::Interchange::Materials::Standard;
+
+	FString TextureUid;
+	ShaderNode->GetStringAttribute(UInterchangeShaderPortsAPI::MakeInputValueKey(Nodes::TextureSample::Inputs::Texture.ToString()), TextureUid);
+
+	FString ExpressionClassName;
+
+	if (UInterchangeTextureNode* TextureNode = Cast<UInterchangeTextureNode>(BaseNodeContainer->GetNode(TextureUid)))
+	{
+		if (TextureNode->IsA<UInterchangeTextureCubeNode>())
+		{
+			ExpressionClassName = UMaterialExpressionTextureSampleParameterCube::StaticClass()->GetName();
+		}
+		else if (TextureNode->IsA<UInterchangeTexture2DArrayNode>())
+		{
+			ExpressionClassName = UMaterialExpressionTextureSampleParameter2DArray::StaticClass()->GetName();
+		}
+		else if (TextureNode->IsA<UInterchangeTexture2DNode>())
+		{
+			ExpressionClassName = UMaterialExpressionTextureSampleParameter2D::StaticClass()->GetName();
+		}
+		else
+		{
+			ExpressionClassName = UMaterialExpressionTextureSampleParameter2D::StaticClass()->GetName();
+		}
+	}
+
+	TextureSampleFactoryNode->SetCustomExpressionClassName(ExpressionClassName);
+
+	HandleTextureCoordinates(ShaderNode, TextureSampleFactoryNode);
+}
+
+void UInterchangeGenericMaterialPipeline::HandleTextureCoordinates(const UInterchangeShaderNode* ShaderNode, UInterchangeMaterialExpressionFactoryNode* TextureSampleFactoryNode)
+{
+	using namespace UE::Interchange::Materials::Standard;
+
+	float UTilingValue = 1.f;
+	bool bHasUTiling = ShaderNode->GetFloatAttribute(UInterchangeShaderPortsAPI::MakeInputValueKey(Nodes::TextureSample::Inputs::UTiling.ToString()), UTilingValue);
+	bHasUTiling = bHasUTiling && (UTilingValue != GetDefault<UMaterialExpressionTextureCoordinate>()->UTiling);
+
+	float VTilingValue = 1.f;
+	bool bHasVTiling = ShaderNode->GetFloatAttribute(UInterchangeShaderPortsAPI::MakeInputValueKey(Nodes::TextureSample::Inputs::VTiling.ToString()), VTilingValue);
+	bHasVTiling = bHasVTiling && (VTilingValue != GetDefault<UMaterialExpressionTextureCoordinate>()->VTiling);
+
+	if (bHasUTiling || bHasVTiling)
+	{
+		UInterchangeMaterialExpressionFactoryNode* TextureCoordinate = NewObject<UInterchangeMaterialExpressionFactoryNode>(BaseNodeContainer);
+		const FString TextureCoordinateUid = TextureSampleFactoryNode->GetUniqueID() + TEXT("Coordinate");
+
+		TextureCoordinate->InitializeNode(TextureCoordinateUid, ShaderNode->GetDisplayLabel(), EInterchangeNodeContainerType::FactoryData);
+		BaseNodeContainer->AddNode(TextureCoordinate);
+		BaseNodeContainer->SetNodeParentUid(TextureCoordinateUid, TextureSampleFactoryNode->GetUniqueID());
+
+		TextureCoordinate->SetCustomExpressionClassName(UMaterialExpressionTextureCoordinate::StaticClass()->GetName());
+
+		TextureCoordinate->AddFloatAttribute(GET_MEMBER_NAME_CHECKED(UMaterialExpressionTextureCoordinate, UTiling).ToString(), UTilingValue);
+		TextureCoordinate->AddFloatAttribute(GET_MEMBER_NAME_CHECKED(UMaterialExpressionTextureCoordinate, VTiling).ToString(), VTilingValue);
+
+		UInterchangeShaderPortsAPI::ConnectDefaultOuputToInput(TextureSampleFactoryNode, GET_MEMBER_NAME_CHECKED(UMaterialExpressionTextureSample, Coordinates).ToString(), TextureCoordinate->GetUniqueID());
+	}
+}
+
+void UInterchangeGenericMaterialPipeline::HandleLerpNode(const UInterchangeShaderNode* ShaderNode, UInterchangeMaterialFactoryNode* MaterialFactoryNode, UInterchangeMaterialExpressionFactoryNode* LerpFactoryNode)
+{
+	using namespace UE::Interchange::Materials::Standard;
+
+	LerpFactoryNode->SetCustomExpressionClassName(UMaterialExpressionLinearInterpolate::StaticClass()->GetName());
+
+	TTuple<UInterchangeMaterialExpressionFactoryNode*, FString> ColorAExpression =
+		CreateMaterialExpressionForInput(MaterialFactoryNode, ShaderNode, Nodes::Lerp::Inputs::A.ToString(), LerpFactoryNode->GetUniqueID());
+
+	if (ColorAExpression.Get<0>())
+	{
+		UInterchangeShaderPortsAPI::ConnectOuputToInput(LerpFactoryNode, Nodes::Lerp::Inputs::A.ToString(), ColorAExpression.Get<0>()->GetUniqueID(), ColorAExpression.Get<1>());
+	}
+
+	TTuple<UInterchangeMaterialExpressionFactoryNode*, FString> ColorBExpression =
+		CreateMaterialExpressionForInput(MaterialFactoryNode, ShaderNode, Nodes::Lerp::Inputs::B.ToString(), LerpFactoryNode->GetUniqueID());
+
+	if (ColorBExpression.Get<0>())
+	{
+		UInterchangeShaderPortsAPI::ConnectOuputToInput(LerpFactoryNode, Nodes::Lerp::Inputs::B.ToString(), ColorBExpression.Get<0>()->GetUniqueID(), ColorBExpression.Get<1>());
+	}
+
+	float LerpFactor = 0.5f;
+	ShaderNode->GetFloatAttribute(UInterchangeShaderPortsAPI::MakeInputValueKey(Nodes::Lerp::Inputs::Factor.ToString()), LerpFactor);
+
+	const FName ConstAlphaMemberName = GET_MEMBER_NAME_CHECKED(UMaterialExpressionLinearInterpolate, ConstAlpha);
+	LerpFactoryNode->AddFloatAttribute(ConstAlphaMemberName.ToString(), LerpFactor);
+	LerpFactoryNode->AddApplyAndFillDelegates<float>(ConstAlphaMemberName.ToString(), UMaterialExpressionLinearInterpolate::StaticClass(), ConstAlphaMemberName);
+}
+
+UInterchangeMaterialExpressionFactoryNode* UInterchangeGenericMaterialPipeline::CreateMaterialExpressionForShaderNode(UInterchangeMaterialFactoryNode* MaterialFactoryNode,
+	const UInterchangeShaderNode* ShaderNode, const FString& ParentUid)
+{
+	using namespace UE::Interchange::Materials::Standard;
+
+	// If we recognize the shader node type
+	// - Create material expression for specific node type
+	//
+	// If we don't recognize the shader node type
+	// - Create material expression by trying to match the node type to a material expression class name
+
+	const FString MaterialExpressionUid = TEXT("Factory_") + ShaderNode->GetUniqueID();
+
+	UInterchangeMaterialExpressionFactoryNode* MaterialExpression = Cast<UInterchangeMaterialExpressionFactoryNode>(BaseNodeContainer->GetNode(MaterialExpressionUid));
+	if (MaterialExpression != nullptr)
+	{
+		return MaterialExpression;
+	}
+
+	MaterialExpression = NewObject<UInterchangeMaterialExpressionFactoryNode>(BaseNodeContainer);
+	MaterialExpression->InitializeNode(MaterialExpressionUid, ShaderNode->GetDisplayLabel(), EInterchangeNodeContainerType::FactoryData);
+	BaseNodeContainer->AddNode(MaterialExpression);
+
+	FString ShaderType;
+	ShaderNode->GetCustomShaderType(ShaderType);
+
+	if (*ShaderType == Nodes::TextureSample::Name)
+	{
+		HandleTextureSampleNode(ShaderNode, MaterialExpression);
+	}
+	else if (*ShaderType == Nodes::Lerp::Name)
+	{
+		HandleLerpNode(ShaderNode, MaterialFactoryNode, MaterialExpression);
+	}
+	else
+	{
+		const FString ExpressionClassName = TEXT("MaterialExpression") + ShaderType;
+		MaterialExpression->SetCustomExpressionClassName(ExpressionClassName);
+
+		TArray<FString> Inputs;
+		UInterchangeShaderPortsAPI::GatherInputs(ShaderNode, Inputs);
+
+		for (const FString& InputName : Inputs)
+		{
+			TTuple<UInterchangeMaterialExpressionFactoryNode*, FString> InputExpression =
+				CreateMaterialExpressionForInput(MaterialFactoryNode, ShaderNode, InputName, MaterialExpressionUid);
+
+			if (InputExpression.Get<0>())
+			{
+				UInterchangeShaderPortsAPI::ConnectOuputToInput(MaterialExpression, InputName, InputExpression.Get<0>()->GetUniqueID(), InputExpression.Get<1>());
+			}
+		}
+	}
+
+	if (!ParentUid.IsEmpty())
+	{
+		BaseNodeContainer->SetNodeParentUid(MaterialExpressionUid, ParentUid);
+	}
+
+	MaterialExpression->AddTargetNodeUid(ShaderNode->GetUniqueID());
+
+	if (*ShaderType == Nodes::TextureSample::Name)
+	{
+		FString TextureUid;
+		ShaderNode->GetStringAttribute(UInterchangeShaderPortsAPI::MakeInputValueKey(Nodes::TextureSample::Inputs::Texture.ToString()), TextureUid);
+
+		// Make the material factory node have a dependency on the texture factory node so that the texture asset gets created first
+		if (UInterchangeTextureNode* TextureNode = Cast<UInterchangeTextureNode>(BaseNodeContainer->GetNode(TextureUid)))
+		{
+			TArray<FString> TextureNodeTargets;
+			TextureNode->GetTargetNodeUids(TextureNodeTargets);
+
+			if (TextureNodeTargets.Num() > 0)
+			{
+				FString TextureFactoryNodeUid = TextureNodeTargets[0];
+
+				if (BaseNodeContainer->IsNodeUidValid(TextureFactoryNodeUid))
+				{
+					TArray<FString> FactoryDependencies;
+					MaterialFactoryNode->GetFactoryDependencies(FactoryDependencies);
+					if (!FactoryDependencies.Contains(TextureFactoryNodeUid))
+					{
+						MaterialFactoryNode->AddFactoryDependencyUid(TextureFactoryNodeUid);
+					}
+				}
+			}
+		}
+	}
+
+	return MaterialExpression;
+}
+
+UInterchangeMaterialExpressionFactoryNode* UInterchangeGenericMaterialPipeline::CreateExpressionNode(const FString& ExpressionName, const FString& ParentUid, UClass* MaterialExpressionClass)
+{
+	const FString MaterialExpressionUid = ParentUid + TEXT("\\") + ExpressionName;
+
+	UInterchangeMaterialExpressionFactoryNode* MaterialExpressionFactoryNode = NewObject<UInterchangeMaterialExpressionFactoryNode>(BaseNodeContainer);
+	MaterialExpressionFactoryNode->SetCustomExpressionClassName(MaterialExpressionClass->GetName());
+	MaterialExpressionFactoryNode->InitializeNode(MaterialExpressionUid, ExpressionName, EInterchangeNodeContainerType::FactoryData);
+	BaseNodeContainer->AddNode(MaterialExpressionFactoryNode);
+	BaseNodeContainer->SetNodeParentUid(MaterialExpressionUid, ParentUid);
+
+	return MaterialExpressionFactoryNode;
+}
+
+UInterchangeMaterialExpressionFactoryNode* UInterchangeGenericMaterialPipeline::CreateScalarParameterExpression(const UInterchangeShaderNode* ShaderNode, const FString& InputName, const FString& ParentUid)
+{
+	UInterchangeMaterialExpressionFactoryNode* MaterialExpressionFactoryNode = CreateExpressionNode(InputName, ParentUid, UMaterialExpressionScalarParameter::StaticClass());
+
+	float InputValue;
+	if (ShaderNode->GetFloatAttribute(UInterchangeShaderPortsAPI::MakeInputValueKey(InputName), InputValue))
+	{
+		const FName DefaultValueMemberName = GET_MEMBER_NAME_CHECKED(UMaterialExpressionScalarParameter, DefaultValue);
+		MaterialExpressionFactoryNode->AddFloatAttribute(DefaultValueMemberName.ToString(), InputValue);
+		MaterialExpressionFactoryNode->AddApplyAndFillDelegates<float>(DefaultValueMemberName.ToString(), UMaterialExpressionScalarParameter::StaticClass(), DefaultValueMemberName);
+	}
+
+	return MaterialExpressionFactoryNode;
+}
+
+UInterchangeMaterialExpressionFactoryNode* UInterchangeGenericMaterialPipeline::CreateVectorParameterExpression(const UInterchangeShaderNode* ShaderNode, const FString& InputName, const FString& ParentUid)
+{
+	UInterchangeMaterialExpressionFactoryNode* MaterialExpressionFactoryNode = CreateExpressionNode(InputName, ParentUid, UMaterialExpressionVectorParameter::StaticClass());
+
+	FLinearColor InputValue;
+	if (ShaderNode->GetLinearColorAttribute(UInterchangeShaderPortsAPI::MakeInputValueKey(InputName), InputValue))
+	{
+		const FName DefaultValueMemberName = GET_MEMBER_NAME_CHECKED(UMaterialExpressionVectorParameter, DefaultValue);
+		MaterialExpressionFactoryNode->AddLinearColorAttribute(DefaultValueMemberName.ToString(), InputValue);
+		MaterialExpressionFactoryNode->AddApplyAndFillDelegates<FLinearColor>(DefaultValueMemberName.ToString(), UMaterialExpressionVectorParameter::StaticClass(), DefaultValueMemberName);
+	}
+
+	return MaterialExpressionFactoryNode;
+}
+
+TTuple<UInterchangeMaterialExpressionFactoryNode*, FString> UInterchangeGenericMaterialPipeline::CreateMaterialExpressionForInput(UInterchangeMaterialFactoryNode* MaterialFactoryNode, const UInterchangeShaderNode* ShaderNode, const FString& InputName, const FString& ParentUid)
+{
+	// If we have a connection
+	// - Create material expression for connected shader node
+	// - Connected material expression to something?
+	//
+	// If we don't have a connection
+	// - Create material expression for the input value
+
+	UInterchangeMaterialExpressionFactoryNode* MaterialExpressionFactoryNode = nullptr;
+
+	FString ConnectedShaderNodeUid;
+	FString OutputName;
+	if (UInterchangeShaderPortsAPI::GetInputConnection(ShaderNode, InputName, ConnectedShaderNodeUid, OutputName))
+	{
+		if (UInterchangeShaderNode* ConnectedShaderNode = Cast<UInterchangeShaderNode>(BaseNodeContainer->GetNode(ConnectedShaderNodeUid)))
+		{
+			MaterialExpressionFactoryNode = CreateMaterialExpressionForShaderNode(MaterialFactoryNode, ConnectedShaderNode, ParentUid);
+		}
+	}
+	else
+	{
+		FString MaterialExpressionUid = ShaderNode->GetUniqueID() + TEXT(".") + InputName;
+		switch(UInterchangeShaderPortsAPI::GetInputType(ShaderNode, InputName))
+		{
+		case UE::Interchange::EAttributeTypes::Float:
+			MaterialExpressionFactoryNode = CreateScalarParameterExpression(ShaderNode, InputName, ParentUid);
+			break;
+		case UE::Interchange::EAttributeTypes::LinearColor:
+			MaterialExpressionFactoryNode = CreateVectorParameterExpression(ShaderNode, InputName, ParentUid);
+			break;
+		}
+	}
+
+	return TTuple<UInterchangeMaterialExpressionFactoryNode*, FString>{MaterialExpressionFactoryNode, OutputName};
+}
+
+UInterchangeMaterialFactoryNode* UInterchangeGenericMaterialPipeline::CreateMaterialFactoryNode(const UInterchangeShaderGraphNode* ShaderGraphNode)
+{
+	UInterchangeMaterialFactoryNode* MaterialFactoryNode = Cast<UInterchangeMaterialFactoryNode>( CreateBaseMaterialFactoryNode(ShaderGraphNode, UInterchangeMaterialFactoryNode::StaticClass()) );
+
+	if (!HandlePhongModel(ShaderGraphNode, MaterialFactoryNode))
+	{
+		if (!HandleLambertModel(ShaderGraphNode, MaterialFactoryNode))
+		{
+			HandlePBRModel(ShaderGraphNode, MaterialFactoryNode);
+		}
+	}
+
+	HandleCommonParameters(ShaderGraphNode, MaterialFactoryNode);
+
+	return MaterialFactoryNode;
+}
+
diff --git a/Engine/Plugins/Experimental/Interchange/Runtime/Source/Pipelines/Private/InterchangeGenericSkeletalMeshPipeline.cpp b/Engine/Plugins/Experimental/Interchange/Runtime/Source/Pipelines/Private/InterchangeGenericSkeletalMeshPipeline.cpp
--- a/Engine/Plugins/Experimental/Interchange/Runtime/Source/Pipelines/Private/InterchangeGenericSkeletalMeshPipeline.cpp
+++ b/Engine/Plugins/Experimental/Interchange/Runtime/Source/Pipelines/Private/InterchangeGenericSkeletalMeshPipeline.cpp
@@ -248,12 +248,12 @@ void UInterchangeGenericAssetsPipeline::ExecutePreImportPipelineSkeletalMesh()
 		return;
 	}
 
-	if (ForceAllMeshHasType != EInterchangeForceMeshType::IFMT_None && ForceAllMeshHasType != EInterchangeForceMeshType::IFMT_SkeletalMesh)
+	if (ForceAllMeshAsType != EInterchangeForceMeshType::IFMT_None && ForceAllMeshAsType != EInterchangeForceMeshType::IFMT_SkeletalMesh)
 	{
 		//Nothing to import
 		return;
 	}
-	const bool bConvertStaticMeshToSkeletalMesh = (ForceAllMeshHasType == EInterchangeForceMeshType::IFMT_SkeletalMesh);
+	const bool bConvertStaticMeshToSkeletalMesh = (ForceAllMeshAsType == EInterchangeForceMeshType::IFMT_SkeletalMesh);
 	TMap<FString, TArray<FString>> SkeletalMeshFactoryDependencyOrderPerSkeletonRootNodeUid;
 
 	auto SetSkeletalMeshDependencies = [&SkeletalMeshFactoryDependencyOrderPerSkeletonRootNodeUid](const FString& JointNodeUid, UInterchangeSkeletalMeshFactoryNode* SkeletalMeshFactoryNode)
diff --git a/Engine/Plugins/Experimental/Interchange/Runtime/Source/Pipelines/Private/InterchangeGenericTexturePipeline.cpp b/Engine/Plugins/Experimental/Interchange/Runtime/Source/Pipelines/Private/InterchangeGenericTexturePipeline.cpp
--- a/Engine/Plugins/Experimental/Interchange/Runtime/Source/Pipelines/Private/InterchangeGenericTexturePipeline.cpp
+++ b/Engine/Plugins/Experimental/Interchange/Runtime/Source/Pipelines/Private/InterchangeGenericTexturePipeline.cpp
@@ -161,6 +161,9 @@ UInterchangeTextureFactoryNode* UInterchangeGenericAssetsPipeline::CreateTexture
 		TextureFactoryNode->SetCustomTranslatedTextureNodeUid(TextureNode->GetUniqueID());
 		BaseNodeContainer->AddNode(TextureFactoryNode);
 		TextureFactoryNodes.Add(TextureFactoryNode);
+
+		TextureFactoryNode->AddTargetNodeUid(TextureNode->GetUniqueID());
+		TextureNode->AddTargetNodeUid(TextureFactoryNode->GetUniqueID());
 	}
 	return TextureFactoryNode;
 }
diff --git a/Engine/Plugins/Experimental/Interchange/Runtime/Source/Pipelines/Public/InterchangeGenericAssetsPipeline.h b/Engine/Plugins/Experimental/Interchange/Runtime/Source/Pipelines/Public/InterchangeGenericAssetsPipeline.h
--- a/Engine/Plugins/Experimental/Interchange/Runtime/Source/Pipelines/Public/InterchangeGenericAssetsPipeline.h
+++ b/Engine/Plugins/Experimental/Interchange/Runtime/Source/Pipelines/Public/InterchangeGenericAssetsPipeline.h
@@ -3,6 +3,7 @@
 #pragma once
 
 #include "CoreMinimal.h"
+#include "InterchangeGenericMaterialPipeline.h"
 #include "InterchangePipelineBase.h"
 #include "InterchangeSourceData.h"
 #include "Nodes/InterchangeBaseNodeContainer.h"
@@ -11,8 +12,6 @@
 
 #include "InterchangeGenericAssetsPipeline.generated.h"
 
-class UInterchangeMaterialFactoryNode;
-class UInterchangeMaterialNode;
 class UInterchangeMeshNode;
 class UInterchangePipelineMeshesUtilities;
 class UInterchangeResult;
@@ -36,7 +35,6 @@ template<class T> class TSubclassOf;
 #define COMMON_SKELETAL_ANIMATIONS_CATEGORY "Common Skeletal Mesh and Animations"
 #define SKELETAL_MESHES_CATEGORY "Skeletal Meshes"
 #define ANIMATIONS_CATEGORY "Animations"
-#define MATERIALS_CATEGORY "Materials"
 #define TEXTURES_CATEGORY "Textures"
 
 /** Force mesh type, if user want to import all meshes as one type*/
@@ -77,6 +75,8 @@ class INTERCHANGEPIPELINES_API UInterchangeGenericAssetsPipeline : public UInter
 	GENERATED_BODY()
 
 public:
+	UInterchangeGenericAssetsPipeline();
+
 	//////////////////////////////////////////////////////////////////////////
 	// BEGIN Pre import pipeline properties, please keep by per category order for properties declaration
 
@@ -94,7 +94,7 @@ public:
 
 	/** Allow to convert mesh to a particular type */
  	UPROPERTY(EditAnywhere, BlueprintReadWrite, Category = COMMON_MESHES_CATEGORY)
- 	TEnumAsByte<enum EInterchangeForceMeshType> ForceAllMeshHasType = EInterchangeForceMeshType::IFMT_None;
+ 	TEnumAsByte<enum EInterchangeForceMeshType> ForceAllMeshAsType = EInterchangeForceMeshType::IFMT_None;
 
 	/** If enable, meshes LODs will be imported. Note that it required the advanced bBakeMesh property to be enabled. */
 	UPROPERTY(EditAnywhere, BlueprintReadWrite, Category = COMMON_MESHES_CATEGORY, meta = (editcondition = "bBakeMeshes"))
@@ -178,15 +178,14 @@ public:
 	//////	MATERIALS_CATEGORY Properties //////
 
 
-	/** If enable, import the material asset find in the sources. */
-	UPROPERTY(EditAnywhere, BlueprintReadWrite, Category = MATERIALS_CATEGORY)
-	bool bImportMaterials = true;
+	UPROPERTY(VisibleAnywhere, Instanced, Category = "Materials")
+	UInterchangeGenericMaterialPipeline* MaterialPipeline;
 
 
 	//////	TEXTURES_CATEGORY Properties //////
 
 
-	/** If enable, import the material asset find in the sources. */
+	/** If enabled, imports the texture assets found in the sources. */
 	UPROPERTY(EditAnywhere, BlueprintReadWrite, Category = TEXTURES_CATEGORY)
 	bool bImportTextures = true;
 
@@ -237,16 +236,6 @@ private:
 	/** Texture factory assets nodes */
 	TArray<UInterchangeTextureFactoryNode*> TextureFactoryNodes;
 	
-	/** Material translated assets nodes */
-	TArray<UInterchangeMaterialNode*> MaterialNodes;
-	
-	/** Material factory assets nodes */
-	TArray<UInterchangeMaterialFactoryNode*> MaterialFactoryNodes;
-
-	
-	UInterchangeMaterialFactoryNode* CreateMaterialFactoryNode(const UInterchangeMaterialNode* MaterialNode);
-
-	
 	/************************************************************************/
 	/* Skeletal mesh API BEGIN                                              */
 
diff --git a/Engine/Plugins/Experimental/Interchange/Runtime/Source/Pipelines/Public/InterchangeGenericMaterialPipeline.h b/Engine/Plugins/Experimental/Interchange/Runtime/Source/Pipelines/Public/InterchangeGenericMaterialPipeline.h
new file mode 100644
--- /dev/null
+++ b/Engine/Plugins/Experimental/Interchange/Runtime/Source/Pipelines/Public/InterchangeGenericMaterialPipeline.h
@@ -0,0 +1,74 @@
+// Copyright Epic Games, Inc. All Rights Reserved.
+
+#pragma once
+
+#include "CoreMinimal.h"
+#include "InterchangePipelineBase.h"
+#include "InterchangeSourceData.h"
+#include "Nodes/InterchangeBaseNodeContainer.h"
+#include "UObject/Object.h"
+#include "UObject/ObjectMacros.h"
+
+#include "InterchangeGenericMaterialPipeline.generated.h"
+
+class UInterchangeBaseMaterialFactoryNode;
+class UInterchangeShaderGraphNode;
+class UInterchangeShaderNode;
+class UInterchangeMaterialFactoryNode;
+class UInterchangeMaterialExpressionFactoryNode;
+class UInterchangeResult;
+
+UENUM(BlueprintType)
+enum class EInterchangeMaterialImportOption : uint8
+{
+	DoNotImport,
+	ImportAsMaterials,
+	//ImportAsMaterialInstances,
+};
+
+UCLASS(BlueprintType, Experimental)
+class INTERCHANGEPIPELINES_API UInterchangeGenericMaterialPipeline : public UInterchangePipelineBase
+{
+	GENERATED_BODY()
+
+public:
+	/** If enabled, imports the material assets found in the sources. */
+	UPROPERTY(EditAnywhere, BlueprintReadWrite, Category = MATERIALS_CATEGORY)
+	EInterchangeMaterialImportOption MaterialImport = EInterchangeMaterialImportOption::ImportAsMaterials;
+
+public:
+	virtual void ExecutePreImportPipeline(UInterchangeBaseNodeContainer* InBaseNodeContainer, const TArray<UInterchangeSourceData*>& InSourceDatas) override;
+
+private:
+	UPROPERTY()
+	UInterchangeBaseNodeContainer* BaseNodeContainer;
+
+	TArray<const UInterchangeSourceData*> SourceDatas;
+	
+private:
+	/** Material translated assets nodes */
+	TArray<UInterchangeShaderGraphNode*> MaterialNodes;
+	
+	/** Material factory assets nodes */
+	TArray<UInterchangeBaseMaterialFactoryNode*> MaterialFactoryNodes;
+
+	UInterchangeBaseMaterialFactoryNode* CreateBaseMaterialFactoryNode(const UInterchangeBaseNode* MaterialNode, TSubclassOf<UInterchangeBaseMaterialFactoryNode> NodeType);
+	UInterchangeMaterialFactoryNode* CreateMaterialFactoryNode(const UInterchangeShaderGraphNode* ShaderGraphNode);
+
+	bool HandlePhongModel(const UInterchangeShaderGraphNode* ShaderGraphNode, UInterchangeMaterialFactoryNode* MaterialFactoryNode);
+	bool HandleLambertModel(const UInterchangeShaderGraphNode* ShaderGraphNode, UInterchangeMaterialFactoryNode* MaterialFactoryNode);
+	bool HandlePBRModel(const UInterchangeShaderGraphNode* ShaderGraphNode, UInterchangeMaterialFactoryNode* MaterialFactoryNode);
+	bool HandleCommonParameters(const UInterchangeShaderGraphNode* ShaderGraphNode, UInterchangeMaterialFactoryNode* MaterialFactoryNode);
+	void HandleTextureSampleNode(const UInterchangeShaderNode* ShaderNode, UInterchangeMaterialExpressionFactoryNode* TextureSampleFactoryNode);
+	void HandleTextureCoordinates(const UInterchangeShaderNode* ShaderNode, UInterchangeMaterialExpressionFactoryNode* TextureSampleFactoryNode);
+	void HandleLerpNode(const UInterchangeShaderNode* ShaderNode, UInterchangeMaterialFactoryNode* MaterialFactoryNode, UInterchangeMaterialExpressionFactoryNode* LerpFactoryNode);
+
+	UInterchangeMaterialExpressionFactoryNode* CreateMaterialExpressionForShaderNode(UInterchangeMaterialFactoryNode* MaterialFactoryNode, const UInterchangeShaderNode* ShaderNode, const FString& ParentUid);
+	TTuple<UInterchangeMaterialExpressionFactoryNode*, FString> CreateMaterialExpressionForInput(UInterchangeMaterialFactoryNode* MaterialFactoryNode, const UInterchangeShaderNode* ShaderNode, const FString& InputName, const FString& ParentUid);
+
+	UInterchangeMaterialExpressionFactoryNode* CreateExpressionNode(const FString& ExpressionName, const FString& ParentUid, UClass* MaterialExpressionClass);
+	UInterchangeMaterialExpressionFactoryNode* CreateScalarParameterExpression(const UInterchangeShaderNode* ShaderNode, const FString& InputName, const FString& ParentUid);
+	UInterchangeMaterialExpressionFactoryNode* CreateVectorParameterExpression(const UInterchangeShaderNode* ShaderNode, const FString& InputName, const FString& ParentUid);
+};
+
+
diff --git a/Engine/Source/Runtime/Experimental/Interchange/Core/Private/Nodes/InterchangeBaseNode.cpp b/Engine/Source/Runtime/Experimental/Interchange/Core/Private/Nodes/InterchangeBaseNode.cpp
--- a/Engine/Source/Runtime/Experimental/Interchange/Core/Private/Nodes/InterchangeBaseNode.cpp
+++ b/Engine/Source/Runtime/Experimental/Interchange/Core/Private/Nodes/InterchangeBaseNode.cpp
@@ -199,7 +199,7 @@ UE::Interchange::EAttributeTypes UInterchangeBaseNode::GetAttributeType(const UE
 	return Attributes->GetAttributeType(NodeAttributeKey);
 }
 
-void UInterchangeBaseNode::GetAttributeKeys(TArray<UE::Interchange::FAttributeKey>& AttributeKeys)
+void UInterchangeBaseNode::GetAttributeKeys(TArray<UE::Interchange::FAttributeKey>& AttributeKeys) const
 {
 	Attributes->GetAttributeKeys(AttributeKeys);
 }
@@ -215,7 +215,7 @@ bool UInterchangeBaseNode::AddBooleanAttribute(const FString& NodeAttributeKey,
 	INTERCHANGE_BASE_NODE_ADD_ATTRIBUTE(bool);
 }
 
-bool UInterchangeBaseNode::GetBooleanAttribute(const FString& NodeAttributeKey, bool& OutValue)
+bool UInterchangeBaseNode::GetBooleanAttribute(const FString& NodeAttributeKey, bool& OutValue) const
 {
 	INTERCHANGE_BASE_NODE_GET_ATTRIBUTE(bool);
 }
@@ -225,7 +225,7 @@ bool UInterchangeBaseNode::AddInt32Attribute(const FString& NodeAttributeKey, co
 	INTERCHANGE_BASE_NODE_ADD_ATTRIBUTE(int32);
 }
 
-bool UInterchangeBaseNode::GetInt32Attribute(const FString& NodeAttributeKey, int32& OutValue)
+bool UInterchangeBaseNode::GetInt32Attribute(const FString& NodeAttributeKey, int32& OutValue) const
 {
 	INTERCHANGE_BASE_NODE_GET_ATTRIBUTE(int32);
 }
@@ -235,7 +235,7 @@ bool UInterchangeBaseNode::AddFloatAttribute(const FString& NodeAttributeKey, co
 	INTERCHANGE_BASE_NODE_ADD_ATTRIBUTE(float);
 }
 
-bool UInterchangeBaseNode::GetFloatAttribute(const FString& NodeAttributeKey, float& OutValue)
+bool UInterchangeBaseNode::GetFloatAttribute(const FString& NodeAttributeKey, float& OutValue) const
 {
 	INTERCHANGE_BASE_NODE_GET_ATTRIBUTE(float);
 }
@@ -245,11 +245,21 @@ bool UInterchangeBaseNode::AddStringAttribute(const FString& NodeAttributeKey, c
 	INTERCHANGE_BASE_NODE_ADD_ATTRIBUTE(FString);
 }
 
-bool UInterchangeBaseNode::GetStringAttribute(const FString& NodeAttributeKey, FString& OutValue)
+bool UInterchangeBaseNode::GetStringAttribute(const FString& NodeAttributeKey, FString& OutValue) const
 {
 	INTERCHANGE_BASE_NODE_GET_ATTRIBUTE(FString);
 }
 
+bool UInterchangeBaseNode::AddLinearColorAttribute(const FString& NodeAttributeKey, const FLinearColor& Value)
+{
+	INTERCHANGE_BASE_NODE_ADD_ATTRIBUTE(FLinearColor);
+}
+
+bool UInterchangeBaseNode::GetLinearColorAttribute(const FString& NodeAttributeKey, FLinearColor& OutValue) const
+{
+	INTERCHANGE_BASE_NODE_GET_ATTRIBUTE(FLinearColor);
+}
+
 FString UInterchangeBaseNode::GetUniqueID() const
 {
 	ensure(bIsInitialized);
@@ -518,26 +528,37 @@ FProperty* InterchangePrivateNodeBase::FindPropertyByPathChecked(TVariant<UObjec
 		}
 
 		PropertyName = FStringView(RestOfPropertyPath.GetData(), SeparatorIndex);
-		RestOfPropertyPath = FStringView(RestOfPropertyPath.GetData() + SeparatorIndex + 1);
+
+		if (SeparatorIndex < RestOfPropertyPath.Len())
+		{
+			RestOfPropertyPath = FStringView(RestOfPropertyPath.GetData() + SeparatorIndex + 1, RestOfPropertyPath.Len() - SeparatorIndex - 1);
+		}
+		else
+		{
+			RestOfPropertyPath = FStringView();
+		}
 
 		Property = FindFieldChecked<FProperty>(Outer, FName(PropertyName.Len(), PropertyName.GetData()));
 
-		if (FStructProperty* StructProperty = CastField<FStructProperty>(Property))
+		if (RestOfPropertyPath.Len() > 0)
 		{
-			Outer = StructProperty->Struct;
-
-			if (Container.IsType<UObject*>())
+			if (FStructProperty* StructProperty = CastField<FStructProperty>(Property))
 			{
-				Container.Set<uint8*>(Property->ContainerPtrToValuePtr<uint8>(Container.Get<UObject*>()));
-			}
-			else
-			{
-				Container.Set<uint8*>(Property->ContainerPtrToValuePtr<uint8>(Container.Get<uint8*>()));
+				Outer = StructProperty->Struct;
+
+				if (Container.IsType<UObject*>())
+				{
+					Container.Set<uint8*>(Property->ContainerPtrToValuePtr<uint8>(Container.Get<UObject*>()));
+				}
+				else
+				{
+					Container.Set<uint8*>(Property->ContainerPtrToValuePtr<uint8>(Container.Get<uint8*>()));
+				}
 			}
 		}
-		
 	} while (RestOfPropertyPath.Len() > 0);
 
 	check(Property);
 	return Property;
 }
+
diff --git a/Engine/Source/Runtime/Experimental/Interchange/Core/Public/InterchangePipelineBase.h b/Engine/Source/Runtime/Experimental/Interchange/Core/Public/InterchangePipelineBase.h
--- a/Engine/Source/Runtime/Experimental/Interchange/Core/Public/InterchangePipelineBase.h
+++ b/Engine/Source/Runtime/Experimental/Interchange/Core/Public/InterchangePipelineBase.h
@@ -136,8 +136,8 @@ protected:
 	}
 
 	/**
-	 * This function can read the node data and apply some change to the imported asset. This is call after the factory create the asset and configure the asset properties.
-	 * The interchange manager is not calling this function directly. It is calling the blueprint native event in case this object is a blueprint derive object.
+	 * This function can read the node data and apply some change to the imported asset. This is called after the factory creates the asset and configures the asset properties.
+	 * The interchange manager is not calling this function directly. It is calling the blueprint native event in case this object is a blueprint derived object.
 	 * By default the scripted implementation is calling this virtual pipeline.
 	 */
 	virtual void ExecutePostImportPipeline(const UInterchangeBaseNodeContainer* BaseNodeContainer, const FString& NodeKey, UObject* CreatedAsset, bool bIsAReimport)
diff --git a/Engine/Source/Runtime/Experimental/Interchange/Core/Public/Nodes/InterchangeBaseNode.h b/Engine/Source/Runtime/Experimental/Interchange/Core/Public/Nodes/InterchangeBaseNode.h
--- a/Engine/Source/Runtime/Experimental/Interchange/Core/Public/Nodes/InterchangeBaseNode.h
+++ b/Engine/Source/Runtime/Experimental/Interchange/Core/Public/Nodes/InterchangeBaseNode.h
@@ -97,71 +97,17 @@ namespace InterchangePrivateNodeBase
 const UE::Interchange::FAttributeKey Macro_Custom##AttributeName##Key = UE::Interchange::FAttributeKey(TEXT(#AttributeName));
 
 #if WITH_ENGINE
-#define IMPLEMENT_NODE_ATTRIBUTE_APPLY_UOBJECT_BYNAME(AttributeName, AttributeType, ObjectType, PropertyName)	\
-bool ApplyCustom##AttributeName##ToAsset(UObject* Asset) const													\
-{																												\
-	if (!Asset)																									\
-	{																											\
-		return false;																							\
-	}																											\
-	ObjectType* TypedObject = Cast<ObjectType>(Asset);															\
-	if (!TypedObject)																							\
-	{																											\
-		return false;																							\
-	}																											\
-	AttributeType ValueData;																					\
-	if (GetCustom##AttributeName(ValueData))																	\
-	{																											\
-		TVariant<UObject*, uint8*> Container;																	\
-		Container.Set<UObject*>(Asset);																			\
-		if(FProperty* Property = InterchangePrivateNodeBase::FindPropertyByPathChecked(Container, Asset->GetClass(), PropertyName))\
-		{																										\
-			if (Container.IsType<UObject*>())																	\
-			{																									\
-				*(Property->ContainerPtrToValuePtr<AttributeType>(Container.Get<UObject*>())) = ValueData;		\
-			}																									\
-			else																								\
-			{																									\
-				*(Property->ContainerPtrToValuePtr<AttributeType>(Container.Get<uint8*>())) = ValueData;		\
-			}																									\
-		}																										\
-		return true;																							\
-	}																											\
-	return false;																								\
-}																												\
-																												\
-bool FillCustom##AttributeName##FromAsset(UObject* Asset)														\
-{																												\
-	if (!Asset)																									\
-	{																											\
-		return false;																							\
-	}																											\
-	ObjectType* TypedObject = Cast<ObjectType>(Asset);															\
-	if (!TypedObject)																							\
-	{																											\
-		return false;																							\
-	}																											\
-	TVariant<UObject*, uint8*> Container;																		\
-	Container.Set<UObject*>(Asset);																				\
-	if(FProperty* Property = InterchangePrivateNodeBase::FindPropertyByPathChecked(Container, Asset->GetClass(), PropertyName))\
-	{																											\
-		AttributeType ValueData;																				\
-		if (Container.IsType<UObject*>())																		\
-		{																										\
-			ValueData = *(Property->ContainerPtrToValuePtr<AttributeType>(Container.Get<UObject*>()));			\
-		}																										\
-		else																									\
-		{																										\
-			ValueData = *(Property->ContainerPtrToValuePtr<AttributeType>(Container.Get<uint8*>()));			\
-		}																										\
-		if (SetCustom##AttributeName(ValueData, false))															\
-		{																										\
-			return true;																						\
-		}																										\
-	}																											\
-																												\
-	return false;																								\
+#define IMPLEMENT_NODE_ATTRIBUTE_APPLY_UOBJECT_BYNAME(AttributeName, AttributeType, ObjectType, PropertyName)		\
+bool ApplyCustom##AttributeName##ToAsset(UObject* Asset) const														\
+{																													\
+	return ApplyAttributeToObject<AttributeType>(Macro_Custom##AttributeName##Key.ToString(), Asset, PropertyName);	\
+}																													\
+																													\
+bool FillCustom##AttributeName##FromAsset(UObject* Asset)															\
+{																													\
+	return FillAttributeFromObject<AttributeType>(Macro_Custom##AttributeName##Key.ToString(), Asset, PropertyName);	\
 }
+
 #define IMPLEMENT_NODE_ATTRIBUTE_APPLY_UOBJECT(AttributeName, AttributeType, ObjectType)						\
 	IMPLEMENT_NODE_ATTRIBUTE_APPLY_UOBJECT_BYNAME(AttributeName, AttributeType, ObjectType, TEXT(#AttributeName))
 
@@ -176,7 +122,7 @@ bool FillCustom##AttributeName##FromAsset(UObject* Asset)														\
 	FString OperationName = GetTypeName() + TEXT(".Get" #AttributeName);																				\
 	return InterchangePrivateNodeBase::GetCustomAttribute<AttributeType>(*Attributes, Macro_Custom##AttributeName##Key, OperationName, AttributeValue);
 
-#define IMPLEMENT_NODE_ATTRIBUTE_SETTER_NODELEGATE(AttributeName, AttributeType)															\
+#define IMPLEMENT_NODE_ATTRIBUTE_SETTER_NODELEGATE(AttributeName, AttributeType)																		\
 	FString OperationName = GetTypeName() + TEXT(".Set" #AttributeName);																				\
 	return InterchangePrivateNodeBase::SetCustomAttribute<AttributeType>(*Attributes, Macro_Custom##AttributeName##Key, OperationName, AttributeValue);
 
@@ -188,10 +134,7 @@ bool FillCustom##AttributeName##FromAsset(UObject* Asset)														\
 	{																																				\
 		if(bAddApplyDelegate)																														\
 		{																																			\
-			TArray<UE::Interchange::FApplyAttributeToAsset>& Delegates = ApplyCustomAttributeDelegates.FindOrAdd(AssetType::StaticClass());			\
-			Delegates.Add(UE::Interchange::FApplyAttributeToAsset::CreateUObject(this, &NodeClassName::ApplyCustom##AttributeName##ToAsset));		\
-			TArray<UE::Interchange::FFillAttributeToAsset>& FillDelegates = FillCustomAttributeDelegates.FindOrAdd(AssetType::StaticClass());		\
-			FillDelegates.Add(UE::Interchange::FFillAttributeToAsset::CreateUObject(this, &NodeClassName::FillCustom##AttributeName##FromAsset));	\
+			AddApplyAndFillDelegates<AttributeType>(Macro_Custom##AttributeName##Key.ToString(), AssetType::StaticClass(), *Macro_Custom##AttributeName##Key.ToString()); \
 		}																																			\
 		return true;																																\
 	}																																				\
@@ -268,8 +211,6 @@ class INTERCHANGECORE_API UInterchangeBaseNode : public UObject
 public:
 	UInterchangeBaseNode();
 
-//	virtual ~UInterchangeBaseNode() = default;
-
 	/**
 	 * Initialize the base data of the node
 	 * @param UniqueID - The uniqueId for this node
@@ -352,7 +293,7 @@ public:
 	/**
 	 * Use this function to query all the node attribute keys
 	 */
-	void GetAttributeKeys(TArray<UE::Interchange::FAttributeKey>& AttributeKeys);
+	void GetAttributeKeys(TArray<UE::Interchange::FAttributeKey>& AttributeKeys) const;
 
 	/**
 	 * Remove any attribute from this node. Return false if we cannot remove it. If the attribute do not exist it will return true.
@@ -370,7 +311,7 @@ public:
 	 * Get a boolean attribute from this node. Return false if the attribute do not exist
 	 */
 	UFUNCTION(BlueprintCallable, Category = "Interchange | Node")
-	bool GetBooleanAttribute(const FString& NodeAttributeKey, bool& OutValue);
+	bool GetBooleanAttribute(const FString& NodeAttributeKey, bool& OutValue) const;
 
 	/**
 	 * Add a int32 attribute to this node. Return false if the attribute do not exist or if we cannot add it
@@ -382,7 +323,7 @@ public:
 	 * Get a int32 attribute from this node. Return false if the attribute do not exist
 	 */
 	UFUNCTION(BlueprintCallable, Category = "Interchange | Node")
-	bool GetInt32Attribute(const FString& NodeAttributeKey, int32& OutValue);
+	bool GetInt32Attribute(const FString& NodeAttributeKey, int32& OutValue) const;
 
 	/**
 	 * Add a float attribute to this node. Return false if the attribute do not exist or if we cannot add it
@@ -394,7 +335,7 @@ public:
 	 * Get a float attribute from this node. Return false if the attribute do not exist
 	 */
 	UFUNCTION(BlueprintCallable, Category = "Interchange | Node")
-	bool GetFloatAttribute(const FString& NodeAttributeKey, float& OutValue);
+	bool GetFloatAttribute(const FString& NodeAttributeKey, float& OutValue) const;
 
 	/**
 	 * Add a string attribute to this node. Return false if the attribute do not exist or if we cannot add it
@@ -406,7 +347,73 @@ public:
 	 * Get a string attribute from this node. Return false if the attribute do not exist
 	 */
 	UFUNCTION(BlueprintCallable, Category = "Interchange | Node")
-	bool GetStringAttribute(const FString& NodeAttributeKey, FString& OutValue);
+	bool GetStringAttribute(const FString& NodeAttributeKey, FString& OutValue) const;
+
+	/**
+	 * Add a FLinearColor attribute to this node. Return false if the attribute do not exist or if we cannot add it
+	 */
+	UFUNCTION(BlueprintCallable, Category = "Interchange | Node")
+	bool AddLinearColorAttribute(const FString& NodeAttributeKey, const FLinearColor& Value);
+	
+	/**
+	 * Get a FLinearColor attribute from this node. Return false if the attribute do not exist
+	 */
+	UFUNCTION(BlueprintCallable, Category = "Interchange | Node")
+	bool GetLinearColorAttribute(const FString& NodeAttributeKey, FLinearColor& OutValue) const;
+
+	template<typename AttributeType>
+	bool GetAttribute(const FString& NodeAttributeKey, AttributeType& OutValue) const
+	{
+		INTERCHANGE_BASE_NODE_GET_ATTRIBUTE(AttributeType);
+	}
+	template<typename AttributeType>
+	bool SetAttribute(const FString& NodeAttributeKey, const AttributeType& Value)
+	{
+		INTERCHANGE_BASE_NODE_ADD_ATTRIBUTE(AttributeType);
+	}
+
+	/**
+	 * Adds the delegates that will read and write the attribute value to a UObject.
+	 * @param NodeAttributeKey	The key of the attribute for which we are adding the delegates.
+	 * @param ObjectClass		The class of the object that we will apply and fill on.
+	 * @param PropertyName		The name of the property of the ObjectClass that we will read and write to.
+	 */
+	template<typename AttributeType>
+	inline void AddApplyAndFillDelegates(const FString& NodeAttributeKey, UClass* ObjectClass, const FName PropertyName);
+
+	/**
+	 * Writes an attribute value to a UObject.
+	 * @param NodeAttributeKey	The key for the attribute to write.
+	 * @param Object			The object to write to.
+	 * @param PropertyName		The name of the property to write to on the Object.
+	 * @result					True if we succeeded.
+	 */
+	template<typename AttributeType>
+	inline bool ApplyAttributeToObject(const FString& NodeAttributeKey, UObject* Object, const FName PropertyName) const;
+
+	/**
+	 * Specialized version of ApplyAttributeToObject for strings.
+	 * If the target property is a FObjectPropertyBase, treat the string as an object path.
+	 */
+	template<>
+	bool ApplyAttributeToObject<FString>(const FString& NodeAttributeKey, UObject* Object, const FName PropertyName) const;
+
+	/**
+	 * Reads an attribute value from a UObject.
+	 * @param NodeAttributeKey	The key for the attribute to update.
+	 * @param Object			The object to read from.
+	 * @param PropertyName		The name of the property to read from on the Object.
+	 * @result					True if we succeeded.
+	 */
+	template<typename AttributeType>
+	inline bool FillAttributeFromObject(const FString& NodeAttributeKey, UObject* Object, const FName PropertyName);
+
+	/**
+	 * Specialized version of FillAttributeFromObject for strings.
+	 * If the target property is a FObjectPropertyBase, treat the string as an object path.
+	 */
+	template<>
+	inline bool FillAttributeFromObject<FString>(const FString& NodeAttributeKey, UObject* Object, const FName PropertyName);
 
 	/**
 	 * Return the unique id pass in the constructor.
@@ -588,3 +595,169 @@ protected:
 	 */
 	mutable UE::Interchange::TArrayAttributeHelper<FString> TargetNodes;
 };
+
+template<typename AttributeType>
+inline void UInterchangeBaseNode::AddApplyAndFillDelegates(const FString& NodeAttributeKey, UClass* ObjectClass, const FName PropertyName)
+{
+	TArray<UE::Interchange::FApplyAttributeToAsset>& ApplyDelegates = ApplyCustomAttributeDelegates.FindOrAdd(ObjectClass);
+	ApplyDelegates.Add(
+		UE::Interchange::FApplyAttributeToAsset::CreateLambda(
+		[NodeAttributeKey, PropertyName, this](UObject* Asset)
+		{
+			return ApplyAttributeToObject<AttributeType>(NodeAttributeKey, Asset, PropertyName);
+		})
+	);
+
+	TArray<UE::Interchange::FFillAttributeToAsset>& FillDelegates = FillCustomAttributeDelegates.FindOrAdd(ObjectClass);
+	FillDelegates.Add(
+		UE::Interchange::FApplyAttributeToAsset::CreateLambda(
+		[NodeAttributeKey, PropertyName, this](UObject* Asset)
+		{
+			return FillAttributeFromObject<AttributeType>(NodeAttributeKey, Asset, PropertyName);
+		})
+	);
+}
+
+template<typename AttributeType>
+inline bool UInterchangeBaseNode::ApplyAttributeToObject(const FString& NodeAttributeKey, UObject* Object, const FName PropertyName) const
+{
+	if (!Object)
+	{
+		return false;
+	}
+
+	AttributeType ValueData;
+	if (GetAttribute<AttributeType>(NodeAttributeKey, ValueData))
+	{
+		TVariant<UObject*, uint8*> Container;
+		Container.Set<UObject*>(Object);
+		if(FProperty* Property = InterchangePrivateNodeBase::FindPropertyByPathChecked(Container, Object->GetClass(), PropertyName.ToString()))
+		{
+			if (Container.IsType<UObject*>())
+			{
+				*(Property->ContainerPtrToValuePtr<AttributeType>(Container.Get<UObject*>())) = ValueData;
+			}
+			else
+			{
+				*(Property->ContainerPtrToValuePtr<AttributeType>(Container.Get<uint8*>())) = ValueData;
+			}
+		}
+		return true;
+	}
+	return false;
+}
+
+template<>
+inline bool UInterchangeBaseNode::ApplyAttributeToObject<FString>(const FString& NodeAttributeKey, UObject* Object, const FName PropertyName) const
+{
+	if (!Object)
+	{
+		return false;
+	}
+
+	FString ValueData;
+	if (GetAttribute<FString>(NodeAttributeKey, ValueData))
+	{
+		TVariant<UObject*, uint8*> Container;
+		Container.Set<UObject*>(Object);
+		if(FProperty* Property = InterchangePrivateNodeBase::FindPropertyByPathChecked(Container, Object->GetClass(), PropertyName.ToString()))
+		{
+			if (FObjectPropertyBase* ObjectProperty = CastField<FObjectPropertyBase>(Property))
+			{
+				void* PropertyValueAddress;
+
+				if (Container.IsType<UObject*>())
+				{
+					PropertyValueAddress = Property->ContainerPtrToValuePtr<void>(Container.Get<UObject*>());
+				}
+				else
+				{
+					PropertyValueAddress =  Property->ContainerPtrToValuePtr<void>(Container.Get<uint8*>());
+				}
+
+				ObjectProperty->SetObjectPropertyValue(PropertyValueAddress, FSoftObjectPath(ValueData).TryLoad());
+			}
+			else
+			{
+				if (Container.IsType<UObject*>())
+				{
+					*(Property->ContainerPtrToValuePtr<FString>(Container.Get<UObject*>())) = ValueData;
+				}
+				else
+				{
+					*(Property->ContainerPtrToValuePtr<FString>(Container.Get<uint8*>())) = ValueData;
+				}
+			}
+		}
+		return true;
+	}
+	return false;
+}
+
+template<typename AttributeType>
+inline bool UInterchangeBaseNode::FillAttributeFromObject(const FString& NodeAttributeKey, UObject* Object, const FName PropertyName)
+{
+	TVariant<UObject*, uint8*> Container;
+	Container.Set<UObject*>(Object);
+	if(FProperty* Property = InterchangePrivateNodeBase::FindPropertyByPathChecked(Container, Object->GetClass(), PropertyName.ToString()))
+	{
+		AttributeType ValueData;
+		if (Container.IsType<UObject*>())
+		{
+			ValueData = *(Property->ContainerPtrToValuePtr<AttributeType>(Container.Get<UObject*>()));
+		}
+		else
+		{
+			ValueData = *(Property->ContainerPtrToValuePtr<AttributeType>(Container.Get<uint8*>()));
+		}
+
+		if (SetAttribute(NodeAttributeKey, ValueData))
+		{
+			return true;
+		}
+	}
+
+	return false;
+}
+
+template<>
+inline bool UInterchangeBaseNode::FillAttributeFromObject<FString>(const FString& NodeAttributeKey, UObject* Object, const FName PropertyName)
+{
+	TVariant<UObject*, uint8*> Container;
+	Container.Set<UObject*>(Object);
+	if(FProperty* Property = InterchangePrivateNodeBase::FindPropertyByPathChecked(Container, Object->GetClass(), PropertyName.ToString()))
+	{
+		if (FObjectPropertyBase* ObjectProperty = CastField<FObjectPropertyBase>(Property))
+		{
+			void* PropertyValueAddress;
+
+			if (Container.IsType<UObject*>())
+			{
+				PropertyValueAddress = Property->ContainerPtrToValuePtr<void>(Container.Get<UObject*>());
+			}
+			else
+			{
+				PropertyValueAddress =  Property->ContainerPtrToValuePtr<void>(Container.Get<uint8*>());
+			}
+
+			UObject* ObjectValue = ObjectProperty->GetObjectPropertyValue(PropertyValueAddress);
+			return SetAttribute(NodeAttributeKey, ObjectValue->GetPathName());
+		}
+		else
+		{
+			FString ValueData;
+			if (Container.IsType<UObject*>())
+			{
+				ValueData = *(Property->ContainerPtrToValuePtr<FString>(Container.Get<UObject*>()));
+			}
+			else
+			{
+				ValueData = *(Property->ContainerPtrToValuePtr<FString>(Container.Get<uint8*>()));
+			}
+
+			return SetAttribute(NodeAttributeKey, ValueData);
+		}
+	}
+
+	return false;
+}
\ No newline at end of file
