Interchange
- Fix for static analysis issues

diff --git a/Engine/Plugins/Experimental/Interchange/Runtime/Source/Import/Private/Material/InterchangeMaterialFactory.cpp b/Engine/Plugins/Experimental/Interchange/Runtime/Source/Import/Private/Material/InterchangeMaterialFactory.cpp
--- a/Engine/Plugins/Experimental/Interchange/Runtime/Source/Import/Private/Material/InterchangeMaterialFactory.cpp
+++ b/Engine/Plugins/Experimental/Interchange/Runtime/Source/Import/Private/Material/InterchangeMaterialFactory.cpp
@@ -58,7 +58,7 @@ namespace UE
 						}
 					}
 
-					if (MaterialExpressionClass->IsChildOf<UMaterialExpression>())
+					if (MaterialExpressionClass && MaterialExpressionClass->IsChildOf<UMaterialExpression>())
 					{
 						return MaterialExpressionClass;
 					}
diff --git a/Engine/Plugins/Experimental/Interchange/Runtime/Source/Pipelines/Public/InterchangeGenericMaterialPipeline.h b/Engine/Plugins/Experimental/Interchange/Runtime/Source/Pipelines/Public/InterchangeGenericMaterialPipeline.h
--- a/Engine/Plugins/Experimental/Interchange/Runtime/Source/Pipelines/Public/InterchangeGenericMaterialPipeline.h
+++ b/Engine/Plugins/Experimental/Interchange/Runtime/Source/Pipelines/Public/InterchangeGenericMaterialPipeline.h
@@ -41,7 +41,7 @@ public:
 
 private:
 	UPROPERTY()
-	UInterchangeBaseNodeContainer* BaseNodeContainer;
+	TObjectPtr<UInterchangeBaseNodeContainer> BaseNodeContainer;
 
 	TArray<const UInterchangeSourceData*> SourceDatas;
 	
