Interchange: replaced a UObject ptr with a TObjectPtr

diff --git a/Engine/Plugins/Experimental/Interchange/Runtime/Source/Pipelines/Public/InterchangeGenericAssetsPipeline.h b/Engine/Plugins/Experimental/Interchange/Runtime/Source/Pipelines/Public/InterchangeGenericAssetsPipeline.h
--- a/Engine/Plugins/Experimental/Interchange/Runtime/Source/Pipelines/Public/InterchangeGenericAssetsPipeline.h
+++ b/Engine/Plugins/Experimental/Interchange/Runtime/Source/Pipelines/Public/InterchangeGenericAssetsPipeline.h
@@ -179,7 +179,7 @@ public:
 
 
 	UPROPERTY(VisibleAnywhere, Instanced, Category = "Materials")
-	UInterchangeGenericMaterialPipeline* MaterialPipeline;
+	TObjectPtr<UInterchangeGenericMaterialPipeline> MaterialPipeline;
 
 
 	//////	TEXTURES_CATEGORY Properties //////
