Interchange initial basic re-import implementation
#rnx

diff --git a/Engine/Plugins/Experimental/Interchange/Runtime/Source/Import/Private/InterchangeImportCommon.cpp b/Engine/Plugins/Experimental/Interchange/Runtime/Source/Import/Private/InterchangeImportCommon.cpp
--- a/Engine/Plugins/Experimental/Interchange/Runtime/Source/Import/Private/InterchangeImportCommon.cpp
+++ b/Engine/Plugins/Experimental/Interchange/Runtime/Source/Import/Private/InterchangeImportCommon.cpp
@@ -47,6 +47,7 @@ namespace UE
 				AssetImportData->NodeUniqueID = Parameters.NodeUniqueID;
 				FObjectDuplicationParameters DupParam(Parameters.NodeContainer, AssetImportData);
 				AssetImportData->NodeContainer = CastChecked<UInterchangeBaseNodeContainer>(StaticDuplicateObjectEx(DupParam));
+				AssetImportData->Pipelines.Reset();
 				for (const UInterchangePipelineBase* Pipeline : Parameters.Pipelines)
 				{
 					UInterchangePipelineBase* DupPipeline = Cast<UInterchangePipelineBase>(StaticDuplicateObject(Pipeline, AssetImportData));
@@ -56,8 +57,6 @@ namespace UE
 					}
 				}
 			}
-
-
 		}
 
 		FFactoryCommon::FUpdateImportAssetDataParameters::FUpdateImportAssetDataParameters(UObject* InAssetImportDataOuter
@@ -174,8 +173,7 @@ namespace UE
 #endif //#if WITH_EDITORONLY_DATA
 
 
-		void FFactoryCommon::ApplyReimportStrategyToAsset(const EReimportStrategyFlags ReimportStrategyFlags
-										  , UObject* Asset
+		void FFactoryCommon::ApplyReimportStrategyToAsset(UObject* Asset
 										  , UInterchangeBaseNode* PreviousAssetNode
 										  , UInterchangeBaseNode* CurrentAssetNode
 										  , UInterchangeBaseNode* PipelineAssetNode)
@@ -184,7 +182,7 @@ namespace UE
 			{
 				return;
 			}
-
+			EReimportStrategyFlags ReimportStrategyFlags = PipelineAssetNode->GetReimportStrategyFlags();
 			switch (ReimportStrategyFlags)
 			{
 				case EReimportStrategyFlags::ApplyNoProperties:
diff --git a/Engine/Plugins/Experimental/Interchange/Runtime/Source/Import/Private/InterchangeImportCommon.h b/Engine/Plugins/Experimental/Interchange/Runtime/Source/Import/Private/InterchangeImportCommon.h
--- a/Engine/Plugins/Experimental/Interchange/Runtime/Source/Import/Private/InterchangeImportCommon.h
+++ b/Engine/Plugins/Experimental/Interchange/Runtime/Source/Import/Private/InterchangeImportCommon.h
@@ -3,7 +3,6 @@
 
 #include "CoreMinimal.h"
 #include "EditorFramework/AssetImportData.h"
-#include "InterchangeFactoryBase.h"
 
 class UAssetImportData;
 class UInterchangeBaseNode;
@@ -73,8 +72,7 @@ namespace UE
 			/**
 			 * Apply the current strategy to the PipelineAssetNode
 			 */
-			static void ApplyReimportStrategyToAsset(const EReimportStrategyFlags ReimportStrategyFlags
-											  , UObject* Asset
+			static void ApplyReimportStrategyToAsset(UObject* Asset
 											  , UInterchangeBaseNode* PreviousAssetNode
 											  , UInterchangeBaseNode* CurrentAssetNode
 											  , UInterchangeBaseNode* PipelineAssetNode);
diff --git a/Engine/Plugins/Experimental/Interchange/Runtime/Source/Import/Private/Mesh/InterchangeSkeletalMeshFactory.cpp b/Engine/Plugins/Experimental/Interchange/Runtime/Source/Import/Private/Mesh/InterchangeSkeletalMeshFactory.cpp
--- a/Engine/Plugins/Experimental/Interchange/Runtime/Source/Import/Private/Mesh/InterchangeSkeletalMeshFactory.cpp
+++ b/Engine/Plugins/Experimental/Interchange/Runtime/Source/Import/Private/Mesh/InterchangeSkeletalMeshFactory.cpp
@@ -5,6 +5,7 @@
 #include "Components.h"
 #include "Engine/SkeletalMesh.h"
 #include "GPUSkinPublicDefs.h"
+#include "InterchangeAssetImportData.h"
 #include "InterchangeImportCommon.h"
 #include "InterchangeImportLog.h"
 #include "InterchangeMaterialFactoryNode.h"
@@ -922,7 +923,7 @@ UObject* UInterchangeSkeletalMeshFactory::CreateAsset(const FCreateAssetParams&
 		return nullptr;
 	}
 
-	const UInterchangeSkeletalMeshFactoryNode* SkeletalMeshFactoryNode = Cast<UInterchangeSkeletalMeshFactoryNode>(Arguments.AssetNode);
+	UInterchangeSkeletalMeshFactoryNode* SkeletalMeshFactoryNode = Cast<UInterchangeSkeletalMeshFactoryNode>(Arguments.AssetNode);
 	if (SkeletalMeshFactoryNode == nullptr)
 	{
 		return nullptr;
@@ -962,333 +963,382 @@ UObject* UInterchangeSkeletalMeshFactory::CreateAsset(const FCreateAssetParams&
 		return nullptr;
 	}
 
-	if (SkeletalMeshObject)
+	USkeletalMesh* SkeletalMesh = Cast<USkeletalMesh>(SkeletalMeshObject);
+
+	const bool bIsReImport = (Arguments.ReimportObject != nullptr);
+
+	if (!ensure(SkeletalMesh))
+	{
+		if (!bIsReImport)
+		{
+			UE_LOG(LogInterchangeImport, Error, TEXT("Could not create skeletalMesh asset %s"), *Arguments.AssetName);
+		}
+		else
+		{
+			UE_LOG(LogInterchangeImport, Error, TEXT("Could not find reimported skeletalMesh asset %s"), *Arguments.AssetName);
+		}
+		return nullptr;
+	}
+
+	//Dirty the DDC Key for any imported Skeletal Mesh
+	SkeletalMesh->InvalidateDeriveDataCacheGUID();
+	USkeleton* SkeletonReference = nullptr;
+		
+	FSkeletalMeshModel* ImportedResource = SkeletalMesh->GetImportedModel();
+	if (!bIsReImport)
+	{
+		if (!ensure(ImportedResource->LODModels.Num() == 0))
+		{
+			ImportedResource->LODModels.Empty();
+		}
+	}
+	else
+	{
+		SkeletalMesh->GetImportedBounds().BoxExtent.Set(0.0, 0.0, 0.0);
+
+		//When we re-import, we force the current skeletalmesh skeleton, to be specified and to be the reference
+		FSoftObjectPath SpecifiedSkeleton = SkeletalMesh->GetSkeleton();
+		SkeletalMeshFactoryNode->SetCustomSkeletonSoftObjectPath(SpecifiedSkeleton);
+	}
+			
+	int32 LodCount = SkeletalMeshFactoryNode->GetLodDataCount();
+	TArray<FString> LodDataUniqueIds;
+	SkeletalMeshFactoryNode->GetLodDataUniqueIds(LodDataUniqueIds);
+	ensure(LodDataUniqueIds.Num() == LodCount);
+	int32 CurrentLodIndex = 0;
+
+	for (int32 LodIndex = 0; LodIndex < LodCount; ++LodIndex)
 	{
-		//Currently material re-import will not touch the material at all
-		//TODO design a re-import process for the material (expressions and input connections)
-		if(!Arguments.ReimportObject)
+		FString LodUniqueId = LodDataUniqueIds[LodIndex];
+		const UInterchangeSkeletalMeshLodDataNode* LodDataNode = Cast<UInterchangeSkeletalMeshLodDataNode>(Arguments.NodeContainer->GetNode(LodUniqueId));
+		if (!LodDataNode)
 		{
-			USkeletalMesh* SkeletalMesh = Cast<USkeletalMesh>(SkeletalMeshObject);
-			if (!ensure(SkeletalMesh))
+			UE_LOG(LogInterchangeImport, Warning, TEXT("Invalid LOD when importing SkeletalMesh asset %s"), *Arguments.AssetName);
+			continue;
+		}
+
+		FString SkeletonNodeUid;
+		if (!LodDataNode->GetCustomSkeletonUid(SkeletonNodeUid))
+		{
+			UE_LOG(LogInterchangeImport, Warning, TEXT("Invalid Skeleton LOD when importing SkeletalMesh asset %s"), *Arguments.AssetName);
+			continue;
+		}
+		const UInterchangeSkeletonFactoryNode* SkeletonNode = Cast<UInterchangeSkeletonFactoryNode>(Arguments.NodeContainer->GetNode(SkeletonNodeUid));
+		if (!SkeletonNode)
+		{
+			UE_LOG(LogInterchangeImport, Warning, TEXT("Invalid Skeleton LOD when importing SkeletalMesh asset %s"), *Arguments.AssetName);
+			continue;
+		}
+
+		FSoftObjectPath SpecifiedSkeleton;
+		SkeletalMeshFactoryNode->GetCustomSkeletonSoftObjectPath(SpecifiedSkeleton);
+		bool bSpecifiedSkeleton = SpecifiedSkeleton.IsValid();
+		if (SkeletonReference == nullptr)
+		{
+			UObject* SkeletonObject = nullptr;
+
+			if (SpecifiedSkeleton.IsValid())
 			{
-				UE_LOG(LogInterchangeImport, Error, TEXT("Could not create SkeletalMesh asset %s"), *Arguments.AssetName);
-				return nullptr;
+				SkeletonObject = SpecifiedSkeleton.TryLoad();
+			}
+			else if (SkeletonNode->ReferenceObject.IsValid())
+			{
+				SkeletonObject = SkeletonNode->ReferenceObject.TryLoad();
 			}
-			//Dirty the DDC Key for any imported Skeletal Mesh
-			SkeletalMesh->InvalidateDeriveDataCacheGUID();
 
-			FSkeletalMeshModel* ImportedResource = SkeletalMesh->GetImportedModel();
-			if (!ensure(ImportedResource->LODModels.Num() == 0))
+			if (SkeletonObject)
 			{
-				ImportedResource->LODModels.Empty();
+				SkeletonReference = Cast<USkeleton>(SkeletonObject);
+
 			}
-			USkeleton* SkeletonReference = nullptr;
-			int32 LodCount = SkeletalMeshFactoryNode->GetLodDataCount();
-			TArray<FString> LodDataUniqueIds;
-			SkeletalMeshFactoryNode->GetLodDataUniqueIds(LodDataUniqueIds);
-			ensure(LodDataUniqueIds.Num() == LodCount);
-			int32 CurrentLodIndex = 0;
-			for (int32 LodIndex = 0; LodIndex < LodCount; ++LodIndex)
+				
+			if (!ensure(SkeletonReference))
 			{
-				FString LodUniqueId = LodDataUniqueIds[LodIndex];
-				const UInterchangeSkeletalMeshLodDataNode* LodDataNode = Cast<UInterchangeSkeletalMeshLodDataNode>(Arguments.NodeContainer->GetNode(LodUniqueId));
-				if (!LodDataNode)
-				{
-					UE_LOG(LogInterchangeImport, Warning, TEXT("Invalid LOD when importing SkeletalMesh asset %s"), *Arguments.AssetName);
-					continue;
-				}
+				UE_LOG(LogInterchangeImport, Warning, TEXT("Invalid Skeleton LOD when importing SkeletalMesh asset %s"), *Arguments.AssetName);
+				break;
+			}
+		}
 
-				FString SkeletonNodeUid;
-				if (!LodDataNode->GetCustomSkeletonUid(SkeletonNodeUid))
-				{
-					UE_LOG(LogInterchangeImport, Warning, TEXT("Invalid Skeleton LOD when importing SkeletalMesh asset %s"), *Arguments.AssetName);
-					continue;
-				}
-				const UInterchangeSkeletonFactoryNode* SkeletonNode = Cast<UInterchangeSkeletonFactoryNode>(Arguments.NodeContainer->GetNode(SkeletonNodeUid));
-				if (!SkeletonNode)
-				{
-					UE_LOG(LogInterchangeImport, Warning, TEXT("Invalid Skeleton LOD when importing SkeletalMesh asset %s"), *Arguments.AssetName);
-					continue;
-				}
+		FString RootJointNodeId;
+		if (!SkeletonNode->GetCustomRootJointUid(RootJointNodeId))
+		{
+			UE_LOG(LogInterchangeImport, Warning, TEXT("Invalid Skeleton LOD Root Joint when importing SkeletalMesh asset %s"), *Arguments.AssetName);
+			continue;
+		}
 
-				FSoftObjectPath SpecifiedSkeleton;
-				SkeletalMeshFactoryNode->GetCustomSkeletonSoftObjectPath(SpecifiedSkeleton);
-				bool bSpecifiedSkeleton = SpecifiedSkeleton.IsValid();
-				if (SkeletonReference == nullptr)
+		int32 SkeletonDepth = 0;
+		TArray<SkeletalMeshImportData::FBone> RefBonesBinary;
+		bool bUseTimeZeroAsBindPose = false;
+		SkeletonNode->GetCustomUseTimeZeroForBindPose(bUseTimeZeroAsBindPose);
+		bool bDiffPose = false;
+		UE::Interchange::Private::ProcessImportMeshSkeleton(SkeletonReference, SkeletalMesh->GetRefSkeleton(), SkeletonDepth, Arguments.NodeContainer, RootJointNodeId, RefBonesBinary, bUseTimeZeroAsBindPose, bDiffPose);
+		if (bSpecifiedSkeleton && !SkeletonReference->IsCompatibleMesh(SkeletalMesh))
+		{
+			UE_LOG(LogInterchangeImport, Warning, TEXT("The skeleton %s is incompatible with the imported skeletalmesh asset %s"), *SkeletonReference->GetName(), *Arguments.AssetName);
+		}
+				
+		TArray<UE::Interchange::Private::FMeshNodeContext> MeshReferences;
+		//Scope to query the mesh node
+		{
+			TArray<FString> MeshUids;
+			LodDataNode->GetMeshUids(MeshUids);
+			MeshReferences.Reserve(MeshUids.Num());
+			for (const FString& MeshUid : MeshUids)
+			{
+				UE::Interchange::Private::FMeshNodeContext MeshReference;
+				MeshReference.MeshNode = Cast<UInterchangeMeshNode>(Arguments.NodeContainer->GetNode(MeshUid));
+				if (!MeshReference.MeshNode)
 				{
-					UObject* SkeletonObject = nullptr;
-
-					if (SpecifiedSkeleton.IsValid())
+					//The reference is a scene node and we need to bake the geometry
+					MeshReference.SceneNode = Cast<UInterchangeSceneNode>(Arguments.NodeContainer->GetNode(MeshUid));
+					if (!ensure(MeshReference.SceneNode != nullptr))
 					{
-						SkeletonObject = SpecifiedSkeleton.TryLoad();
-					}
-					else if (SkeletonNode->ReferenceObject.IsValid())
-					{
-						SkeletonObject = SkeletonNode->ReferenceObject.TryLoad();
+						UE_LOG(LogInterchangeImport, Warning, TEXT("Invalid LOD mesh reference when importing SkeletalMesh asset %s"), *Arguments.AssetName);
+						continue;
 					}
-
-					if (SkeletonObject)
+					FString MeshDependencyUid;
+					MeshReference.SceneNode->GetCustomAssetInstanceUid(MeshDependencyUid);
+					MeshReference.MeshNode = Cast<UInterchangeMeshNode>(Arguments.NodeContainer->GetNode(MeshDependencyUid));
+					//Cache the scene node global matrix, we will use this matrix to bake the vertices, add the node geometric mesh offset to this matrix to bake it properly
+					FTransform SceneNodeGlobalTransform;
+					if (!bUseTimeZeroAsBindPose || !MeshReference.SceneNode->GetCustomTimeZeroGlobalTransform(SceneNodeGlobalTransform))
 					{
-						SkeletonReference = Cast<USkeleton>(SkeletonObject);
-
+						ensure(MeshReference.SceneNode->GetCustomGlobalTransform(SceneNodeGlobalTransform));
 					}
-					if (!ensure(SkeletonReference))
+					FTransform SceneNodeGeometricTransform;
+					if(MeshReference.SceneNode->GetCustomGeometricTransform(SceneNodeGeometricTransform))
 					{
-						UE_LOG(LogInterchangeImport, Warning, TEXT("Invalid Skeleton LOD when importing SkeletalMesh asset %s"), *Arguments.AssetName);
-						break;
+						SceneNodeGlobalTransform *= SceneNodeGeometricTransform;
 					}
+					MeshReference.SceneGlobalTransform = SceneNodeGlobalTransform;
 				}
-
-				FString RootJointNodeId;
-				if (!SkeletonNode->GetCustomRootJointUid(RootJointNodeId))
+				if (!ensure(MeshReference.MeshNode != nullptr))
 				{
-					UE_LOG(LogInterchangeImport, Warning, TEXT("Invalid Skeleton LOD Root Joint when importing SkeletalMesh asset %s"), *Arguments.AssetName);
+					UE_LOG(LogInterchangeImport, Warning, TEXT("Invalid LOD mesh reference when importing SkeletalMesh asset %s"), *Arguments.AssetName);
 					continue;
 				}
-
-				int32 SkeletonDepth = 0;
-				TArray<SkeletalMeshImportData::FBone> RefBonesBinary;
-				bool bUseTimeZeroAsBindPose = false;
-				SkeletonNode->GetCustomUseTimeZeroForBindPose(bUseTimeZeroAsBindPose);
-				bool bDiffPose = false;
-				UE::Interchange::Private::ProcessImportMeshSkeleton(SkeletonReference, SkeletalMesh->GetRefSkeleton(), SkeletonDepth, Arguments.NodeContainer, RootJointNodeId, RefBonesBinary, bUseTimeZeroAsBindPose, bDiffPose);
-				if (bSpecifiedSkeleton && !SkeletonReference->IsCompatibleMesh(SkeletalMesh))
+				TOptional<FString> MeshPayloadKey = MeshReference.MeshNode->GetPayLoadKey();
+				if (MeshPayloadKey.IsSet())
 				{
-					UE_LOG(LogInterchangeImport, Warning, TEXT("The skeleton %s is incompatible with the imported skeletalmesh asset %s"), *SkeletonReference->GetName(), *Arguments.AssetName);
+					MeshReference.TranslatorPayloadKey = MeshPayloadKey.GetValue();
 				}
-				
-				TArray<UE::Interchange::Private::FMeshNodeContext> MeshReferences;
-				//Scope to query the mesh node
+				else
 				{
-					TArray<FString> MeshUids;
-					LodDataNode->GetMeshUids(MeshUids);
-					MeshReferences.Reserve(MeshUids.Num());
-					for (const FString& MeshUid : MeshUids)
-					{
-						UE::Interchange::Private::FMeshNodeContext MeshReference;
-						MeshReference.MeshNode = Cast<UInterchangeMeshNode>(Arguments.NodeContainer->GetNode(MeshUid));
-						if (!MeshReference.MeshNode)
-						{
-							//The reference is a scene node and we need to bake the geometry
-							MeshReference.SceneNode = Cast<UInterchangeSceneNode>(Arguments.NodeContainer->GetNode(MeshUid));
-							if (!ensure(MeshReference.SceneNode != nullptr))
-							{
-								UE_LOG(LogInterchangeImport, Warning, TEXT("Invalid LOD mesh reference when importing SkeletalMesh asset %s"), *Arguments.AssetName);
-								continue;
-							}
-							FString MeshDependencyUid;
-							MeshReference.SceneNode->GetCustomAssetInstanceUid(MeshDependencyUid);
-							MeshReference.MeshNode = Cast<UInterchangeMeshNode>(Arguments.NodeContainer->GetNode(MeshDependencyUid));
-							//Cache the scene node global matrix, we will use this matrix to bake the vertices, add the node geometric mesh offset to this matrix to bake it properly
-							FTransform SceneNodeGlobalTransform;
-							if (!bUseTimeZeroAsBindPose || !MeshReference.SceneNode->GetCustomTimeZeroGlobalTransform(SceneNodeGlobalTransform))
-							{
-								ensure(MeshReference.SceneNode->GetCustomGlobalTransform(SceneNodeGlobalTransform));
-							}
-							FTransform SceneNodeGeometricTransform;
-							if(MeshReference.SceneNode->GetCustomGeometricTransform(SceneNodeGeometricTransform))
-							{
-								SceneNodeGlobalTransform *= SceneNodeGeometricTransform;
-							}
-							MeshReference.SceneGlobalTransform = SceneNodeGlobalTransform;
-						}
-						if (!ensure(MeshReference.MeshNode != nullptr))
-						{
-							UE_LOG(LogInterchangeImport, Warning, TEXT("Invalid LOD mesh reference when importing SkeletalMesh asset %s"), *Arguments.AssetName);
-							continue;
-						}
-						TOptional<FString> MeshPayloadKey = MeshReference.MeshNode->GetPayLoadKey();
-						if (MeshPayloadKey.IsSet())
-						{
-							MeshReference.TranslatorPayloadKey = MeshPayloadKey.GetValue();
-						}
-						else
-						{
-							UE_LOG(LogInterchangeImport, Warning, TEXT("Empty LOD mesh reference payload when importing SkeletalMesh asset %s"), *Arguments.AssetName);
-							continue;
-						}
-						MeshReferences.Add(MeshReference);
-					}
-				}
-
-				//Add the lod mesh data to the skeletalmesh
-				FSkeletalMeshImportData SkeletalMeshImportData;
-				const bool bSkinControlPointToTimeZero = bUseTimeZeroAsBindPose && bDiffPose;
-				//Get all meshes and blend shapes payload and fill the SkeletalMeshImportData structure
-				UE::Interchange::Private::RetrieveAllSkeletalMeshPayloadsAndFillImportData(SkeletalMeshFactoryNode
-																						   , SkeletalMeshImportData
-																						   , MeshReferences
-																						   , RefBonesBinary
-																						   , Arguments
-																						   , SkeletalMeshTranslatorPayloadInterface
-																						   , bSkinControlPointToTimeZero
-																						   , Arguments.NodeContainer
-																						   , RootJointNodeId);
-				//////////////////////////////////////////////////////////////////////////
-				//Manage vertex color, we want to use the translated source data
-				//Replace -> do nothing
-				//Ignore -> remove vertex color from import data (when we re-import, ignore have to put back the current mesh vertex color)
-				//Override -> replace the vertex color by the override color
-				{
-					bool bReplaceVertexColor = false;
-					SkeletalMeshFactoryNode->GetCustomVertexColorReplace(bReplaceVertexColor);
-					if (!bReplaceVertexColor)
-					{
-						bool bIgnoreVertexColor = false;
-						SkeletalMeshFactoryNode->GetCustomVertexColorIgnore(bIgnoreVertexColor);
-						if (bIgnoreVertexColor)
-						{
-							//Flush the vertex color, if we re-import we have to fill it with the old data
-							SkeletalMeshImportData.bHasVertexColors = false;
-							for (SkeletalMeshImportData::FVertex& Wedge : SkeletalMeshImportData.Wedges)
-							{
-								Wedge.Color = FColor::White;
-							}
-						}
-						else
-						{
-							FColor OverrideVertexColor;
-							if (SkeletalMeshFactoryNode->GetCustomVertexColorOverride(OverrideVertexColor))
-							{
-								SkeletalMeshImportData.bHasVertexColors = true;
-								for(SkeletalMeshImportData::FVertex& Wedge : SkeletalMeshImportData.Wedges)
-								{
-									Wedge.Color = OverrideVertexColor;
-								}
-							}
-						}
-					}
-					// Store whether or not this mesh has vertex colors
-					SkeletalMesh->SetHasVertexColors(SkeletalMeshImportData.bHasVertexColors);
-					SkeletalMesh->SetVertexColorGuid(SkeletalMesh->GetHasVertexColors() ? FGuid::NewGuid() : FGuid());
+					UE_LOG(LogInterchangeImport, Warning, TEXT("Empty LOD mesh reference payload when importing SkeletalMesh asset %s"), *Arguments.AssetName);
+					continue;
 				}
+				MeshReferences.Add(MeshReference);
+			}
+		}
 
-				ensure(ImportedResource->LODModels.Add(new FSkeletalMeshLODModel()) == CurrentLodIndex);
-				FSkeletalMeshLODModel& LODModel = ImportedResource->LODModels[CurrentLodIndex];
-
-				TMap<FString, UMaterialInterface*> AvailableMaterials;
-				TArray<FString> FactoryDependencies;
-				SkeletalMeshFactoryNode->GetFactoryDependencies(FactoryDependencies);
-				for (int32 DependencyIndex = 0; DependencyIndex < FactoryDependencies.Num(); ++DependencyIndex)
+		//Add the lod mesh data to the skeletalmesh
+		FSkeletalMeshImportData SkeletalMeshImportData;
+		const bool bSkinControlPointToTimeZero = bUseTimeZeroAsBindPose && bDiffPose;
+		//Get all meshes and blend shapes payload and fill the SkeletalMeshImportData structure
+		UE::Interchange::Private::RetrieveAllSkeletalMeshPayloadsAndFillImportData(SkeletalMeshFactoryNode
+																					, SkeletalMeshImportData
+																					, MeshReferences
+																					, RefBonesBinary
+																					, Arguments
+																					, SkeletalMeshTranslatorPayloadInterface
+																					, bSkinControlPointToTimeZero
+																					, Arguments.NodeContainer
+																					, RootJointNodeId);
+		//////////////////////////////////////////////////////////////////////////
+		//Manage vertex color, we want to use the translated source data
+		//Replace -> do nothing
+		//Ignore -> remove vertex color from import data (when we re-import, ignore have to put back the current mesh vertex color)
+		//Override -> replace the vertex color by the override color
+		{
+			bool bReplaceVertexColor = false;
+			SkeletalMeshFactoryNode->GetCustomVertexColorReplace(bReplaceVertexColor);
+			if (!bReplaceVertexColor)
+			{
+				bool bIgnoreVertexColor = false;
+				SkeletalMeshFactoryNode->GetCustomVertexColorIgnore(bIgnoreVertexColor);
+				if (bIgnoreVertexColor)
 				{
-					const UInterchangeMaterialFactoryNode* MaterialFactoryNode = Cast<UInterchangeMaterialFactoryNode>(Arguments.NodeContainer->GetNode(FactoryDependencies[DependencyIndex]));
-					if (!MaterialFactoryNode || !MaterialFactoryNode->ReferenceObject.IsValid())
+					//Flush the vertex color, if we re-import we have to fill it with the old data
+					SkeletalMeshImportData.bHasVertexColors = false;
+					for (SkeletalMeshImportData::FVertex& Wedge : SkeletalMeshImportData.Wedges)
 					{
-						continue;
-					}
-					if (!MaterialFactoryNode->IsEnabled())
-					{
-						continue;
+						Wedge.Color = FColor::White;
 					}
-					UMaterialInterface* MaterialInterface = Cast<UMaterialInterface>(MaterialFactoryNode->ReferenceObject.ResolveObject());
-					if (!MaterialInterface)
-					{
-						continue;
-					}
-					AvailableMaterials.Add(MaterialFactoryNode->GetDisplayLabel(), MaterialInterface);
 				}
-
-				UE::Interchange::Private::ProcessImportMeshMaterials(SkeletalMesh->GetMaterials(), SkeletalMeshImportData, AvailableMaterials);
-				UE::Interchange::Private::ProcessImportMeshInfluences(SkeletalMeshImportData.Wedges.Num(), SkeletalMeshImportData.Influences);
-				//Store the original fbx import data the SkelMeshImportDataPtr should not be modified after this
-				SkeletalMesh->SaveLODImportedData(CurrentLodIndex, SkeletalMeshImportData);
-				//We reimport both
-				SkeletalMesh->SetLODImportedDataVersions(CurrentLodIndex, ESkeletalMeshGeoImportVersions::LatestVersion, ESkeletalMeshSkinningImportVersions::LatestVersion);
-				FSkeletalMeshLODInfo& NewLODInfo = SkeletalMesh->AddLODInfo();
-				NewLODInfo.ReductionSettings.NumOfTrianglesPercentage = 1.0f;
-				NewLODInfo.ReductionSettings.NumOfVertPercentage = 1.0f;
-				NewLODInfo.ReductionSettings.MaxDeviationPercentage = 0.0f;
-				NewLODInfo.LODHysteresis = 0.02f;
-				NewLODInfo.bImportWithBaseMesh = true;
-
-				//Make sure each LOD material index point on the correct skeletalMesh materials list index.
-				TArray<FSkeletalMaterial>& Materials = SkeletalMesh->GetMaterials();
-				const TArray<SkeletalMeshImportData::FMaterial>& ImportedMaterials = SkeletalMeshImportData.Materials;
-				if (FSkeletalMeshLODInfo* LodInfo = SkeletalMesh->GetLODInfo(CurrentLodIndex))
+				else
 				{
-					LodInfo->LODMaterialMap.Empty();
-					// Now set up the material mapping array.
-					for (int32 ImportedMaterialIndex = 0; ImportedMaterialIndex < ImportedMaterials.Num(); ImportedMaterialIndex++)
+					FColor OverrideVertexColor;
+					if (SkeletalMeshFactoryNode->GetCustomVertexColorOverride(OverrideVertexColor))
 					{
-						FName ImportedMaterialName = *(ImportedMaterials[ImportedMaterialIndex].MaterialImportName);
-						//Match by name
-						int32 LODMatIndex = INDEX_NONE;
-						for (int32 MaterialIndex = 0; MaterialIndex < Materials.Num(); ++MaterialIndex)
+						SkeletalMeshImportData.bHasVertexColors = true;
+						for(SkeletalMeshImportData::FVertex& Wedge : SkeletalMeshImportData.Wedges)
 						{
-							const FSkeletalMaterial& SkeletalMaterial = Materials[MaterialIndex];
-							if (SkeletalMaterial.ImportedMaterialSlotName != NAME_None && SkeletalMaterial.ImportedMaterialSlotName == ImportedMaterialName)
-							{
-								LODMatIndex = MaterialIndex;
-								break;
-							}
-						}
-						// If we dont have a match, add a new entry to the material list.
-						if (LODMatIndex == INDEX_NONE)
-						{
-							LODMatIndex = Materials.Add(FSkeletalMaterial(ImportedMaterials[ImportedMaterialIndex].Material.Get(), true, false, ImportedMaterialName, ImportedMaterialName));
+							Wedge.Color = OverrideVertexColor;
 						}
-						LodInfo->LODMaterialMap.Add(LODMatIndex);
 					}
 				}
+			}
+			// Store whether or not this mesh has vertex colors
+			SkeletalMesh->SetHasVertexColors(SkeletalMeshImportData.bHasVertexColors);
+			SkeletalMesh->SetVertexColorGuid(SkeletalMesh->GetHasVertexColors() ? FGuid::NewGuid() : FGuid());
+		}
 
-				//Add the bound to the skeletal mesh
-				if (SkeletalMesh->GetImportedBounds().BoxExtent.IsNearlyZero())
-				{
-					FBox3f BoundingBox(SkeletalMeshImportData.Points.GetData(), SkeletalMeshImportData.Points.Num());
-					const FVector3f BoundingBoxSize = BoundingBox.GetSize();
-
-					if (SkeletalMeshImportData.Points.Num() > 2 && BoundingBoxSize.X < THRESH_POINTS_ARE_SAME && BoundingBoxSize.Y < THRESH_POINTS_ARE_SAME && BoundingBoxSize.Z < THRESH_POINTS_ARE_SAME)
-					{
-						//TODO log a user error
-						//AddTokenizedErrorMessage(FTokenizedMessage::Create(EMessageSeverity::Error, FText::Format(LOCTEXT("FbxSkeletaLMeshimport_ErrorMeshTooSmall", "Cannot import this mesh, the bounding box of this mesh is smaller than the supported threshold[{0}]."), FText::FromString(FString::Printf(TEXT("%f"), THRESH_POINTS_ARE_SAME)))), FFbxErrors::SkeletalMesh_FillImportDataFailed);
-					}
-					SkeletalMesh->SetImportedBounds(FBoxSphereBounds((FBox)BoundingBox));
-				}
-
-				CurrentLodIndex++;
+		if (bIsReImport)
+		{
+			while (ImportedResource->LODModels.Num() <= CurrentLodIndex)
+			{
+				ImportedResource->LODModels.Add(new FSkeletalMeshLODModel());
 			}
+		}
+		else
+		{
+			ensure(ImportedResource->LODModels.Add(new FSkeletalMeshLODModel()) == CurrentLodIndex);
+		}
+		FSkeletalMeshLODModel& LODModel = ImportedResource->LODModels[CurrentLodIndex];
 
-			SkeletonReference->MergeAllBonesToBoneTree(SkeletalMesh);
-			if (SkeletalMesh->GetSkeleton() != SkeletonReference)
+		TMap<FString, UMaterialInterface*> AvailableMaterials;
+		TArray<FString> FactoryDependencies;
+		SkeletalMeshFactoryNode->GetFactoryDependencies(FactoryDependencies);
+		for (int32 DependencyIndex = 0; DependencyIndex < FactoryDependencies.Num(); ++DependencyIndex)
+		{
+			const UInterchangeMaterialFactoryNode* MaterialFactoryNode = Cast<UInterchangeMaterialFactoryNode>(Arguments.NodeContainer->GetNode(FactoryDependencies[DependencyIndex]));
+			if (!MaterialFactoryNode || !MaterialFactoryNode->ReferenceObject.IsValid())
+			{
+				continue;
+			}
+			if (!MaterialFactoryNode->IsEnabled())
+			{
+				continue;
+			}
+			UMaterialInterface* MaterialInterface = Cast<UMaterialInterface>(MaterialFactoryNode->ReferenceObject.ResolveObject());
+			if (!MaterialInterface)
 			{
-				SkeletalMesh->SetSkeleton(SkeletonReference);
+				continue;
 			}
+			AvailableMaterials.Add(MaterialFactoryNode->GetDisplayLabel(), MaterialInterface);
+		}
 
-			SkeletalMesh->CalculateInvRefMatrices();
+		UE::Interchange::Private::ProcessImportMeshMaterials(SkeletalMesh->GetMaterials(), SkeletalMeshImportData, AvailableMaterials);
+		UE::Interchange::Private::ProcessImportMeshInfluences(SkeletalMeshImportData.Wedges.Num(), SkeletalMeshImportData.Influences);
+		//Store the original fbx import data the SkelMeshImportDataPtr should not be modified after this
+		SkeletalMesh->SaveLODImportedData(CurrentLodIndex, SkeletalMeshImportData);
+		//We reimport both
+		SkeletalMesh->SetLODImportedDataVersions(CurrentLodIndex, ESkeletalMeshGeoImportVersions::LatestVersion, ESkeletalMeshSkinningImportVersions::LatestVersion);
 
-			/** Apply all SkeletalMeshFactoryNode custom attributes to the skeletal mesh asset */
-			SkeletalMeshFactoryNode->ApplyAllCustomAttributeToObject(SkeletalMesh);
+		auto AddLodInfo = [&SkeletalMesh]()
+		{
+			FSkeletalMeshLODInfo& NewLODInfo = SkeletalMesh->AddLODInfo();
+			NewLODInfo.ReductionSettings.NumOfTrianglesPercentage = 1.0f;
+			NewLODInfo.ReductionSettings.NumOfVertPercentage = 1.0f;
+			NewLODInfo.ReductionSettings.MaxDeviationPercentage = 0.0f;
+			NewLODInfo.LODHysteresis = 0.02f;
+			NewLODInfo.bImportWithBaseMesh = true;
+		};
+
+		if (bIsReImport)
+		{
+			while (SkeletalMesh->GetLODNum() <= CurrentLodIndex)
+			{
+				AddLodInfo();
+			}
+		}
+		else
+		{
+			AddLodInfo();
+		}
 
-			bool bCreatePhysicsAsset = false;
-			SkeletalMeshFactoryNode->GetCustomCreatePhysicsAsset(bCreatePhysicsAsset);
-			
-			if (!bCreatePhysicsAsset)
+		TArray<FSkeletalMaterial>& Materials = SkeletalMesh->GetMaterials();
+		const TArray<SkeletalMeshImportData::FMaterial>& ImportedMaterials = SkeletalMeshImportData.Materials;
+		if (FSkeletalMeshLODInfo* LodInfo = SkeletalMesh->GetLODInfo(CurrentLodIndex))
+		{
+			LodInfo->LODMaterialMap.Empty();
+			// Now set up the material mapping array.
+			for (int32 ImportedMaterialIndex = 0; ImportedMaterialIndex < ImportedMaterials.Num(); ImportedMaterialIndex++)
 			{
-				FSoftObjectPath SpecifiedPhysicAsset;
-				SkeletalMeshFactoryNode->GetCustomPhysicAssetSoftObjectPath(SpecifiedPhysicAsset);
-				if (SpecifiedPhysicAsset.IsValid())
+				FName ImportedMaterialName = *(ImportedMaterials[ImportedMaterialIndex].MaterialImportName);
+				//Match by name
+				int32 LODMatIndex = INDEX_NONE;
+				for (int32 MaterialIndex = 0; MaterialIndex < Materials.Num(); ++MaterialIndex)
 				{
-					UPhysicsAsset* PhysicsAsset = Cast<UPhysicsAsset>(SpecifiedPhysicAsset.TryLoad());
-					SkeletalMesh->SetPhysicsAsset(PhysicsAsset);
+					const FSkeletalMaterial& SkeletalMaterial = Materials[MaterialIndex];
+					if (SkeletalMaterial.ImportedMaterialSlotName != NAME_None && SkeletalMaterial.ImportedMaterialSlotName == ImportedMaterialName)
+					{
+						LODMatIndex = MaterialIndex;
+						break;
+					}
 				}
+				// If we dont have a match, add a new entry to the material list.
+				if (LODMatIndex == INDEX_NONE)
+				{
+					LODMatIndex = Materials.Add(FSkeletalMaterial(ImportedMaterials[ImportedMaterialIndex].Material.Get(), true, false, ImportedMaterialName, ImportedMaterialName));
+				}
+				LodInfo->LODMaterialMap.Add(LODMatIndex);
 			}
 		}
-		else //Re-Import
+
+		//Add the bound to the skeletal mesh
+		if (SkeletalMesh->GetImportedBounds().BoxExtent.IsNearlyZero())
 		{
-			//We want to swap the source data and set properly any skeleton change
+			FBox3f BoundingBox(SkeletalMeshImportData.Points.GetData(), SkeletalMeshImportData.Points.Num());
+			const FVector3f BoundingBoxSize = BoundingBox.GetSize();
+
+			if (SkeletalMeshImportData.Points.Num() > 2 && BoundingBoxSize.X < THRESH_POINTS_ARE_SAME && BoundingBoxSize.Y < THRESH_POINTS_ARE_SAME && BoundingBoxSize.Z < THRESH_POINTS_ARE_SAME)
+			{
+				//TODO log a user error
+				//AddTokenizedErrorMessage(FTokenizedMessage::Create(EMessageSeverity::Error, FText::Format(LOCTEXT("FbxSkeletaLMeshimport_ErrorMeshTooSmall", "Cannot import this mesh, the bounding box of this mesh is smaller than the supported threshold[{0}]."), FText::FromString(FString::Printf(TEXT("%f"), THRESH_POINTS_ARE_SAME)))), FFbxErrors::SkeletalMesh_FillImportDataFailed);
+			}
+			SkeletalMesh->SetImportedBounds(FBoxSphereBounds((FBox)BoundingBox));
 		}
-		
-		//Getting the file Hash will cache it into the source data
-		Arguments.SourceData->GetFileContentHash();
 
-		//The interchange completion task (call in the GameThread after the factories pass), will call PostEditChange which will trig another asynchronous system that will build all material in parallel
+		CurrentLodIndex++;
+	}
+
+	SkeletonReference->MergeAllBonesToBoneTree(SkeletalMesh);
+	if (SkeletalMesh->GetSkeleton() != SkeletonReference)
+	{
+		SkeletalMesh->SetSkeleton(SkeletonReference);
+	}
+
+	SkeletalMesh->CalculateInvRefMatrices();
+
+	if (!bIsReImport)
+	{
+		/** Apply all SkeletalMeshFactoryNode custom attributes to the skeletal mesh asset */
+		SkeletalMeshFactoryNode->ApplyAllCustomAttributeToObject(SkeletalMesh);
+
+		bool bCreatePhysicsAsset = false;
+		SkeletalMeshFactoryNode->GetCustomCreatePhysicsAsset(bCreatePhysicsAsset);
+
+		if (!bCreatePhysicsAsset)
+		{
+			FSoftObjectPath SpecifiedPhysicAsset;
+			SkeletalMeshFactoryNode->GetCustomPhysicAssetSoftObjectPath(SpecifiedPhysicAsset);
+			if (SpecifiedPhysicAsset.IsValid())
+			{
+				UPhysicsAsset* PhysicsAsset = Cast<UPhysicsAsset>(SpecifiedPhysicAsset.TryLoad());
+				SkeletalMesh->SetPhysicsAsset(PhysicsAsset);
+			}
+		}
 	}
 	else
 	{
-		//The material is not a UMaterialInterface
-		SkeletalMeshObject->RemoveFromRoot();
-		SkeletalMeshObject->MarkAsGarbage();
+		//Apply the re import strategy 
+		UInterchangeAssetImportData* InterchangeAssetImportData = Cast<UInterchangeAssetImportData>(SkeletalMesh->GetAssetImportData());
+		UInterchangeBaseNode* PreviousNode = nullptr;
+		if (InterchangeAssetImportData)
+		{
+			PreviousNode = InterchangeAssetImportData->NodeContainer->GetNode(InterchangeAssetImportData->NodeUniqueID);
+		}
+		UInterchangeBaseNode* CurrentNode = NewObject<UInterchangeBaseNode>(GetTransientPackage(), UInterchangeSkeletalMeshFactoryNode::StaticClass());
+		UInterchangeBaseNode::CopyStorage(SkeletalMeshFactoryNode, CurrentNode);
+		CurrentNode->FillAllCustomAttributeFromObject(SkeletalMesh);
+		UE::Interchange::FFactoryCommon::ApplyReimportStrategyToAsset(SkeletalMesh, PreviousNode, CurrentNode, SkeletalMeshFactoryNode);
 	}
+		
+	//Getting the file Hash will cache it into the source data
+	Arguments.SourceData->GetFileContentHash();
+
+	//The interchange completion task (call in the GameThread after the factories pass), will call PostEditChange which will trig another asynchronous system that will build all material in parallel
 	return SkeletalMeshObject;
 
 #endif //else !WITH_EDITOR || !WITH_EDITORONLY_DATA
diff --git a/Engine/Plugins/Experimental/Interchange/Runtime/Source/Import/Private/Mesh/InterchangeSkeletonFactory.cpp b/Engine/Plugins/Experimental/Interchange/Runtime/Source/Import/Private/Mesh/InterchangeSkeletonFactory.cpp
--- a/Engine/Plugins/Experimental/Interchange/Runtime/Source/Import/Private/Mesh/InterchangeSkeletonFactory.cpp
+++ b/Engine/Plugins/Experimental/Interchange/Runtime/Source/Import/Private/Mesh/InterchangeSkeletonFactory.cpp
@@ -82,7 +82,9 @@ UObject* UInterchangeSkeletonFactory::CreateAsset(const FCreateAssetParams& Argu
 
 	// create an asset if it doesn't exist
 	UObject* ExistingAsset = StaticFindObject(nullptr, Arguments.Parent, *Arguments.AssetName);
-
+	
+	bool bIsReImport = Arguments.ReimportObject != nullptr;
+	
 	UObject* SkeletonObject = nullptr;
 	// create a new skeleton or overwrite existing asset, if possible
 	if (!ExistingAsset)
@@ -90,6 +92,8 @@ UObject* UInterchangeSkeletonFactory::CreateAsset(const FCreateAssetParams& Argu
 		//NewObject is not thread safe, the asset registry directory watcher tick on the main thread can trig before we finish initializing the UObject and will crash
 		//The UObject should have been create by calling CreateEmptyAsset on the main thread.
 		check(IsInGameThread());
+		//We should not do a NewObject if we are doing a reimport
+		check(!bIsReImport);
 		SkeletonObject = NewObject<UObject>(Arguments.Parent, SkeletonClass, *Arguments.AssetName, RF_Public | RF_Standalone);
 	}
 	else if(ExistingAsset->GetClass()->IsChildOf(SkeletonClass))
@@ -106,26 +110,31 @@ UObject* UInterchangeSkeletonFactory::CreateAsset(const FCreateAssetParams& Argu
 
 	if (SkeletonObject)
 	{
-		//Currently skeleton re-import will not touch the skeleton at all
-		//TODO design a re-import process for the skeleton (expressions and input connections)
-		if(!Arguments.ReimportObject)
+		USkeleton* Skeleton = Cast<USkeleton>(SkeletonObject);
+		if (!ensure(Skeleton))
 		{
-			USkeleton* Skeleton = Cast<USkeleton>(SkeletonObject);
-			if (!ensure(Skeleton))
-			{
-				UE_LOG(LogInterchangeImport, Warning, TEXT("Could not create Skeleton asset %s"), *Arguments.AssetName);
-				return nullptr;
-			}
-			FString RootJointUid;
-			if (!SkeletonNode->GetCustomRootJointUid(RootJointUid))
+			UE_LOG(LogInterchangeImport, Warning, TEXT("Could not create Skeleton asset %s"), *Arguments.AssetName);
+			return nullptr;
+		}
+		FString RootJointUid;
+		if (!SkeletonNode->GetCustomRootJointUid(RootJointUid))
+		{
+			UE_LOG(LogInterchangeImport, Warning, TEXT("Could not create Skeleton asset %s, because there is no valid root joint node"), *Arguments.AssetName);
+			return nullptr;
+		}
+		if (bIsReImport)
+		{
+			if (const UInterchangeBaseNode* RootJointNode = Arguments.NodeContainer->GetNode(RootJointUid))
 			{
-				UE_LOG(LogInterchangeImport, Warning, TEXT("Could not create Skeleton asset %s, because there is no valid root joint node"), *Arguments.AssetName);
-				return nullptr;
+				FName RootJointName = FName(*RootJointNode->GetDisplayLabel());
+				if (RootJointName != Skeleton->GetReferenceSkeleton().GetBoneName(0))
+				{
+					return nullptr;
+				}
 			}
-			
-			//The joint will be added by the skeletalmesh factory since we need a valid skeletalmesh to add joint to a skeleton
-			SkeletonNode->ReferenceObject = Skeleton;
 		}
+		//The joint will be added by the skeletalmesh factory since we need a valid skeletalmesh to add joint to a skeleton
+		SkeletonNode->ReferenceObject = Skeleton;
 		
 		//Getting the file Hash will cache it into the source data
 		Arguments.SourceData->GetFileContentHash();
diff --git a/Engine/Plugins/Experimental/Interchange/Runtime/Source/Import/Private/Mesh/InterchangeStaticMeshFactory.cpp b/Engine/Plugins/Experimental/Interchange/Runtime/Source/Import/Private/Mesh/InterchangeStaticMeshFactory.cpp
--- a/Engine/Plugins/Experimental/Interchange/Runtime/Source/Import/Private/Mesh/InterchangeStaticMeshFactory.cpp
+++ b/Engine/Plugins/Experimental/Interchange/Runtime/Source/Import/Private/Mesh/InterchangeStaticMeshFactory.cpp
@@ -3,6 +3,7 @@
 
 #include "Components.h"
 #include "Engine/StaticMesh.h"
+#include "InterchangeAssetImportData.h"
 #include "InterchangeImportCommon.h"
 #include "InterchangeImportLog.h"
 #include "InterchangeMaterialFactoryNode.h"
@@ -95,7 +96,7 @@ UObject* UInterchangeStaticMeshFactory::CreateAsset(const FCreateAssetParams& Ar
 		return nullptr;
 	}
 
-	const UInterchangeStaticMeshFactoryNode* StaticMeshFactoryNode = Cast<UInterchangeStaticMeshFactoryNode>(Arguments.AssetNode);
+	UInterchangeStaticMeshFactoryNode* StaticMeshFactoryNode = Cast<UInterchangeStaticMeshFactoryNode>(Arguments.AssetNode);
 	if (StaticMeshFactoryNode == nullptr)
 	{
 		return nullptr;
@@ -136,256 +137,265 @@ UObject* UInterchangeStaticMeshFactory::CreateAsset(const FCreateAssetParams& Ar
 		return nullptr;
 	}
 
-	if (StaticMeshObject)
+	UStaticMesh* StaticMesh = Cast<UStaticMesh>(StaticMeshObject);
+	if (!ensure(StaticMesh))
 	{
-		// Currently static mesh re-import will not touch the static mesh at all
-		// TODO design a re-import process for the static mesh (expressions and input connections)
-		if (!Arguments.ReimportObject)
-		{
-			UStaticMesh* StaticMesh = Cast<UStaticMesh>(StaticMeshObject);
-			if (!ensure(StaticMesh))
-			{
-				UE_LOG(LogInterchangeImport, Error, TEXT("Could not create StaticMesh asset %s"), *Arguments.AssetName);
-				return nullptr;
-			}
+		UE_LOG(LogInterchangeImport, Error, TEXT("Could not create StaticMesh asset %s"), *Arguments.AssetName);
+		return nullptr;
+	}
 			
-			int32 LodCount = StaticMeshFactoryNode->GetLodDataCount();
-			StaticMesh->SetNumSourceModels(LodCount);
+	int32 LodCount = StaticMeshFactoryNode->GetLodDataCount();
+	int32 FinalLodCount = LodCount;
+	if (Arguments.ReimportObject)
+	{
+		//When we reimport we dont want to reduce the number of existing LODs
+		FinalLodCount = FMath::Max(StaticMesh->GetNumLODs(), LodCount);
+	}
+	StaticMesh->SetNumSourceModels(FinalLodCount);
 
-			TArray<FString> LodDataUniqueIds;
-			StaticMeshFactoryNode->GetLodDataUniqueIds(LodDataUniqueIds);
-			ensure(LodDataUniqueIds.Num() == LodCount);
+	TArray<FString> LodDataUniqueIds;
+	StaticMeshFactoryNode->GetLodDataUniqueIds(LodDataUniqueIds);
+	ensure(LodDataUniqueIds.Num() == LodCount);
 
-			int32 CurrentLodIndex = 0;
-			for (int32 LodIndex = 0; LodIndex < LodCount; ++LodIndex)
-			{
-				FString LodUniqueId = LodDataUniqueIds[LodIndex];
-				const UInterchangeStaticMeshLodDataNode* LodDataNode = Cast<UInterchangeStaticMeshLodDataNode>(Arguments.NodeContainer->GetNode(LodUniqueId));
-				if (!LodDataNode)
-				{
-					UE_LOG(LogInterchangeImport, Warning, TEXT("Invalid LOD when importing StaticMesh asset %s"), *Arguments.AssetName);
-					continue;
-				}
+	int32 CurrentLodIndex = 0;
+	for (int32 LodIndex = 0; LodIndex < LodCount; ++LodIndex)
+	{
+		FString LodUniqueId = LodDataUniqueIds[LodIndex];
+		const UInterchangeStaticMeshLodDataNode* LodDataNode = Cast<UInterchangeStaticMeshLodDataNode>(Arguments.NodeContainer->GetNode(LodUniqueId));
+		if (!LodDataNode)
+		{
+			UE_LOG(LogInterchangeImport, Warning, TEXT("Invalid LOD when importing StaticMesh asset %s"), *Arguments.AssetName);
+			continue;
+		}
 
-				// Get the mesh node context for each MeshUids
-				struct FMeshNodeContext
-				{
-					const UInterchangeMeshNode* MeshNode = nullptr;
-					const UInterchangeSceneNode* SceneNode = nullptr;
-					TOptional<FTransform> SceneGlobalTransform;
-					FString TranslatorPayloadKey;
-				};
-				TArray<FMeshNodeContext> MeshReferences;
-
-				// Scope to query the mesh node
-				{
-					TArray<FString> MeshUids;
-					LodDataNode->GetMeshUids(MeshUids);
-					MeshReferences.Reserve(MeshUids.Num());
+		// Get the mesh node context for each MeshUids
+		struct FMeshNodeContext
+		{
+			const UInterchangeMeshNode* MeshNode = nullptr;
+			const UInterchangeSceneNode* SceneNode = nullptr;
+			TOptional<FTransform> SceneGlobalTransform;
+			FString TranslatorPayloadKey;
+		};
+		TArray<FMeshNodeContext> MeshReferences;
+
+		// Scope to query the mesh node
+		{
+			TArray<FString> MeshUids;
+			LodDataNode->GetMeshUids(MeshUids);
+			MeshReferences.Reserve(MeshUids.Num());
 
-					for (const FString& MeshUid : MeshUids)
+			for (const FString& MeshUid : MeshUids)
+			{
+				FMeshNodeContext MeshReference;
+				MeshReference.MeshNode = Cast<UInterchangeMeshNode>(Arguments.NodeContainer->GetNode(MeshUid));
+				if (!MeshReference.MeshNode)
+				{
+					// The reference is a scene node and we need to bake the geometry
+					MeshReference.SceneNode = Cast<UInterchangeSceneNode>(Arguments.NodeContainer->GetNode(MeshUid));
+					if (!ensure(MeshReference.SceneNode != nullptr))
 					{
-						FMeshNodeContext MeshReference;
-						MeshReference.MeshNode = Cast<UInterchangeMeshNode>(Arguments.NodeContainer->GetNode(MeshUid));
-						if (!MeshReference.MeshNode)
-						{
-							// The reference is a scene node and we need to bake the geometry
-							MeshReference.SceneNode = Cast<UInterchangeSceneNode>(Arguments.NodeContainer->GetNode(MeshUid));
-							if (!ensure(MeshReference.SceneNode != nullptr))
-							{
-								UE_LOG(LogInterchangeImport, Warning, TEXT("Invalid LOD mesh reference when importing StaticMesh asset %s"), *Arguments.AssetName);
-								continue;
-							}
-
-							FString MeshDependencyUid;
-							MeshReference.SceneNode->GetCustomAssetInstanceUid(MeshDependencyUid);
-							MeshReference.MeshNode = Cast<UInterchangeMeshNode>(Arguments.NodeContainer->GetNode(MeshDependencyUid));
-
-							// Cache the scene node global matrix, we will use this matrix to bake the vertices
-							FTransform SceneNodeGlobalTransform;
-							if (MeshReference.SceneNode->GetCustomGlobalTransform(SceneNodeGlobalTransform))
-							{
-								MeshReference.SceneGlobalTransform = SceneNodeGlobalTransform;
-							}
-						}
+						UE_LOG(LogInterchangeImport, Warning, TEXT("Invalid LOD mesh reference when importing StaticMesh asset %s"), *Arguments.AssetName);
+						continue;
+					}
 
-						if (!ensure(MeshReference.MeshNode != nullptr))
-						{
-							UE_LOG(LogInterchangeImport, Warning, TEXT("Invalid LOD mesh reference when importing StaticMesh asset %s"), *Arguments.AssetName);
-							continue;
-						}
+					FString MeshDependencyUid;
+					MeshReference.SceneNode->GetCustomAssetInstanceUid(MeshDependencyUid);
+					MeshReference.MeshNode = Cast<UInterchangeMeshNode>(Arguments.NodeContainer->GetNode(MeshDependencyUid));
 
-						TOptional<FString> MeshPayloadKey = MeshReference.MeshNode->GetPayLoadKey();
-						if (MeshPayloadKey.IsSet())
-						{
-							MeshReference.TranslatorPayloadKey = MeshPayloadKey.GetValue();
-						}
-						else
-						{
-							UE_LOG(LogInterchangeImport, Warning, TEXT("Empty LOD mesh reference payload when importing StaticMesh asset %s"), *Arguments.AssetName);
-							continue;
-						}
-
-						MeshReferences.Add(MeshReference);
+					// Cache the scene node global matrix, we will use this matrix to bake the vertices
+					FTransform SceneNodeGlobalTransform;
+					if (MeshReference.SceneNode->GetCustomGlobalTransform(SceneNodeGlobalTransform))
+					{
+						MeshReference.SceneGlobalTransform = SceneNodeGlobalTransform;
 					}
 				}
 
-				// Add the lod mesh data to the static mesh
-				FMeshDescription* LodMeshDescription = StaticMesh->CreateMeshDescription(CurrentLodIndex);
-
-				FStaticMeshOperations::FAppendSettings AppendSettings;
-				for (int32 ChannelIdx = 0; ChannelIdx < FStaticMeshOperations::FAppendSettings::MAX_NUM_UV_CHANNELS; ++ChannelIdx)
+				if (!ensure(MeshReference.MeshNode != nullptr))
 				{
-					AppendSettings.bMergeUVChannels[ChannelIdx] = true;
+					UE_LOG(LogInterchangeImport, Warning, TEXT("Invalid LOD mesh reference when importing StaticMesh asset %s"), *Arguments.AssetName);
+					continue;
 				}
 
-				// Fetch the payloads for all mesh references in parallel
-				TMap<FString, TFuture<TOptional<UE::Interchange::FStaticMeshPayloadData>>> MeshPayloads;
-				for (const FMeshNodeContext& MeshNodeContext : MeshReferences)
+				TOptional<FString> MeshPayloadKey = MeshReference.MeshNode->GetPayLoadKey();
+				if (MeshPayloadKey.IsSet())
 				{
-					const FString& PayloadKey = MeshNodeContext.TranslatorPayloadKey;
-					MeshPayloads.Add(PayloadKey, StaticMeshTranslatorPayloadInterface->GetStaticMeshPayloadData(PayloadKey));
+					MeshReference.TranslatorPayloadKey = MeshPayloadKey.GetValue();
 				}
-				
-				// Fill the lod mesh description using all combined mesh parts
-				for (const FMeshNodeContext& MeshNodeContext : MeshReferences)
+				else
 				{
-					TOptional<UE::Interchange::FStaticMeshPayloadData> LodMeshPayload = MeshPayloads.FindChecked(MeshNodeContext.TranslatorPayloadKey).Get();
-					if (!LodMeshPayload.IsSet())
-					{
-						UE_LOG(LogInterchangeImport, Warning, TEXT("Invalid static mesh payload key [%s] StaticMesh asset %s"), *MeshNodeContext.TranslatorPayloadKey, *Arguments.AssetName);
-						continue;
-					}
-
-					// Bake the payload, with the provided transform
-					if (MeshNodeContext.SceneGlobalTransform.IsSet())
-					{
-						AppendSettings.MeshTransform = MeshNodeContext.SceneGlobalTransform;
-					}
-					else
-					{
-						AppendSettings.MeshTransform.Reset();
-					}
-
-					FStaticMeshOperations::AppendMeshDescription(LodMeshPayload->MeshDescription, *LodMeshDescription, AppendSettings);
+					UE_LOG(LogInterchangeImport, Warning, TEXT("Empty LOD mesh reference payload when importing StaticMesh asset %s"), *Arguments.AssetName);
+					continue;
 				}
 
-				//////////////////////////////////////////////////////////////////////////
-				//Manage vertex color
-				//Replace -> do nothing, we want to use the translated source data
-				//Ignore -> remove vertex color from import data (when we re-import, ignore have to put back the current mesh vertex color)
-				//Override -> replace the vertex color by the override color
-				{
-					FStaticMeshAttributes Attributes(*LodMeshDescription);
-					TVertexInstanceAttributesRef<FVector4f> VertexInstanceColors = Attributes.GetVertexInstanceColors();
-					bool bReplaceVertexColor = false;
-					StaticMeshFactoryNode->GetCustomVertexColorReplace(bReplaceVertexColor);
-					if (!bReplaceVertexColor)
-					{
-						bool bIgnoreVertexColor = false;
-						StaticMeshFactoryNode->GetCustomVertexColorIgnore(bIgnoreVertexColor);
-						if (bIgnoreVertexColor)
-						{
-							//Flush the vertex color, if we re-import we have to fill it with the old data
-							for (const FVertexInstanceID& VertexInstanceID : LodMeshDescription->VertexInstances().GetElementIDs())
-							{
-								VertexInstanceColors[VertexInstanceID] = FVector4f(FLinearColor(FColor::White));
-							}
-						}
-						else
-						{
-							FColor OverrideVertexColor;
-							if (StaticMeshFactoryNode->GetCustomVertexColorOverride(OverrideVertexColor))
-							{
-								for (const FVertexInstanceID& VertexInstanceID : LodMeshDescription->VertexInstances().GetElementIDs())
-								{
-									VertexInstanceColors[VertexInstanceID] = FVector4f(FLinearColor(OverrideVertexColor));
-								}
-							}
-						}
-					}
-				}
+				MeshReferences.Add(MeshReference);
+			}
+		}
 
-				StaticMesh->CommitMeshDescription(CurrentLodIndex);
+		// Add the lod mesh data to the static mesh
+		FMeshDescription* LodMeshDescription = StaticMesh->CreateMeshDescription(CurrentLodIndex);
 
-				// Handle materials
-				// @TODO: move this outside of the LOD loop?
-				TArray<FString> FactoryDependencies;
-				StaticMeshFactoryNode->GetFactoryDependencies(FactoryDependencies);
+		FStaticMeshOperations::FAppendSettings AppendSettings;
+		for (int32 ChannelIdx = 0; ChannelIdx < FStaticMeshOperations::FAppendSettings::MAX_NUM_UV_CHANNELS; ++ChannelIdx)
+		{
+			AppendSettings.bMergeUVChannels[ChannelIdx] = true;
+		}
 
-				for (int32 DependencyIndex = 0; DependencyIndex < FactoryDependencies.Num(); ++DependencyIndex)
-				{
-					const UInterchangeMaterialFactoryNode* MaterialFactoryNode = Cast<UInterchangeMaterialFactoryNode>(Arguments.NodeContainer->GetNode(FactoryDependencies[DependencyIndex]));
-					if (!MaterialFactoryNode || !MaterialFactoryNode->ReferenceObject.IsValid())
-					{
-						continue;
-					}
-					if (!MaterialFactoryNode->IsEnabled())
-					{
-						continue;
-					}
+		// Fetch the payloads for all mesh references in parallel
+		TMap<FString, TFuture<TOptional<UE::Interchange::FStaticMeshPayloadData>>> MeshPayloads;
+		for (const FMeshNodeContext& MeshNodeContext : MeshReferences)
+		{
+			const FString& PayloadKey = MeshNodeContext.TranslatorPayloadKey;
+			MeshPayloads.Add(PayloadKey, StaticMeshTranslatorPayloadInterface->GetStaticMeshPayloadData(PayloadKey));
+		}
+				
+		// Fill the lod mesh description using all combined mesh parts
+		for (const FMeshNodeContext& MeshNodeContext : MeshReferences)
+		{
+			TOptional<UE::Interchange::FStaticMeshPayloadData> LodMeshPayload = MeshPayloads.FindChecked(MeshNodeContext.TranslatorPayloadKey).Get();
+			if (!LodMeshPayload.IsSet())
+			{
+				UE_LOG(LogInterchangeImport, Warning, TEXT("Invalid static mesh payload key [%s] StaticMesh asset %s"), *MeshNodeContext.TranslatorPayloadKey, *Arguments.AssetName);
+				continue;
+			}
 
-					FName MaterialSlotName = *MaterialFactoryNode->GetDisplayLabel();
-					UMaterialInterface* MaterialInterface = Cast<UMaterialInterface>(MaterialFactoryNode->ReferenceObject.ResolveObject());
-					if (!MaterialInterface)
-					{
-						MaterialInterface = UMaterial::GetDefaultMaterial(MD_Surface);
-					}
+			// Bake the payload, with the provided transform
+			if (MeshNodeContext.SceneGlobalTransform.IsSet())
+			{
+				AppendSettings.MeshTransform = MeshNodeContext.SceneGlobalTransform;
+			}
+			else
+			{
+				AppendSettings.MeshTransform.Reset();
+			}
 
-					int32 MaterialSlotIndex = StaticMesh->GetMaterialIndexFromImportedMaterialSlotName(MaterialSlotName);
-					if (MaterialSlotIndex == INDEX_NONE)
+			FStaticMeshOperations::AppendMeshDescription(LodMeshPayload->MeshDescription, *LodMeshDescription, AppendSettings);
+		}
+
+		//////////////////////////////////////////////////////////////////////////
+		//Manage vertex color
+		//Replace -> do nothing, we want to use the translated source data
+		//Ignore -> remove vertex color from import data (when we re-import, ignore have to put back the current mesh vertex color)
+		//Override -> replace the vertex color by the override color
+		{
+			FStaticMeshAttributes Attributes(*LodMeshDescription);
+			TVertexInstanceAttributesRef<FVector4f> VertexInstanceColors = Attributes.GetVertexInstanceColors();
+			bool bReplaceVertexColor = false;
+			StaticMeshFactoryNode->GetCustomVertexColorReplace(bReplaceVertexColor);
+			if (!bReplaceVertexColor)
+			{
+				bool bIgnoreVertexColor = false;
+				StaticMeshFactoryNode->GetCustomVertexColorIgnore(bIgnoreVertexColor);
+				if (bIgnoreVertexColor)
+				{
+					//Flush the vertex color, if we re-import we have to fill it with the old data
+					for (const FVertexInstanceID& VertexInstanceID : LodMeshDescription->VertexInstances().GetElementIDs())
 					{
-						StaticMesh->GetStaticMaterials().Emplace(MaterialInterface, MaterialSlotName);
+						VertexInstanceColors[VertexInstanceID] = FVector4f(FLinearColor(FColor::White));
 					}
 				}
-
-				FStaticMeshConstAttributes StaticMeshAttributes(*LodMeshDescription);
-				TPolygonGroupAttributesRef<const FName> SlotNames = StaticMeshAttributes.GetPolygonGroupMaterialSlotNames();
-				int32 SectionIndex = 0;
-				for (FPolygonGroupID PolygonGroupID : LodMeshDescription->PolygonGroups().GetElementIDs())
+				else
 				{
-					int32 MaterialSlotIndex = StaticMesh->GetMaterialIndexFromImportedMaterialSlotName(SlotNames[PolygonGroupID]);
-					
-					if (MaterialSlotIndex == INDEX_NONE)
+					FColor OverrideVertexColor;
+					if (StaticMeshFactoryNode->GetCustomVertexColorOverride(OverrideVertexColor))
 					{
-						// If no material was found with this slot name, it is probably because the pipeline is configured to not import materials.
-						// Fill out a blank slot instead.
-						MaterialSlotIndex = StaticMesh->GetStaticMaterials().Emplace(nullptr, SlotNames[PolygonGroupID]);
+						for (const FVertexInstanceID& VertexInstanceID : LodMeshDescription->VertexInstances().GetElementIDs())
+						{
+							VertexInstanceColors[VertexInstanceID] = FVector4f(FLinearColor(OverrideVertexColor));
+						}
 					}
+				}
+			}
+		}
 
-					FMeshSectionInfo Info = StaticMesh->GetSectionInfoMap().Get(CurrentLodIndex, SectionIndex);
-					Info.MaterialIndex = MaterialSlotIndex;
-					StaticMesh->GetSectionInfoMap().Remove(CurrentLodIndex, SectionIndex);
-					StaticMesh->GetSectionInfoMap().Set(CurrentLodIndex, SectionIndex, Info);
+		StaticMesh->CommitMeshDescription(CurrentLodIndex);
 
-					SectionIndex++;
-				}
+		// Handle materials
+		// @TODO: move this outside of the LOD loop?
+		TArray<FString> FactoryDependencies;
+		StaticMeshFactoryNode->GetFactoryDependencies(FactoryDependencies);
 
-				// Set static mesh build settings
-				FStaticMeshSourceModel& SrcModel = StaticMesh->GetSourceModel(CurrentLodIndex);
-				SrcModel.ReductionSettings.MaxDeviation = 0.0f;
-				SrcModel.ReductionSettings.PercentTriangles = 1.0f;
-				SrcModel.ReductionSettings.PercentVertices = 1.0f;
+		for (int32 DependencyIndex = 0; DependencyIndex < FactoryDependencies.Num(); ++DependencyIndex)
+		{
+			const UInterchangeMaterialFactoryNode* MaterialFactoryNode = Cast<UInterchangeMaterialFactoryNode>(Arguments.NodeContainer->GetNode(FactoryDependencies[DependencyIndex]));
+			if (!MaterialFactoryNode || !MaterialFactoryNode->ReferenceObject.IsValid())
+			{
+				continue;
+			}
+			if (!MaterialFactoryNode->IsEnabled())
+			{
+				continue;
+			}
 
-				CurrentLodIndex++;
+			FName MaterialSlotName = *MaterialFactoryNode->GetDisplayLabel();
+			UMaterialInterface* MaterialInterface = Cast<UMaterialInterface>(MaterialFactoryNode->ReferenceObject.ResolveObject());
+			if (!MaterialInterface)
+			{
+				MaterialInterface = UMaterial::GetDefaultMaterial(MD_Surface);
 			}
 
-			// Apply all StaticMeshFactoryNode custom attributes to the static mesh asset
-			StaticMeshFactoryNode->ApplyAllCustomAttributeToObject(StaticMesh);
+			int32 MaterialSlotIndex = StaticMesh->GetMaterialIndexFromImportedMaterialSlotName(MaterialSlotName);
+			if (MaterialSlotIndex == INDEX_NONE)
+			{
+				StaticMesh->GetStaticMaterials().Emplace(MaterialInterface, MaterialSlotName);
+			}
 		}
-		
-		// Getting the file Hash will cache it into the source data
-		Arguments.SourceData->GetFileContentHash();
 
-		// The interchange completion task (called in the GameThread after the factories pass), will call PostEditChange
-		// which will trigger another asynchronous system that will build all static meshes in parallel
+		FStaticMeshConstAttributes StaticMeshAttributes(*LodMeshDescription);
+		TPolygonGroupAttributesRef<const FName> SlotNames = StaticMeshAttributes.GetPolygonGroupMaterialSlotNames();
+		int32 SectionIndex = 0;
+		for (FPolygonGroupID PolygonGroupID : LodMeshDescription->PolygonGroups().GetElementIDs())
+		{
+			int32 MaterialSlotIndex = StaticMesh->GetMaterialIndexFromImportedMaterialSlotName(SlotNames[PolygonGroupID]);
+					
+			if (MaterialSlotIndex == INDEX_NONE)
+			{
+				// If no material was found with this slot name, it is probably because the pipeline is configured to not import materials.
+				// Fill out a blank slot instead.
+				MaterialSlotIndex = StaticMesh->GetStaticMaterials().Emplace(nullptr, SlotNames[PolygonGroupID]);
+			}
+
+			FMeshSectionInfo Info = StaticMesh->GetSectionInfoMap().Get(CurrentLodIndex, SectionIndex);
+			Info.MaterialIndex = MaterialSlotIndex;
+			StaticMesh->GetSectionInfoMap().Remove(CurrentLodIndex, SectionIndex);
+			StaticMesh->GetSectionInfoMap().Set(CurrentLodIndex, SectionIndex, Info);
+
+			SectionIndex++;
+		}
+		if (CurrentLodIndex >= FinalLodCount)
+		{
+			// Set static mesh build settings
+			FStaticMeshSourceModel& SrcModel = StaticMesh->GetSourceModel(CurrentLodIndex);
+			SrcModel.ReductionSettings.MaxDeviation = 0.0f;
+			SrcModel.ReductionSettings.PercentTriangles = 1.0f;
+			SrcModel.ReductionSettings.PercentVertices = 1.0f;
+		}
+
+		CurrentLodIndex++;
+	}
+
+	if (!Arguments.ReimportObject)
+	{
+		// Apply all StaticMeshFactoryNode custom attributes to the static mesh asset
+		StaticMeshFactoryNode->ApplyAllCustomAttributeToObject(StaticMesh);
 	}
 	else
 	{
-		// The static mesh is not a UStaticMesh
-		StaticMeshObject->RemoveFromRoot();
-		StaticMeshObject->MarkAsGarbage();
+		//Apply the re import strategy 
+		UInterchangeAssetImportData* InterchangeAssetImportData = Cast<UInterchangeAssetImportData>(StaticMesh->GetAssetImportData());
+		UInterchangeBaseNode* PreviousNode = nullptr;
+		if (InterchangeAssetImportData)
+		{
+			PreviousNode = InterchangeAssetImportData->NodeContainer->GetNode(InterchangeAssetImportData->NodeUniqueID);
+		}
+		UInterchangeBaseNode* CurrentNode = NewObject<UInterchangeBaseNode>(GetTransientPackage(), UInterchangeStaticMeshFactoryNode::StaticClass());
+		UInterchangeBaseNode::CopyStorage(StaticMeshFactoryNode, CurrentNode);
+		CurrentNode->FillAllCustomAttributeFromObject(StaticMesh);
+		UE::Interchange::FFactoryCommon::ApplyReimportStrategyToAsset(StaticMesh, PreviousNode, CurrentNode, StaticMeshFactoryNode);
 	}
+		
+	// Getting the file Hash will cache it into the source data
+	Arguments.SourceData->GetFileContentHash();
+
 	return StaticMeshObject;
 
 #endif //else !WITH_EDITOR || !WITH_EDITORONLY_DATA
diff --git a/Engine/Plugins/Experimental/Interchange/Runtime/Source/Import/Private/Texture/InterchangeTextureFactory.cpp b/Engine/Plugins/Experimental/Interchange/Runtime/Source/Import/Private/Texture/InterchangeTextureFactory.cpp
--- a/Engine/Plugins/Experimental/Interchange/Runtime/Source/Import/Private/Texture/InterchangeTextureFactory.cpp
+++ b/Engine/Plugins/Experimental/Interchange/Runtime/Source/Import/Private/Texture/InterchangeTextureFactory.cpp
@@ -1020,7 +1020,7 @@ void UInterchangeTextureFactory::PreImportPreCompletedCallback(const FImportPreC
 			UInterchangeBaseNode::CopyStorage(TextureFactoryNode, CurrentNode);
 			CurrentNode->FillAllCustomAttributeFromObject(Texture);
 			//Apply reimport strategy
-			UE::Interchange::FFactoryCommon::ApplyReimportStrategyToAsset(Arguments.ReimportStrategyFlags, Texture, PreviousNode, CurrentNode, TextureFactoryNode);
+			UE::Interchange::FFactoryCommon::ApplyReimportStrategyToAsset(Texture, PreviousNode, CurrentNode, TextureFactoryNode);
 		}
 	}
 #endif //WITH_EDITOR
diff --git a/Engine/Plugins/Experimental/Interchange/Runtime/Source/Pipelines/Private/InterchangeGenericAssetsPipeline.cpp b/Engine/Plugins/Experimental/Interchange/Runtime/Source/Pipelines/Private/InterchangeGenericAssetsPipeline.cpp
--- a/Engine/Plugins/Experimental/Interchange/Runtime/Source/Pipelines/Private/InterchangeGenericAssetsPipeline.cpp
+++ b/Engine/Plugins/Experimental/Interchange/Runtime/Source/Pipelines/Private/InterchangeGenericAssetsPipeline.cpp
@@ -214,6 +214,14 @@ void UInterchangeGenericAssetsPipeline::ExecutePreImportPipeline(UInterchangeBas
 	}
 
 	ImplementUseSourceNameForAssetOption();
+	//Make sure all factory nodes have the specified strategy
+	BaseNodeContainer->IterateNodes([ReimportStrategyClosure = ReimportStrategy](const FString& NodeUid, UInterchangeBaseNode* Node)
+		{
+			if (Node->GetNodeContainerType() == EInterchangeNodeContainerType::FactoryData)
+			{
+				Node->SetReimportStrategyFlags(ReimportStrategyClosure);
+			}
+		});
 }
 
 void UInterchangeGenericAssetsPipeline::ExecutePostImportPipeline(const UInterchangeBaseNodeContainer* InBaseNodeContainer, const FString& NodeKey, UObject* CreatedAsset, bool bIsAReimport)
diff --git a/Engine/Plugins/Experimental/Interchange/Runtime/Source/Pipelines/Public/InterchangeGenericAssetsPipeline.h b/Engine/Plugins/Experimental/Interchange/Runtime/Source/Pipelines/Public/InterchangeGenericAssetsPipeline.h
--- a/Engine/Plugins/Experimental/Interchange/Runtime/Source/Pipelines/Public/InterchangeGenericAssetsPipeline.h
+++ b/Engine/Plugins/Experimental/Interchange/Runtime/Source/Pipelines/Public/InterchangeGenericAssetsPipeline.h
@@ -6,6 +6,7 @@
 #include "InterchangeGenericMaterialPipeline.h"
 #include "InterchangePipelineBase.h"
 #include "InterchangeSourceData.h"
+#include "Nodes/InterchangeBaseNode.h"
 #include "Nodes/InterchangeBaseNodeContainer.h"
 #include "UObject/Object.h"
 #include "UObject/ObjectMacros.h"
@@ -88,6 +89,10 @@ public:
 	UPROPERTY(EditAnywhere, BlueprintReadWrite, Category = COMMON_CATEGORY)
 	bool bUseSourceNameForAsset = true;
 
+	/* Allow user to choose the re-import strategy. */
+	UPROPERTY(EditAnywhere, BlueprintReadWrite, Category = COMMON_CATEGORY)
+	EReimportStrategyFlags ReimportStrategy = EReimportStrategyFlags::ApplyNoProperties;
+
 
 	//////	COMMON_MESHES_CATEGORY Properties //////
 
diff --git a/Engine/Source/Editor/UnrealEd/Classes/Factories/ReimportInterchangeFactory.h b/Engine/Source/Editor/UnrealEd/Classes/Factories/ReimportInterchangeFactory.h
new file mode 100644
--- /dev/null
+++ b/Engine/Source/Editor/UnrealEd/Classes/Factories/ReimportInterchangeFactory.h
@@ -0,0 +1,46 @@
+// Copyright Epic Games, Inc. All Rights Reserved.
+
+//~=============================================================================
+// ReimportFbxSkeletalMeshFactory
+//~=============================================================================
+
+#pragma once
+
+#include "CoreMinimal.h"
+#include "UObject/ObjectMacros.h"
+#include "EditorReimportHandler.h"
+#include "ReimportInterchangeFactory.generated.h"
+
+/**
+ * This FReimportHandler class is temporary until we remove the UFactory import/reimport code.
+ * We use this ReimportHandler factory to make sure the set reimport is not touching the UInterchangeAssetImportData when reimporting an asset
+ */
+UCLASS(MinimalAPI, collapsecategories)
+class UReimportInterchangeFactory : public UObject, public FReimportHandler
+{
+	GENERATED_UCLASS_BODY()
+
+	//~ Begin FReimportHandler Interface
+	virtual bool CanReimport( UObject* Obj, TArray<FString>& OutFilenames ) override;
+	virtual void SetReimportPaths(UObject* Obj, const TArray<FString>& NewReimportPaths) override
+	{
+		for (int32 SourceIndex = 0; SourceIndex < NewReimportPaths.Num(); ++SourceIndex)
+		{
+			SetReimportPaths(Obj, NewReimportPaths[SourceIndex], SourceIndex);
+		}
+	}
+	virtual void SetReimportPaths(UObject* Obj, const FString& NewReimportPath, const int32 SourceIndex) override;
+	virtual EReimportResult::Type Reimport(UObject* Obj) override
+	{
+		return Reimport(Obj, INDEX_NONE);
+	}
+	virtual EReimportResult::Type Reimport(UObject* Obj, int32 SourceFileIndex);
+	virtual int32 GetPriority() const override;
+
+	virtual bool IsInterchangeFactory() const override
+	{
+		return true;
+	}
+	//~ End FReimportHandler Interface
+
+};
diff --git a/Engine/Source/Editor/UnrealEd/Private/Editor.cpp b/Engine/Source/Editor/UnrealEd/Private/Editor.cpp
--- a/Engine/Source/Editor/UnrealEd/Private/Editor.cpp
+++ b/Engine/Source/Editor/UnrealEd/Private/Editor.cpp
@@ -390,7 +390,7 @@ bool FReimportManager::Reimport( UObject* Obj, bool bAskForNewFileIfMissing, boo
 
 			if ( bValidSourceFilename )
 			{
-				if (bUseInterchangeFramework)
+				if (bUseInterchangeFramework && CanReimportHandler->IsInterchangeFactory())
 				{
 					UE::Interchange::FScopedSourceData ScopedSourceData(SourceFilenames[0]);
 
diff --git a/Engine/Source/Editor/UnrealEd/Private/Factories/ReimportInterchangeFactory.cpp b/Engine/Source/Editor/UnrealEd/Private/Factories/ReimportInterchangeFactory.cpp
new file mode 100644
--- /dev/null
+++ b/Engine/Source/Editor/UnrealEd/Private/Factories/ReimportInterchangeFactory.cpp
@@ -0,0 +1,140 @@
+// Copyright Epic Games, Inc. All Rights Reserved.
+
+#include "Factories/ReimportInterchangeFactory.h"
+
+#include "CoreMinimal.h"
+#include "EngineDefines.h"
+#include "InterchangeAssetImportData.h"
+#include "InterchangeManager.h"
+#include "Settings/EditorExperimentalSettings.h"
+#include "UObject/UObjectGlobals.h"
+#include "UObject/ObjectMacros.h"
+
+#define LOCTEXT_NAMESPACE "ReimportInterchangeFactory"
+
+UReimportInterchangeFactory::UReimportInterchangeFactory(const FObjectInitializer& ObjectInitializer)
+	: Super(ObjectInitializer)
+{
+}
+
+//~ Begin FReimportHandler Interface
+bool UReimportInterchangeFactory::CanReimport(UObject* Obj, TArray<FString>& OutFilenames)
+{
+	auto CanReimportAsset = [&OutFilenames](UAssetImportData* AssetImportData)
+	{
+		if (!AssetImportData)
+		{
+			return false;
+		}
+		UInterchangeAssetImportData* InterchangeAssetImportData = Cast<UInterchangeAssetImportData>(AssetImportData);
+		if (InterchangeAssetImportData == nullptr)
+		{
+			return false;
+		}
+		AssetImportData->ExtractFilenames(OutFilenames);
+		return true;
+	};
+
+	
+	const UEditorExperimentalSettings* EditorExperimentalSettings = GetDefault<UEditorExperimentalSettings>();
+	const bool bUseInterchangeFramework = EditorExperimentalSettings->bEnableInterchangeFramework;
+	const bool bUseInterchangeFrameworkForTextureOnly = (!bUseInterchangeFramework) && EditorExperimentalSettings->bEnableInterchangeFrameworkForTextureOnly;
+
+	USkeletalMesh* SkeletalMesh = Cast<USkeletalMesh>(Obj);
+	UStaticMesh* StaticMesh = Cast<UStaticMesh>(Obj);
+	//UAnimSequence* AnimSequence = Cast<UAnimSequence>(Obj);
+	UTexture* Texture = Cast<UTexture>(Obj);
+	if (bUseInterchangeFramework && SkeletalMesh)
+	{
+		if (UAssetImportData* AssetImportData = SkeletalMesh->GetAssetImportData())
+		{
+			return CanReimportAsset(AssetImportData);
+		}
+	}
+	else if (bUseInterchangeFramework && StaticMesh)
+	{
+		if (UAssetImportData* AssetImportData = StaticMesh->GetAssetImportData())
+		{
+			return CanReimportAsset(AssetImportData);
+		}
+	}
+// 	else if (bUseInterchangeFramework && AnimSequence)
+// 	{
+// 		if (UAssetImportData* AssetImportData = AnimSequence->AssetImportData)
+// 		{
+// 			return CanReimportAsset(AssetImportData);
+// 		}
+// 	}
+	else if ((bUseInterchangeFramework || bUseInterchangeFrameworkForTextureOnly) && Texture)
+	{
+		if (UAssetImportData* AssetImportData = Texture->AssetImportData)
+		{
+			return CanReimportAsset(AssetImportData);
+		}
+	}
+	return false;
+}
+
+void UReimportInterchangeFactory::SetReimportPaths(UObject* Obj, const FString& NewReimportPath, const int32 SourceFileIndex)
+{
+	auto SetAssetImportDataPath = [SourceFileIndex, &NewReimportPath](UAssetImportData* AssetImportData)
+	{
+		int32 RealSourceFileIndex = SourceFileIndex == INDEX_NONE ? 0 : SourceFileIndex;
+		if (RealSourceFileIndex < AssetImportData->GetSourceFileCount())
+		{
+			AssetImportData->UpdateFilenameOnly(NewReimportPath, SourceFileIndex);
+		}
+		else
+		{
+			//Create a source file entry, this case happen when user import a specific content for the first time
+			FString SourceIndexLabel = USkeletalMesh::GetSourceFileLabelFromIndex(RealSourceFileIndex).ToString();
+			AssetImportData->AddFileName(NewReimportPath, RealSourceFileIndex, SourceIndexLabel);
+		}
+	};
+
+	USkeletalMesh* SkeletalMesh = Cast<USkeletalMesh>(Obj);
+	UStaticMesh* StaticMesh = Cast<UStaticMesh>(Obj);
+//	UAnimSequence* AnimSequence = Cast<UAnimSequence>(Obj);
+	UTexture* Texture = Cast<UTexture>(Obj);
+	if (SkeletalMesh)
+	{
+		if (UAssetImportData* AssetImportData = SkeletalMesh->GetAssetImportData())
+		{
+			SetAssetImportDataPath(AssetImportData);
+		}
+	}
+	else if (StaticMesh)
+	{
+		if (UAssetImportData* AssetImportData = StaticMesh->GetAssetImportData())
+		{
+			SetAssetImportDataPath(AssetImportData);
+		}
+	}
+// 	else if (AnimSequence)
+// 	{
+// 		if (UAssetImportData* AssetImportData = AnimSequence->AssetImportData)
+// 		{
+// 			SetAssetImportDataPath(AssetImportData);
+// 		}
+// 	}
+	else if (Texture)
+	{
+		if (UAssetImportData* AssetImportData = Texture->AssetImportData)
+		{
+			SetAssetImportDataPath(AssetImportData);
+		}
+	}
+}
+
+EReimportResult::Type UReimportInterchangeFactory::Reimport(UObject* Obj, int32 SourceFileIndex)
+{
+	return EReimportResult::Failed;
+}
+
+int32 UReimportInterchangeFactory::GetPriority() const
+{
+	//We want a high priority to surpass other legacy re-import handlers
+	return UFactory::GetDefaultImportPriority() + 10;
+}
+
+#undef LOCTEXT_NAMESPACE
\ No newline at end of file
diff --git a/Engine/Source/Editor/UnrealEd/Public/EditorReimportHandler.h b/Engine/Source/Editor/UnrealEd/Public/EditorReimportHandler.h
--- a/Engine/Source/Editor/UnrealEd/Public/EditorReimportHandler.h
+++ b/Engine/Source/Editor/UnrealEd/Public/EditorReimportHandler.h
@@ -272,6 +272,11 @@ public:
 		return bAutomatedReimport;
 	}
 
+	virtual bool IsInterchangeFactory() const
+	{
+		return false;
+	}
+
 protected:
 	FString PreferredReimportPath;
 	bool bAutomatedReimport;
diff --git a/Engine/Source/Runtime/Experimental/Interchange/Core/Private/Nodes/InterchangeBaseNode.cpp b/Engine/Source/Runtime/Experimental/Interchange/Core/Private/Nodes/InterchangeBaseNode.cpp
--- a/Engine/Source/Runtime/Experimental/Interchange/Core/Private/Nodes/InterchangeBaseNode.cpp
+++ b/Engine/Source/Runtime/Experimental/Interchange/Core/Private/Nodes/InterchangeBaseNode.cpp
@@ -66,6 +66,12 @@ namespace UE
 			return AttributeKey;
 		}
 
+		const FAttributeKey& FBaseNodeStaticData::ReimportStrategyFlagsKey()
+		{
+			static FAttributeKey AttributeKey(TEXT("__Reimport_Strategy_Flags_Key__"));
+			return AttributeKey;
+		}
+
 	} //ns Interchange
 } //ns UE
 
@@ -76,6 +82,7 @@ UInterchangeBaseNode::UInterchangeBaseNode()
 	TargetNodes.Initialize(Attributes, UE::Interchange::FBaseNodeStaticData::TargetAssetIDsKey());
 	RegisterAttribute<bool>(UE::Interchange::FBaseNodeStaticData::IsEnabledKey(), true);
 	RegisterAttribute<uint8>(UE::Interchange::FBaseNodeStaticData::NodeContainerTypeKey(), static_cast<uint8>(EInterchangeNodeContainerType::None));
+	RegisterAttribute<uint8>(UE::Interchange::FBaseNodeStaticData::ReimportStrategyFlagsKey(), static_cast<uint8>(EReimportStrategyFlags::ApplyNoProperties));
 }
 
 void UInterchangeBaseNode::InitializeNode(const FString& UniqueID, const FString& DisplayLabel, const EInterchangeNodeContainerType NodeContainerType)
@@ -171,6 +178,10 @@ FString UInterchangeBaseNode::GetKeyDisplayName(const UE::Interchange::FAttribut
 	{
 		KeyDisplayName = TEXT("Node Container Type");
 	}
+	else if (NodeAttributeKey == UE::Interchange::FBaseNodeStaticData::ReimportStrategyFlagsKey())
+	{
+		KeyDisplayName = TEXT("Re-Import Strategy");
+	}
 	return KeyDisplayName;
 }
 
@@ -288,6 +299,25 @@ bool UInterchangeBaseNode::SetDisplayLabel(const FString& DisplayLabel)
 	return false;
 }
 
+EReimportStrategyFlags UInterchangeBaseNode::GetReimportStrategyFlags() const
+{
+	checkSlow(Attributes->ContainAttribute(UE::Interchange::FBaseNodeStaticData::ReimportStrategyFlagsKey()));
+	uint8 ReimportStrategyFlags = static_cast<uint8>(EReimportStrategyFlags::ApplyNoProperties);
+	Attributes->GetAttributeHandle<uint8>(UE::Interchange::FBaseNodeStaticData::ReimportStrategyFlagsKey()).Get(ReimportStrategyFlags);
+	return static_cast<EReimportStrategyFlags>(ReimportStrategyFlags);
+}
+
+bool UInterchangeBaseNode::SetReimportStrategyFlags(const EReimportStrategyFlags& ReimportStrategyFlags)
+{
+	UE::Interchange::EAttributeStorageResult Result = Attributes->RegisterAttribute(UE::Interchange::FBaseNodeStaticData::ReimportStrategyFlagsKey(), static_cast<uint8>(ReimportStrategyFlags));
+	if (IsAttributeStorageResultSuccess(Result))
+	{
+		UE::Interchange::FAttributeStorage::TAttributeHandle<FString> Handle = Attributes->GetAttributeHandle<FString>(UE::Interchange::FBaseNodeStaticData::ReimportStrategyFlagsKey());
+		return Handle.IsValid();
+	}
+	return false;
+}
+
 FString UInterchangeBaseNode::GetParentUid() const
 {
 	if (!Attributes->ContainAttribute(UE::Interchange::FBaseNodeStaticData::ParentIDKey()))
diff --git a/Engine/Source/Runtime/Experimental/Interchange/Core/Public/InterchangeFactoryBase.h b/Engine/Source/Runtime/Experimental/Interchange/Core/Public/InterchangeFactoryBase.h
--- a/Engine/Source/Runtime/Experimental/Interchange/Core/Public/InterchangeFactoryBase.h
+++ b/Engine/Source/Runtime/Experimental/Interchange/Core/Public/InterchangeFactoryBase.h
@@ -17,15 +17,6 @@ class UInterchangePipelineBase;
 class UInterchangeSourceData;
 class UInterchangeTranslatorBase;
 
-
-UENUM()
-enum class EReimportStrategyFlags : uint8
-{
-	ApplyNoProperties, //Do not apply any property when re-importing, simply change the source data
-	ApplyPipelineProperties, //Always apply all pipeline specified properties
-	ApplyEditorChangedProperties //Always apply all pipeline properties, but leave the properties modified in editor since the last import
-};
-
 UCLASS(BlueprintType, Blueprintable, Abstract, Experimental)
 class INTERCHANGECORE_API UInterchangeFactoryBase : public UObject
 {
@@ -69,8 +60,6 @@ public:
 		 * the asset we want to re-import. The re-import should just change the source data and not any asset settings.
 		 */
 		UObject* ReimportObject = nullptr;
-
-		EReimportStrategyFlags ReimportStrategyFlags;
 	};
 
 	/**
@@ -155,9 +144,7 @@ public:
 		FString NodeUniqueID;
 		UInterchangeBaseNodeContainer* NodeContainer = nullptr;
 		TArray<UInterchangePipelineBase*> Pipelines;
-
-
-		EReimportStrategyFlags ReimportStrategyFlags;
+		
 		bool bIsReimport  = false;
 ;
 	};
diff --git a/Engine/Source/Runtime/Experimental/Interchange/Core/Public/Nodes/InterchangeBaseNode.h b/Engine/Source/Runtime/Experimental/Interchange/Core/Public/Nodes/InterchangeBaseNode.h
--- a/Engine/Source/Runtime/Experimental/Interchange/Core/Public/Nodes/InterchangeBaseNode.h
+++ b/Engine/Source/Runtime/Experimental/Interchange/Core/Public/Nodes/InterchangeBaseNode.h
@@ -184,6 +184,7 @@ namespace UE
 			static const FAttributeKey& ClassTypeAttributeKey();
 			static const FAttributeKey& AssetNameKey();
 			static const FAttributeKey& NodeContainerTypeKey();
+			static const FAttributeKey& ReimportStrategyFlagsKey();
 		};
 
 	} //ns Interchange
@@ -198,6 +199,14 @@ enum class EInterchangeNodeContainerType : uint8
 	FactoryData
 };
 
+UENUM()
+enum class EReimportStrategyFlags : uint8
+{
+	ApplyNoProperties, //Do not apply any property when re-importing, simply change the source data
+	ApplyPipelineProperties, //Always apply all pipeline specified properties
+	ApplyEditorChangedProperties //Always apply all pipeline properties, but leave the properties modified in editor since the last import
+};
+
 /**
  * This struct is used to store and retrieve key value attributes. The attributes are store in a generic FAttributeStorage which serialize the value in a TArray64<uint8>
  * See UE::Interchange::EAttributeTypes to know the supported template types
@@ -433,6 +442,20 @@ public:
 	UFUNCTION(BlueprintCallable, Category = "Interchange | Node")
 	bool SetDisplayLabel(const FString& DisplayName);
 
+	/**
+	 * Return the reimport strategy flags.
+	 *
+	 */
+	UFUNCTION(BlueprintCallable, Category = "Interchange | Node")
+	EReimportStrategyFlags GetReimportStrategyFlags() const;
+
+	/**
+	 * Change the reimport strategy flags.
+	 *
+	 */
+	UFUNCTION(BlueprintCallable, Category = "Interchange | Node")
+	bool SetReimportStrategyFlags(const EReimportStrategyFlags& ReimportStrategyFlags);
+
 	/**
 	 * Return the parent unique id. In case the attribute does not exist it will return InvalidNodeUid()
 	 */
diff --git a/Engine/Source/Runtime/Experimental/Interchange/Engine/Private/Tasks/InterchangeTaskCreateAsset.cpp b/Engine/Source/Runtime/Experimental/Interchange/Engine/Private/Tasks/InterchangeTaskCreateAsset.cpp
--- a/Engine/Source/Runtime/Experimental/Interchange/Engine/Private/Tasks/InterchangeTaskCreateAsset.cpp
+++ b/Engine/Source/Runtime/Experimental/Interchange/Engine/Private/Tasks/InterchangeTaskCreateAsset.cpp
@@ -322,10 +322,6 @@ void UE::Interchange::FTaskCreateAsset::DoTask(ENamedThreads::Type CurrentThread
 		}
 		CreateAssetParams.ReimportObject = AsyncHelper->TaskData.ReimportObject;
 
-		CreateAssetParams.ReimportStrategyFlags = EReimportStrategyFlags::ApplyNoProperties;
-		//CreateAssetParams.ReimportStrategyFlags = EReimportStrategyFlags::ApplyPipelineProperties;
-		//CreateAssetParams.ReimportStrategyFlags = EReimportStrategyFlags::ApplyEditorChangedProperties;
-
 		NodeAsset = Factory->CreateAsset(CreateAssetParams);
 	}
 	if (NodeAsset)
