Interchange
- Made sure we don't call MarkPackageDirty from non game threads as its not threadsafe.
- Update the texture sampler sampling types during the factory PreImportPreCompleted callback so that we have the correct sampling type for textures that are being upgraded to virtual.
- Fixed some issues with material creation where we had some early returns in functions that shouldn't.

diff --git a/Engine/Plugins/Experimental/Interchange/Runtime/Source/Import/Private/Material/InterchangeMaterialFactory.cpp b/Engine/Plugins/Experimental/Interchange/Runtime/Source/Import/Private/Material/InterchangeMaterialFactory.cpp
--- a/Engine/Plugins/Experimental/Interchange/Runtime/Source/Import/Private/Material/InterchangeMaterialFactory.cpp
+++ b/Engine/Plugins/Experimental/Interchange/Runtime/Source/Import/Private/Material/InterchangeMaterialFactory.cpp
@@ -145,6 +145,17 @@ namespace UE
 
 					TextureExpression->AutoSetSampleType();
 				}
+
+				UMaterialExpression* CreateMaterialExpression(UMaterial* Material, const TSubclassOf<UMaterialExpression>& ExpressionClass)
+				{
+					UMaterialFunction* const MaterialFunction = nullptr;
+					UObject* const SelectedAsset = nullptr;
+					const int32 NodePosX = 0;
+					const int32 NodePosY = 0;
+					const bool bAllowMarkingPackageDirty = false;
+
+					return UMaterialEditingLibrary::CreateMaterialExpressionEx(Material, MaterialFunction, ExpressionClass, SelectedAsset, NodePosX, NodePosY, bAllowMarkingPackageDirty);
+				}
 #endif // #if WITH_EDITOR
 			}
 		}
@@ -289,16 +300,28 @@ void UInterchangeMaterialFactory::PreImportPreCompletedCallback(const FImportPre
 	if (ensure(Arguments.ImportedObject && Arguments.SourceData))
 	{
 		//We must call the Update of the asset source file in the main thread because UAssetImportData::Update execute some delegate we do not control
-		UMaterialInterface* ImportedMaterial = CastChecked<UMaterialInterface>(Arguments.ImportedObject);
+		UMaterialInterface* ImportedMaterialInterface = CastChecked<UMaterialInterface>(Arguments.ImportedObject);
+
+		//Update the samplers type in case the textures were changed during their PreImportPreCompletedCallback
+		if (UMaterial* ImportedMaterial = CastChecked<UMaterial>(ImportedMaterialInterface))
+		{
+			for (UMaterialExpression* Expression : ImportedMaterial->Expressions)
+			{
+				if (UMaterialExpressionTextureBase* TextureSample = Cast<UMaterialExpressionTextureBase>(Expression))
+				{
+					TextureSample->AutoSetSampleType();
+				}
+			}
+		}
 
-		UE::Interchange::FFactoryCommon::FUpdateImportAssetDataParameters UpdateImportAssetDataParameters(ImportedMaterial
-																										  , ImportedMaterial->AssetImportData
+		UE::Interchange::FFactoryCommon::FUpdateImportAssetDataParameters UpdateImportAssetDataParameters(ImportedMaterialInterface
+																										  , ImportedMaterialInterface->AssetImportData
 																										  , Arguments.SourceData
 																										  , Arguments.NodeUniqueID
 																										  , Arguments.NodeContainer
 																										  , Arguments.Pipelines);
 
-		ImportedMaterial->AssetImportData = UE::Interchange::FFactoryCommon::UpdateImportAssetData(UpdateImportAssetDataParameters);
+		ImportedMaterialInterface->AssetImportData = UE::Interchange::FFactoryCommon::UpdateImportAssetData(UpdateImportAssetDataParameters);
 	}
 #endif
 }
@@ -474,7 +497,7 @@ UMaterialExpression* UInterchangeMaterialFactory::CreateExpression(UMaterial* Ma
 		return nullptr;
 	}
 
-	UMaterialExpression* MaterialExpression = UMaterialEditingLibrary::CreateMaterialExpression(Material, ExpressionClass);
+	UMaterialExpression* MaterialExpression = CreateMaterialExpression(Material, ExpressionClass);
 
 	if (!MaterialExpression)
 	{
@@ -541,7 +564,7 @@ UMaterialExpression* UInterchangeMaterialFactory::CreateExpressionsForNode(UMate
 				}
 
 				const int32 InputIndex = GetInputIndex(*MaterialExpression, InputName);
-				const int32 OutputIndex = GetOutputIndex(*MaterialExpression, OutputName);
+				const int32 OutputIndex = GetOutputIndex(*ConnectedExpression, OutputName);
 
 				if (InputIndex != INDEX_NONE)
 				{
diff --git a/Engine/Plugins/Experimental/Interchange/Runtime/Source/Import/Private/Mesh/InterchangeStaticMeshFactory.cpp b/Engine/Plugins/Experimental/Interchange/Runtime/Source/Import/Private/Mesh/InterchangeStaticMeshFactory.cpp
--- a/Engine/Plugins/Experimental/Interchange/Runtime/Source/Import/Private/Mesh/InterchangeStaticMeshFactory.cpp
+++ b/Engine/Plugins/Experimental/Interchange/Runtime/Source/Import/Private/Mesh/InterchangeStaticMeshFactory.cpp
@@ -307,7 +307,9 @@ UObject* UInterchangeStaticMeshFactory::CreateAsset(const FCreateAssetParams& Ar
 			}
 		}
 
-		StaticMesh->CommitMeshDescription(CurrentLodIndex);
+		UStaticMesh::FCommitMeshDescriptionParams CommitMeshDescriptionParams;
+		CommitMeshDescriptionParams.bMarkPackageDirty = false; // Marking packages dirty isn't threadsafe
+		StaticMesh->CommitMeshDescription(CurrentLodIndex, CommitMeshDescriptionParams);
 
 		// Handle materials
 		// @TODO: move this outside of the LOD loop?
diff --git a/Engine/Plugins/Experimental/Interchange/Runtime/Source/Pipelines/Private/InterchangeGenericMaterialPipeline.cpp b/Engine/Plugins/Experimental/Interchange/Runtime/Source/Pipelines/Private/InterchangeGenericMaterialPipeline.cpp
--- a/Engine/Plugins/Experimental/Interchange/Runtime/Source/Pipelines/Private/InterchangeGenericMaterialPipeline.cpp
+++ b/Engine/Plugins/Experimental/Interchange/Runtime/Source/Pipelines/Private/InterchangeGenericMaterialPipeline.cpp
@@ -217,6 +217,8 @@ bool UInterchangeGenericMaterialPipeline::HandlePBRModel(const UInterchangeShade
 {
 	using namespace UE::Interchange::Materials::PBR;
 
+	bool bShadingModelHandled = false;
+
 	// BaseColor
 	{
 		const bool bHasInput = UInterchangeShaderPortsAPI::HasInput(ShaderGraphNode, Parameters::BaseColor);
@@ -231,7 +233,7 @@ bool UInterchangeGenericMaterialPipeline::HandlePBRModel(const UInterchangeShade
 				MaterialFactoryNode->ConnectOutputToBaseColor(ExpressionFactoryNode.Get<0>()->GetUniqueID(), ExpressionFactoryNode.Get<1>());
 			}
 
-			return true;
+			bShadingModelHandled = true;
 		}
 	}
 
@@ -249,7 +251,7 @@ bool UInterchangeGenericMaterialPipeline::HandlePBRModel(const UInterchangeShade
 				MaterialFactoryNode->ConnectOutputToMetallic(ExpressionFactoryNode.Get<0>()->GetUniqueID(), ExpressionFactoryNode.Get<1>());
 			}
 
-			return true;
+			bShadingModelHandled = true;
 		}
 	}
 
@@ -267,7 +269,7 @@ bool UInterchangeGenericMaterialPipeline::HandlePBRModel(const UInterchangeShade
 				MaterialFactoryNode->ConnectOutputToSpecular(ExpressionFactoryNode.Get<0>()->GetUniqueID(), ExpressionFactoryNode.Get<1>());
 			}
 
-			return true;
+			bShadingModelHandled = true;
 		}
 	}
 
@@ -285,14 +287,14 @@ bool UInterchangeGenericMaterialPipeline::HandlePBRModel(const UInterchangeShade
 				MaterialFactoryNode->ConnectOutputToRoughness(ExpressionFactoryNode.Get<0>()->GetUniqueID(), ExpressionFactoryNode.Get<1>());
 			}
 
-			return true;
+			bShadingModelHandled = true;
 		}
 	}
 
-	return false;
+	return bShadingModelHandled;
 }
 
-bool UInterchangeGenericMaterialPipeline::HandleCommonParameters(const UInterchangeShaderGraphNode* ShaderGraphNode, UInterchangeMaterialFactoryNode* MaterialFactoryNode)
+void UInterchangeGenericMaterialPipeline::HandleCommonParameters(const UInterchangeShaderGraphNode* ShaderGraphNode, UInterchangeMaterialFactoryNode* MaterialFactoryNode)
 {
 	using namespace UE::Interchange::Materials::Common;
 
@@ -309,8 +311,6 @@ bool UInterchangeGenericMaterialPipeline::HandleCommonParameters(const UIntercha
 			{
 				MaterialFactoryNode->ConnectOutputToEmissiveColor(EmissiveExpressionFactoryNode.Get<0>()->GetUniqueID(), EmissiveExpressionFactoryNode.Get<1>());
 			}
-
-			return true;
 		}
 	}
 
@@ -327,8 +327,6 @@ bool UInterchangeGenericMaterialPipeline::HandleCommonParameters(const UIntercha
 			{
 				MaterialFactoryNode->ConnectOutputToNormal(ExpressionFactoryNode.Get<0>()->GetUniqueID(), ExpressionFactoryNode.Get<1>());
 			}
-
-			return true;
 		}
 	}
 
@@ -345,12 +343,8 @@ bool UInterchangeGenericMaterialPipeline::HandleCommonParameters(const UIntercha
 			{
 				MaterialFactoryNode->ConnectOutputToOpacity(OpacityExpressionFactoryNode.Get<0>()->GetUniqueID(), OpacityExpressionFactoryNode.Get<1>());
 			}
-
-			return true;
 		}
 	}
-
-	return false;
 }
 
 void UInterchangeGenericMaterialPipeline::HandleTextureSampleNode(const UInterchangeShaderNode* ShaderNode, UInterchangeMaterialExpressionFactoryNode* TextureSampleFactoryNode)
diff --git a/Engine/Plugins/Experimental/Interchange/Runtime/Source/Pipelines/Public/InterchangeGenericMaterialPipeline.h b/Engine/Plugins/Experimental/Interchange/Runtime/Source/Pipelines/Public/InterchangeGenericMaterialPipeline.h
--- a/Engine/Plugins/Experimental/Interchange/Runtime/Source/Pipelines/Public/InterchangeGenericMaterialPipeline.h
+++ b/Engine/Plugins/Experimental/Interchange/Runtime/Source/Pipelines/Public/InterchangeGenericMaterialPipeline.h
@@ -58,7 +58,7 @@ private:
 	bool HandlePhongModel(const UInterchangeShaderGraphNode* ShaderGraphNode, UInterchangeMaterialFactoryNode* MaterialFactoryNode);
 	bool HandleLambertModel(const UInterchangeShaderGraphNode* ShaderGraphNode, UInterchangeMaterialFactoryNode* MaterialFactoryNode);
 	bool HandlePBRModel(const UInterchangeShaderGraphNode* ShaderGraphNode, UInterchangeMaterialFactoryNode* MaterialFactoryNode);
-	bool HandleCommonParameters(const UInterchangeShaderGraphNode* ShaderGraphNode, UInterchangeMaterialFactoryNode* MaterialFactoryNode);
+	void HandleCommonParameters(const UInterchangeShaderGraphNode* ShaderGraphNode, UInterchangeMaterialFactoryNode* MaterialFactoryNode);
 	void HandleTextureSampleNode(const UInterchangeShaderNode* ShaderNode, UInterchangeMaterialExpressionFactoryNode* TextureSampleFactoryNode);
 	void HandleTextureCoordinates(const UInterchangeShaderNode* ShaderNode, UInterchangeMaterialExpressionFactoryNode* TextureSampleFactoryNode);
 	void HandleLerpNode(const UInterchangeShaderNode* ShaderNode, UInterchangeMaterialFactoryNode* MaterialFactoryNode, UInterchangeMaterialExpressionFactoryNode* LerpFactoryNode);
diff --git a/Engine/Source/Editor/MaterialEditor/Private/MaterialEditingLibrary.cpp b/Engine/Source/Editor/MaterialEditor/Private/MaterialEditingLibrary.cpp
--- a/Engine/Source/Editor/MaterialEditor/Private/MaterialEditingLibrary.cpp
+++ b/Engine/Source/Editor/MaterialEditor/Private/MaterialEditingLibrary.cpp
@@ -504,7 +504,8 @@ UMaterialExpression* UMaterialEditingLibrary::CreateMaterialExpressionInFunction
 }
 
 
-UMaterialExpression* UMaterialEditingLibrary::CreateMaterialExpressionEx(UMaterial* Material, UMaterialFunction* MaterialFunction, TSubclassOf<UMaterialExpression> ExpressionClass, UObject* SelectedAsset, int32 NodePosX, int32 NodePosY)
+UMaterialExpression* UMaterialEditingLibrary::CreateMaterialExpressionEx(UMaterial* Material, UMaterialFunction* MaterialFunction, TSubclassOf<UMaterialExpression> ExpressionClass,
+	UObject* SelectedAsset, int32 NodePosX, int32 NodePosY, bool bAllowMarkingPackageDirty)
 {
 	UMaterialExpression* NewExpression = nullptr;
 	if (Material || MaterialFunction)
@@ -532,7 +533,7 @@ UMaterialExpression* UMaterialEditingLibrary::CreateMaterialExpressionEx(UMateri
 		NewExpression->MaterialExpressionEditorY = NodePosY;
 
 		// Create a GUID for the node
-		NewExpression->UpdateMaterialExpressionGuid(true, true);
+		NewExpression->UpdateMaterialExpressionGuid(true, bAllowMarkingPackageDirty);
 
 		if (SelectedAsset)
 		{
@@ -574,7 +575,7 @@ UMaterialExpression* UMaterialEditingLibrary::CreateMaterialExpressionEx(UMateri
 			FunctionOutput->ValidateName();
 		}
 
-		NewExpression->UpdateParameterGuid(true, true);
+		NewExpression->UpdateParameterGuid(true, bAllowMarkingPackageDirty);
 
 		if (NewExpression->HasAParameterName())
 		{
@@ -619,7 +620,10 @@ UMaterialExpression* UMaterialEditingLibrary::CreateMaterialExpressionEx(UMateri
 			Material->AddExpressionParameter(NewExpression, Material->EditorParameters);
 		}
 
-		NewExpression->MarkPackageDirty();
+		if (bAllowMarkingPackageDirty)
+		{
+			NewExpression->MarkPackageDirty();
+		}
 	}
 	return NewExpression;
 }
diff --git a/Engine/Source/Editor/MaterialEditor/Public/MaterialEditingLibrary.h b/Engine/Source/Editor/MaterialEditor/Public/MaterialEditingLibrary.h
--- a/Engine/Source/Editor/MaterialEditor/Public/MaterialEditingLibrary.h
+++ b/Engine/Source/Editor/MaterialEditor/Public/MaterialEditingLibrary.h
@@ -63,14 +63,16 @@ public:
 	/**
 	*	Create a new material expression node within the supplied material, optionally specifying asset to use
 	*	@note	If a MaterialFunction and Material are specified, expression is added to Material and not MaterialFunction, assuming Material is a preview that will be copied to Function later by user.
-	*	@param	Material			Material asset to add an expression to
-	*	@param	MaterialFunction	Specified if adding an expression to a MaterialFunction, used as Outer for new expression object
-	*	@param	SelectedAsset		If specified, new node will attempt to use this asset, if of the appropriate type (e.g. Texture for a TextureSampler)
-	*	@param	ExpressionClass		Class of expression to add
-	*	@param	NodePosX			X position of new expression node
-	*	@param	NodePosY			Y position of new expression node
+	*	@param	Material					Material asset to add an expression to
+	*	@param	MaterialFunction			Specified if adding an expression to a MaterialFunction, used as Outer for new expression object
+	*	@param	SelectedAsset				If specified, new node will attempt to use this asset, if of the appropriate type (e.g. Texture for a TextureSampler)
+	*	@param	ExpressionClass				Class of expression to add
+	*	@param	NodePosX					X position of new expression node
+	*	@param	NodePosY					Y position of new expression node
+	*	@param	bAllowMarkingPackageDirty	Packages can't be marked dirty outside of the game thread. If this is false, package will need to be dirtied through other means. 
 	*/
-	static UMaterialExpression* CreateMaterialExpressionEx(UMaterial* Material, UMaterialFunction* MaterialFunction, TSubclassOf<UMaterialExpression> ExpressionClass, UObject* SelectedAsset = nullptr, int32 NodePosX = 0, int32 NodePosY = 0);
+	static UMaterialExpression* CreateMaterialExpressionEx(UMaterial* Material, UMaterialFunction* MaterialFunction, TSubclassOf<UMaterialExpression> ExpressionClass,
+		UObject* SelectedAsset = nullptr, int32 NodePosX = 0, int32 NodePosY = 0, bool bAllowMarkingPackageDirty = true);
 
 	/**
 	*	Rebuilds dependent Material Instance Editors
