Interchange Core
- Added support for bitfields in Apply & Fill delegates.

Interchange Materials
- Added support for two sidedness, blend mode, translucency lighting mode, shading model, ambient occlusion and IOR.
- Textures sampled for a normal input will get set to compression normalmap by the material pipeline.
- Added a multiply shader node.

diff --git a/Engine/Plugins/Experimental/Interchange/Runtime/Source/Import/Private/Material/InterchangeMaterialFactory.cpp b/Engine/Plugins/Experimental/Interchange/Runtime/Source/Import/Private/Material/InterchangeMaterialFactory.cpp
--- a/Engine/Plugins/Experimental/Interchange/Runtime/Source/Import/Private/Material/InterchangeMaterialFactory.cpp
+++ b/Engine/Plugins/Experimental/Interchange/Runtime/Source/Import/Private/Material/InterchangeMaterialFactory.cpp
@@ -112,34 +112,16 @@ namespace UE
 
 				void SetupTextureExpression(const UInterchangeFactoryBase::FCreateAssetParams& Arguments, const UInterchangeMaterialExpressionFactoryNode* ExpressionNode, UMaterialExpressionTextureBase* TextureExpression)
 				{
-					const UInterchangeTextureNode* TextureNode = nullptr;
+					using namespace UE::Interchange::Materials::Standard::Nodes::TextureSample;
 
-					FString TextureDependency;
+					FString TextureFactoryNodeUid;
+					ExpressionNode->GetStringAttribute(UInterchangeShaderPortsAPI::MakeInputValueKey(Inputs::Texture.ToString()), TextureFactoryNodeUid);
 
-					TArray<FString> TargetNodes;
-					ExpressionNode->GetTargetNodeUids(TargetNodes);
-					if (!TargetNodes.IsEmpty())
+					if (const UInterchangeTextureFactoryNode* TextureFactoryNode = Cast<UInterchangeTextureFactoryNode>(Arguments.NodeContainer->GetNode(TextureFactoryNodeUid)))
 					{
-						if (const UInterchangeShaderNode* ShaderNode = Cast<UInterchangeShaderNode>(Arguments.NodeContainer->GetNode(TargetNodes[0])))
+						if (UTexture* Texture = Cast<UTexture>(TextureFactoryNode->ReferenceObject.TryLoad()))
 						{
-							FString TextureUid;
-							ShaderNode->GetStringAttribute(UInterchangeShaderPortsAPI::MakeInputValueKey(TEXT("TextureUid")), TextureUid);
-
-							TextureNode = Cast<UInterchangeTextureNode>(Arguments.NodeContainer->GetNode(TextureUid));
-						}
-					}
-
-					if (TextureNode)
-					{
-						TArray<FString> TextureNodeTargetNodes;
-						TextureNode->GetTargetNodeUids(TextureNodeTargetNodes);
-
-						if (const UInterchangeTextureFactoryNode* TextureFactoryNode = Cast<UInterchangeTextureFactoryNode>(Arguments.NodeContainer->GetNode(TextureNodeTargetNodes[0])))
-						{
-							if (UTexture* Texture = Cast<UTexture>(TextureFactoryNode->ReferenceObject.TryLoad()))
-							{
-								TextureExpression->Texture = Texture;
-							}
+							TextureExpression->Texture = Texture;
 						}
 					}
 
@@ -339,17 +321,19 @@ void UInterchangeMaterialFactory::SetupMaterial(UMaterial* Material, const FCrea
 	{
 		FString BaseColorUid;
 		FString OutputName;
-		MaterialFactoryNode->GetBaseColorConnection(BaseColorUid, OutputName);
 
-		const UInterchangeMaterialExpressionFactoryNode* BaseColor = Cast<UInterchangeMaterialExpressionFactoryNode>(Arguments.NodeContainer->GetNode(BaseColorUid));
-
-		if (BaseColor)
+		if (MaterialFactoryNode->GetBaseColorConnection(BaseColorUid, OutputName))
 		{
-			if (UMaterialExpression* BaseColorExpression = CreateExpressionsForNode(Material, Arguments, BaseColor, Expressions))
+			const UInterchangeMaterialExpressionFactoryNode* BaseColor = Cast<UInterchangeMaterialExpressionFactoryNode>(Arguments.NodeContainer->GetNode(BaseColorUid));
+
+			if (BaseColor)
 			{
-				if (FExpressionInput* BaseColorInput = Material->GetExpressionInputForProperty(MP_BaseColor))
+				if (UMaterialExpression* BaseColorExpression = CreateExpressionsForNode(Material, Arguments, BaseColor, Expressions))
 				{
-					BaseColorExpression->ConnectExpression(BaseColorInput, GetOutputIndex(*BaseColorExpression, OutputName));
+					if (FExpressionInput* BaseColorInput = Material->GetExpressionInputForProperty(MP_BaseColor))
+					{
+						BaseColorExpression->ConnectExpression(BaseColorInput, GetOutputIndex(*BaseColorExpression, OutputName));
+					}
 				}
 			}
 		}
@@ -359,17 +343,19 @@ void UInterchangeMaterialFactory::SetupMaterial(UMaterial* Material, const FCrea
 	{
 		FString MetallicUid;
 		FString OutputName;
-		MaterialFactoryNode->GetMetallicConnection(MetallicUid, OutputName);
 
-		const UInterchangeMaterialExpressionFactoryNode* MetallicNode = Cast<UInterchangeMaterialExpressionFactoryNode>(Arguments.NodeContainer->GetNode(MetallicUid));
-
-		if (MetallicNode)
+		if (MaterialFactoryNode->GetMetallicConnection(MetallicUid, OutputName))
 		{
-			if (UMaterialExpression* MetallicExpression = CreateExpressionsForNode(Material, Arguments, MetallicNode, Expressions))
+			const UInterchangeMaterialExpressionFactoryNode* MetallicNode = Cast<UInterchangeMaterialExpressionFactoryNode>(Arguments.NodeContainer->GetNode(MetallicUid));
+
+			if (MetallicNode)
 			{
-				if (FExpressionInput* MetallicInput = Material->GetExpressionInputForProperty(MP_Metallic))
+				if (UMaterialExpression* MetallicExpression = CreateExpressionsForNode(Material, Arguments, MetallicNode, Expressions))
 				{
-					MetallicExpression->ConnectExpression(MetallicInput, GetOutputIndex(*MetallicExpression, OutputName));
+					if (FExpressionInput* MetallicInput = Material->GetExpressionInputForProperty(MP_Metallic))
+					{
+						MetallicExpression->ConnectExpression(MetallicInput, GetOutputIndex(*MetallicExpression, OutputName));
+					}
 				}
 			}
 		}
@@ -379,17 +365,19 @@ void UInterchangeMaterialFactory::SetupMaterial(UMaterial* Material, const FCrea
 	{
 		FString SpecularUid;
 		FString OutputName;
-		MaterialFactoryNode->GetSpecularConnection(SpecularUid, OutputName);
-
-		const UInterchangeMaterialExpressionFactoryNode* SpecularNode = Cast<UInterchangeMaterialExpressionFactoryNode>(Arguments.NodeContainer->GetNode(SpecularUid));
-
-		if (SpecularNode)
+		
+		if (MaterialFactoryNode->GetSpecularConnection(SpecularUid, OutputName))
 		{
-			if (UMaterialExpression* SpecularExpression = CreateExpressionsForNode(Material, Arguments, SpecularNode, Expressions))
+			const UInterchangeMaterialExpressionFactoryNode* SpecularNode = Cast<UInterchangeMaterialExpressionFactoryNode>(Arguments.NodeContainer->GetNode(SpecularUid));
+
+			if (SpecularNode)
 			{
-				if (FExpressionInput* SpecularInput = Material->GetExpressionInputForProperty(MP_Specular))
+				if (UMaterialExpression* SpecularExpression = CreateExpressionsForNode(Material, Arguments, SpecularNode, Expressions))
 				{
-					SpecularExpression->ConnectExpression(SpecularInput, GetOutputIndex(*SpecularExpression, OutputName));
+					if (FExpressionInput* SpecularInput = Material->GetExpressionInputForProperty(MP_Specular))
+					{
+						SpecularExpression->ConnectExpression(SpecularInput, GetOutputIndex(*SpecularExpression, OutputName));
+					}
 				}
 			}
 		}
@@ -399,17 +387,19 @@ void UInterchangeMaterialFactory::SetupMaterial(UMaterial* Material, const FCrea
 	{
 		FString RoughnessUid;
 		FString OutputName;
-		MaterialFactoryNode->GetRoughnessConnection(RoughnessUid, OutputName);
-
-		const UInterchangeMaterialExpressionFactoryNode* Roughness = Cast<UInterchangeMaterialExpressionFactoryNode>(Arguments.NodeContainer->GetNode(RoughnessUid));
 
-		if (Roughness)
+		if (MaterialFactoryNode->GetRoughnessConnection(RoughnessUid, OutputName))
 		{
-			if (UMaterialExpression* RoughnessExpression = CreateExpressionsForNode(Material, Arguments, Roughness, Expressions))
+			const UInterchangeMaterialExpressionFactoryNode* Roughness = Cast<UInterchangeMaterialExpressionFactoryNode>(Arguments.NodeContainer->GetNode(RoughnessUid));
+
+			if (Roughness)
 			{
-				if (FExpressionInput* RoughnessInput = Material->GetExpressionInputForProperty(MP_Roughness))
+				if (UMaterialExpression* RoughnessExpression = CreateExpressionsForNode(Material, Arguments, Roughness, Expressions))
 				{
-					RoughnessExpression->ConnectExpression(RoughnessInput, GetOutputIndex(*RoughnessExpression, OutputName));
+					if (FExpressionInput* RoughnessInput = Material->GetExpressionInputForProperty(MP_Roughness))
+					{
+						RoughnessExpression->ConnectExpression(RoughnessInput, GetOutputIndex(*RoughnessExpression, OutputName));
+					}
 				}
 			}
 		}
@@ -419,17 +409,19 @@ void UInterchangeMaterialFactory::SetupMaterial(UMaterial* Material, const FCrea
 	{
 		FString EmissiveColorUid;
 		FString OutputName;
-		MaterialFactoryNode->GetEmissiveColorConnection(EmissiveColorUid, OutputName);
-
-		const UInterchangeMaterialExpressionFactoryNode* EmissiveColor = Cast<UInterchangeMaterialExpressionFactoryNode>(Arguments.NodeContainer->GetNode(EmissiveColorUid));
-
-		if (EmissiveColor)
+		
+		if (MaterialFactoryNode->GetEmissiveColorConnection(EmissiveColorUid, OutputName))
 		{
-			if (UMaterialExpression* EmissiveExpression = CreateExpressionsForNode(Material, Arguments, EmissiveColor, Expressions))
+			const UInterchangeMaterialExpressionFactoryNode* EmissiveColor = Cast<UInterchangeMaterialExpressionFactoryNode>(Arguments.NodeContainer->GetNode(EmissiveColorUid));
+
+			if (EmissiveColor)
 			{
-				if (FExpressionInput* EmissiveInput = Material->GetExpressionInputForProperty(MP_EmissiveColor))
+				if (UMaterialExpression* EmissiveExpression = CreateExpressionsForNode(Material, Arguments, EmissiveColor, Expressions))
 				{
-					EmissiveExpression->ConnectExpression(EmissiveInput, GetOutputIndex(*EmissiveExpression, OutputName));
+					if (FExpressionInput* EmissiveInput = Material->GetExpressionInputForProperty(MP_EmissiveColor))
+					{
+						EmissiveExpression->ConnectExpression(EmissiveInput, GetOutputIndex(*EmissiveExpression, OutputName));
+					}
 				}
 			}
 		}
@@ -439,17 +431,19 @@ void UInterchangeMaterialFactory::SetupMaterial(UMaterial* Material, const FCrea
 	{
 		FString NormalUid;
 		FString OutputName;
-		MaterialFactoryNode->GetNormalConnection(NormalUid, OutputName);
-
-		const UInterchangeMaterialExpressionFactoryNode* NormalNode = Cast<UInterchangeMaterialExpressionFactoryNode>(Arguments.NodeContainer->GetNode(NormalUid));
-
-		if (NormalNode)
+		
+		if (MaterialFactoryNode->GetNormalConnection(NormalUid, OutputName))
 		{
-			if (UMaterialExpression* NormalExpression = CreateExpressionsForNode(Material, Arguments, NormalNode, Expressions))
+			const UInterchangeMaterialExpressionFactoryNode* NormalNode = Cast<UInterchangeMaterialExpressionFactoryNode>(Arguments.NodeContainer->GetNode(NormalUid));
+
+			if (NormalNode)
 			{
-				if (FExpressionInput* NormalInput = Material->GetExpressionInputForProperty(MP_Normal))
+				if (UMaterialExpression* NormalExpression = CreateExpressionsForNode(Material, Arguments, NormalNode, Expressions))
 				{
-					NormalExpression->ConnectExpression(NormalInput, GetOutputIndex(*NormalExpression, OutputName));
+					if (FExpressionInput* NormalInput = Material->GetExpressionInputForProperty(MP_Normal))
+					{
+						NormalExpression->ConnectExpression(NormalInput, GetOutputIndex(*NormalExpression, OutputName));
+					}
 				}
 			}
 		}
@@ -459,17 +453,63 @@ void UInterchangeMaterialFactory::SetupMaterial(UMaterial* Material, const FCrea
 	{
 		FString OpacityUid;
 		FString OutputName;
-		MaterialFactoryNode->GetOpacityConnection(OpacityUid, OutputName);
 
-		const UInterchangeMaterialExpressionFactoryNode* OpacityNode = Cast<UInterchangeMaterialExpressionFactoryNode>(Arguments.NodeContainer->GetNode(OpacityUid));
+		if (MaterialFactoryNode->GetOpacityConnection(OpacityUid, OutputName))
+		{
+			const UInterchangeMaterialExpressionFactoryNode* OpacityNode = Cast<UInterchangeMaterialExpressionFactoryNode>(Arguments.NodeContainer->GetNode(OpacityUid));
 
-		if (OpacityNode)
+			if (OpacityNode)
+			{
+				if (UMaterialExpression* OpacityExpression = CreateExpressionsForNode(Material, Arguments, OpacityNode, Expressions))
+				{
+					if (FExpressionInput* OpacityInput = Material->GetExpressionInputForProperty(MP_Opacity))
+					{
+						OpacityExpression->ConnectExpression(OpacityInput, GetOutputIndex(*OpacityExpression, OutputName));
+					}
+				}
+			}
+		}
+	}
+
+	// Occlusion
+	{
+		FString OcclusionUid;
+		FString OutputName;
+
+		if (MaterialFactoryNode->GetOcclusionConnection(OcclusionUid, OutputName))
 		{
-			if (UMaterialExpression* OpacityExpression = CreateExpressionsForNode(Material, Arguments, OpacityNode, Expressions))
+			const UInterchangeMaterialExpressionFactoryNode* OcclusionNode = Cast<UInterchangeMaterialExpressionFactoryNode>(Arguments.NodeContainer->GetNode(OcclusionUid));
+
+			if (OcclusionNode)
 			{
-				if (FExpressionInput* OpacityInput = Material->GetExpressionInputForProperty(MP_Opacity))
+				if (UMaterialExpression* OcclusionExpression = CreateExpressionsForNode(Material, Arguments, OcclusionNode, Expressions))
 				{
-					OpacityExpression->ConnectExpression(OpacityInput, GetOutputIndex(*OpacityExpression, OutputName));
+					if (FExpressionInput* OcclusionInput = Material->GetExpressionInputForProperty(MP_AmbientOcclusion))
+					{
+						OcclusionExpression->ConnectExpression(OcclusionInput, GetOutputIndex(*OcclusionExpression, OutputName));
+					}
+				}
+			}
+		}
+	}
+
+	// Refraction
+	{
+		FString RefractionUid;
+		FString OutputName;
+
+		if (MaterialFactoryNode->GetRefractionConnection(RefractionUid, OutputName))
+		{
+			const UInterchangeMaterialExpressionFactoryNode* RefractionNode = Cast<UInterchangeMaterialExpressionFactoryNode>(Arguments.NodeContainer->GetNode(RefractionUid));
+
+			if (RefractionNode)
+			{
+				if (UMaterialExpression* RefractionExpression = CreateExpressionsForNode(Material, Arguments, RefractionNode, Expressions))
+				{
+					if (FExpressionInput* RefractionInput = Material->GetExpressionInputForProperty(MP_Refraction))
+					{
+						RefractionExpression->ConnectExpression(RefractionInput, GetOutputIndex(*RefractionExpression, OutputName));
+					}
 				}
 			}
 		}
diff --git a/Engine/Plugins/Experimental/Interchange/Runtime/Source/Nodes/Private/InterchangeMaterialFactoryNode.cpp b/Engine/Plugins/Experimental/Interchange/Runtime/Source/Nodes/Private/InterchangeMaterialFactoryNode.cpp
--- a/Engine/Plugins/Experimental/Interchange/Runtime/Source/Nodes/Private/InterchangeMaterialFactoryNode.cpp
+++ b/Engine/Plugins/Experimental/Interchange/Runtime/Source/Nodes/Private/InterchangeMaterialFactoryNode.cpp
@@ -136,6 +136,76 @@ bool UInterchangeMaterialFactoryNode::ConnectOutputToOpacity(const FString& Expr
 	return UInterchangeShaderPortsAPI::ConnectOuputToInput(this, UE::Interchange::Materials::PBR::Parameters::Opacity.ToString(), ExpressionNodeUid, OutputName);
 }
 
+bool UInterchangeMaterialFactoryNode::GetOcclusionConnection(FString& ExpressionNodeUid, FString& OutputName) const
+{
+	return UInterchangeShaderPortsAPI::GetInputConnection(this, UE::Interchange::Materials::Common::Parameters::Occlusion.ToString(), ExpressionNodeUid, OutputName);
+}
+
+bool UInterchangeMaterialFactoryNode::ConnectToOcclusion(const FString& AttributeValue)
+{
+	return UInterchangeShaderPortsAPI::ConnectDefaultOuputToInput(this, UE::Interchange::Materials::Common::Parameters::Occlusion.ToString(), AttributeValue);
+}
+
+bool UInterchangeMaterialFactoryNode::ConnectOutputToOcclusion(const FString& ExpressionNodeUid, const FString& OutputName)
+{
+	return UInterchangeShaderPortsAPI::ConnectOuputToInput(this, UE::Interchange::Materials::PBR::Parameters::Occlusion.ToString(), ExpressionNodeUid, OutputName);
+}
+
+bool UInterchangeMaterialFactoryNode::GetRefractionConnection(FString& ExpressionNodeUid, FString& OutputName) const
+{
+	return UInterchangeShaderPortsAPI::GetInputConnection(this, UE::Interchange::Materials::Common::Parameters::IndexOfRefraction.ToString(), ExpressionNodeUid, OutputName);
+}
+
+bool UInterchangeMaterialFactoryNode::ConnectToRefraction(const FString& AttributeValue)
+{
+	return UInterchangeShaderPortsAPI::ConnectDefaultOuputToInput(this, UE::Interchange::Materials::Common::Parameters::IndexOfRefraction.ToString(), AttributeValue);
+}
+
+bool UInterchangeMaterialFactoryNode::ConnectOutputToRefraction(const FString& ExpressionNodeUid, const FString& OutputName)
+{
+	return UInterchangeShaderPortsAPI::ConnectOuputToInput(this, UE::Interchange::Materials::PBR::Parameters::IndexOfRefraction.ToString(), ExpressionNodeUid, OutputName);
+}
+
+bool UInterchangeMaterialFactoryNode::GetCustomShadingModel(TEnumAsByte<EMaterialShadingModel>& AttributeValue) const
+{
+	IMPLEMENT_NODE_ATTRIBUTE_GETTER(ShadingModel, TEnumAsByte<EMaterialShadingModel>);
+}
+
+bool UInterchangeMaterialFactoryNode::SetCustomShadingModel(const TEnumAsByte<EMaterialShadingModel>& AttributeValue, bool bAddApplyDelegate)
+{
+	IMPLEMENT_NODE_ATTRIBUTE_SETTER(UInterchangeMaterialFactoryNode, ShadingModel, TEnumAsByte<EMaterialShadingModel>, UMaterial);
+}
+
+bool UInterchangeMaterialFactoryNode::GetCustomTranslucencyLightingMode(TEnumAsByte<ETranslucencyLightingMode>& AttributeValue) const
+{
+	IMPLEMENT_NODE_ATTRIBUTE_GETTER(TranslucencyLightingMode, TEnumAsByte<ETranslucencyLightingMode>);
+}
+
+bool UInterchangeMaterialFactoryNode::SetCustomTranslucencyLightingMode(const TEnumAsByte<ETranslucencyLightingMode>& AttributeValue, bool bAddApplyDelegate)
+{
+	IMPLEMENT_NODE_ATTRIBUTE_SETTER(UInterchangeMaterialFactoryNode, TranslucencyLightingMode, TEnumAsByte<ETranslucencyLightingMode>, UMaterial);
+}
+
+bool UInterchangeMaterialFactoryNode::GetCustomBlendMode(TEnumAsByte<EBlendMode>& AttributeValue) const
+{
+	IMPLEMENT_NODE_ATTRIBUTE_GETTER(BlendMode, TEnumAsByte<EBlendMode>);
+}
+
+bool UInterchangeMaterialFactoryNode::SetCustomBlendMode(const TEnumAsByte<EBlendMode>& AttributeValue, bool bAddApplyDelegate)
+{
+	IMPLEMENT_NODE_ATTRIBUTE_SETTER(UInterchangeMaterialFactoryNode, BlendMode, TEnumAsByte<EBlendMode>, UMaterial);
+}
+
+bool UInterchangeMaterialFactoryNode::GetCustomTwoSided(bool& AttributeValue) const
+{
+	IMPLEMENT_NODE_ATTRIBUTE_GETTER(TwoSided, bool);
+}
+
+bool UInterchangeMaterialFactoryNode::SetCustomTwoSided(const bool& AttributeValue, bool bAddApplyDelegate)
+{
+	IMPLEMENT_NODE_ATTRIBUTE_SETTER(UInterchangeMaterialFactoryNode, TwoSided, bool, UMaterial);
+}
+
 FString UInterchangeMaterialExpressionFactoryNode::GetTypeName() const
 {
 	const FString TypeName = TEXT("MaterialExpressionFactoryNode");
diff --git a/Engine/Plugins/Experimental/Interchange/Runtime/Source/Nodes/Private/InterchangeShaderGraphNode.cpp b/Engine/Plugins/Experimental/Interchange/Runtime/Source/Nodes/Private/InterchangeShaderGraphNode.cpp
--- a/Engine/Plugins/Experimental/Interchange/Runtime/Source/Nodes/Private/InterchangeShaderGraphNode.cpp
+++ b/Engine/Plugins/Experimental/Interchange/Runtime/Source/Nodes/Private/InterchangeShaderGraphNode.cpp
@@ -133,3 +133,13 @@ FString UInterchangeShaderGraphNode::GetTypeName() const
 	const FString TypeName = TEXT("ShaderGraphNode");
 	return TypeName;
 }
+
+bool UInterchangeShaderGraphNode::GetCustomTwoSided(bool& AttributeValue) const
+{
+	IMPLEMENT_NODE_ATTRIBUTE_GETTER(TwoSided, bool);
+}
+
+bool UInterchangeShaderGraphNode::SetCustomTwoSided(const bool& AttributeValue)
+{
+	IMPLEMENT_NODE_ATTRIBUTE_SETTER_NODELEGATE(TwoSided, bool);
+}
diff --git a/Engine/Plugins/Experimental/Interchange/Runtime/Source/Nodes/Public/InterchangeMaterialDefinitions.h b/Engine/Plugins/Experimental/Interchange/Runtime/Source/Nodes/Public/InterchangeMaterialDefinitions.h
--- a/Engine/Plugins/Experimental/Interchange/Runtime/Source/Nodes/Public/InterchangeMaterialDefinitions.h
+++ b/Engine/Plugins/Experimental/Interchange/Runtime/Source/Nodes/Public/InterchangeMaterialDefinitions.h
@@ -24,6 +24,17 @@ namespace Interchange
 					}
 				}
 
+				namespace Multiply
+				{
+					const FName Name = TEXT("Multiply");
+
+					namespace Inputs
+					{
+						const FName A = TEXT("A");
+						const FName B = TEXT("B");
+					}
+				}
+
 				namespace TextureSample
 				{
 					const FName Name = TEXT("TextureSample");
@@ -34,6 +45,16 @@ namespace Interchange
 						const FName UTiling = TEXT("UTiling"); // Type: float
 						const FName VTiling = TEXT("VTiling"); // Type: float
 					}
+
+					namespace Outputs
+					{
+						const FName RGB = TEXT("RGB"); // Type: linear color
+						const FName R = TEXT("R"); // Type: float
+						const FName G = TEXT("G"); // Type: float
+						const FName B = TEXT("B"); // Type: float
+						const FName A = TEXT("A"); // Type: float
+						const FName RGBA = TEXT("RGBA"); // Type: linear color
+					}
 				}
 			}
 		}
@@ -45,6 +66,8 @@ namespace Interchange
 				const FName EmissiveColor = TEXT("EmissiveColor"); // Type: linear color
 				const FName Normal = TEXT("Normal"); // Type: vector3f
 				const FName Opacity = TEXT("Opacity"); // Type: float
+				const FName Occlusion = TEXT("Occlusion"); // Type: float
+				const FName IndexOfRefraction = TEXT("IOR"); // Type: float
 			}
 		}
 
diff --git a/Engine/Plugins/Experimental/Interchange/Runtime/Source/Nodes/Public/InterchangeMaterialFactoryNode.h b/Engine/Plugins/Experimental/Interchange/Runtime/Source/Nodes/Public/InterchangeMaterialFactoryNode.h
--- a/Engine/Plugins/Experimental/Interchange/Runtime/Source/Nodes/Public/InterchangeMaterialFactoryNode.h
+++ b/Engine/Plugins/Experimental/Interchange/Runtime/Source/Nodes/Public/InterchangeMaterialFactoryNode.h
@@ -34,6 +34,8 @@ public:
 	UFUNCTION(BlueprintCallable, Category = "Interchange | Node | MaterialFactory")
 	virtual class UClass* GetObjectClass() const override;
 
+// Material Inputs
+public:
 	UFUNCTION(BlueprintCallable, Category = "Interchange | Node | MaterialFactory")
 	bool GetBaseColorConnection(FString& ExpressionNodeUid, FString& OutputName) const;
 
@@ -96,6 +98,57 @@ public:
 
 	UFUNCTION(BlueprintCallable, Category = "Interchange | Node | MaterialFactory")
 	bool ConnectOutputToOpacity(const FString& ExpressionNodeUid, const FString& OutputName);
+
+	UFUNCTION(BlueprintCallable, Category = "Interchange | Node | MaterialFactory")
+	bool GetOcclusionConnection(FString& ExpressionNodeUid, FString& OutputName) const;
+
+	UFUNCTION(BlueprintCallable, Category = "Interchange | Node | MaterialFactory")
+	bool ConnectToOcclusion(const FString& AttributeValue);
+
+	UFUNCTION(BlueprintCallable, Category = "Interchange | Node | MaterialFactory")
+	bool ConnectOutputToOcclusion(const FString& ExpressionNodeUid, const FString& OutputName);
+
+	UFUNCTION(BlueprintCallable, Category = "Interchange | Node | MaterialFactory")
+	bool GetRefractionConnection(FString& ExpressionNodeUid, FString& OutputName) const;
+
+	UFUNCTION(BlueprintCallable, Category = "Interchange | Node | MaterialFactory")
+	bool ConnectToRefraction(const FString& AttributeValue);
+
+	UFUNCTION(BlueprintCallable, Category = "Interchange | Node | MaterialFactory")
+	bool ConnectOutputToRefraction(const FString& ExpressionNodeUid, const FString& OutputName);
+
+// Material parameters
+public:
+	UFUNCTION(BlueprintCallable, Category = "Interchange | Node | Material")
+	bool GetCustomShadingModel(TEnumAsByte<EMaterialShadingModel>& AttributeValue) const;
+ 
+	UFUNCTION(BlueprintCallable, Category = "Interchange | Node | Material")
+	bool SetCustomShadingModel(const TEnumAsByte<EMaterialShadingModel>& AttributeValue, bool bAddApplyDelegate = true);
+
+	UFUNCTION(BlueprintCallable, Category = "Interchange | Node | Material")
+	bool GetCustomTranslucencyLightingMode(TEnumAsByte<ETranslucencyLightingMode>& AttributeValue) const;
+ 
+	UFUNCTION(BlueprintCallable, Category = "Interchange | Node | Material")
+	bool SetCustomTranslucencyLightingMode(const TEnumAsByte<ETranslucencyLightingMode>& AttributeValue, bool bAddApplyDelegate = true);
+
+	UFUNCTION(BlueprintCallable, Category = "Interchange | Node | Material")
+	bool GetCustomBlendMode(TEnumAsByte<EBlendMode>& AttributeValue) const;
+ 
+	UFUNCTION(BlueprintCallable, Category = "Interchange | Node | Material")
+	bool SetCustomBlendMode(const TEnumAsByte<EBlendMode>& AttributeValue, bool bAddApplyDelegate = true);
+
+	UFUNCTION(BlueprintCallable, Category = "Interchange | Node | Material")
+	bool GetCustomTwoSided(bool& AttributeValue) const;
+ 
+	/** Sets if this shader graph should be rendered two sided or not. Defaults to off. */
+	UFUNCTION(BlueprintCallable, Category = "Interchange | Node | Material")
+	bool SetCustomTwoSided(const bool& AttributeValue, bool bAddApplyDelegate = true);
+
+private:
+	const UE::Interchange::FAttributeKey Macro_CustomShadingModelKey = UE::Interchange::FAttributeKey(TEXT("ShadingModel"));
+	const UE::Interchange::FAttributeKey Macro_CustomTranslucencyLightingModeKey = UE::Interchange::FAttributeKey(TEXT("TranslucencyLightingMode"));
+	const UE::Interchange::FAttributeKey Macro_CustomBlendModeKey = UE::Interchange::FAttributeKey(TEXT("BlendMode"));
+	const UE::Interchange::FAttributeKey Macro_CustomTwoSidedKey = UE::Interchange::FAttributeKey(TEXT("TwoSided"));
 };
 
 UCLASS(BlueprintType, Experimental)
diff --git a/Engine/Plugins/Experimental/Interchange/Runtime/Source/Nodes/Public/InterchangeShaderGraphNode.h b/Engine/Plugins/Experimental/Interchange/Runtime/Source/Nodes/Public/InterchangeShaderGraphNode.h
--- a/Engine/Plugins/Experimental/Interchange/Runtime/Source/Nodes/Public/InterchangeShaderGraphNode.h
+++ b/Engine/Plugins/Experimental/Interchange/Runtime/Source/Nodes/Public/InterchangeShaderGraphNode.h
@@ -129,4 +129,17 @@ public:
 	 * Return the node type name of the class, we use this when reporting error
 	 */
 	virtual FString GetTypeName() const override;
+
+public:
+	UFUNCTION(BlueprintCallable, Category = "Interchange | Node | Material")
+	bool GetCustomTwoSided(bool& AttributeValue) const;
+ 
+	/**
+	 * Sets if this shader graph should be rendered two sided or not. Defaults to off.
+	 */
+	UFUNCTION(BlueprintCallable, Category = "Interchange | Node | Material")
+	bool SetCustomTwoSided(const bool& AttributeValue);
+
+private:
+	const UE::Interchange::FAttributeKey Macro_CustomTwoSidedKey = UE::Interchange::FAttributeKey(TEXT("TwoSided"));
 };
\ No newline at end of file
diff --git a/Engine/Plugins/Experimental/Interchange/Runtime/Source/Pipelines/Private/InterchangeGenericMaterialPipeline.cpp b/Engine/Plugins/Experimental/Interchange/Runtime/Source/Pipelines/Private/InterchangeGenericMaterialPipeline.cpp
--- a/Engine/Plugins/Experimental/Interchange/Runtime/Source/Pipelines/Private/InterchangeGenericMaterialPipeline.cpp
+++ b/Engine/Plugins/Experimental/Interchange/Runtime/Source/Pipelines/Private/InterchangeGenericMaterialPipeline.cpp
@@ -10,6 +10,7 @@
 #include "InterchangeTexture2DNode.h"
 #include "InterchangeTexture2DArrayNode.h"
 #include "InterchangeTextureCubeNode.h"
+#include "InterchangeTextureFactoryNode.h"
 #include "InterchangeTextureNode.h"
 #include "InterchangePipelineLog.h"
 #include "InterchangeSourceData.h"
@@ -298,6 +299,13 @@ void UInterchangeGenericMaterialPipeline::HandleCommonParameters(const UIntercha
 {
 	using namespace UE::Interchange::Materials::Common;
 
+	// Two sidedness
+	{
+		bool bTwoSided = false;
+		ShaderGraphNode->GetCustomTwoSided(bTwoSided);
+		MaterialFactoryNode->SetCustomTwoSided(bTwoSided);
+	}
+
 	// Emissive
 	{
 		const bool bHasEmissiveInput = UInterchangeShaderPortsAPI::HasInput(ShaderGraphNode, Parameters::EmissiveColor);
@@ -320,6 +328,8 @@ void UInterchangeGenericMaterialPipeline::HandleCommonParameters(const UIntercha
 
 		if (bHasNormalInput)
 		{
+			TGuardValue<bool> ParsingForNormalInputGuard(bParsingForNormalInput, true);
+
 			TTuple<UInterchangeMaterialExpressionFactoryNode*, FString> ExpressionFactoryNode =
 				CreateMaterialExpressionForInput(MaterialFactoryNode, ShaderGraphNode, Parameters::Normal.ToString(), MaterialFactoryNode->GetUniqueID());
 
@@ -342,6 +352,40 @@ void UInterchangeGenericMaterialPipeline::HandleCommonParameters(const UIntercha
 			if (OpacityExpressionFactoryNode.Get<0>())
 			{
 				MaterialFactoryNode->ConnectOutputToOpacity(OpacityExpressionFactoryNode.Get<0>()->GetUniqueID(), OpacityExpressionFactoryNode.Get<1>());
+				MaterialFactoryNode->SetCustomBlendMode(EBlendMode::BLEND_Translucent);
+				MaterialFactoryNode->SetCustomTranslucencyLightingMode(ETranslucencyLightingMode::TLM_Surface);
+			}
+		}
+	}
+
+	// Ambient Occlusion
+	{
+		const bool bHasOcclusionInput = UInterchangeShaderPortsAPI::HasInput(ShaderGraphNode, Parameters::Occlusion);
+
+		if (bHasOcclusionInput)
+		{
+			TTuple<UInterchangeMaterialExpressionFactoryNode*, FString> ExpressionFactoryNode =
+				CreateMaterialExpressionForInput(MaterialFactoryNode, ShaderGraphNode, Parameters::Occlusion.ToString(), MaterialFactoryNode->GetUniqueID());
+
+			if (ExpressionFactoryNode.Get<0>())
+			{
+				MaterialFactoryNode->ConnectOutputToOcclusion(ExpressionFactoryNode.Get<0>()->GetUniqueID(), ExpressionFactoryNode.Get<1>());
+			}
+		}
+	}
+
+	// Index of Refraction
+	{
+		const bool bHasIorInput = UInterchangeShaderPortsAPI::HasInput(ShaderGraphNode, Parameters::IndexOfRefraction);
+
+		if (bHasIorInput)
+		{
+			TTuple<UInterchangeMaterialExpressionFactoryNode*, FString> ExpressionFactoryNode =
+				CreateMaterialExpressionForInput(MaterialFactoryNode, ShaderGraphNode, Parameters::IndexOfRefraction.ToString(), MaterialFactoryNode->GetUniqueID());
+
+			if (ExpressionFactoryNode.Get<0>())
+			{
+				MaterialFactoryNode->ConnectOutputToRefraction(ExpressionFactoryNode.Get<0>()->GetUniqueID(), ExpressionFactoryNode.Get<1>());
 			}
 		}
 	}
@@ -349,12 +393,13 @@ void UInterchangeGenericMaterialPipeline::HandleCommonParameters(const UIntercha
 
 void UInterchangeGenericMaterialPipeline::HandleTextureSampleNode(const UInterchangeShaderNode* ShaderNode, UInterchangeMaterialExpressionFactoryNode* TextureSampleFactoryNode)
 {
-	using namespace UE::Interchange::Materials::Standard;
+	using namespace UE::Interchange::Materials::Standard::Nodes::TextureSample;
 
 	FString TextureUid;
-	ShaderNode->GetStringAttribute(UInterchangeShaderPortsAPI::MakeInputValueKey(Nodes::TextureSample::Inputs::Texture.ToString()), TextureUid);
+	ShaderNode->GetStringAttribute(UInterchangeShaderPortsAPI::MakeInputValueKey(Inputs::Texture.ToString()), TextureUid);
 
 	FString ExpressionClassName;
+	FString TextureFactoryUid;
 
 	if (UInterchangeTextureNode* TextureNode = Cast<UInterchangeTextureNode>(BaseNodeContainer->GetNode(TextureUid)))
 	{
@@ -374,9 +419,26 @@ void UInterchangeGenericMaterialPipeline::HandleTextureSampleNode(const UInterch
 		{
 			ExpressionClassName = UMaterialExpressionTextureSampleParameter2D::StaticClass()->GetName();
 		}
+
+		TArray<FString> TextureTargetNodes;
+		TextureNode->GetTargetNodeUids(TextureTargetNodes);
+
+		if (TextureTargetNodes.Num() > 0)
+		{
+			TextureFactoryUid = TextureTargetNodes[0];
+		}
 	}
 
 	TextureSampleFactoryNode->SetCustomExpressionClassName(ExpressionClassName);
+	TextureSampleFactoryNode->AddStringAttribute(UInterchangeShaderPortsAPI::MakeInputValueKey(Inputs::Texture.ToString()), TextureFactoryUid);
+
+	if (bParsingForNormalInput)
+	{
+		if (UInterchangeTextureFactoryNode* TextureFactoryNode = Cast<UInterchangeTextureFactoryNode>(BaseNodeContainer->GetNode(TextureFactoryUid)))
+		{
+			TextureFactoryNode->SetCustomCompressionSettings(TextureCompressionSettings::TC_Normalmap);
+		}
+	}
 
 	HandleTextureCoordinates(ShaderNode, TextureSampleFactoryNode);
 }
diff --git a/Engine/Plugins/Experimental/Interchange/Runtime/Source/Pipelines/Private/InterchangeGenericTexturePipeline.cpp b/Engine/Plugins/Experimental/Interchange/Runtime/Source/Pipelines/Private/InterchangeGenericTexturePipeline.cpp
--- a/Engine/Plugins/Experimental/Interchange/Runtime/Source/Pipelines/Private/InterchangeGenericTexturePipeline.cpp
+++ b/Engine/Plugins/Experimental/Interchange/Runtime/Source/Pipelines/Private/InterchangeGenericTexturePipeline.cpp
@@ -176,29 +176,32 @@ void UInterchangeGenericAssetsPipeline::PostImportTextureAssetImport(UObject* Cr
 		// Verify if the texture is a normal map
 		if (UTexture* Texture = Cast<UTexture>(CreatedAsset))
 		{
-			// This can create 2 build of the texture (we should revisit this at some point)
-			if (FTextureCompilingManager::Get().IsCompilingTexture(Texture))
+			if (!Texture->IsNormalMap())
 			{
-				TWeakObjectPtr<UTexture> WeakTexturePtr = Texture;
-				TSharedRef<FDelegateHandle> HandlePtr = MakeShared<FDelegateHandle>();
-				HandlePtr.Get() = FTextureCompilingManager::Get().OnTexturePostCompileEvent().AddLambda([this, WeakTexturePtr, HandlePtr](const TArrayView<UTexture* const>&)
-					{
-						if (UTexture* TextureToTest = WeakTexturePtr.Get())
+				// This can create 2 build of the texture (we should revisit this at some point)
+				if (FTextureCompilingManager::Get().IsCompilingTexture(Texture))
+				{
+					TWeakObjectPtr<UTexture> WeakTexturePtr = Texture;
+					TSharedRef<FDelegateHandle> HandlePtr = MakeShared<FDelegateHandle>();
+					HandlePtr.Get() = FTextureCompilingManager::Get().OnTexturePostCompileEvent().AddLambda([this, WeakTexturePtr, HandlePtr](const TArrayView<UTexture* const>&)
 						{
-							if (FTextureCompilingManager::Get().IsCompilingTexture(TextureToTest))
+							if (UTexture* TextureToTest = WeakTexturePtr.Get())
 							{
-								return;
-							}
+								if (FTextureCompilingManager::Get().IsCompilingTexture(TextureToTest))
+								{
+									return;
+								}
 
-							UE::Interchange::Private::AdjustTextureForNormalMap(TextureToTest, bFlipNormalMapGreenChannel);
-						}
+								UE::Interchange::Private::AdjustTextureForNormalMap(TextureToTest, bFlipNormalMapGreenChannel);
+							}
 
-						FTextureCompilingManager::Get().OnTexturePostCompileEvent().Remove(HandlePtr.Get());
-					});
-			}
-			else
-			{
-				UE::Interchange::Private::AdjustTextureForNormalMap(Texture, bFlipNormalMapGreenChannel);
+							FTextureCompilingManager::Get().OnTexturePostCompileEvent().Remove(HandlePtr.Get());
+						});
+				}
+				else
+				{
+					UE::Interchange::Private::AdjustTextureForNormalMap(Texture, bFlipNormalMapGreenChannel);
+				}
 			}
 		}
 	}
diff --git a/Engine/Plugins/Experimental/Interchange/Runtime/Source/Pipelines/Public/InterchangeGenericMaterialPipeline.h b/Engine/Plugins/Experimental/Interchange/Runtime/Source/Pipelines/Public/InterchangeGenericMaterialPipeline.h
--- a/Engine/Plugins/Experimental/Interchange/Runtime/Source/Pipelines/Public/InterchangeGenericMaterialPipeline.h
+++ b/Engine/Plugins/Experimental/Interchange/Runtime/Source/Pipelines/Public/InterchangeGenericMaterialPipeline.h
@@ -69,6 +69,9 @@ private:
 	UInterchangeMaterialExpressionFactoryNode* CreateExpressionNode(const FString& ExpressionName, const FString& ParentUid, UClass* MaterialExpressionClass);
 	UInterchangeMaterialExpressionFactoryNode* CreateScalarParameterExpression(const UInterchangeShaderNode* ShaderNode, const FString& InputName, const FString& ParentUid);
 	UInterchangeMaterialExpressionFactoryNode* CreateVectorParameterExpression(const UInterchangeShaderNode* ShaderNode, const FString& InputName, const FString& ParentUid);
+
+private:
+	bool bParsingForNormalInput;
 };
 
 
diff --git a/Engine/Source/Runtime/Experimental/Interchange/Core/Public/Nodes/InterchangeBaseNode.h b/Engine/Source/Runtime/Experimental/Interchange/Core/Public/Nodes/InterchangeBaseNode.h
--- a/Engine/Source/Runtime/Experimental/Interchange/Core/Public/Nodes/InterchangeBaseNode.h
+++ b/Engine/Source/Runtime/Experimental/Interchange/Core/Public/Nodes/InterchangeBaseNode.h
@@ -400,13 +400,6 @@ public:
 	template<typename AttributeType>
 	inline bool ApplyAttributeToObject(const FString& NodeAttributeKey, UObject* Object, const FName PropertyName) const;
 
-	/**
-	 * Specialized version of ApplyAttributeToObject for strings.
-	 * If the target property is a FObjectPropertyBase, treat the string as an object path.
-	 */
-	template<>
-	bool ApplyAttributeToObject<FString>(const FString& NodeAttributeKey, UObject* Object, const FName PropertyName) const;
-
 	/**
 	 * Reads an attribute value from a UObject.
 	 * @param NodeAttributeKey	The key for the attribute to update.
@@ -417,13 +410,6 @@ public:
 	template<typename AttributeType>
 	inline bool FillAttributeFromObject(const FString& NodeAttributeKey, UObject* Object, const FName PropertyName);
 
-	/**
-	 * Specialized version of FillAttributeFromObject for strings.
-	 * If the target property is a FObjectPropertyBase, treat the string as an object path.
-	 */
-	template<>
-	inline bool FillAttributeFromObject<FString>(const FString& NodeAttributeKey, UObject* Object, const FName PropertyName);
-
 	/**
 	 * Return the unique id pass in the constructor.
 	 */
@@ -656,20 +642,27 @@ inline bool UInterchangeBaseNode::ApplyAttributeToObject(const FString& NodeAttr
 		Container.Set<UObject*>(Object);
 		if(FProperty* Property = InterchangePrivateNodeBase::FindPropertyByPathChecked(Container, Object->GetClass(), PropertyName.ToString()))
 		{
+			AttributeType* PropertyValueAddress;
 			if (Container.IsType<UObject*>())
 			{
-				*(Property->ContainerPtrToValuePtr<AttributeType>(Container.Get<UObject*>())) = ValueData;
+				PropertyValueAddress = Property->ContainerPtrToValuePtr<AttributeType>(Container.Get<UObject*>());
 			}
 			else
 			{
-				*(Property->ContainerPtrToValuePtr<AttributeType>(Container.Get<uint8*>())) = ValueData;
+				PropertyValueAddress = Property->ContainerPtrToValuePtr<AttributeType>(Container.Get<uint8*>());
 			}
+
+			*PropertyValueAddress = ValueData;
 		}
 		return true;
 	}
 	return false;
 }
 
+/**
+ * Specialized version of ApplyAttributeToObject for strings.
+ * If the target property is a FObjectPropertyBase, treat the string as an object path.
+ */
 template<>
 inline bool UInterchangeBaseNode::ApplyAttributeToObject<FString>(const FString& NodeAttributeKey, UObject* Object, const FName PropertyName) const
 {
@@ -685,31 +678,67 @@ inline bool UInterchangeBaseNode::ApplyAttributeToObject<FString>(const FString&
 		Container.Set<UObject*>(Object);
 		if(FProperty* Property = InterchangePrivateNodeBase::FindPropertyByPathChecked(Container, Object->GetClass(), PropertyName.ToString()))
 		{
+			FString* PropertyValueAddress;
+			if (Container.IsType<UObject*>())
+			{
+				PropertyValueAddress = Property->ContainerPtrToValuePtr<FString>(Container.Get<UObject*>());
+			}
+			else
+			{
+				PropertyValueAddress = Property->ContainerPtrToValuePtr<FString>(Container.Get<uint8*>());
+			}
+
 			if (FObjectPropertyBase* ObjectProperty = CastField<FObjectPropertyBase>(Property))
 			{
-				void* PropertyValueAddress;
+				ObjectProperty->SetObjectPropertyValue(PropertyValueAddress, FSoftObjectPath(ValueData).TryLoad());
+			}
+			else
+			{
+				*PropertyValueAddress = ValueData;
+			}
+		}
+		return true;
+	}
+	return false;
+}
 
-				if (Container.IsType<UObject*>())
-				{
-					PropertyValueAddress = Property->ContainerPtrToValuePtr<void>(Container.Get<UObject*>());
-				}
-				else
-				{
-					PropertyValueAddress =  Property->ContainerPtrToValuePtr<void>(Container.Get<uint8*>());
-				}
+/**
+ * Specialized version of ApplyAttributeToObject for bools.
+ * If the target property is a FBoolProperty, treat the propertyg as a bitfield.
+ */
+template<>
+inline bool UInterchangeBaseNode::ApplyAttributeToObject<bool>(const FString& NodeAttributeKey, UObject* Object, const FName PropertyName) const
+{
+	if (!Object)
+	{
+		return false;
+	}
 
-				ObjectProperty->SetObjectPropertyValue(PropertyValueAddress, FSoftObjectPath(ValueData).TryLoad());
+	bool bValueData;
+	if (GetAttribute<bool>(NodeAttributeKey, bValueData))
+	{
+		TVariant<UObject*, uint8*> Container;
+		Container.Set<UObject*>(Object);
+		if(FProperty* Property = InterchangePrivateNodeBase::FindPropertyByPathChecked(Container, Object->GetClass(), PropertyName.ToString()))
+		{
+			bool* PropertyValueAddress;
+			if (Container.IsType<UObject*>())
+			{
+				PropertyValueAddress = Property->ContainerPtrToValuePtr<bool>(Container.Get<UObject*>());
+			}
+			else
+			{
+				PropertyValueAddress = Property->ContainerPtrToValuePtr<bool>(Container.Get<uint8*>());
+			}
+
+			// Support for bitfields
+			if (FBoolProperty* BoolProperty = CastField<FBoolProperty>(Property))
+			{
+				BoolProperty->SetPropertyValue(PropertyValueAddress, bValueData);
 			}
 			else
 			{
-				if (Container.IsType<UObject*>())
-				{
-					*(Property->ContainerPtrToValuePtr<FString>(Container.Get<UObject*>())) = ValueData;
-				}
-				else
-				{
-					*(Property->ContainerPtrToValuePtr<FString>(Container.Get<uint8*>())) = ValueData;
-				}
+				*PropertyValueAddress = bValueData;
 			}
 		}
 		return true;
@@ -724,25 +753,35 @@ inline bool UInterchangeBaseNode::FillAttributeFromObject(const FString& NodeAtt
 	Container.Set<UObject*>(Object);
 	if(FProperty* Property = InterchangePrivateNodeBase::FindPropertyByPathChecked(Container, Object->GetClass(), PropertyName.ToString()))
 	{
-		AttributeType ValueData;
+		AttributeType* PropertyValueAddress;
 		if (Container.IsType<UObject*>())
 		{
-			ValueData = *(Property->ContainerPtrToValuePtr<AttributeType>(Container.Get<UObject*>()));
+			PropertyValueAddress = Property->ContainerPtrToValuePtr<AttributeType>(Container.Get<UObject*>());
 		}
 		else
 		{
-			ValueData = *(Property->ContainerPtrToValuePtr<AttributeType>(Container.Get<uint8*>()));
+			PropertyValueAddress = Property->ContainerPtrToValuePtr<AttributeType>(Container.Get<uint8*>());
 		}
 
-		if (SetAttribute(NodeAttributeKey, ValueData))
+		// Support for bitfields
+		if (FBoolProperty* BoolProperty = CastField<FBoolProperty>(Property))
 		{
-			return true;
+			const bool bPropertyValue = BoolProperty->GetPropertyValue(PropertyValueAddress);
+			return SetAttribute(NodeAttributeKey, bPropertyValue);
+		}
+		else
+		{
+			return SetAttribute(NodeAttributeKey, *PropertyValueAddress);
 		}
 	}
 
 	return false;
 }
 
+/**
+ * Specialized version of FillAttributeFromObject for strings.
+ * If the target property is a FObjectPropertyBase, treat the string as an object path.
+ */
 template<>
 inline bool UInterchangeBaseNode::FillAttributeFromObject<FString>(const FString& NodeAttributeKey, UObject* Object, const FName PropertyName)
 {
@@ -750,37 +789,62 @@ inline bool UInterchangeBaseNode::FillAttributeFromObject<FString>(const FString
 	Container.Set<UObject*>(Object);
 	if(FProperty* Property = InterchangePrivateNodeBase::FindPropertyByPathChecked(Container, Object->GetClass(), PropertyName.ToString()))
 	{
-		if (FObjectPropertyBase* ObjectProperty = CastField<FObjectPropertyBase>(Property))
+		FString* PropertyValueAddress;
+		if (Container.IsType<UObject*>())
 		{
-			void* PropertyValueAddress;
-
-			if (Container.IsType<UObject*>())
-			{
-				PropertyValueAddress = Property->ContainerPtrToValuePtr<void>(Container.Get<UObject*>());
-			}
-			else
-			{
-				PropertyValueAddress =  Property->ContainerPtrToValuePtr<void>(Container.Get<uint8*>());
-			}
+			PropertyValueAddress = Property->ContainerPtrToValuePtr<FString>(Container.Get<UObject*>());
+		}
+		else
+		{
+			PropertyValueAddress = Property->ContainerPtrToValuePtr<FString>(Container.Get<uint8*>());
+		}
 
+		if (FObjectPropertyBase* ObjectProperty = CastField<FObjectPropertyBase>(Property))
+		{
 			UObject* ObjectValue = ObjectProperty->GetObjectPropertyValue(PropertyValueAddress);
 			return SetAttribute(NodeAttributeKey, ObjectValue->GetPathName());
 		}
 		else
 		{
-			FString ValueData;
-			if (Container.IsType<UObject*>())
-			{
-				ValueData = *(Property->ContainerPtrToValuePtr<FString>(Container.Get<UObject*>()));
-			}
-			else
-			{
-				ValueData = *(Property->ContainerPtrToValuePtr<FString>(Container.Get<uint8*>()));
-			}
+			return SetAttribute(NodeAttributeKey, *PropertyValueAddress);
+		}
+	}
+
+	return false;
+}
+
+/**
+ * Specialized version of FillAttributeFromObject for bools.
+ * If the target property is a FBoolProperty, treat the property as a bitfield.
+ */
+template<>
+inline bool UInterchangeBaseNode::FillAttributeFromObject<bool>(const FString& NodeAttributeKey, UObject* Object, const FName PropertyName)
+{
+	TVariant<UObject*, uint8*> Container;
+	Container.Set<UObject*>(Object);
+	if(FProperty* Property = InterchangePrivateNodeBase::FindPropertyByPathChecked(Container, Object->GetClass(), PropertyName.ToString()))
+	{
+		bool* PropertyValueAddress;
+		if (Container.IsType<UObject*>())
+		{
+			PropertyValueAddress = Property->ContainerPtrToValuePtr<bool>(Container.Get<UObject*>());
+		}
+		else
+		{
+			PropertyValueAddress = Property->ContainerPtrToValuePtr<bool>(Container.Get<uint8*>());
+		}
 
-			return SetAttribute(NodeAttributeKey, ValueData);
+		// Support for bitfields
+		if (FBoolProperty* BoolProperty = CastField<FBoolProperty>(Property))
+		{
+			const bool bPropertyValue = BoolProperty->GetPropertyValue(PropertyValueAddress);
+			return SetAttribute(NodeAttributeKey, bPropertyValue);
+		}
+		else
+		{
+			return SetAttribute(NodeAttributeKey, *PropertyValueAddress);
 		}
 	}
 
 	return false;
-}
\ No newline at end of file
+}
