Build fix for InterchangeWorker: Removed engine enum from UInterchangeMaterialFactoryNode since the InterchangeWorker doesn't compile against the engine.

diff --git a/Engine/Plugins/Experimental/Interchange/Runtime/Source/Nodes/Private/InterchangeMaterialFactoryNode.cpp b/Engine/Plugins/Experimental/Interchange/Runtime/Source/Nodes/Private/InterchangeMaterialFactoryNode.cpp
--- a/Engine/Plugins/Experimental/Interchange/Runtime/Source/Nodes/Private/InterchangeMaterialFactoryNode.cpp
+++ b/Engine/Plugins/Experimental/Interchange/Runtime/Source/Nodes/Private/InterchangeMaterialFactoryNode.cpp
@@ -166,34 +166,34 @@ bool UInterchangeMaterialFactoryNode::ConnectOutputToRefraction(const FString& E
 	return UInterchangeShaderPortsAPI::ConnectOuputToInput(this, UE::Interchange::Materials::PBR::Parameters::IndexOfRefraction.ToString(), ExpressionNodeUid, OutputName);
 }
 
-bool UInterchangeMaterialFactoryNode::GetCustomShadingModel(TEnumAsByte<EMaterialShadingModel>& AttributeValue) const
+bool UInterchangeMaterialFactoryNode::GetCustomShadingModel(uint8& AttributeValue) const
 {
-	IMPLEMENT_NODE_ATTRIBUTE_GETTER(ShadingModel, TEnumAsByte<EMaterialShadingModel>);
+	IMPLEMENT_NODE_ATTRIBUTE_GETTER(ShadingModel, uint8);
 }
 
-bool UInterchangeMaterialFactoryNode::SetCustomShadingModel(const TEnumAsByte<EMaterialShadingModel>& AttributeValue, bool bAddApplyDelegate)
+bool UInterchangeMaterialFactoryNode::SetCustomShadingModel(const uint8& AttributeValue, bool bAddApplyDelegate)
 {
-	IMPLEMENT_NODE_ATTRIBUTE_SETTER(UInterchangeMaterialFactoryNode, ShadingModel, TEnumAsByte<EMaterialShadingModel>, UMaterial);
+	IMPLEMENT_NODE_ATTRIBUTE_SETTER(UInterchangeMaterialFactoryNode, ShadingModel, uint8, UMaterial);
 }
 
-bool UInterchangeMaterialFactoryNode::GetCustomTranslucencyLightingMode(TEnumAsByte<ETranslucencyLightingMode>& AttributeValue) const
+bool UInterchangeMaterialFactoryNode::GetCustomTranslucencyLightingMode(uint8& AttributeValue) const
 {
-	IMPLEMENT_NODE_ATTRIBUTE_GETTER(TranslucencyLightingMode, TEnumAsByte<ETranslucencyLightingMode>);
+	IMPLEMENT_NODE_ATTRIBUTE_GETTER(TranslucencyLightingMode, uint8);
 }
 
-bool UInterchangeMaterialFactoryNode::SetCustomTranslucencyLightingMode(const TEnumAsByte<ETranslucencyLightingMode>& AttributeValue, bool bAddApplyDelegate)
+bool UInterchangeMaterialFactoryNode::SetCustomTranslucencyLightingMode(const uint8& AttributeValue, bool bAddApplyDelegate)
 {
-	IMPLEMENT_NODE_ATTRIBUTE_SETTER(UInterchangeMaterialFactoryNode, TranslucencyLightingMode, TEnumAsByte<ETranslucencyLightingMode>, UMaterial);
+	IMPLEMENT_NODE_ATTRIBUTE_SETTER(UInterchangeMaterialFactoryNode, TranslucencyLightingMode, uint8, UMaterial);
 }
 
-bool UInterchangeMaterialFactoryNode::GetCustomBlendMode(TEnumAsByte<EBlendMode>& AttributeValue) const
+bool UInterchangeMaterialFactoryNode::GetCustomBlendMode(uint8& AttributeValue) const
 {
-	IMPLEMENT_NODE_ATTRIBUTE_GETTER(BlendMode, TEnumAsByte<EBlendMode>);
+	IMPLEMENT_NODE_ATTRIBUTE_GETTER(BlendMode, uint8);
 }
 
-bool UInterchangeMaterialFactoryNode::SetCustomBlendMode(const TEnumAsByte<EBlendMode>& AttributeValue, bool bAddApplyDelegate)
+bool UInterchangeMaterialFactoryNode::SetCustomBlendMode(const uint8& AttributeValue, bool bAddApplyDelegate)
 {
-	IMPLEMENT_NODE_ATTRIBUTE_SETTER(UInterchangeMaterialFactoryNode, BlendMode, TEnumAsByte<EBlendMode>, UMaterial);
+	IMPLEMENT_NODE_ATTRIBUTE_SETTER(UInterchangeMaterialFactoryNode, BlendMode, uint8, UMaterial);
 }
 
 bool UInterchangeMaterialFactoryNode::GetCustomTwoSided(bool& AttributeValue) const
diff --git a/Engine/Plugins/Experimental/Interchange/Runtime/Source/Nodes/Public/InterchangeMaterialFactoryNode.h b/Engine/Plugins/Experimental/Interchange/Runtime/Source/Nodes/Public/InterchangeMaterialFactoryNode.h
--- a/Engine/Plugins/Experimental/Interchange/Runtime/Source/Nodes/Public/InterchangeMaterialFactoryNode.h
+++ b/Engine/Plugins/Experimental/Interchange/Runtime/Source/Nodes/Public/InterchangeMaterialFactoryNode.h
@@ -120,22 +120,22 @@ public:
 // Material parameters
 public:
 	UFUNCTION(BlueprintCallable, Category = "Interchange | Node | Material")
-	bool GetCustomShadingModel(TEnumAsByte<EMaterialShadingModel>& AttributeValue) const;
+	bool GetCustomShadingModel(uint8& AttributeValue) const;
  
 	UFUNCTION(BlueprintCallable, Category = "Interchange | Node | Material")
-	bool SetCustomShadingModel(const TEnumAsByte<EMaterialShadingModel>& AttributeValue, bool bAddApplyDelegate = true);
+	bool SetCustomShadingModel(const uint8& AttributeValue, bool bAddApplyDelegate = true);
 
 	UFUNCTION(BlueprintCallable, Category = "Interchange | Node | Material")
-	bool GetCustomTranslucencyLightingMode(TEnumAsByte<ETranslucencyLightingMode>& AttributeValue) const;
+	bool GetCustomTranslucencyLightingMode(uint8& AttributeValue) const;
  
 	UFUNCTION(BlueprintCallable, Category = "Interchange | Node | Material")
-	bool SetCustomTranslucencyLightingMode(const TEnumAsByte<ETranslucencyLightingMode>& AttributeValue, bool bAddApplyDelegate = true);
+	bool SetCustomTranslucencyLightingMode(const uint8& AttributeValue, bool bAddApplyDelegate = true);
 
 	UFUNCTION(BlueprintCallable, Category = "Interchange | Node | Material")
-	bool GetCustomBlendMode(TEnumAsByte<EBlendMode>& AttributeValue) const;
+	bool GetCustomBlendMode(uint8& AttributeValue) const;
  
 	UFUNCTION(BlueprintCallable, Category = "Interchange | Node | Material")
-	bool SetCustomBlendMode(const TEnumAsByte<EBlendMode>& AttributeValue, bool bAddApplyDelegate = true);
+	bool SetCustomBlendMode(const uint8& AttributeValue, bool bAddApplyDelegate = true);
 
 	UFUNCTION(BlueprintCallable, Category = "Interchange | Node | Material")
 	bool GetCustomTwoSided(bool& AttributeValue) const;
