Interchange
- Removed IMPLEMENT_NODE_ATTRIBUTE_DELEGATE since it wasn't needed anymore (default delegate is templated).
- Added IMPLEMENT_NODE_ATTRIBUTE_SETTER_WITH_CUSTOM_DELEGATE for setters that want to provide a custom delegate instead of using the default one.

diff --git a/Engine/Plugins/Experimental/Interchange/Runtime/Source/Nodes/Public/InterchangeCineCameraFactoryNode.h b/Engine/Plugins/Experimental/Interchange/Runtime/Source/Nodes/Public/InterchangeCineCameraFactoryNode.h
--- a/Engine/Plugins/Experimental/Interchange/Runtime/Source/Nodes/Public/InterchangeCineCameraFactoryNode.h
+++ b/Engine/Plugins/Experimental/Interchange/Runtime/Source/Nodes/Public/InterchangeCineCameraFactoryNode.h
@@ -26,7 +26,7 @@ public:
 	UFUNCTION(BlueprintCallable, Category = "Interchange | Node | CameraFactory")
 	bool SetCustomFocalLength(const float& AttributeValue, bool bAddApplyDelegate = true)
 	{
-		IMPLEMENT_NODE_ATTRIBUTE_SETTER(UInterchangeCineCameraFactoryNode, FocalLength, float, UCineCameraComponent);
+		IMPLEMENT_NODE_ATTRIBUTE_SETTER_WITH_CUSTOM_DELEGATE(UInterchangeCineCameraFactoryNode, FocalLength, float, UCineCameraComponent);
 	}
 
 	UFUNCTION(BlueprintCallable, Category = "Interchange | Node | CameraFactory")
@@ -38,7 +38,7 @@ public:
 	UFUNCTION(BlueprintCallable, Category = "Interchange | Node | CameraFactory")
 	bool SetCustomSensorWidth(const float& AttributeValue, bool bAddApplyDelegate = true)
 	{
-		IMPLEMENT_NODE_ATTRIBUTE_SETTER(UInterchangeCineCameraFactoryNode, SensorWidth, float, UCineCameraComponent);
+		IMPLEMENT_NODE_ATTRIBUTE_SETTER_WITH_CUSTOM_DELEGATE(UInterchangeCineCameraFactoryNode, SensorWidth, float, UCineCameraComponent);
 	}
 
 	UFUNCTION(BlueprintCallable, Category = "Interchange | Node | CameraFactory")
@@ -50,7 +50,7 @@ public:
 	UFUNCTION(BlueprintCallable, Category = "Interchange | Node | CameraFactory")
 	bool SetCustomSensorHeight(const float& AttributeValue, bool bAddApplyDelegate = true)
 	{
-		IMPLEMENT_NODE_ATTRIBUTE_SETTER(UInterchangeCineCameraFactoryNode, SensorHeight, float, UCineCameraComponent);
+		IMPLEMENT_NODE_ATTRIBUTE_SETTER_WITH_CUSTOM_DELEGATE(UInterchangeCineCameraFactoryNode, SensorHeight, float, UCineCameraComponent);
 	}
 
 private:
@@ -59,7 +59,7 @@ private:
 	const UE::Interchange::FAttributeKey Macro_CustomSensorHeightKey = UE::Interchange::FAttributeKey(TEXT("SensorHeight"));
 
 private:
-	IMPLEMENT_NODE_ATTRIBUTE_APPLY_UOBJECT_BYNAME(FocalLength, float, UCineCameraComponent, TEXT("CurrentFocalLength"));
-	IMPLEMENT_NODE_ATTRIBUTE_APPLY_UOBJECT_BYNAME(SensorWidth, float, UCineCameraComponent, TEXT("Filmback.SensorWidth"));
-	IMPLEMENT_NODE_ATTRIBUTE_APPLY_UOBJECT_BYNAME(SensorHeight, float, UCineCameraComponent, TEXT("Filmback.SensorHeight"));
+	IMPLEMENT_NODE_ATTRIBUTE_DELEGATE_BY_PROPERTYNAME(FocalLength, float, UCineCameraComponent, TEXT("CurrentFocalLength"));
+	IMPLEMENT_NODE_ATTRIBUTE_DELEGATE_BY_PROPERTYNAME(SensorWidth, float, UCineCameraComponent, TEXT("Filmback.SensorWidth"));
+	IMPLEMENT_NODE_ATTRIBUTE_DELEGATE_BY_PROPERTYNAME(SensorHeight, float, UCineCameraComponent, TEXT("Filmback.SensorHeight"));
 };
diff --git a/Engine/Plugins/Experimental/Interchange/Runtime/Source/Nodes/Public/InterchangeTexture2DArrayFactoryNode.h b/Engine/Plugins/Experimental/Interchange/Runtime/Source/Nodes/Public/InterchangeTexture2DArrayFactoryNode.h
--- a/Engine/Plugins/Experimental/Interchange/Runtime/Source/Nodes/Public/InterchangeTexture2DArrayFactoryNode.h
+++ b/Engine/Plugins/Experimental/Interchange/Runtime/Source/Nodes/Public/InterchangeTexture2DArrayFactoryNode.h
@@ -73,9 +73,4 @@ private:
 	IMPLEMENT_NODE_ATTRIBUTE_KEY(AddressY);
 	IMPLEMENT_NODE_ATTRIBUTE_KEY(AddressZ);
 
-#if WITH_EDITORONLY_DATA
-	IMPLEMENT_NODE_ATTRIBUTE_APPLY_UOBJECT(AddressX, uint8, UTexture2DArray);
-	IMPLEMENT_NODE_ATTRIBUTE_APPLY_UOBJECT(AddressY, uint8, UTexture2DArray);
-	IMPLEMENT_NODE_ATTRIBUTE_APPLY_UOBJECT(AddressZ, uint8, UTexture2DArray);
-#endif
 };
\ No newline at end of file
diff --git a/Engine/Plugins/Experimental/Interchange/Runtime/Source/Nodes/Public/InterchangeTextureFactoryNode.h b/Engine/Plugins/Experimental/Interchange/Runtime/Source/Nodes/Public/InterchangeTextureFactoryNode.h
--- a/Engine/Plugins/Experimental/Interchange/Runtime/Source/Nodes/Public/InterchangeTextureFactoryNode.h
+++ b/Engine/Plugins/Experimental/Interchange/Runtime/Source/Nodes/Public/InterchangeTextureFactoryNode.h
@@ -659,7 +659,7 @@ public:
 	UFUNCTION(BlueprintCallable, Category = "Interchange | Node | Texture")
 	bool SetCustomDownscale(const float& AttributeValue, bool bAddApplyDelegate = true)
 	{
-		IMPLEMENT_NODE_ATTRIBUTE_SETTER(UInterchangeTextureFactoryNode, Downscale, float, UTexture)
+		IMPLEMENT_NODE_ATTRIBUTE_SETTER_WITH_CUSTOM_DELEGATE(UInterchangeTextureFactoryNode, Downscale, float, UTexture)
 	}
 
 	/** Return false if the Attribute was not set previously.*/
@@ -804,52 +804,6 @@ private:
 	const UE::Interchange::FAttributeKey Macro_CustomCompositeTextureModeKey = UE::Interchange::FAttributeKey(TEXT("CompositeTextureMode"));
 	const UE::Interchange::FAttributeKey Macro_CustomCompositePowerKey = UE::Interchange::FAttributeKey(TEXT("CompositePower"));
 
-#if WITH_EDITORONLY_DATA
-	IMPLEMENT_NODE_ATTRIBUTE_APPLY_UOBJECT(AdjustBrightness, float, UTexture);
-	IMPLEMENT_NODE_ATTRIBUTE_APPLY_UOBJECT(AdjustBrightnessCurve, float, UTexture);
-	IMPLEMENT_NODE_ATTRIBUTE_APPLY_UOBJECT(AdjustVibrance, float, UTexture);
-	IMPLEMENT_NODE_ATTRIBUTE_APPLY_UOBJECT(AdjustSaturation, float, UTexture);
-	IMPLEMENT_NODE_ATTRIBUTE_APPLY_UOBJECT(AdjustRGBCurve, float, UTexture);
-	IMPLEMENT_NODE_ATTRIBUTE_APPLY_UOBJECT(AdjustHue, float, UTexture);
-	IMPLEMENT_NODE_ATTRIBUTE_APPLY_UOBJECT(AdjustMinAlpha, float, UTexture);
-	IMPLEMENT_NODE_ATTRIBUTE_APPLY_UOBJECT(AdjustMaxAlpha, float, UTexture);
-	IMPLEMENT_NODE_ATTRIBUTE_APPLY_UOBJECT(bChromaKeyTexture, bool, UTexture);
-	IMPLEMENT_NODE_ATTRIBUTE_APPLY_UOBJECT(ChromaKeyThreshold, float, UTexture);
-	IMPLEMENT_NODE_ATTRIBUTE_APPLY_UOBJECT(ChromaKeyColor, FColor, UTexture);
-
-	//Texture Compression
-	IMPLEMENT_NODE_ATTRIBUTE_APPLY_UOBJECT(CompressionNoAlpha, bool, UTexture);
-	IMPLEMENT_NODE_ATTRIBUTE_APPLY_UOBJECT(DeferCompression, bool, UTexture);
-	IMPLEMENT_NODE_ATTRIBUTE_APPLY_UOBJECT(LossyCompressionAmount, uint8, UTexture);
-	IMPLEMENT_NODE_ATTRIBUTE_APPLY_UOBJECT(MaxTextureSize, int32, UTexture);
-	IMPLEMENT_NODE_ATTRIBUTE_APPLY_UOBJECT(CompressionQuality, uint8, UTexture);
-#endif
-	IMPLEMENT_NODE_ATTRIBUTE_APPLY_UOBJECT(CompressionSettings, uint8, UTexture);
-#if WITH_EDITORONLY_DATA
-	//Texture general
-	IMPLEMENT_NODE_ATTRIBUTE_APPLY_UOBJECT(bDitherMipMapAlpha, bool, UTexture);
-	IMPLEMENT_NODE_ATTRIBUTE_APPLY_UOBJECT(AlphaCoverageThresholds, FVector4, UTexture);
-	IMPLEMENT_NODE_ATTRIBUTE_APPLY_UOBJECT(bFlipGreenChannel, bool, UTexture);
-	IMPLEMENT_NODE_ATTRIBUTE_APPLY_UOBJECT(PowerOfTwoMode, uint8, UTexture);
-	IMPLEMENT_NODE_ATTRIBUTE_APPLY_UOBJECT(PaddingColor, FColor, UTexture);
-#endif //#if WITH_EDITORONLY_DATA
-	IMPLEMENT_NODE_ATTRIBUTE_APPLY_UOBJECT(Filter, uint8, UTexture);
-	IMPLEMENT_NODE_ATTRIBUTE_APPLY_UOBJECT(MipLoadOptions, uint8, UTexture);
-	IMPLEMENT_NODE_ATTRIBUTE_APPLY_UOBJECT(SRGB, bool, UTexture);
-#if WITH_EDITORONLY_DATA
-	IMPLEMENT_NODE_ATTRIBUTE_APPLY_UOBJECT(bUseLegacyGamma, bool, UTexture);
-#endif //#if WITH_EDITORONLY_DATA
-	IMPLEMENT_NODE_ATTRIBUTE_APPLY_UOBJECT(VirtualTextureStreaming, bool, UTexture);
-	IMPLEMENT_NODE_ATTRIBUTE_APPLY_UOBJECT(AddressX, uint8, UTexture2D);
-	IMPLEMENT_NODE_ATTRIBUTE_APPLY_UOBJECT(AddressY, uint8, UTexture2D);
-#if WITH_EDITORONLY_DATA
-	//Level of Detail
-	IMPLEMENT_NODE_ATTRIBUTE_APPLY_UOBJECT(bPreserveBorder, bool, UTexture);
-	IMPLEMENT_NODE_ATTRIBUTE_APPLY_UOBJECT(MipGenSettings, uint8, UTexture);
-#endif //#if WITH_EDITORONLY_DATA
-	IMPLEMENT_NODE_ATTRIBUTE_APPLY_UOBJECT(LODBias, int32, UTexture);
-	IMPLEMENT_NODE_ATTRIBUTE_APPLY_UOBJECT(LODGroup, uint8, UTexture);
-
 	//TODO support per platform data in the FAttributeStorage, so we can set different value per platform at the pipeline stage, We set only the default value for now
 	//IMPLEMENT_NODE_ATTRIBUTE_APPLY_UOBJECT(Downscale, float, UTexture, );
 #if WITH_ENGINE
@@ -892,12 +846,6 @@ private:
 	}
 #endif //WITH_ENGINE
 
-	IMPLEMENT_NODE_ATTRIBUTE_APPLY_UOBJECT(DownscaleOptions, uint8, UTexture);
-#if WITH_EDITORONLY_DATA
-	//Compositing
-	IMPLEMENT_NODE_ATTRIBUTE_APPLY_UOBJECT(CompositeTextureMode, uint8, UTexture);
-	IMPLEMENT_NODE_ATTRIBUTE_APPLY_UOBJECT(CompositePower, float, UTexture);
-#endif //#if WITH_EDITORONLY_DATA
 protected:
 #if WITH_ENGINE
 	TSubclassOf<UTexture> AssetClass = nullptr;
diff --git a/Engine/Plugins/Experimental/Interchange/Runtime/Source/Nodes/Public/InterchangeTextureLightProfileFactoryNode.h b/Engine/Plugins/Experimental/Interchange/Runtime/Source/Nodes/Public/InterchangeTextureLightProfileFactoryNode.h
--- a/Engine/Plugins/Experimental/Interchange/Runtime/Source/Nodes/Public/InterchangeTextureLightProfileFactoryNode.h
+++ b/Engine/Plugins/Experimental/Interchange/Runtime/Source/Nodes/Public/InterchangeTextureLightProfileFactoryNode.h
@@ -71,7 +71,4 @@ private:
 	IMPLEMENT_NODE_ATTRIBUTE_KEY(Brightness);
 	IMPLEMENT_NODE_ATTRIBUTE_KEY(TextureMultiplier);
 
-	IMPLEMENT_NODE_ATTRIBUTE_APPLY_UOBJECT(Brightness, float, UTextureLightProfile);
-	IMPLEMENT_NODE_ATTRIBUTE_APPLY_UOBJECT(TextureMultiplier, float, UTextureLightProfile);
-
 };
\ No newline at end of file
diff --git a/Engine/Source/Runtime/Experimental/Interchange/Core/Public/Nodes/InterchangeBaseNode.h b/Engine/Source/Runtime/Experimental/Interchange/Core/Public/Nodes/InterchangeBaseNode.h
--- a/Engine/Source/Runtime/Experimental/Interchange/Core/Public/Nodes/InterchangeBaseNode.h
+++ b/Engine/Source/Runtime/Experimental/Interchange/Core/Public/Nodes/InterchangeBaseNode.h
@@ -97,7 +97,7 @@ namespace InterchangePrivateNodeBase
 const UE::Interchange::FAttributeKey Macro_Custom##AttributeName##Key = UE::Interchange::FAttributeKey(TEXT(#AttributeName));
 
 #if WITH_ENGINE
-#define IMPLEMENT_NODE_ATTRIBUTE_APPLY_UOBJECT_BYNAME(AttributeName, AttributeType, ObjectType, PropertyName)		\
+#define IMPLEMENT_NODE_ATTRIBUTE_DELEGATE_BY_PROPERTYNAME(AttributeName, AttributeType, ObjectType, PropertyName)		\
 bool ApplyCustom##AttributeName##ToAsset(UObject* Asset) const														\
 {																													\
 	return ApplyAttributeToObject<AttributeType>(Macro_Custom##AttributeName##Key.ToString(), Asset, PropertyName);	\
@@ -105,16 +105,12 @@ bool ApplyCustom##AttributeName##ToAsset(UObject* Asset) const														\
 																													\
 bool FillCustom##AttributeName##FromAsset(UObject* Asset)															\
 {																													\
-	return FillAttributeFromObject<AttributeType>(Macro_Custom##AttributeName##Key.ToString(), Asset, PropertyName);	\
+	return FillAttributeFromObject<AttributeType>(Macro_Custom##AttributeName##Key.ToString(), Asset, PropertyName);\
 }
 
-#define IMPLEMENT_NODE_ATTRIBUTE_APPLY_UOBJECT(AttributeName, AttributeType, ObjectType)						\
-	IMPLEMENT_NODE_ATTRIBUTE_APPLY_UOBJECT_BYNAME(AttributeName, AttributeType, ObjectType, TEXT(#AttributeName))
-
 #else //#if WITH_ENGINE
 
-#define IMPLEMENT_NODE_ATTRIBUTE_APPLY_UOBJECT(AttributeName, AttributeType, AssetType)
-#define IMPLEMENT_NODE_ATTRIBUTE_APPLY_UOBJECT_BYNAME(AttributeName, AttributeType, AssetType, PropertyName)
+#define IMPLEMENT_NODE_ATTRIBUTE_DELEGATE_BY_PROPERTYNAME(AttributeName, AttributeType, AssetType, PropertyName)
 
 #endif //#if WITH_ENGINE
 
@@ -140,12 +136,30 @@ bool FillCustom##AttributeName##FromAsset(UObject* Asset)															\
 	}																																				\
 	return false;
 
+#define IMPLEMENT_NODE_ATTRIBUTE_SETTER_WITH_CUSTOM_DELEGATE(NodeClassName, AttributeName, AttributeType, AssetType)								\
+	FString OperationName = GetTypeName() + TEXT(".Set" #AttributeName);																			\
+	if(InterchangePrivateNodeBase::SetCustomAttribute<AttributeType>(*Attributes, Macro_Custom##AttributeName##Key, OperationName, AttributeValue))	\
+	{																																				\
+		if(bAddApplyDelegate)																														\
+		{																																			\
+			TArray<UE::Interchange::FApplyAttributeToAsset>& Delegates = ApplyCustomAttributeDelegates.FindOrAdd(AssetType::StaticClass());			\
+			Delegates.Add(UE::Interchange::FApplyAttributeToAsset::CreateUObject(this, &NodeClassName::ApplyCustom##AttributeName##ToAsset));		\
+			TArray<UE::Interchange::FFillAttributeToAsset>& FillDelegates = FillCustomAttributeDelegates.FindOrAdd(AssetType::StaticClass());		\
+			FillDelegates.Add(UE::Interchange::FFillAttributeToAsset::CreateUObject(this, &NodeClassName::FillCustom##AttributeName##FromAsset));	\
+		}																																			\
+		return true;																																\
+	}																																				\
+	return false;
+
 #else //#if WITH_ENGINE
 
 #define IMPLEMENT_NODE_ATTRIBUTE_SETTER(NodeClassName, AttributeName, AttributeType, AssetType)															\
 	FString OperationName = GetTypeName() + TEXT(".Set" #AttributeName);																				\
 	return InterchangePrivateNodeBase::SetCustomAttribute<AttributeType>(*Attributes, Macro_Custom##AttributeName##Key, OperationName, AttributeValue);
 
+#define IMPLEMENT_NODE_ATTRIBUTE_SETTER_WITH_CUSTOM_DELEGATE(NodeClassName, AttributeName, AttributeType, AssetType)								\
+	IMPLEMENT_NODE_ATTRIBUTE_SETTER(NodeClassName, AttributeName, AttributeType, AssetType)
+
 #endif //#if WITH_ENGINE
 
 #define INTERCHANGE_BASE_NODE_ADD_ATTRIBUTE(ValueType)																									\
