InterchangeBasePipeline
- Save and load settings for its subobject pipelines.

Interchange Materials
- Added a UInterchangeMaterialInstanceFactoryNode to allow the Material Factory to create material instances.
- Added a ParentMaterial setting to UInterchangeGenericMaterialPipeline. Specifies which parent material to use when creating instances. This setting is optional. If left empty, the pipeline will choose between a Phong, Lambert or PBR parent material.

diff --git a/Engine/Plugins/Experimental/Interchange/Runtime/Interchange.uplugin b/Engine/Plugins/Experimental/Interchange/Runtime/Interchange.uplugin
--- a/Engine/Plugins/Experimental/Interchange/Runtime/Interchange.uplugin
+++ b/Engine/Plugins/Experimental/Interchange/Runtime/Interchange.uplugin
@@ -11,7 +11,7 @@
 	"MarketplaceURL": "",
 	"SupportURL": "",
 	"EnabledByDefault": false,
-	"CanContainContent": false,
+	"CanContainContent": true,
 	"IsBetaVersion": false,
 	"Installed": false,
 	"IsHidden": true,
diff --git a/Engine/Plugins/Experimental/Interchange/Runtime/Source/Import/Private/Material/InterchangeMaterialFactory.cpp b/Engine/Plugins/Experimental/Interchange/Runtime/Source/Import/Private/Material/InterchangeMaterialFactory.cpp
--- a/Engine/Plugins/Experimental/Interchange/Runtime/Source/Import/Private/Material/InterchangeMaterialFactory.cpp
+++ b/Engine/Plugins/Experimental/Interchange/Runtime/Source/Import/Private/Material/InterchangeMaterialFactory.cpp
@@ -15,6 +15,8 @@
 #include "Materials/MaterialExpressionTextureSample.h"
 #include "Materials/MaterialExpressionVectorParameter.h"
 #include "Materials/MaterialInterface.h"
+#include "Materials/MaterialInstanceConstant.h"
+#include "Materials/MaterialInstanceDynamic.h"
 #include "Nodes/InterchangeBaseNode.h"
 #include "Nodes/InterchangeBaseNodeContainer.h"
 
@@ -153,7 +155,6 @@ UObject* UInterchangeMaterialFactory::CreateEmptyAsset(const FCreateAssetParams&
 {
 	UObject* Material = nullptr;
 
-#if WITH_EDITOR
 	if (!Arguments.AssetNode || !Arguments.AssetNode->GetObjectClass()->IsChildOf(GetFactoryClass()))
 	{
 		return nullptr;
@@ -177,7 +178,22 @@ UObject* UInterchangeMaterialFactory::CreateEmptyAsset(const FCreateAssetParams&
 	// create a new material or overwrite existing asset, if possible
 	if (!ExistingAsset)
 	{
-		Material = NewObject<UObject>(Arguments.Parent, MaterialClass, *Arguments.AssetName, RF_Public | RF_Standalone);
+		if (MaterialClass->IsChildOf<UMaterialInstanceDynamic>())
+		{
+			if (const UInterchangeMaterialInstanceFactoryNode* MaterialInstanceFactoryNode = Cast<UInterchangeMaterialInstanceFactoryNode>(MaterialFactoryNode))
+			{
+				FString ParentPath;
+				if (MaterialInstanceFactoryNode->GetCustomParent(ParentPath))
+				{
+					FSoftObjectPath ParentMaterial(ParentPath);
+					Material = UMaterialInstanceDynamic::Create(Cast<UMaterialInterface>(ParentMaterial.TryLoad()), Arguments.Parent);
+				}
+			}
+		}
+		else
+		{
+			Material = NewObject<UObject>(Arguments.Parent, MaterialClass, *Arguments.AssetName, RF_Public | RF_Standalone);
+		}
 	}
 	else if (ExistingAsset->GetClass()->IsChildOf(MaterialClass))
 	{
@@ -190,8 +206,22 @@ UObject* UInterchangeMaterialFactory::CreateEmptyAsset(const FCreateAssetParams&
 		UE_LOG(LogInterchangeImport, Warning, TEXT("Could not create Material asset %s"), *Arguments.AssetName);
 		return nullptr;
 	}
+
+#if WITH_EDITOR
 	Material->PreEditChange(nullptr);
 
+	if (UMaterialInstanceConstant* MaterialInstanceConstant = Cast<UMaterialInstanceConstant>(Material))
+	{
+		if (const UInterchangeMaterialInstanceFactoryNode* MaterialInstanceFactoryNode = Cast<UInterchangeMaterialInstanceFactoryNode>(MaterialFactoryNode))
+		{
+			FString ParentPath;
+			if (MaterialInstanceFactoryNode->GetCustomParent(ParentPath))
+			{
+				FSoftObjectPath ParentMaterial(ParentPath);
+				MaterialInstanceConstant->SetParentEditorOnly(Cast<UMaterialInterface>(ParentMaterial.TryLoad()));
+			}
+		}
+	}
 #endif //WITH_EDITOR
 
 	return Material;
@@ -199,18 +229,11 @@ UObject* UInterchangeMaterialFactory::CreateEmptyAsset(const FCreateAssetParams&
 
 UObject* UInterchangeMaterialFactory::CreateAsset(const FCreateAssetParams& Arguments)
 {
-#if !WITH_EDITOR
-
-	UE_LOG(LogInterchangeImport, Error, TEXT("Cannot import Material asset in runtime, this is an editor only feature."));
-	return nullptr;
-
-#else
-
 	if (!Arguments.AssetNode || !Arguments.AssetNode->GetObjectClass()->IsChildOf(GetFactoryClass()))
 	{
 		return nullptr;
 	}
-	
+
 	const UInterchangeBaseMaterialFactoryNode* MaterialFactoryNode = Cast<UInterchangeBaseMaterialFactoryNode>(Arguments.AssetNode);
 	if (MaterialFactoryNode == nullptr)
 	{
@@ -252,10 +275,12 @@ UObject* UInterchangeMaterialFactory::CreateAsset(const FCreateAssetParams& Argu
 		{
 			if (UMaterialInterface* MaterialInterface = Cast<UMaterialInterface>(MaterialObject))
 			{
+#if WITH_EDITOR
 				if (UMaterial* Material = Cast<UMaterial>(MaterialObject))
 				{
 					SetupMaterial(Material, Arguments, MaterialFactoryNode);
 				}
+#endif // #if WITH_EDITOR
 
 				MaterialFactoryNode->ApplyAllCustomAttributeToObject(MaterialObject);
 			}
@@ -268,7 +293,6 @@ UObject* UInterchangeMaterialFactory::CreateAsset(const FCreateAssetParams& Argu
 	}
 
 	return MaterialObject;
-#endif // WITH_EDITOR
 }
 
 /* This function is call in the completion task on the main thread, use it to call main thread post creation step for your assets*/
@@ -285,7 +309,7 @@ void UInterchangeMaterialFactory::PreImportPreCompletedCallback(const FImportPre
 		UMaterialInterface* ImportedMaterialInterface = CastChecked<UMaterialInterface>(Arguments.ImportedObject);
 
 		//Update the samplers type in case the textures were changed during their PreImportPreCompletedCallback
-		if (UMaterial* ImportedMaterial = CastChecked<UMaterial>(ImportedMaterialInterface))
+		if (UMaterial* ImportedMaterial = Cast<UMaterial>(ImportedMaterialInterface))
 		{
 			for (UMaterialExpression* Expression : ImportedMaterial->Expressions)
 			{
@@ -295,6 +319,11 @@ void UInterchangeMaterialFactory::PreImportPreCompletedCallback(const FImportPre
 				}
 			}
 		}
+		else if (UMaterialInstance* MaterialInstance = Cast<UMaterialInstance>(ImportedMaterialInterface))
+		{
+			// Material instances expect their parameters to only be updated from the game thread
+			SetupMaterialInstance(MaterialInstance, Arguments.NodeContainer, Cast<UInterchangeBaseMaterialFactoryNode>(Arguments.FactoryNode));
+		}
 
 		UE::Interchange::FFactoryCommon::FUpdateImportAssetDataParameters UpdateImportAssetDataParameters(ImportedMaterialInterface
 																										  , ImportedMaterialInterface->AssetImportData
@@ -629,4 +658,107 @@ UMaterialExpression* UInterchangeMaterialFactory::CreateExpressionsForNode(UMate
 }
 #endif // #if WITH_EDITOR
 
+void UInterchangeMaterialFactory::SetupMaterialInstance(UMaterialInstance* MaterialInstance, const UInterchangeBaseNodeContainer* NodeContainer, const UInterchangeBaseMaterialFactoryNode* MaterialFactoryNode)
+{
+	if (!MaterialFactoryNode || !NodeContainer)
+	{
+		return;
+	}
+
+	TArray<FString> Inputs;
+	UInterchangeShaderPortsAPI::GatherInputs(MaterialFactoryNode, Inputs);
+
+	for (const FString& InputName : Inputs)
+	{
+		const FName ParameterName = *InputName;
+
+		switch(UInterchangeShaderPortsAPI::GetInputType(MaterialFactoryNode, InputName))
+		{
+		case UE::Interchange::EAttributeTypes::Float:
+			{
+				float InstanceValue;
+				if (MaterialInstance->GetScalarParameterValue(ParameterName, InstanceValue))
+				{
+					float InputValue = 0.f;
+					MaterialFactoryNode->GetFloatAttribute(UInterchangeShaderPortsAPI::MakeInputValueKey(InputName), InputValue);
+
+					if ( !FMath::IsNearlyEqual(InputValue, InstanceValue) )
+					{
+#if WITH_EDITOR
+						if (UMaterialInstanceConstant* MaterialInstanceConstant = Cast<UMaterialInstanceConstant>(MaterialInstance))
+						{
+							MaterialInstanceConstant->SetScalarParameterValueEditorOnly(ParameterName, InputValue);
+						}
+						else
+#endif // #if WITH_EDITOR
+						if (UMaterialInstanceDynamic* MaterialInstanceDynamic = Cast<UMaterialInstanceDynamic>(MaterialInstance))
+						{
+							MaterialInstanceDynamic->SetScalarParameterValue(ParameterName, InputValue);
+						}
+					}
+				}
+			}
+			break;
+		case UE::Interchange::EAttributeTypes::LinearColor:
+			{
+				FLinearColor InstanceValue;
+				if (MaterialInstance->GetVectorParameterValue(ParameterName, InstanceValue))
+				{
+					FLinearColor InputValue;
+					if (MaterialFactoryNode->GetLinearColorAttribute(UInterchangeShaderPortsAPI::MakeInputValueKey(InputName), InputValue))
+					{
+						if ( !InputValue.Equals(InstanceValue) )
+						{
+#if WITH_EDITOR
+							if (UMaterialInstanceConstant* MaterialInstanceConstant = Cast<UMaterialInstanceConstant>(MaterialInstance))
+							{
+								MaterialInstanceConstant->SetVectorParameterValueEditorOnly(ParameterName, InputValue);
+							}
+							else
+#endif // #if WITH_EDITOR
+							if (UMaterialInstanceDynamic* MaterialInstanceDynamic = Cast<UMaterialInstanceDynamic>(MaterialInstance))
+							{
+								MaterialInstanceDynamic->SetVectorParameterValue(ParameterName, InputValue);
+							}
+						}
+					}
+				}
+			}
+			break;
+		case UE::Interchange::EAttributeTypes::String:
+			{
+				UTexture* InstanceValue;
+				if (MaterialInstance->GetTextureParameterValue(ParameterName, InstanceValue))
+				{
+					FString InputValue;
+					if (MaterialFactoryNode->GetStringAttribute(UInterchangeShaderPortsAPI::MakeInputValueKey(InputName), InputValue))
+					{
+						if (const UInterchangeTextureFactoryNode* TextureFactoryNode = Cast<UInterchangeTextureFactoryNode>(NodeContainer->GetNode(InputValue)))
+						{
+							if (UTexture* InputTexture = Cast<UTexture>(TextureFactoryNode->ReferenceObject.TryLoad()))
+							{
+								if ( InputTexture != InstanceValue )
+								{
+#if WITH_EDITOR
+									if (UMaterialInstanceConstant* MaterialInstanceConstant = Cast<UMaterialInstanceConstant>(MaterialInstance))
+									{
+										MaterialInstanceConstant->SetTextureParameterValueEditorOnly(ParameterName, InputTexture);
+									}
+									else
+#endif // #if WITH_EDITOR
+									if (UMaterialInstanceDynamic* MaterialInstanceDynamic = Cast<UMaterialInstanceDynamic>(MaterialInstance))
+									{
+										MaterialInstanceDynamic->SetTextureParameterValue(ParameterName, InputTexture);
+									}
+								}
+							}
+						}
+					}
+				}
+			}
+			break;
+		}
+	}
+}
+
 #undef LOCTEXT_NAMESPACE
diff --git a/Engine/Plugins/Experimental/Interchange/Runtime/Source/Import/Private/Mesh/InterchangeStaticMeshFactory.cpp b/Engine/Plugins/Experimental/Interchange/Runtime/Source/Import/Private/Mesh/InterchangeStaticMeshFactory.cpp
--- a/Engine/Plugins/Experimental/Interchange/Runtime/Source/Import/Private/Mesh/InterchangeStaticMeshFactory.cpp
+++ b/Engine/Plugins/Experimental/Interchange/Runtime/Source/Import/Private/Mesh/InterchangeStaticMeshFactory.cpp
@@ -318,7 +318,7 @@ UObject* UInterchangeStaticMeshFactory::CreateAsset(const FCreateAssetParams& Ar
 
 		for (int32 DependencyIndex = 0; DependencyIndex < FactoryDependencies.Num(); ++DependencyIndex)
 		{
-			const UInterchangeMaterialFactoryNode* MaterialFactoryNode = Cast<UInterchangeMaterialFactoryNode>(Arguments.NodeContainer->GetNode(FactoryDependencies[DependencyIndex]));
+			const UInterchangeBaseMaterialFactoryNode* MaterialFactoryNode = Cast<UInterchangeBaseMaterialFactoryNode>(Arguments.NodeContainer->GetNode(FactoryDependencies[DependencyIndex]));
 			if (!MaterialFactoryNode || !MaterialFactoryNode->ReferenceObject.IsValid())
 			{
 				continue;
diff --git a/Engine/Plugins/Experimental/Interchange/Runtime/Source/Import/Public/Material/InterchangeMaterialFactory.h b/Engine/Plugins/Experimental/Interchange/Runtime/Source/Import/Public/Material/InterchangeMaterialFactory.h
--- a/Engine/Plugins/Experimental/Interchange/Runtime/Source/Import/Public/Material/InterchangeMaterialFactory.h
+++ b/Engine/Plugins/Experimental/Interchange/Runtime/Source/Import/Public/Material/InterchangeMaterialFactory.h
@@ -13,6 +13,7 @@ class UInterchangeBaseMaterialFactoryNode;
 class UInterchangeMaterialExpressionFactoryNode;
 class UMaterial;
 class UMaterialExpression;
+class UMaterialInstance;
 
 UCLASS(BlueprintType, Experimental)
 class INTERCHANGEIMPORT_API UInterchangeMaterialFactory : public UInterchangeFactoryBase
@@ -39,6 +40,8 @@ private:
 
 	UMaterialExpression* CreateExpression(UMaterial* Material, const UInterchangeFactoryBase::FCreateAssetParams& Arguments, const UInterchangeMaterialExpressionFactoryNode& ExpressionNode);
 #endif // #if WITH_EDITOR
+
+	void SetupMaterialInstance(UMaterialInstance* MaterialInstance, const UInterchangeBaseNodeContainer* NodeContainer, const UInterchangeBaseMaterialFactoryNode* MaterialFactoryNode);
 };
 
 
diff --git a/Engine/Plugins/Experimental/Interchange/Runtime/Source/Nodes/Private/InterchangeMaterialFactoryNode.cpp b/Engine/Plugins/Experimental/Interchange/Runtime/Source/Nodes/Private/InterchangeMaterialFactoryNode.cpp
--- a/Engine/Plugins/Experimental/Interchange/Runtime/Source/Nodes/Private/InterchangeMaterialFactoryNode.cpp
+++ b/Engine/Plugins/Experimental/Interchange/Runtime/Source/Nodes/Private/InterchangeMaterialFactoryNode.cpp
@@ -6,6 +6,7 @@
 
 #include "Materials/Material.h"
 #include "Materials/MaterialInstance.h"
+#include "Materials/MaterialInstanceConstant.h"
 #include "Materials/MaterialInterface.h"
 
 #endif
@@ -221,3 +222,48 @@ bool UInterchangeMaterialExpressionFactoryNode::SetCustomExpressionClassName(con
 {
 	IMPLEMENT_NODE_ATTRIBUTE_SETTER_NODELEGATE(ExpressionClassName, FString);
 }
+
+FString UInterchangeMaterialInstanceFactoryNode::GetTypeName() const
+{
+	const FString TypeName = TEXT("MaterialInstanceFactoryNode");
+	return TypeName;
+}
+
+UClass* UInterchangeMaterialInstanceFactoryNode::GetObjectClass() const
+{
+#if WITH_ENGINE
+	FString InstanceClassName;
+	if (GetCustomInstanceClassName(InstanceClassName))
+	{
+		UClass* InstanceClass = FindObject<UClass>(nullptr, *InstanceClassName);
+		if (InstanceClass->IsChildOf<UMaterialInstance>())
+		{
+			return InstanceClass;
+		}
+	}
+
+	return UMaterialInstanceConstant::StaticClass();
+#else
+	return nullptr;
+#endif
+}
+
+bool UInterchangeMaterialInstanceFactoryNode::GetCustomInstanceClassName(FString& AttributeValue) const
+{
+	IMPLEMENT_NODE_ATTRIBUTE_GETTER(InstanceClassName, FString);
+}
+
+bool UInterchangeMaterialInstanceFactoryNode::SetCustomInstanceClassName(const FString& AttributeValue)
+{
+	IMPLEMENT_NODE_ATTRIBUTE_SETTER_NODELEGATE(InstanceClassName, FString);
+}
+
+bool UInterchangeMaterialInstanceFactoryNode::GetCustomParent(FString& AttributeValue) const
+{
+	IMPLEMENT_NODE_ATTRIBUTE_GETTER(Parent, FString);
+}
+
+bool UInterchangeMaterialInstanceFactoryNode::SetCustomParent(const FString& AttributeValue)
+{
+	IMPLEMENT_NODE_ATTRIBUTE_SETTER_NODELEGATE(Parent, FString);
+}
diff --git a/Engine/Plugins/Experimental/Interchange/Runtime/Source/Nodes/Public/InterchangeMaterialFactoryNode.h b/Engine/Plugins/Experimental/Interchange/Runtime/Source/Nodes/Public/InterchangeMaterialFactoryNode.h
--- a/Engine/Plugins/Experimental/Interchange/Runtime/Source/Nodes/Public/InterchangeMaterialFactoryNode.h
+++ b/Engine/Plugins/Experimental/Interchange/Runtime/Source/Nodes/Public/InterchangeMaterialFactoryNode.h
@@ -171,3 +171,30 @@ private:
 
 };
 
+UCLASS(BlueprintType, Experimental)
+class INTERCHANGENODES_API UInterchangeMaterialInstanceFactoryNode : public UInterchangeBaseMaterialFactoryNode
+{
+	GENERATED_BODY()
+
+public:
+	virtual FString GetTypeName() const override;
+	virtual UClass* GetObjectClass() const override;
+
+public:
+	UFUNCTION(BlueprintCallable, Category = "Interchange | Node | MaterialInstanceFactory")
+	bool GetCustomInstanceClassName(FString& AttributeValue) const;
+
+	UFUNCTION(BlueprintCallable, Category = "Interchange | Node | MaterialInstanceFactory")
+	bool SetCustomInstanceClassName(const FString& AttributeValue);
+
+	UFUNCTION(BlueprintCallable, Category = "Interchange | Node | MaterialInstanceFactory")
+	bool GetCustomParent(FString& AttributeValue) const;
+
+	UFUNCTION(BlueprintCallable, Category = "Interchange | Node | MaterialInstanceFactory")
+	bool SetCustomParent(const FString& AttributeValue);
+
+private:
+	const UE::Interchange::FAttributeKey Macro_CustomInstanceClassNameKey = UE::Interchange::FAttributeKey(TEXT("InstanceClassName"));
+	const UE::Interchange::FAttributeKey Macro_CustomParentKey = UE::Interchange::FAttributeKey(TEXT("Parent"));
+	
+};
diff --git a/Engine/Plugins/Experimental/Interchange/Runtime/Source/Pipelines/Private/InterchangeGenericMaterialPipeline.cpp b/Engine/Plugins/Experimental/Interchange/Runtime/Source/Pipelines/Private/InterchangeGenericMaterialPipeline.cpp
--- a/Engine/Plugins/Experimental/Interchange/Runtime/Source/Pipelines/Private/InterchangeGenericMaterialPipeline.cpp
+++ b/Engine/Plugins/Experimental/Interchange/Runtime/Source/Pipelines/Private/InterchangeGenericMaterialPipeline.cpp
@@ -25,6 +25,8 @@
 #include "Materials/MaterialExpressionTextureSampleParameter2DArray.h"
 #include "Materials/MaterialExpressionTextureSampleParameterCube.h"
 #include "Materials/MaterialExpressionVectorParameter.h"
+#include "Materials/MaterialInstanceConstant.h"
+#include "Materials/MaterialInstanceDynamic.h"
 #include "Misc/Paths.h"
 #include "Nodes/InterchangeBaseNode.h"
 #include "Nodes/InterchangeBaseNodeContainer.h"
@@ -63,7 +65,6 @@ void UInterchangeGenericMaterialPipeline::ExecutePreImportPipeline(UInterchangeB
 		}
 	});
 
-	//import materials
 	if (MaterialImport == EInterchangeMaterialImportOption::ImportAsMaterials)
 	{
 		for (const UInterchangeShaderGraphNode* ShaderGraphNode : MaterialNodes)
@@ -75,6 +76,17 @@ void UInterchangeGenericMaterialPipeline::ExecutePreImportPipeline(UInterchangeB
 			}
 		}
 	}
+	else if (MaterialImport == EInterchangeMaterialImportOption::ImportAsMaterialInstances)
+	{
+		for (const UInterchangeShaderGraphNode* ShaderGraphNode : MaterialNodes)
+		{
+			if (UInterchangeMaterialInstanceFactoryNode* MaterialInstanceFactoryNode = CreateMaterialInstanceFactoryNode(ShaderGraphNode))
+			{
+				//By default we do not create the materials, every node with mesh attribute can enable them. So we wont create unused materials.
+				MaterialInstanceFactoryNode->SetEnabled(false);
+			}
+		}
+	}
 
 	return;
 }
@@ -111,14 +123,39 @@ UInterchangeBaseMaterialFactoryNode* UInterchangeGenericMaterialPipeline::Create
 	return MaterialFactoryNode;
 }
 
-bool UInterchangeGenericMaterialPipeline::HandlePhongModel(const UInterchangeShaderGraphNode* ShaderGraphNode, UInterchangeMaterialFactoryNode* MaterialFactoryNode)
+bool UInterchangeGenericMaterialPipeline::IsPBRModel(const UInterchangeShaderGraphNode* ShaderGraphNode) const
+{
+	using namespace UE::Interchange::Materials::PBR;
+
+	const bool bHasBaseColorInput = UInterchangeShaderPortsAPI::HasInput(ShaderGraphNode, Parameters::BaseColor);
+
+	return bHasBaseColorInput;
+}
+
+bool UInterchangeGenericMaterialPipeline::IsPhongModel(const UInterchangeShaderGraphNode* ShaderGraphNode) const
 {
 	using namespace UE::Interchange::Materials::Phong;
 
 	const bool bHasDiffuseInput = UInterchangeShaderPortsAPI::HasInput(ShaderGraphNode, Parameters::DiffuseColor);
 	const bool bHasSpecularInput = UInterchangeShaderPortsAPI::HasInput(ShaderGraphNode, Parameters::SpecularColor);
 
-	if (bHasDiffuseInput && bHasSpecularInput)
+	return bHasDiffuseInput && bHasSpecularInput;
+}
+
+bool UInterchangeGenericMaterialPipeline::IsLambertModel(const UInterchangeShaderGraphNode* ShaderGraphNode) const
+{
+	using namespace UE::Interchange::Materials::Lambert;
+
+	const bool bHasDiffuseInput = UInterchangeShaderPortsAPI::HasInput(ShaderGraphNode, Parameters::DiffuseColor);
+
+	return bHasDiffuseInput;
+}
+
+bool UInterchangeGenericMaterialPipeline::HandlePhongModel(const UInterchangeShaderGraphNode* ShaderGraphNode, UInterchangeMaterialFactoryNode* MaterialFactoryNode)
+{
+	using namespace UE::Interchange::Materials::Phong;
+
+	if (IsPhongModel(ShaderGraphNode))
 	{
 		// ConvertFromDiffSpec function call
 		UInterchangeMaterialExpressionFactoryNode* FunctionCallExpression = NewObject<UInterchangeMaterialExpressionFactoryNode>(BaseNodeContainer, NAME_None);
@@ -195,10 +232,7 @@ bool UInterchangeGenericMaterialPipeline::HandleLambertModel(const UInterchangeS
 {
 	using namespace UE::Interchange::Materials::Lambert;
 
-	// Diffuse
-	const bool bHasDiffuseInput = UInterchangeShaderPortsAPI::HasInput(ShaderGraphNode, Parameters::DiffuseColor);
-
-	if (bHasDiffuseInput)
+	if (IsLambertModel(ShaderGraphNode))
 	{
 		TTuple<UInterchangeMaterialExpressionFactoryNode*, FString> DiffuseExpressionFactoryNode =
 			CreateMaterialExpressionForInput(MaterialFactoryNode, ShaderGraphNode, Parameters::DiffuseColor.ToString(), MaterialFactoryNode->GetUniqueID());
@@ -641,8 +675,7 @@ UInterchangeMaterialExpressionFactoryNode* UInterchangeGenericMaterialPipeline::
 TTuple<UInterchangeMaterialExpressionFactoryNode*, FString> UInterchangeGenericMaterialPipeline::CreateMaterialExpressionForInput(UInterchangeMaterialFactoryNode* MaterialFactoryNode, const UInterchangeShaderNode* ShaderNode, const FString& InputName, const FString& ParentUid)
 {
 	// If we have a connection
-	// - Create material expression for connected shader node
-	// - Connected material expression to something?
+	// - Create material expression for the connected shader node
 	//
 	// If we don't have a connection
 	// - Create material expression for the input value
@@ -660,7 +693,6 @@ TTuple<UInterchangeMaterialExpressionFactoryNode*, FString> UInterchangeGenericM
 	}
 	else
 	{
-		FString MaterialExpressionUid = ShaderNode->GetUniqueID() + TEXT(".") + InputName;
 		switch(UInterchangeShaderPortsAPI::GetInputType(ShaderNode, InputName))
 		{
 		case UE::Interchange::EAttributeTypes::Float:
@@ -692,3 +724,180 @@ UInterchangeMaterialFactoryNode* UInterchangeGenericMaterialPipeline::CreateMate
 	return MaterialFactoryNode;
 }
 
+UInterchangeMaterialInstanceFactoryNode* UInterchangeGenericMaterialPipeline::CreateMaterialInstanceFactoryNode(const UInterchangeShaderGraphNode* ShaderGraphNode)
+{
+	UInterchangeMaterialInstanceFactoryNode* MaterialInstanceFactoryNode =
+		Cast<UInterchangeMaterialInstanceFactoryNode>( CreateBaseMaterialFactoryNode(ShaderGraphNode, UInterchangeMaterialInstanceFactoryNode::StaticClass()) );
+
+	if (UMaterialInterface* ParentMaterialObj = Cast<UMaterialInterface>(ParentMaterial.TryLoad()))
+	{
+		MaterialInstanceFactoryNode->SetCustomParent(ParentMaterialObj->GetPathName());
+	}
+	else if (IsPBRModel(ShaderGraphNode))
+	{
+		MaterialInstanceFactoryNode->SetCustomParent(TEXT("Material'/Interchange/Materials/PBRSurfaceMaterial.PBRSurfaceMaterial'"));
+	}
+	else if (IsPhongModel(ShaderGraphNode))
+	{
+		MaterialInstanceFactoryNode->SetCustomParent(TEXT("Material'/Interchange/Materials/PhongSurfaceMaterial.PhongSurfaceMaterial'"));
+	}
+	else if (IsLambertModel(ShaderGraphNode))
+	{
+		MaterialInstanceFactoryNode->SetCustomParent(TEXT("Material'/Interchange/Materials/LambertSurfaceMaterial.LambertSurfaceMaterial'"));
+	}
+	else
+	{
+		// Default to PBR
+		MaterialInstanceFactoryNode->SetCustomParent(TEXT("Material'/Interchange/Materials/PBRSurfaceMaterial.PBRSurfaceMaterial'"));
+	}
+
+#if WITH_EDITOR
+	MaterialInstanceFactoryNode->SetCustomInstanceClassName(UMaterialInstanceConstant::StaticClass()->GetPathName());
+#else
+	MaterialInstanceFactoryNode->SetCustomInstanceClassName(UMaterialInstanceDynamic::StaticClass()->GetPathName());
+#endif
+
+	TArray<FString> Inputs;
+	UInterchangeShaderPortsAPI::GatherInputs(ShaderGraphNode, Inputs);
+
+	for (const FString& InputName : Inputs)
+	{
+		TVariant<FString, FLinearColor, float> InputValue;
+
+		FString ConnectedShaderNodeUid;
+		FString OutputName;
+		if (UInterchangeShaderPortsAPI::GetInputConnection(ShaderGraphNode, InputName, ConnectedShaderNodeUid, OutputName))
+		{
+			if (UInterchangeShaderNode* ConnectedShaderNode = Cast<UInterchangeShaderNode>(BaseNodeContainer->GetNode(ConnectedShaderNodeUid)))
+			{
+				InputValue = VisitShaderNode(ConnectedShaderNode);
+			}
+		}
+		else
+		{
+			switch(UInterchangeShaderPortsAPI::GetInputType(ShaderGraphNode, InputName))
+			{
+			case UE::Interchange::EAttributeTypes::Float:
+				{
+					float AttributeValue = 0.f;
+					ShaderGraphNode->GetFloatAttribute(UInterchangeShaderPortsAPI::MakeInputValueKey(InputName), AttributeValue);
+					InputValue.Set<float>(AttributeValue);
+				}
+				break;
+			case UE::Interchange::EAttributeTypes::LinearColor:
+				{
+					FLinearColor AttributeValue = FLinearColor::White;
+					ShaderGraphNode->GetLinearColorAttribute(UInterchangeShaderPortsAPI::MakeInputValueKey(InputName), AttributeValue);
+					InputValue.Set<FLinearColor>(AttributeValue);
+				}
+				break;
+			}
+		}
+
+		if (InputValue.IsType<float>())
+		{
+			MaterialInstanceFactoryNode->AddFloatAttribute(UInterchangeShaderPortsAPI::MakeInputValueKey(InputName), InputValue.Get<float>());
+		}
+		else if (InputValue.IsType<FLinearColor>())
+		{
+			MaterialInstanceFactoryNode->AddLinearColorAttribute(UInterchangeShaderPortsAPI::MakeInputValueKey(InputName), InputValue.Get<FLinearColor>());
+		}
+		else if (InputValue.IsType<FString>())
+		{
+			const FString MapName(InputName + TEXT("Map"));
+			MaterialInstanceFactoryNode->AddStringAttribute(UInterchangeShaderPortsAPI::MakeInputValueKey(MapName), InputValue.Get<FString>());
+
+			const FString MapWeightName(MapName + TEXT("Weight"));
+			MaterialInstanceFactoryNode->AddFloatAttribute(UInterchangeShaderPortsAPI::MakeInputValueKey(MapWeightName), 1.f);
+		}
+	}
+
+	return MaterialInstanceFactoryNode;
+}
+
+TVariant<FString, FLinearColor, float> UInterchangeGenericMaterialPipeline::VisitShaderNode(const UInterchangeShaderNode* ShaderNode)
+{
+	using namespace UE::Interchange::Materials::Standard::Nodes;
+
+	TVariant<FString, FLinearColor, float> Result;
+
+	FString ShaderType;
+	if (ShaderNode->GetCustomShaderType(ShaderType))
+	{
+		if (*ShaderType == TextureSample::Name)
+		{
+			return VisitTextureSampleNode(ShaderNode);
+		}
+	}
+
+	{
+		TArray<FString> Inputs;
+		UInterchangeShaderPortsAPI::GatherInputs(ShaderNode, Inputs);
+
+		if (Inputs.Num() > 0)
+		{
+			const FString& InputName = Inputs[0];
+			FString ConnectedShaderNodeUid;
+			FString OutputName;
+			if (UInterchangeShaderPortsAPI::GetInputConnection(ShaderNode, InputName, ConnectedShaderNodeUid, OutputName))
+			{
+				if (UInterchangeShaderNode* ConnectedShaderNode = Cast<UInterchangeShaderNode>(BaseNodeContainer->GetNode(ConnectedShaderNodeUid)))
+				{
+					Result = VisitShaderNode(ConnectedShaderNode);
+				}
+			}
+			else
+			{
+				switch(UInterchangeShaderPortsAPI::GetInputType(ShaderNode, InputName))
+				{
+				case UE::Interchange::EAttributeTypes::Float:
+					{
+						float InputValue = 0.f;
+						ShaderNode->GetFloatAttribute(UInterchangeShaderPortsAPI::MakeInputValueKey(InputName), InputValue);
+						Result.Set<float>(InputValue);
+					}
+					break;
+				case UE::Interchange::EAttributeTypes::LinearColor:
+					{
+						FLinearColor InputValue = FLinearColor::White;
+						ShaderNode->GetLinearColorAttribute(UInterchangeShaderPortsAPI::MakeInputValueKey(InputName), InputValue);
+						Result.Set<FLinearColor>(InputValue);
+					}
+					break;
+				}
+			}
+		}
+	}
+
+	return Result;
+}
+
+TVariant<FString, FLinearColor, float> UInterchangeGenericMaterialPipeline::VisitTextureSampleNode(const UInterchangeShaderNode* ShaderNode)
+{
+	using namespace UE::Interchange::Materials::Standard::Nodes;
+
+	TVariant<FString, FLinearColor, float> Result;
+
+	FString TextureUid;
+	if (ShaderNode->GetStringAttribute(UInterchangeShaderPortsAPI::MakeInputValueKey(TextureSample::Inputs::Texture.ToString()), TextureUid))
+	{
+		if (!TextureUid.IsEmpty())
+		{
+			FString TextureFactoryUid;
+			if (UInterchangeTextureNode* TextureNode = Cast<UInterchangeTextureNode>(BaseNodeContainer->GetNode(TextureUid)))
+			{
+				TArray<FString> TextureTargetNodes;
+				TextureNode->GetTargetNodeUids(TextureTargetNodes);
+
+				if (TextureTargetNodes.Num() > 0)
+				{
+					TextureFactoryUid = TextureTargetNodes[0];
+				}
+			}
+
+			Result.Set<FString>(TextureFactoryUid);
+		}
+	}
+
+	return Result;
+}
diff --git a/Engine/Plugins/Experimental/Interchange/Runtime/Source/Pipelines/Public/InterchangeGenericMaterialPipeline.h b/Engine/Plugins/Experimental/Interchange/Runtime/Source/Pipelines/Public/InterchangeGenericMaterialPipeline.h
--- a/Engine/Plugins/Experimental/Interchange/Runtime/Source/Pipelines/Public/InterchangeGenericMaterialPipeline.h
+++ b/Engine/Plugins/Experimental/Interchange/Runtime/Source/Pipelines/Public/InterchangeGenericMaterialPipeline.h
@@ -16,6 +16,7 @@ class UInterchangeShaderGraphNode;
 class UInterchangeShaderNode;
 class UInterchangeMaterialFactoryNode;
 class UInterchangeMaterialExpressionFactoryNode;
+class UInterchangeMaterialInstanceFactoryNode;
 class UInterchangeResult;
 
 UENUM(BlueprintType)
@@ -23,7 +24,7 @@ enum class EInterchangeMaterialImportOption : uint8
 {
 	DoNotImport,
 	ImportAsMaterials,
-	//ImportAsMaterialInstances,
+	ImportAsMaterialInstances,
 };
 
 UCLASS(BlueprintType, Experimental)
@@ -32,10 +33,13 @@ class INTERCHANGEPIPELINES_API UInterchangeGenericMaterialPipeline : public UInt
 	GENERATED_BODY()
 
 public:
-	/** If enabled, imports the material assets found in the sources. */
-	UPROPERTY(EditAnywhere, BlueprintReadWrite, Category = MATERIALS_CATEGORY)
+	UPROPERTY(EditAnywhere, BlueprintReadWrite, Category = "Materials")
 	EInterchangeMaterialImportOption MaterialImport = EInterchangeMaterialImportOption::ImportAsMaterials;
 
+	/** Optional material used as the parent when importing materials as instances. If no parent material is specified, one will be automatically selected during the import process. */
+	UPROPERTY(EditAnywhere, BlueprintReadWrite, Category = "Materials", Meta= (EditCondition="MaterialImport==EInterchangeMaterialImportOption::ImportAsMaterialInstances", AllowedClasses="MaterialInterface"))
+	FSoftObjectPath ParentMaterial;
+
 public:
 	virtual void ExecutePreImportPipeline(UInterchangeBaseNodeContainer* InBaseNodeContainer, const TArray<UInterchangeSourceData*>& InSourceDatas) override;
 
@@ -54,6 +58,16 @@ private:
 
 	UInterchangeBaseMaterialFactoryNode* CreateBaseMaterialFactoryNode(const UInterchangeBaseNode* MaterialNode, TSubclassOf<UInterchangeBaseMaterialFactoryNode> NodeType);
 	UInterchangeMaterialFactoryNode* CreateMaterialFactoryNode(const UInterchangeShaderGraphNode* ShaderGraphNode);
+	UInterchangeMaterialInstanceFactoryNode* CreateMaterialInstanceFactoryNode(const UInterchangeShaderGraphNode* ShaderGraphNode);
+
+	/** True if the shader graph has a base color input. */
+	bool IsPBRModel(const UInterchangeShaderGraphNode* ShaderGraphNode) const;
+
+	/** True if the shader graph has diffuse color and specular color inputs. */
+	bool IsPhongModel(const UInterchangeShaderGraphNode* ShaderGraphNode) const;
+
+	/** True if the shader graph has a diffuse color input. */
+	bool IsLambertModel(const UInterchangeShaderGraphNode* ShaderGraphNode) const;
 
 	bool HandlePhongModel(const UInterchangeShaderGraphNode* ShaderGraphNode, UInterchangeMaterialFactoryNode* MaterialFactoryNode);
 	bool HandleLambertModel(const UInterchangeShaderGraphNode* ShaderGraphNode, UInterchangeMaterialFactoryNode* MaterialFactoryNode);
@@ -70,6 +84,14 @@ private:
 	UInterchangeMaterialExpressionFactoryNode* CreateScalarParameterExpression(const UInterchangeShaderNode* ShaderNode, const FString& InputName, const FString& ParentUid);
 	UInterchangeMaterialExpressionFactoryNode* CreateVectorParameterExpression(const UInterchangeShaderNode* ShaderNode, const FString& InputName, const FString& ParentUid);
 
+	/**
+	 * Visits a given shader node and its connections to find its strongest value.
+	 * Only its first input is visited as it's assumed that it's the most impactful.
+	 * The goal is to simplify a branch of a node graph to a single value, to be used for material instancing.
+	 */
+	TVariant<FString, FLinearColor, float> VisitShaderNode(const UInterchangeShaderNode* ShaderNode);
+	TVariant<FString, FLinearColor, float> VisitTextureSampleNode(const UInterchangeShaderNode* ShaderNode);
+
 private:
 	bool bParsingForNormalInput;
 };
diff --git a/Engine/Source/Runtime/Experimental/Interchange/Core/Private/InterchangePipelineBase.cpp b/Engine/Source/Runtime/Experimental/Interchange/Core/Private/InterchangePipelineBase.cpp
--- a/Engine/Source/Runtime/Experimental/Interchange/Core/Private/InterchangePipelineBase.cpp
+++ b/Engine/Source/Runtime/Experimental/Interchange/Core/Private/InterchangePipelineBase.cpp
@@ -95,6 +95,17 @@ void UInterchangePipelineBase::LoadSettings(const FName PipelineStackName)
 				}
 			}
 		}
+		else if (FObjectProperty* SubObject = CastField<FObjectProperty>(Property))
+		{
+			if (UInterchangePipelineBase* SubPipeline = Cast<UInterchangePipelineBase>(SubObject->GetObjectPropertyValue_InContainer(this)))
+			{
+				// Load the settings if the referenced pipeline is a subobject of ours
+				if (SubPipeline->IsInOuter(this))
+				{
+					SubPipeline->LoadSettings(PipelineStackName);
+				}
+			}
+		}
 		else
 		{
 			for (int32 i = 0; i < Property->ArrayDim; i++)
@@ -155,6 +166,17 @@ void UInterchangePipelineBase::SaveSettings(const FName PipelineStackName)
 				Sec->Add(*Key, *Buffer);
 			}
 		}
+		else if (FObjectProperty* SubObject = CastField<FObjectProperty>(Property))
+		{
+			if (UInterchangePipelineBase* SubPipeline = Cast<UInterchangePipelineBase>(SubObject->GetObjectPropertyValue_InContainer(this)))
+			{
+				// Save the settings if the referenced pipeline is a subobject of ours
+				if (SubPipeline->IsInOuter(this))
+				{
+					SubPipeline->SaveSettings(PipelineStackName);
+				}
+			}
+		}
 		else
 		{
 			TCHAR TempKey[MAX_SPRINTF] = TEXT("");
