Add a re-import dialog showing the original pipelinestack use to import the asset. User can currently change the pipeline options, but cannot change the stack like in the import dialog.

Make sure the vertex color is re-import correctly. I put back the same workflow we have before in UE4. Where ignore will put back the vertex color we have in the asset before reimport happen.
To remap the vertex color SM only match vertex position(FVector::Equals) and sk is using an octree to find the closest vertex position.

#rnx

diff --git a/Engine/Plugins/Experimental/Interchange/Editor/Source/Pipelines/Private/InterchangePipelineConfigurationGeneric.cpp b/Engine/Plugins/Experimental/Interchange/Editor/Source/Pipelines/Private/InterchangePipelineConfigurationGeneric.cpp
--- a/Engine/Plugins/Experimental/Interchange/Editor/Source/Pipelines/Private/InterchangePipelineConfigurationGeneric.cpp
+++ b/Engine/Plugins/Experimental/Interchange/Editor/Source/Pipelines/Private/InterchangePipelineConfigurationGeneric.cpp
@@ -27,6 +27,7 @@ EInterchangePipelineConfigurationDialogResult UInterchangePipelineConfigurationG
 	(
 		SAssignNew(InterchangePipelineConfigurationDialog, SInterchangePipelineConfigurationDialog)
 		.OwnerWindow(Window)
+		.bReimport(false)
 	);
 
 	FSlateApplication::Get().AddModalWindow(Window, ParentWindow, false);
@@ -43,3 +44,39 @@ EInterchangePipelineConfigurationDialogResult UInterchangePipelineConfigurationG
 
 	return EInterchangePipelineConfigurationDialogResult::Import;
 }
+
+EInterchangePipelineConfigurationDialogResult UInterchangePipelineConfigurationGeneric::ShowReimportPipelineConfigurationDialog(TArray<UInterchangePipelineBase*>& PipelineStack)
+{
+	//Create and show the graph inspector UI dialog
+	TSharedPtr<SWindow> ParentWindow;
+	if (IMainFrameModule* MainFrame = FModuleManager::LoadModulePtr<IMainFrameModule>("MainFrame"))
+	{
+		ParentWindow = MainFrame->GetParentWindow();
+	}
+	TSharedRef<SWindow> Window = SNew(SWindow)
+		.ClientSize(FVector2D(1000.f, 650.f))
+		.Title(NSLOCTEXT("Interchange", "PipelineConfigurationGenericTitle", "Interchange Pipeline Configuration"));
+	TSharedPtr<SInterchangePipelineConfigurationDialog> InterchangePipelineConfigurationDialog;
+
+	Window->SetContent
+	(
+		SAssignNew(InterchangePipelineConfigurationDialog, SInterchangePipelineConfigurationDialog)
+		.OwnerWindow(Window)
+		.bReimport(true)
+		.PipelineStack(PipelineStack)
+	);
+
+	FSlateApplication::Get().AddModalWindow(Window, ParentWindow, false);
+
+	if (InterchangePipelineConfigurationDialog->IsCanceled())
+	{
+		return EInterchangePipelineConfigurationDialogResult::Cancel;
+	}
+
+	if (InterchangePipelineConfigurationDialog->IsImportAll())
+	{
+		return EInterchangePipelineConfigurationDialogResult::ImportAll;
+	}
+
+	return EInterchangePipelineConfigurationDialogResult::Import;
+}
\ No newline at end of file
diff --git a/Engine/Plugins/Experimental/Interchange/Editor/Source/Pipelines/Private/SInterchangePipelineConfigurationDialog.cpp b/Engine/Plugins/Experimental/Interchange/Editor/Source/Pipelines/Private/SInterchangePipelineConfigurationDialog.cpp
--- a/Engine/Plugins/Experimental/Interchange/Editor/Source/Pipelines/Private/SInterchangePipelineConfigurationDialog.cpp
+++ b/Engine/Plugins/Experimental/Interchange/Editor/Source/Pipelines/Private/SInterchangePipelineConfigurationDialog.cpp
@@ -46,35 +46,60 @@ SInterchangePipelineStacksTreeView::~SInterchangePipelineStacksTreeView()
 void SInterchangePipelineStacksTreeView::Construct(const FArguments& InArgs)
 {
 	OnSelectionChangedDelegate = InArgs._OnSelectionChangedDelegate;
-	//Build the FbxNodeInfoPtr tree data
+	PipelineStack = InArgs._PipelineStack;
+	bReimport = InArgs._bReimport && PipelineStack.Num() > 0;
 
-	const FName& DefaultPipelineStackName = GetDefault<UInterchangeProjectSettings>()->DefaultPipelineStack;
+	//Build the FbxNodeInfoPtr tree data
+	const FName ReimportPipelineName = TEXT("ReimportPipeline");
+	const FName& DefaultPipelineStackName = bReimport ? ReimportPipelineName : GetDefault<UInterchangeProjectSettings>()->DefaultPipelineStack;
 	const TMap<FName, FInterchangePipelineStack>& DefaultPipelineStacks = GetDefault<UInterchangeProjectSettings>()->PipelineStacks;
 
-	 
-	for (const TPair<FName, FInterchangePipelineStack>& NameAndPipelineStack : DefaultPipelineStacks)
+	//In reimport we modify directly the asset pipeline
+	if (bReimport)
 	{
 		TSharedPtr<FInterchangePipelineStacksTreeNodeItem> StackNode = MakeShared<FInterchangePipelineStacksTreeNodeItem>();
-		StackNode->StackName = NameAndPipelineStack.Key;
+		StackNode->StackName = ReimportPipelineName;
 		StackNode->Pipeline = nullptr;
-		const FInterchangePipelineStack& PipelineStack = NameAndPipelineStack.Value;
-		for (int32 PipelineIndex = 0; PipelineIndex < PipelineStack.Pipelines.Num(); ++PipelineIndex)
+		for (int32 PipelineIndex = 0; PipelineIndex < PipelineStack.Num(); ++PipelineIndex)
 		{
-			const UClass* PipelineClass = PipelineStack.Pipelines[PipelineIndex].LoadSynchronous();
-			if (PipelineClass)
+			if(UInterchangePipelineBase* GeneratedPipeline = PipelineStack[PipelineIndex])
 			{
-				UInterchangePipelineBase* GeneratedPipeline = NewObject<UInterchangePipelineBase>(GetTransientPackage(), PipelineClass, NAME_None, RF_NoFlags);
+				GeneratedPipeline->SaveSettings(ReimportPipelineName);
 				//Load the settings for this pipeline
-				GeneratedPipeline->LoadSettings(NameAndPipelineStack.Key);
-				GeneratedPipeline->PreDialogCleanup(NameAndPipelineStack.Key);
 				TSharedPtr<FInterchangePipelineStacksTreeNodeItem> PipelineNode = MakeShared<FInterchangePipelineStacksTreeNodeItem>();
-				PipelineNode->StackName = NameAndPipelineStack.Key;
+				PipelineNode->StackName = ReimportPipelineName;
 				PipelineNode->Pipeline = GeneratedPipeline;
 				StackNode->Childrens.Add(PipelineNode);
 			}
 		}
 		RootNodeArray.Add(StackNode);
 	}
+	else
+	{
+		for (const TPair<FName, FInterchangePipelineStack>& NameAndPipelineStack : DefaultPipelineStacks)
+		{
+			TSharedPtr<FInterchangePipelineStacksTreeNodeItem> StackNode = MakeShared<FInterchangePipelineStacksTreeNodeItem>();
+			StackNode->StackName = NameAndPipelineStack.Key;
+			StackNode->Pipeline = nullptr;
+			const FInterchangePipelineStack& InterchangePipelineStack = NameAndPipelineStack.Value;
+			for (int32 PipelineIndex = 0; PipelineIndex < InterchangePipelineStack.Pipelines.Num(); ++PipelineIndex)
+			{
+				const UClass* PipelineClass = InterchangePipelineStack.Pipelines[PipelineIndex].LoadSynchronous();
+				if (PipelineClass)
+				{
+					UInterchangePipelineBase* GeneratedPipeline = NewObject<UInterchangePipelineBase>(GetTransientPackage(), PipelineClass, NAME_None, RF_NoFlags);
+					//Load the settings for this pipeline
+					GeneratedPipeline->LoadSettings(NameAndPipelineStack.Key);
+					GeneratedPipeline->PreDialogCleanup(NameAndPipelineStack.Key);
+					TSharedPtr<FInterchangePipelineStacksTreeNodeItem> PipelineNode = MakeShared<FInterchangePipelineStacksTreeNodeItem>();
+					PipelineNode->StackName = NameAndPipelineStack.Key;
+					PipelineNode->Pipeline = GeneratedPipeline;
+					StackNode->Childrens.Add(PipelineNode);
+				}
+			}
+			RootNodeArray.Add(StackNode);
+		}
+	}
 
 	STreeView::Construct
 	(
@@ -234,6 +259,11 @@ FReply SInterchangePipelineStacksTreeView::OnCollapseAll()
 
 TSharedPtr<SWidget> SInterchangePipelineStacksTreeView::OnOpenContextMenu()
 {
+	if (bReimport)
+	{
+		return SNullWidget::NullWidget;
+	}
+
 	// Build up the menu for a selection
 	const bool bCloseAfterSelection = true;
 	FMenuBuilder MenuBuilder(bCloseAfterSelection, TSharedPtr<FUICommandList>());
@@ -308,7 +338,9 @@ TSharedRef<SBox> SInterchangePipelineConfigurationDialog::SpawnPipelineConfigura
 {
 	//Create the treeview
 	PipelineConfigurationTreeView = SNew(SInterchangePipelineStacksTreeView)
-		.OnSelectionChangedDelegate(this, &SInterchangePipelineConfigurationDialog::OnSelectionChanged);
+		.OnSelectionChangedDelegate(this, &SInterchangePipelineConfigurationDialog::OnSelectionChanged)
+		.bReimport(bReimport)
+		.PipelineStack(PipelineStack);
 
 	TSharedPtr<SBox> InspectorBox;
 	TSharedRef<SBox> PipelineConfigurationPanelBox = SNew(SBox)
@@ -372,6 +404,8 @@ void SInterchangePipelineConfigurationDialog::Construct(const FArguments& InArgs
 	//Make sure there is a valid default value
 
 	OwnerWindow = InArgs._OwnerWindow;
+	bReimport = InArgs._bReimport;
+	PipelineStack = InArgs._PipelineStack;
 
 	check(OwnerWindow.IsValid());
 
@@ -444,12 +478,33 @@ void SInterchangePipelineConfigurationDialog::RecursiveSavePipelineSettings(cons
 	}
 }
 
+void SInterchangePipelineConfigurationDialog::RecursiveLoadPipelineSettings(const TSharedPtr<FInterchangePipelineStacksTreeNodeItem>& ParentNode, const int32 PipelineIndex) const
+{
+	if (ParentNode->Pipeline)
+	{
+		ParentNode->Pipeline->LoadSettings(ParentNode->StackName);
+	}
+	for (int32 ChildIndex = 0; ChildIndex < ParentNode->Childrens.Num(); ++ChildIndex)
+	{
+		RecursiveLoadPipelineSettings(ParentNode->Childrens[ChildIndex], ChildIndex);
+	}
+}
+
 void SInterchangePipelineConfigurationDialog::ClosePipelineConfiguration(const ECloseEventType CloseEventType)
 {
 	if (CloseEventType == ECloseEventType::Cancel)
 	{
 		bCanceled = true;
 		bImportAll = false;
+		if (bReimport && PipelineConfigurationTreeView)
+		{
+			//If user cancel the reimport we have to put back the pipeline settings
+			const TArray<TSharedPtr<FInterchangePipelineStacksTreeNodeItem>>& RootNodeArray = PipelineConfigurationTreeView->GetRootNodeArray();
+			for (const TSharedPtr<FInterchangePipelineStacksTreeNodeItem>& RootNode : RootNodeArray)
+			{
+				RecursiveLoadPipelineSettings(RootNode, 0);
+			}
+		}
 	}
 	else if (CloseEventType == ECloseEventType::ImportAll)
 	{
@@ -463,7 +518,7 @@ void SInterchangePipelineConfigurationDialog::ClosePipelineConfiguration(const E
 		bImportAll = false;
 	}
 
-	if (PipelineConfigurationTreeView)
+	if (!bReimport && PipelineConfigurationTreeView)
 	{
 		const TArray<TSharedPtr<FInterchangePipelineStacksTreeNodeItem>>& RootNodeArray = PipelineConfigurationTreeView->GetRootNodeArray();
 		for (const TSharedPtr<FInterchangePipelineStacksTreeNodeItem>& RootNode : RootNodeArray)
diff --git a/Engine/Plugins/Experimental/Interchange/Editor/Source/Pipelines/Private/SInterchangePipelineConfigurationDialog.h b/Engine/Plugins/Experimental/Interchange/Editor/Source/Pipelines/Private/SInterchangePipelineConfigurationDialog.h
--- a/Engine/Plugins/Experimental/Interchange/Editor/Source/Pipelines/Private/SInterchangePipelineConfigurationDialog.h
+++ b/Engine/Plugins/Experimental/Interchange/Editor/Source/Pipelines/Private/SInterchangePipelineConfigurationDialog.h
@@ -50,6 +50,8 @@ public:
 		: _OnSelectionChangedDelegate()
 	{}
 		SLATE_EVENT(FOnPipelineConfigurationSelectionChanged, OnSelectionChangedDelegate)
+		SLATE_ARGUMENT(bool, bReimport)
+		SLATE_ARGUMENT(TArray<UInterchangePipelineBase*>, PipelineStack)
 	SLATE_END_ARGS()
 
 	/** Construct this widget */
@@ -64,6 +66,8 @@ public:
 protected:
 	/** Delegate to invoke when selection changes. */
 	FOnPipelineConfigurationSelectionChanged OnSelectionChangedDelegate;
+	bool bReimport = false;
+	TArray<UInterchangePipelineBase*> PipelineStack;
 
 	/** the elements we show in the tree view */
 	TArray<TSharedPtr<FInterchangePipelineStacksTreeNodeItem>> RootNodeArray;
@@ -91,6 +95,8 @@ public:
 	{}
 
 		SLATE_ARGUMENT(TWeakPtr<SWindow>, OwnerWindow)
+		SLATE_ARGUMENT(bool, bReimport)
+		SLATE_ARGUMENT(TArray<UInterchangePipelineBase*>, PipelineStack)
 	SLATE_END_ARGS()
 
 public:
@@ -115,11 +121,14 @@ public:
 	bool IsImportAll() { return bImportAll; }
 private:
 	TWeakPtr< SWindow > OwnerWindow;
+	bool bReimport = false;
+	TArray<UInterchangePipelineBase*> PipelineStack;
 
 	TSharedRef<SBox> SpawnPipelineConfiguration();
 	void OnSelectionChanged(TSharedPtr<FInterchangePipelineStacksTreeNodeItem> Item, ESelectInfo::Type SelectionType);
 
 	void RecursiveSavePipelineSettings(const TSharedPtr<FInterchangePipelineStacksTreeNodeItem>& ParentNode, const int32 PipelineIndex) const;
+	void RecursiveLoadPipelineSettings(const TSharedPtr<FInterchangePipelineStacksTreeNodeItem>& ParentNode, const int32 PipelineIndex) const;
 
 	//Graph Inspector UI elements
 	TSharedPtr<SInterchangePipelineStacksTreeView> PipelineConfigurationTreeView;
diff --git a/Engine/Plugins/Experimental/Interchange/Editor/Source/Pipelines/Public/InterchangePipelineConfigurationGeneric.h b/Engine/Plugins/Experimental/Interchange/Editor/Source/Pipelines/Public/InterchangePipelineConfigurationGeneric.h
--- a/Engine/Plugins/Experimental/Interchange/Editor/Source/Pipelines/Public/InterchangePipelineConfigurationGeneric.h
+++ b/Engine/Plugins/Experimental/Interchange/Editor/Source/Pipelines/Public/InterchangePipelineConfigurationGeneric.h
@@ -19,4 +19,5 @@ public:
 protected:
 
 	virtual EInterchangePipelineConfigurationDialogResult ShowPipelineConfigurationDialog() override;
+	virtual EInterchangePipelineConfigurationDialogResult ShowReimportPipelineConfigurationDialog(TArray<UInterchangePipelineBase*>& PipelineStack) override;
 };
diff --git a/Engine/Plugins/Experimental/Interchange/Runtime/Source/Import/Private/Mesh/InterchangeSkeletalMeshFactory.cpp b/Engine/Plugins/Experimental/Interchange/Runtime/Source/Import/Private/Mesh/InterchangeSkeletalMeshFactory.cpp
--- a/Engine/Plugins/Experimental/Interchange/Runtime/Source/Import/Private/Mesh/InterchangeSkeletalMeshFactory.cpp
+++ b/Engine/Plugins/Experimental/Interchange/Runtime/Source/Import/Private/Mesh/InterchangeSkeletalMeshFactory.cpp
@@ -16,12 +16,14 @@
 #include "InterchangeSkeletonFactoryNode.h"
 #include "InterchangeSourceData.h"
 #include "InterchangeTranslatorBase.h"
+#include "Math/GenericOctree.h"
 #include "Mesh/InterchangeSkeletalMeshPayload.h"
 #include "Mesh/InterchangeSkeletalMeshPayloadInterface.h"
 #include "Nodes/InterchangeBaseNode.h"
 #include "Nodes/InterchangeBaseNodeContainer.h"
 #include "PhysicsEngine/PhysicsAsset.h"
 #include "Rendering/SkeletalMeshLODImporterData.h"
+#include "Rendering/SkeletalMeshLODModel.h"
 #include "Rendering/SkeletalMeshModel.h"
 #include "SkeletalMeshAttributes.h"
 #include "SkeletalMeshOperations.h"
@@ -849,6 +851,149 @@ namespace UE
 					}
 				}
 			}
+
+			/** Helper struct for the mesh component vert position octree */
+			struct FSkeletalMeshVertPosOctreeSemantics
+			{
+				enum { MaxElementsPerLeaf = 16 };
+				enum { MinInclusiveElementsPerNode = 7 };
+				enum { MaxNodeDepth = 12 };
+
+				typedef TInlineAllocator<MaxElementsPerLeaf> ElementAllocator;
+
+				/**
+				 * Get the bounding box of the provided octree element. In this case, the box
+				 * is merely the point specified by the element.
+				 *
+				 * @param	Element	Octree element to get the bounding box for
+				 *
+				 * @return	Bounding box of the provided octree element
+				 */
+				FORCEINLINE static FBoxCenterAndExtent GetBoundingBox(const FSoftSkinVertex& Element)
+				{
+					return FBoxCenterAndExtent(Element.Position, FVector::ZeroVector);
+				}
+
+				/**
+				 * Determine if two octree elements are equal
+				 *
+				 * @param	A	First octree element to check
+				 * @param	B	Second octree element to check
+				 *
+				 * @return	true if both octree elements are equal, false if they are not
+				 */
+				FORCEINLINE static bool AreElementsEqual(const FSoftSkinVertex& A, const FSoftSkinVertex& B)
+				{
+					return (A.Position == B.Position && A.UVs[0] == B.UVs[0]);
+				}
+
+				/** Ignored for this implementation */
+				FORCEINLINE static void SetElementId(const FSoftSkinVertex& Element, FOctreeElementId2 Id)
+				{
+				}
+			};
+			typedef TOctree2<FSoftSkinVertex, FSkeletalMeshVertPosOctreeSemantics> TSKCVertPosOctree;
+
+			void RemapSkeletalMeshVertexColorToImportData(const USkeletalMesh* SkeletalMesh, const int32 LODIndex, FSkeletalMeshImportData* SkelMeshImportData)
+			{
+				//Make sure we have all the source data we need to do the remap
+				if (!SkeletalMesh->GetImportedModel() || !SkeletalMesh->GetImportedModel()->LODModels.IsValidIndex(LODIndex) || !SkeletalMesh->GetHasVertexColors())
+				{
+					return;
+				}
+
+				// Find the extents formed by the cached vertex positions in order to optimize the octree used later
+				FBox Bounds(ForceInitToZero);
+				SkelMeshImportData->bHasVertexColors = true;
+
+				int32 WedgeNumber = SkelMeshImportData->Wedges.Num();
+				for (int32 WedgeIndex = 0; WedgeIndex < WedgeNumber; ++WedgeIndex)
+				{
+					SkeletalMeshImportData::FVertex& Wedge = SkelMeshImportData->Wedges[WedgeIndex];
+					const FVector& Position = SkelMeshImportData->Points[Wedge.VertexIndex];
+					Bounds += Position;
+				}
+
+				TArray<FSoftSkinVertex> Vertices;
+				SkeletalMesh->GetImportedModel()->LODModels[LODIndex].GetVertices(Vertices);
+				for (int32 SkinVertexIndex = 0; SkinVertexIndex < Vertices.Num(); ++SkinVertexIndex)
+				{
+					const FSoftSkinVertex& SkinVertex = Vertices[SkinVertexIndex];
+					Bounds += SkinVertex.Position;
+				}
+
+				TSKCVertPosOctree VertPosOctree(Bounds.GetCenter(), Bounds.GetExtent().GetMax());
+
+				// Add each old vertex to the octree
+				for (int32 SkinVertexIndex = 0; SkinVertexIndex < Vertices.Num(); ++SkinVertexIndex)
+				{
+					const FSoftSkinVertex& SkinVertex = Vertices[SkinVertexIndex];
+					VertPosOctree.AddElement(SkinVertex);
+				}
+
+				TMap<int32, FVector3f> WedgeIndexToNormal;
+				WedgeIndexToNormal.Reserve(WedgeNumber);
+				for (int32 FaceIndex = 0; FaceIndex < SkelMeshImportData->Faces.Num(); ++FaceIndex)
+				{
+					const SkeletalMeshImportData::FTriangle& Triangle = SkelMeshImportData->Faces[FaceIndex];
+					for (int32 Corner = 0; Corner < 3; ++Corner)
+					{
+						WedgeIndexToNormal.Add(Triangle.WedgeIndex[Corner], Triangle.TangentZ[Corner]);
+					}
+				}
+
+				// Iterate over each new vertex position, attempting to find the old vertex it is closest to, applying
+				// the color of the old vertex to the new position if possible.
+				for (int32 WedgeIndex = 0; WedgeIndex < WedgeNumber; ++WedgeIndex)
+				{
+					SkeletalMeshImportData::FVertex& Wedge = SkelMeshImportData->Wedges[WedgeIndex];
+					const FVector& Position = SkelMeshImportData->Points[Wedge.VertexIndex];
+					const FVector2f UV = Wedge.UVs[0];
+					const FVector3f& Normal = WedgeIndexToNormal.FindChecked(WedgeIndex);
+
+					TArray<FSoftSkinVertex> PointsToConsider;
+					VertPosOctree.FindNearbyElements(Position, [&PointsToConsider](const FSoftSkinVertex& Vertex)
+						{
+							PointsToConsider.Add(Vertex);
+						});
+
+					if (PointsToConsider.Num() > 0)
+					{
+						//Get the closest position
+						float MaxNormalDot = -MAX_FLT;
+						float MinUVDistance = MAX_FLT;
+						int32 MatchIndex = INDEX_NONE;
+						for (int32 ConsiderationIndex = 0; ConsiderationIndex < PointsToConsider.Num(); ++ConsiderationIndex)
+						{
+							const FSoftSkinVertex& SkinVertex = PointsToConsider[ConsiderationIndex];
+							const FVector2f& SkinVertexUV = SkinVertex.UVs[0];
+							const float UVDistanceSqr = FVector2f::DistSquared(UV, SkinVertexUV);
+							if (UVDistanceSqr < MinUVDistance)
+							{
+								MinUVDistance = FMath::Min(MinUVDistance, UVDistanceSqr);
+								MatchIndex = ConsiderationIndex;
+								MaxNormalDot = Normal | SkinVertex.TangentZ;
+							}
+							else if (FMath::IsNearlyEqual(UVDistanceSqr, MinUVDistance, KINDA_SMALL_NUMBER))
+							{
+								//This case is useful when we have hard edge that shared vertice, somtime not all the shared wedge have the same paint color
+								//Think about a cube where each face have different vertex color.
+								float NormalDot = Normal | SkinVertex.TangentZ;
+								if (NormalDot > MaxNormalDot)
+								{
+									MaxNormalDot = NormalDot;
+									MatchIndex = ConsiderationIndex;
+								}
+							}
+						}
+						if (PointsToConsider.IsValidIndex(MatchIndex))
+						{
+							Wedge.Color = PointsToConsider[MatchIndex].Color;
+						}
+					}
+				}
+			}
+
 		} //Namespace Private
 	} //namespace Interchange
 } //namespace UE
@@ -1158,11 +1303,19 @@ UObject* UInterchangeSkeletalMeshFactory::CreateAsset(const FCreateAssetParams&
 				SkeletalMeshFactoryNode->GetCustomVertexColorIgnore(bIgnoreVertexColor);
 				if (bIgnoreVertexColor)
 				{
-					//Flush the vertex color, if we re-import we have to fill it with the old data
-					SkeletalMeshImportData.bHasVertexColors = false;
-					for (SkeletalMeshImportData::FVertex& Wedge : SkeletalMeshImportData.Wedges)
+					if (bIsReImport)
 					{
-						Wedge.Color = FColor::White;
+						//Get the vertex color we have in the current asset, 
+						UE::Interchange::Private::RemapSkeletalMeshVertexColorToImportData(SkeletalMesh, LodIndex, &SkeletalMeshImportData);
+					}
+					else
+					{
+						//Flush the vertex color
+						SkeletalMeshImportData.bHasVertexColors = false;
+						for (SkeletalMeshImportData::FVertex& Wedge : SkeletalMeshImportData.Wedges)
+						{
+							Wedge.Color = FColor::White;
+						}
 					}
 				}
 				else
diff --git a/Engine/Plugins/Experimental/Interchange/Runtime/Source/Import/Private/Mesh/InterchangeStaticMeshFactory.cpp b/Engine/Plugins/Experimental/Interchange/Runtime/Source/Import/Private/Mesh/InterchangeStaticMeshFactory.cpp
--- a/Engine/Plugins/Experimental/Interchange/Runtime/Source/Import/Private/Mesh/InterchangeStaticMeshFactory.cpp
+++ b/Engine/Plugins/Experimental/Interchange/Runtime/Source/Import/Private/Mesh/InterchangeStaticMeshFactory.cpp
@@ -146,8 +146,13 @@ UObject* UInterchangeStaticMeshFactory::CreateAsset(const FCreateAssetParams& Ar
 			
 	int32 LodCount = StaticMeshFactoryNode->GetLodDataCount();
 	int32 FinalLodCount = LodCount;
+	TMap<FVector, FColor> ExisitingVertexColorData;
 	if (Arguments.ReimportObject)
 	{
+		if(ExistingAsset)
+		{
+			StaticMesh->GetVertexColorData(ExisitingVertexColorData);
+		}
 		//When we reimport we dont want to reduce the number of existing LODs
 		FinalLodCount = FMath::Max(StaticMesh->GetNumLODs(), LodCount);
 	}
@@ -287,10 +292,29 @@ UObject* UInterchangeStaticMeshFactory::CreateAsset(const FCreateAssetParams& Ar
 				StaticMeshFactoryNode->GetCustomVertexColorIgnore(bIgnoreVertexColor);
 				if (bIgnoreVertexColor)
 				{
-					//Flush the vertex color, if we re-import we have to fill it with the old data
 					for (const FVertexInstanceID& VertexInstanceID : LodMeshDescription->VertexInstances().GetElementIDs())
 					{
-						VertexInstanceColors[VertexInstanceID] = FVector4f(FLinearColor(FColor::White));
+						//If we have old vertex color (reimport), we want to keep it if the option is ignore
+						if (ExisitingVertexColorData.Num() > 0)
+						{
+							FVector3f VertexPosition = LodMeshDescription->GetVertexPosition(LodMeshDescription->GetVertexInstanceVertex(VertexInstanceID));
+							const FColor* PaintedColor = ExisitingVertexColorData.Find(VertexPosition);
+							if (PaintedColor)
+							{
+								// A matching color for this vertex was found
+								VertexInstanceColors[VertexInstanceID] = FVector4f(FLinearColor(*PaintedColor));
+							}
+							else
+							{
+								//Flush the vertex color
+								VertexInstanceColors[VertexInstanceID] = FVector4f(FLinearColor(FColor::White));
+							}
+						}
+						else
+						{
+							//Flush the vertex color
+							VertexInstanceColors[VertexInstanceID] = FVector4f(FLinearColor(FColor::White));
+						}
 					}
 				}
 				else
diff --git a/Engine/Plugins/Experimental/Interchange/Runtime/Source/Pipelines/Public/InterchangeGenericAssetsPipeline.h b/Engine/Plugins/Experimental/Interchange/Runtime/Source/Pipelines/Public/InterchangeGenericAssetsPipeline.h
--- a/Engine/Plugins/Experimental/Interchange/Runtime/Source/Pipelines/Public/InterchangeGenericAssetsPipeline.h
+++ b/Engine/Plugins/Experimental/Interchange/Runtime/Source/Pipelines/Public/InterchangeGenericAssetsPipeline.h
@@ -30,14 +30,6 @@ class UInterchangeTextureFactoryNode;
 class UInterchangeTextureNode;
 template<class T> class TSubclassOf;
 
-#define COMMON_CATEGORY "Common"
-#define COMMON_MESHES_CATEGORY "Common Meshes"
-#define STATIC_MESHES_CATEGORY "Static Meshes"
-#define COMMON_SKELETAL_ANIMATIONS_CATEGORY "Common Skeletal Mesh and Animations"
-#define SKELETAL_MESHES_CATEGORY "Skeletal Meshes"
-#define ANIMATIONS_CATEGORY "Animations"
-#define TEXTURES_CATEGORY "Textures"
-
 /** Force mesh type, if user want to import all meshes as one type*/
 UENUM(BlueprintType)
 enum EInterchangeForceMeshType
@@ -86,11 +78,11 @@ public:
 
 
 	/** If enable and there is only one asset and one source data, we will name the asset like the source data name. */
-	UPROPERTY(EditAnywhere, BlueprintReadWrite, Category = COMMON_CATEGORY)
+	UPROPERTY(EditAnywhere, BlueprintReadWrite, Category = "Common")
 	bool bUseSourceNameForAsset = true;
 
 	/* Allow user to choose the re-import strategy. */
-	UPROPERTY(EditAnywhere, BlueprintReadWrite, Category = COMMON_CATEGORY)
+	UPROPERTY(EditAnywhere, BlueprintReadWrite, Category = "Common")
 	EReimportStrategyFlags ReimportStrategy = EReimportStrategyFlags::ApplyNoProperties;
 
 
@@ -98,34 +90,34 @@ public:
 
 
 	/** Allow to convert mesh to a particular type */
- 	UPROPERTY(EditAnywhere, BlueprintReadWrite, Category = COMMON_MESHES_CATEGORY)
+ 	UPROPERTY(EditAnywhere, BlueprintReadWrite, Category = "Common Meshes")
  	TEnumAsByte<enum EInterchangeForceMeshType> ForceAllMeshAsType = EInterchangeForceMeshType::IFMT_None;
 
 	/** If enable, meshes LODs will be imported. Note that it required the advanced bBakeMesh property to be enabled. */
-	UPROPERTY(EditAnywhere, BlueprintReadWrite, Category = COMMON_MESHES_CATEGORY, meta = (editcondition = "bBakeMeshes"))
+	UPROPERTY(EditAnywhere, BlueprintReadWrite, Category = "Common Meshes", meta = (editcondition = "bBakeMeshes"))
 	bool bImportLods = true;
 
 	/** If enable, meshes will be baked with the scene instance hierarchy transform. */
-	UPROPERTY(EditAnywhere, BlueprintReadWrite, Category = COMMON_MESHES_CATEGORY)
+	UPROPERTY(EditAnywhere, BlueprintReadWrite, Category = "Common Meshes")
 	bool bBakeMeshes = true;
 
 	/** Specify how vertex colors should be imported */
-	UPROPERTY(EditAnywhere, BlueprintReadWrite, Category = COMMON_MESHES_CATEGORY)
+	UPROPERTY(EditAnywhere, BlueprintReadWrite, Category = "Common Meshes")
 	TEnumAsByte<enum EInterchangeVertexColorImportOption> VertexColorImportOption = EInterchangeVertexColorImportOption::IVCIO_Replace;
 
 	/** Specify override color in the case that VertexColorImportOption is set to Override */
-	UPROPERTY(EditAnywhere, BlueprintReadWrite, Category = COMMON_MESHES_CATEGORY)
+	UPROPERTY(EditAnywhere, BlueprintReadWrite, Category = "Common Meshes")
 	FColor VertexOverrideColor;
 	
 	//////	STATIC_MESHES_CATEGORY Properties //////
 
 
 	/** If enable, import the animation asset find in the sources. */
-	UPROPERTY(EditAnywhere, BlueprintReadWrite, Category = STATIC_MESHES_CATEGORY)
+	UPROPERTY(EditAnywhere, BlueprintReadWrite, Category = "Static Meshes")
 	bool bImportStaticMeshes = true;
 
 	/** If enable all translated static mesh node will be imported has a one static mesh. */
-	UPROPERTY(EditAnywhere, BlueprintReadWrite, Category = STATIC_MESHES_CATEGORY)
+	UPROPERTY(EditAnywhere, BlueprintReadWrite, Category = "Static Meshes")
 	bool bCombineStaticMeshes = false;
 
 	
@@ -133,7 +125,7 @@ public:
 
 
 	/** Skeleton to use for imported asset. When importing a skeletal mesh, leaving this as "None" will create a new skeleton. When importing an animation this MUST be specified to import the asset. */
- 	UPROPERTY(EditAnywhere, BlueprintReadWrite, Category = COMMON_SKELETAL_ANIMATIONS_CATEGORY)
+ 	UPROPERTY(EditAnywhere, BlueprintReadWrite, Category = "Common Skeletal Mesh and Animations")
  	TObjectPtr<class USkeleton> Skeleton;
 
 	
@@ -141,42 +133,42 @@ public:
 
 	
 	/** If enable, import the animation asset find in the sources. */
-	UPROPERTY(EditAnywhere, BlueprintReadWrite, Category = SKELETAL_MESHES_CATEGORY)
+	UPROPERTY(EditAnywhere, BlueprintReadWrite, Category = "Skeletal Meshes")
 	bool bImportSkeletalMeshes = true;
 
 	/** If enable all translated skinned mesh node will be imported has a one skeletal mesh, note that it can still create several skeletal mesh for each different skeleton root joint. */
-	UPROPERTY(EditAnywhere, BlueprintReadWrite, Category = SKELETAL_MESHES_CATEGORY)
+	UPROPERTY(EditAnywhere, BlueprintReadWrite, Category = "Skeletal Meshes")
 	bool bCombineSkeletalMeshes = true;
 
 	/** If enable any morph target shape will be imported. */
-	UPROPERTY(EditAnywhere, BlueprintReadWrite, Category = SKELETAL_MESHES_CATEGORY)
+	UPROPERTY(EditAnywhere, BlueprintReadWrite, Category = "Skeletal Meshes")
 	bool bImportMorphTargets = true;
 
 	/** Enable this option to update Skeleton (of the mesh)'s reference pose. Mesh's reference pose is always updated.  */
-	UPROPERTY(EditAnywhere, BlueprintReadWrite, Category = SKELETAL_MESHES_CATEGORY)
+	UPROPERTY(EditAnywhere, BlueprintReadWrite, Category = "Skeletal Meshes")
 	bool bUpdateSkeletonReferencePose = false;
 
 	/** If checked, meshes nested in bone hierarchies will be imported instead of being converted to bones. */
-	UPROPERTY(EditAnywhere, BlueprintReadWrite, Category = SKELETAL_MESHES_CATEGORY)
+	UPROPERTY(EditAnywhere, BlueprintReadWrite, Category = "Skeletal Meshes")
 	bool bImportMeshesInBoneHierarchy = true;
 
 	/** Enable this option to use frame 0 as reference pose */
-	UPROPERTY(EditAnywhere, BlueprintReadWrite, Category = SKELETAL_MESHES_CATEGORY)
+	UPROPERTY(EditAnywhere, BlueprintReadWrite, Category = "Skeletal Meshes")
 	bool bUseT0AsRefPose = false;
 
 	/** If checked, create new PhysicsAsset if it doesn't have it */
-	UPROPERTY(EditAnywhere, BlueprintReadWrite, Category = SKELETAL_MESHES_CATEGORY)
+	UPROPERTY(EditAnywhere, BlueprintReadWrite, Category = "Skeletal Meshes")
 	bool bCreatePhysicsAsset = true;
 
 	/** If this is set, use this specified PhysicsAsset. If its not set and bCreatePhysicsAsset is false, the importer will not generate or set any physic asset. */
-	UPROPERTY(EditAnywhere, BlueprintReadWrite, Category = SKELETAL_MESHES_CATEGORY, meta = (editcondition = "!bCreatePhysicsAsset"))
+	UPROPERTY(EditAnywhere, BlueprintReadWrite, Category = "Skeletal Meshes", meta = (editcondition = "!bCreatePhysicsAsset"))
 	TObjectPtr<class UPhysicsAsset> PhysicsAsset;
 	
 	//////	ANIMATIONS_CATEGORY Properties //////
 
 
 	/** If enable, import the animation asset find in the sources. */
-// 	UPROPERTY(EditAnywhere, BlueprintReadWrite, Category = ANIMATIONS_CATEGORY)
+// 	UPROPERTY(EditAnywhere, BlueprintReadWrite, Category = "Animations")
 // 	bool bImportAnimations = true;
 
 
@@ -191,7 +183,7 @@ public:
 
 
 	/** If enabled, imports the texture assets found in the sources. */
-	UPROPERTY(EditAnywhere, BlueprintReadWrite, Category = TEXTURES_CATEGORY)
+	UPROPERTY(EditAnywhere, BlueprintReadWrite, Category = "Textures")
 	bool bImportTextures = true;
 
 #if WITH_EDITORONLY_DATA
@@ -199,15 +191,15 @@ public:
 	 * If enable, after a new import a test will be run to see if the texture is a normal map
 	 * If the texture is a normal map the SRG, CompressionSettings and LODGroup settings will be adjusted.
 	 */
-	UPROPERTY(EditAnywhere, BlueprintReadWrite, Category = TEXTURES_CATEGORY)
+	UPROPERTY(EditAnywhere, BlueprintReadWrite, Category = "Textures")
 	bool bDetectNormalMapTexture = true;
 
 	/** If enabled, the texture's green channel will be inverted for normal maps. */
-	UPROPERTY(EditAnywhere, Category = TEXTURES_CATEGORY)
+	UPROPERTY(EditAnywhere, Category = "Textures")
 	bool bFlipNormalMapGreenChannel = false;
 
 	/** Specify the files type that should be imported as long/lat cubemap */
-	UPROPERTY(EditAnywhere, Category = TEXTURES_CATEGORY)
+	UPROPERTY(EditAnywhere, Category = "Textures")
 	TSet<FString> FileExtensionsToImportAsLongLatCubemap = {"hdr"};
 #endif
 
diff --git a/Engine/Plugins/MeshPainting/MeshPainting.uplugin b/Engine/Plugins/MeshPainting/MeshPainting.uplugin
--- a/Engine/Plugins/MeshPainting/MeshPainting.uplugin
+++ b/Engine/Plugins/MeshPainting/MeshPainting.uplugin
@@ -30,6 +30,10 @@
 		{
 			"Name": "GeometryProcessing",
 			"Enabled": true
+		},
+		{
+			"Name": "Interchange",
+			"Enabled": true
 		}
 	],
 	"IsExperimentalVersion": false
diff --git a/Engine/Plugins/MeshPainting/Source/MeshPaintEditorMode/MeshPaintEditorMode.Build.cs b/Engine/Plugins/MeshPainting/Source/MeshPaintEditorMode/MeshPaintEditorMode.Build.cs
--- a/Engine/Plugins/MeshPainting/Source/MeshPaintEditorMode/MeshPaintEditorMode.Build.cs
+++ b/Engine/Plugins/MeshPainting/Source/MeshPaintEditorMode/MeshPaintEditorMode.Build.cs
@@ -34,6 +34,8 @@ namespace UnrealBuildTool.Rules
                     "RenderCore",
                     "RHI",
 					"ToolWidgets",
+					"InterchangeEngine",
+					"InterchangePipelines"
                 }
             );
 
diff --git a/Engine/Plugins/MeshPainting/Source/MeshPaintEditorMode/Private/MeshPaintModeHelpers.cpp b/Engine/Plugins/MeshPainting/Source/MeshPaintEditorMode/Private/MeshPaintModeHelpers.cpp
--- a/Engine/Plugins/MeshPainting/Source/MeshPaintEditorMode/Private/MeshPaintModeHelpers.cpp
+++ b/Engine/Plugins/MeshPainting/Source/MeshPaintEditorMode/Private/MeshPaintModeHelpers.cpp
@@ -46,6 +46,8 @@
 #include "MeshPaintMode.h"
 #include "MeshVertexPaintingTool.h"
 #include "MeshTexturePaintingTool.h"
+#include "InterchangeAssetImportData.h"
+#include "InterchangeGenericAssetsPipeline.h"
 
 
 void UMeshPaintModeSubsystem::SetViewportColorMode(EMeshPaintDataColorViewMode ColorViewMode, FEditorViewportClient* ViewportClient)
@@ -311,6 +313,23 @@ void UMeshPaintModeSubsystem::ImportVertexColorsToSkeletalMesh(USkeletalMesh* Sk
 			ImportData->Modify();
 			ImportData->VertexColorImportOption = EVertexColorImportOption::Ignore;
 		}
+
+		UInterchangeAssetImportData* InterchangeAssetImportData = Cast<UInterchangeAssetImportData>(SkeletalMesh->GetAssetImportData());
+		if (InterchangeAssetImportData)
+		{
+			for (TObjectPtr<UInterchangePipelineBase> PipelineBase : InterchangeAssetImportData->Pipelines)
+			{
+				if (UInterchangeGenericAssetsPipeline* GenericAssetPipeline = Cast<UInterchangeGenericAssetsPipeline>(PipelineBase.Get()))
+				{
+					if (GenericAssetPipeline->VertexColorImportOption != EInterchangeVertexColorImportOption::IVCIO_Ignore)
+					{
+						InterchangeAssetImportData->SetFlags(RF_Transactional);
+						InterchangeAssetImportData->Modify();
+						GenericAssetPipeline->VertexColorImportOption = EInterchangeVertexColorImportOption::IVCIO_Ignore;
+					}
+				}
+			}
+		}
 	}
 }
 
diff --git a/Engine/Plugins/MeshPainting/Source/MeshPaintingToolset/MeshPaintingToolset.Build.cs b/Engine/Plugins/MeshPainting/Source/MeshPaintingToolset/MeshPaintingToolset.Build.cs
--- a/Engine/Plugins/MeshPainting/Source/MeshPaintingToolset/MeshPaintingToolset.Build.cs
+++ b/Engine/Plugins/MeshPainting/Source/MeshPaintingToolset/MeshPaintingToolset.Build.cs
@@ -23,8 +23,10 @@ namespace UnrealBuildTool.Rules
 					"StaticMeshDescription",
 					"TypedElementFramework",
 					"TypedElementRuntime",
-					"UnrealEd"
-                }
+					"UnrealEd",
+					"InterchangeEngine",
+					"InterchangePipelines"
+				}
                 );
 
                 PublicDependencyModuleNames.AddRange(
diff --git a/Engine/Plugins/MeshPainting/Source/MeshPaintingToolset/Private/MeshPaintSkeletalMeshAdapter.cpp b/Engine/Plugins/MeshPainting/Source/MeshPaintingToolset/Private/MeshPaintSkeletalMeshAdapter.cpp
--- a/Engine/Plugins/MeshPainting/Source/MeshPaintingToolset/Private/MeshPaintSkeletalMeshAdapter.cpp
+++ b/Engine/Plugins/MeshPainting/Source/MeshPaintingToolset/Private/MeshPaintSkeletalMeshAdapter.cpp
@@ -9,6 +9,8 @@
 #include "Rendering/SkeletalMeshModel.h"
 #include "Factories/FbxSkeletalMeshImportData.h"
 #include "IndexTypes.h"
+#include "InterchangeAssetImportData.h"
+#include "InterchangeGenericAssetsPipeline.h"
 
 DEFINE_LOG_CATEGORY_STATIC(LogMeshPaintSkeletalMeshAdapter, Log, All);
 //////////////////////////////////////////////////////////////////////////
@@ -360,6 +362,23 @@ void FMeshPaintSkeletalMeshComponentAdapter::PreEdit()
 			ImportData->Modify();
 			ImportData->VertexColorImportOption = EVertexColorImportOption::Ignore;
 		}
+
+		UInterchangeAssetImportData* InterchangeAssetImportData = Cast<UInterchangeAssetImportData>(ReferencedSkeletalMesh->GetAssetImportData());
+		if (InterchangeAssetImportData)
+		{
+			for (TObjectPtr<UInterchangePipelineBase> PipelineBase : InterchangeAssetImportData->Pipelines)
+			{
+				if (UInterchangeGenericAssetsPipeline* GenericAssetPipeline = Cast<UInterchangeGenericAssetsPipeline>(PipelineBase.Get()))
+				{
+					if (GenericAssetPipeline->VertexColorImportOption != EInterchangeVertexColorImportOption::IVCIO_Ignore)
+					{
+						GenericAssetPipeline->SetFlags(RF_Transactional);
+						GenericAssetPipeline->Modify();
+						GenericAssetPipeline->VertexColorImportOption = EInterchangeVertexColorImportOption::IVCIO_Ignore;
+					}
+				}
+			}
+		}
 	}
 }
 
diff --git a/Engine/Source/Runtime/Experimental/Interchange/Engine/Private/InterchangeManager.cpp b/Engine/Source/Runtime/Experimental/Interchange/Engine/Private/InterchangeManager.cpp
--- a/Engine/Source/Runtime/Experimental/Interchange/Engine/Private/InterchangeManager.cpp
+++ b/Engine/Source/Runtime/Experimental/Interchange/Engine/Private/InterchangeManager.cpp
@@ -708,18 +708,23 @@ UInterchangeManager::ImportInternal(const FString& ContentPath, const UInterchan
 				}
 			}
 
-			// Simply move the existing pipeline for now. To be revisited for the MVP release.
-			AsyncHelper->Pipelines = MoveTemp(PipelineStack);
-
-			//if (RegisteredPipelineConfiguration && (bShowPipelineStacksConfigurationDialog || (!DefaultPipelineStacks.Contains(DefaultPipelineStackName) && !bIsUnattended)))
-			//{
-				//Show the re-import dialog to let the user make change in the pipelines
-				//PipelineStackName = RegisteredPipelineConfiguration->ScriptedShowReimportPipelineConfigurationDialog(PipelineStack);
-			//}
+			if (RegisteredPipelineConfiguration && bShowPipelineStacksConfigurationDialog && !bIsUnattended)
+			{
+				//Show the dialog, a plugin should have registered this dialog. We use a plugin to be able to use editor code when doing UI
+				EInterchangePipelineConfigurationDialogResult DialogResult = RegisteredPipelineConfiguration->ScriptedShowReimportPipelineConfigurationDialog(PipelineStack);
+				if (DialogResult == EInterchangePipelineConfigurationDialogResult::Cancel)
+				{
+					bImportCancel = true;
+				}
+				if (DialogResult == EInterchangePipelineConfigurationDialogResult::ImportAll)
+				{
+					bImportAllWithDefault = true;
+				}
+			}
 
+			// Simply move the existing pipeline
+			AsyncHelper->Pipelines = MoveTemp(PipelineStack);
 
-			//If the Stack name is empty it mean we want to use the re-import stack (PipelineStack). If there is a name we will
-			//extract the pipeline stack the user want to use.
 		}
 		else
 		{
diff --git a/Engine/Source/Runtime/Experimental/Interchange/Engine/Public/InterchangePipelineConfigurationBase.h b/Engine/Source/Runtime/Experimental/Interchange/Engine/Public/InterchangePipelineConfigurationBase.h
--- a/Engine/Source/Runtime/Experimental/Interchange/Engine/Public/InterchangePipelineConfigurationBase.h
+++ b/Engine/Source/Runtime/Experimental/Interchange/Engine/Public/InterchangePipelineConfigurationBase.h
@@ -28,6 +28,7 @@ public:
 	 */
 	UFUNCTION(BlueprintNativeEvent, BlueprintCallable, Category = "Interchange | Translator")
 	EInterchangePipelineConfigurationDialogResult ScriptedShowPipelineConfigurationDialog();
+
 	/** The default implementation (call if the blueprint do not have any implementation) will call the virtual ExecuteImportPipeline */
 	EInterchangePipelineConfigurationDialogResult ScriptedShowPipelineConfigurationDialog_Implementation()
 	{
@@ -35,6 +36,19 @@ public:
 		return ShowPipelineConfigurationDialog();
 	}
 
+	/**
+	 * Non virtual helper to allow blueprint to implement event base function to implement the ShowPipelineConfigurationDialog,
+	 */
+	UFUNCTION(BlueprintNativeEvent, BlueprintCallable, Category = "Interchange | Translator")
+	EInterchangePipelineConfigurationDialogResult ScriptedShowReimportPipelineConfigurationDialog(TArray<UInterchangePipelineBase*>& PipelineStack);
+
+	/** The default implementation (call if the blueprint do not have any implementation) will call the virtual ExecuteImportPipeline */
+	EInterchangePipelineConfigurationDialogResult ScriptedShowReimportPipelineConfigurationDialog_Implementation(TArray<UInterchangePipelineBase*>& PipelineStack)
+	{
+		//By default we call the virtual import pipeline execution
+		return ShowReimportPipelineConfigurationDialog(PipelineStack);
+	}
+
 protected:
 
 	/**
@@ -45,4 +59,13 @@ protected:
 		//Not implemented
 		return EInterchangePipelineConfigurationDialogResult::Cancel;
 	}
+
+	/**
+	 * This function show a dialog use to configure pipeline stacks and return a stack name that tell the caller the user choice.
+	 */
+	virtual EInterchangePipelineConfigurationDialogResult ShowReimportPipelineConfigurationDialog(TArray<UInterchangePipelineBase*>& PipelineStack)
+	{
+		//Not implemented
+		return EInterchangePipelineConfigurationDialogResult::Cancel;
+	}
 };
