Interchange Materials
- Added a OneMinus shader node.
- In the material factory, don't set the material blend mode as translucent if the only thing connected to the opacity is a 1.f value.

Interchange FBX
- Flip the transparentcolor value using a OneMinus node before connecting it to the opacity input.
- When parsing generic material properties, ignore properties that would end up being connected to a channel that we already connected something to.

diff --git a/Engine/Plugins/Experimental/Interchange/Runtime/Source/Nodes/Public/InterchangeMaterialDefinitions.h b/Engine/Plugins/Experimental/Interchange/Runtime/Source/Nodes/Public/InterchangeMaterialDefinitions.h
--- a/Engine/Plugins/Experimental/Interchange/Runtime/Source/Nodes/Public/InterchangeMaterialDefinitions.h
+++ b/Engine/Plugins/Experimental/Interchange/Runtime/Source/Nodes/Public/InterchangeMaterialDefinitions.h
@@ -35,6 +35,16 @@ namespace Interchange
 					}
 				}
 
+				namespace OneMinus
+				{
+					const FName Name = TEXT("OneMinus");
+
+					namespace Inputs
+					{
+						const FName Input = TEXT("Input");
+					}
+				}
+
 				namespace TextureSample
 				{
 					const FName Name = TEXT("TextureSample");
diff --git a/Engine/Plugins/Experimental/Interchange/Runtime/Source/Parsers/Fbx/Private/FbxMaterial.cpp b/Engine/Plugins/Experimental/Interchange/Runtime/Source/Parsers/Fbx/Private/FbxMaterial.cpp
--- a/Engine/Plugins/Experimental/Interchange/Runtime/Source/Parsers/Fbx/Private/FbxMaterial.cpp
+++ b/Engine/Plugins/Experimental/Interchange/Runtime/Source/Parsers/Fbx/Private/FbxMaterial.cpp
@@ -59,33 +59,40 @@ namespace UE
 				return TextureNode;
 			}
 
-			void FFbxMaterial::ConvertPropertyToShaderNode(UInterchangeBaseNodeContainer& NodeContainer, UInterchangeShaderGraphNode* ShaderGraphNode, FbxProperty& Property, float Factor, FName InputName, TVariant<FLinearColor, float> DefaultValue)
+			void FFbxMaterial::ConvertPropertyToShaderNode(UInterchangeBaseNodeContainer& NodeContainer, UInterchangeShaderGraphNode* ShaderGraphNode, FbxProperty& Property, float Factor, FName InputName,
+				TVariant<FLinearColor, float> DefaultValue, bool bInverse)
 			{
 				using namespace Materials::Standard::Nodes;
 
-				UInterchangeShaderNode* LerpNode = nullptr;
+				UInterchangeShaderNode* NodeToConnectTo = ShaderGraphNode;
+				FString InputToConnectTo = InputName.ToString();
 
 				const int32 TextureCount = Property.GetSrcObjectCount<FbxFileTexture>();
 				const FbxDataType DataType = Property.GetPropertyDataType();
 
+				if (bInverse)
+				{
+					const FString OneMinusNodeName = InputName.ToString() + TEXT("OneMinus");
+					UInterchangeShaderNode* OneMinusNode = CreateShaderNode(NodeContainer, ShaderGraphNode->GetUniqueID(), OneMinusNodeName);
+					OneMinusNode->SetCustomShaderType(OneMinus::Name.ToString());
+
+					UInterchangeShaderPortsAPI::ConnectDefaultOuputToInput(ShaderGraphNode, InputName.ToString(), OneMinusNode->GetUniqueID());
+
+					NodeToConnectTo = OneMinusNode;
+					InputToConnectTo = OneMinus::Inputs::Input.ToString();
+				}
+
 				if (!FMath::IsNearlyEqual(Factor, 1.f))
 				{
 					FString LerpNodeName = InputName.ToString() + TEXT("Lerp");
-					LerpNode = CreateShaderNode(NodeContainer, ShaderGraphNode->GetUniqueID(), LerpNodeName);
+					UInterchangeShaderNode* LerpNode = CreateShaderNode(NodeContainer, ShaderGraphNode->GetUniqueID(), LerpNodeName);
 					LerpNode->SetCustomShaderType(Lerp::Name.ToString());
 
-					if (TextureCount == 0)
+					if (DefaultValue.IsType<float>())
 					{
-						if (DataType.GetType() == eFbxDouble || DataType.GetType() == eFbxFloat || DataType.GetType() == eFbxInt)
-						{
-							LerpNode->AddFloatAttribute(UInterchangeShaderPortsAPI::MakeInputValueKey(Lerp::Inputs::B.ToString()), DefaultValue.Get<float>());
-						}
-						else if (DataType.GetType() == eFbxDouble3)
-						{
-							LerpNode->AddLinearColorAttribute(UInterchangeShaderPortsAPI::MakeInputValueKey(Lerp::Inputs::B.ToString()), DefaultValue.Get<FLinearColor>());
-						}
+						LerpNode->AddFloatAttribute(UInterchangeShaderPortsAPI::MakeInputValueKey(Lerp::Inputs::B.ToString()), DefaultValue.Get<float>());
 					}
-					else
+					else if (DefaultValue.IsType<FLinearColor>())
 					{
 						LerpNode->AddLinearColorAttribute(UInterchangeShaderPortsAPI::MakeInputValueKey(Lerp::Inputs::B.ToString()), DefaultValue.Get<FLinearColor>());
 					}
@@ -93,7 +100,10 @@ namespace UE
 					const float InverseFactor = 1.f - Factor; // We lerp from A to B and prefer to put the strongest input in A so we need to flip the lerp factor
 					LerpNode->AddFloatAttribute(UInterchangeShaderPortsAPI::MakeInputValueKey(Lerp::Inputs::Factor.ToString()), InverseFactor);
 
-					UInterchangeShaderPortsAPI::ConnectDefaultOuputToInput(ShaderGraphNode, InputName.ToString(), LerpNode->GetUniqueID());
+					UInterchangeShaderPortsAPI::ConnectDefaultOuputToInput(NodeToConnectTo, InputToConnectTo, LerpNode->GetUniqueID());
+
+					NodeToConnectTo = LerpNode;
+					InputToConnectTo = Lerp::Inputs::A.ToString();
 				}
 
 				// Handles max one texture per property.
@@ -124,37 +134,25 @@ namespace UE
 					TextureSampleShader->AddFloatAttribute(UInterchangeShaderPortsAPI::MakeInputValueKey(TextureSample::Inputs::UTiling.ToString()), (float)FbxTextureFilePath->GetScaleU());
 					TextureSampleShader->AddFloatAttribute(UInterchangeShaderPortsAPI::MakeInputValueKey(TextureSample::Inputs::VTiling.ToString()), (float)FbxTextureFilePath->GetScaleV());
 
-					if (LerpNode)
-					{
-						UInterchangeShaderPortsAPI::ConnectDefaultOuputToInput(LerpNode, Lerp::Inputs::A.ToString(), TextureSampleShaderUid);
-					}
-					else
-					{
-						UInterchangeShaderPortsAPI::ConnectDefaultOuputToInput(ShaderGraphNode, InputName.ToString(), TextureSampleShaderUid);
-					}
+					UInterchangeShaderPortsAPI::ConnectDefaultOuputToInput(NodeToConnectTo, InputToConnectTo, TextureSampleShaderUid);
+
 				}
 				else if (DataType.GetType() == eFbxDouble || DataType.GetType() == eFbxFloat || DataType.GetType() == eFbxInt)
 				{
-					if (LerpNode)
-					{
-						LerpNode->AddFloatAttribute(UInterchangeShaderPortsAPI::MakeInputValueKey(Lerp::Inputs::A.ToString()), Property.Get<float>());
-					}
-					else
-					{
-						ShaderGraphNode->AddFloatAttribute(UInterchangeShaderPortsAPI::MakeInputValueKey(InputName.ToString()), Property.Get<float>());
-					}
+					NodeToConnectTo->AddFloatAttribute(UInterchangeShaderPortsAPI::MakeInputValueKey(InputToConnectTo), Property.Get<float>());
 				}
 				else if (DataType.GetType() == eFbxDouble3)
 				{
 					const FLinearColor PropertyValue = FFbxConvert::ConvertColor(Property.Get<FbxDouble3>());
 
-					if (LerpNode)
+					if (DefaultValue.IsType<FLinearColor>())
 					{
-						LerpNode->AddLinearColorAttribute(UInterchangeShaderPortsAPI::MakeInputValueKey(Lerp::Inputs::A.ToString()), PropertyValue);
+						NodeToConnectTo->AddLinearColorAttribute(UInterchangeShaderPortsAPI::MakeInputValueKey(InputToConnectTo), PropertyValue);
 					}
-					else
+					else if (DefaultValue.IsType<float>())
 					{
-						ShaderGraphNode->AddLinearColorAttribute(UInterchangeShaderPortsAPI::MakeInputValueKey(InputName.ToString()), PropertyValue);
+						// We're connecting a linear color to a float input. Ideally, we'd go through a desaturate, but for now we'll just take the red channel and ignore the rest.
+						NodeToConnectTo->AddFloatAttribute(UInterchangeShaderPortsAPI::MakeInputValueKey(InputToConnectTo), PropertyValue.R);
 					}
 				}
 			}
@@ -228,22 +226,26 @@ namespace UE
 								}
 								else if (PropertyName == FbxSurfaceMaterial::sTransparentColor)
 								{
-									const float OpacityFactor = 1.f - (float)((FbxSurfaceLambert*)SurfaceMaterial)->TransparencyFactor.Get();
+									const float TransparencyFactor = (float)((FbxSurfaceLambert*)SurfaceMaterial)->TransparencyFactor.Get();
 									TVariant<FLinearColor, float> DefaultValue;
-									DefaultValue.Set<float>(1.f); // Opaque
-									ConvertPropertyToShaderNode(NodeContainer, ShaderGraphNode, CurrentProperty, OpacityFactor, Phong::Parameters::Opacity, DefaultValue);
+									DefaultValue.Set<float>(0.f); // Opaque
+									const bool bInverse = true;
+									ConvertPropertyToShaderNode(NodeContainer, ShaderGraphNode, CurrentProperty, TransparencyFactor, Phong::Parameters::Opacity, DefaultValue, bInverse);
 								}
 								else if (PropertyName == FbxSurfaceMaterial::sDiffuseFactor || PropertyName == FbxSurfaceMaterial::sSpecularFactor ||
-										 PropertyName == FbxSurfaceMaterial::sTransparencyFactor)
+										 PropertyName == FbxSurfaceMaterial::sEmissiveFactor || PropertyName == FbxSurfaceMaterial::sTransparencyFactor)
 								{
 									// Skip factors as they were processed with their corresponding inputs
 								}
 								else
 								{
-									const float DummyFactor = 1.f;
-									TVariant<FLinearColor, float> DummyDefaultValue;
-									DummyDefaultValue.Set<float>(1.f); // Will be ignore because we have a factor of 1
-									ConvertPropertyToShaderNode(NodeContainer, ShaderGraphNode, CurrentProperty, DummyFactor, PropertyName, DummyDefaultValue);
+									if (!UInterchangeShaderPortsAPI::HasInput(ShaderGraphNode, PropertyName))
+									{
+										const float DummyFactor = 1.f;
+										TVariant<FLinearColor, float> DummyDefaultValue;
+										DummyDefaultValue.Set<float>(1.f); // Will be ignored because we have a factor of 1
+										ConvertPropertyToShaderNode(NodeContainer, ShaderGraphNode, CurrentProperty, DummyFactor, PropertyName, DummyDefaultValue);
+									}
 								}
 							}
 						}
diff --git a/Engine/Plugins/Experimental/Interchange/Runtime/Source/Parsers/Fbx/Private/FbxMaterial.h b/Engine/Plugins/Experimental/Interchange/Runtime/Source/Parsers/Fbx/Private/FbxMaterial.h
--- a/Engine/Plugins/Experimental/Interchange/Runtime/Source/Parsers/Fbx/Private/FbxMaterial.h
+++ b/Engine/Plugins/Experimental/Interchange/Runtime/Source/Parsers/Fbx/Private/FbxMaterial.h
@@ -56,7 +56,7 @@ namespace UE
 				UInterchangeShaderNode* CreateShaderNode(UInterchangeBaseNodeContainer& NodeContainer, const FString& NodeUid, const FString& NodeName);
 				UInterchangeTexture2DNode* CreateTexture2DNode(UInterchangeBaseNodeContainer& NodeContainer, const FString& NodeUID, const FString& TextureFilePath);
 				void ConvertPropertyToShaderNode(UInterchangeBaseNodeContainer& NodeContainer, UInterchangeShaderGraphNode* ShaderGraphNode,
-					FbxProperty& Property, float Factor, FName PropertyName, TVariant<FLinearColor, float> DefaultValue);
+					FbxProperty& Property, float Factor, FName PropertyName, TVariant<FLinearColor, float> DefaultValue, bool bInverse = false);
 
 			private:
 				UInterchangeShaderGraphNode* AddShaderGraphNode(FbxSurfaceMaterial* SurfaceMaterial, UInterchangeBaseNodeContainer& NodeContainer);
diff --git a/Engine/Plugins/Experimental/Interchange/Runtime/Source/Pipelines/Private/InterchangeGenericMaterialPipeline.cpp b/Engine/Plugins/Experimental/Interchange/Runtime/Source/Pipelines/Private/InterchangeGenericMaterialPipeline.cpp
--- a/Engine/Plugins/Experimental/Interchange/Runtime/Source/Pipelines/Private/InterchangeGenericMaterialPipeline.cpp
+++ b/Engine/Plugins/Experimental/Interchange/Runtime/Source/Pipelines/Private/InterchangeGenericMaterialPipeline.cpp
@@ -380,14 +380,25 @@ void UInterchangeGenericMaterialPipeline::HandleCommonParameters(const UIntercha
 
 		if (bHasOpacityInput)
 		{
-			TTuple<UInterchangeMaterialExpressionFactoryNode*, FString> OpacityExpressionFactoryNode =
-				CreateMaterialExpressionForInput(MaterialFactoryNode, ShaderGraphNode, Parameters::Opacity.ToString(), MaterialFactoryNode->GetUniqueID());
+			bool bHasSomeTransparency = true;
+
+			float OpacityValue;
+			if (ShaderGraphNode->GetFloatAttribute(UInterchangeShaderPortsAPI::MakeInputValueKey(Parameters::Opacity.ToString()), OpacityValue))
+			{
+				bHasSomeTransparency = !FMath::IsNearlyEqual(OpacityValue, 1.f);
+			}
 
-			if (OpacityExpressionFactoryNode.Get<0>())
+			if (bHasSomeTransparency)
 			{
-				MaterialFactoryNode->ConnectOutputToOpacity(OpacityExpressionFactoryNode.Get<0>()->GetUniqueID(), OpacityExpressionFactoryNode.Get<1>());
-				MaterialFactoryNode->SetCustomBlendMode(EBlendMode::BLEND_Translucent);
-				MaterialFactoryNode->SetCustomTranslucencyLightingMode(ETranslucencyLightingMode::TLM_Surface);
+				TTuple<UInterchangeMaterialExpressionFactoryNode*, FString> OpacityExpressionFactoryNode =
+				CreateMaterialExpressionForInput(MaterialFactoryNode, ShaderGraphNode, Parameters::Opacity.ToString(), MaterialFactoryNode->GetUniqueID());
+
+				if (OpacityExpressionFactoryNode.Get<0>())
+				{
+					MaterialFactoryNode->ConnectOutputToOpacity(OpacityExpressionFactoryNode.Get<0>()->GetUniqueID(), OpacityExpressionFactoryNode.Get<1>());
+					MaterialFactoryNode->SetCustomBlendMode(EBlendMode::BLEND_Translucent);
+					MaterialFactoryNode->SetCustomTranslucencyLightingMode(ETranslucencyLightingMode::TLM_Surface);
+				}
 			}
 		}
 	}
