Interchange: Moved the factory nodes from the Nodes module to a new FactoryNodes module so that they don't get compiled in the InterchangeWorker and can use engine content, as opposed to the translator nodes, which can't.

diff --git a/Engine/Plugins/Enterprise/DatasmithImporter/Source/DatasmithImporter/DatasmithImporter.Build.cs b/Engine/Plugins/Enterprise/DatasmithImporter/Source/DatasmithImporter/DatasmithImporter.Build.cs
--- a/Engine/Plugins/Enterprise/DatasmithImporter/Source/DatasmithImporter/DatasmithImporter.Build.cs
+++ b/Engine/Plugins/Enterprise/DatasmithImporter/Source/DatasmithImporter/DatasmithImporter.Build.cs
@@ -28,6 +28,7 @@ namespace UnrealBuildTool.Rules
 					"InputCore",
 					"InterchangeCore",
 					"InterchangeEngine",
+					"InterchangeFactoryNodes",
 					"InterchangeNodes",
 					"Json",
 					"Landscape",
diff --git a/Engine/Plugins/Experimental/Interchange/Runtime/Interchange.uplugin b/Engine/Plugins/Experimental/Interchange/Runtime/Interchange.uplugin
--- a/Engine/Plugins/Experimental/Interchange/Runtime/Interchange.uplugin
+++ b/Engine/Plugins/Experimental/Interchange/Runtime/Interchange.uplugin
@@ -21,6 +21,11 @@
 			"Type": "Runtime",
 			"LoadingPhase": "Default"
 		},
+		{
+			"Name": "InterchangeFactoryNodes",
+			"Type": "Runtime",
+			"LoadingPhase": "Default"
+		},
 		{
 			"Name": "InterchangeImport",
 			"Type": "Runtime",
diff --git a/Engine/Plugins/Experimental/Interchange/Runtime/Source/FactoryNodes/InterchangeFactoryNodes.Build.cs b/Engine/Plugins/Experimental/Interchange/Runtime/Source/FactoryNodes/InterchangeFactoryNodes.Build.cs
new file mode 100644
--- /dev/null
+++ b/Engine/Plugins/Experimental/Interchange/Runtime/Source/FactoryNodes/InterchangeFactoryNodes.Build.cs
@@ -0,0 +1,28 @@
+// Copyright Epic Games, Inc. All Rights Reserved.
+
+namespace UnrealBuildTool.Rules
+{
+	public class InterchangeFactoryNodes : ModuleRules
+	{
+		public InterchangeFactoryNodes(ReadOnlyTargetRules Target) : base(Target)
+		{
+			PublicDependencyModuleNames.AddRange(
+				new string[]
+				{
+					"Core",
+					"CoreUObject",
+					"InterchangeCore",
+					"InterchangeNodes",
+				}
+			);
+
+			PublicDependencyModuleNames.AddRange(
+				new string[]
+				{
+					"CinematicCamera",
+					"Engine"
+				}
+			);
+		}
+	}
+}
diff --git a/Engine/Plugins/Experimental/Interchange/Runtime/Source/FactoryNodes/Private/InterchangeActorFactoryNode.cpp b/Engine/Plugins/Experimental/Interchange/Runtime/Source/FactoryNodes/Private/InterchangeActorFactoryNode.cpp
new file mode 100644
--- /dev/null
+++ b/Engine/Plugins/Experimental/Interchange/Runtime/Source/FactoryNodes/Private/InterchangeActorFactoryNode.cpp
@@ -0,0 +1,26 @@
+// Copyright Epic Games, Inc. All Rights Reserved.
+
+#include "InterchangeActorFactoryNode.h"
+
+#if WITH_ENGINE
+	#include "GameFramework/Actor.h"
+#endif
+
+UClass* UInterchangeActorFactoryNode::GetObjectClass() const
+{
+#if WITH_ENGINE
+	FString ActorClassName;
+	if (GetCustomActorClassName(ActorClassName))
+	{
+		UClass* ActorClass = FindObject<UClass>(nullptr, *ActorClassName);
+		if (ActorClass->IsChildOf<AActor>())
+		{
+			return ActorClass;
+		}
+	}
+
+	return AActor::StaticClass();
+#else
+	return nullptr;
+#endif
+}
\ No newline at end of file
diff --git a/Engine/Plugins/Experimental/Interchange/Runtime/Source/FactoryNodes/Private/InterchangeFactoryNodesModule.cpp b/Engine/Plugins/Experimental/Interchange/Runtime/Source/FactoryNodes/Private/InterchangeFactoryNodesModule.cpp
new file mode 100644
--- /dev/null
+++ b/Engine/Plugins/Experimental/Interchange/Runtime/Source/FactoryNodes/Private/InterchangeFactoryNodesModule.cpp
@@ -0,0 +1,31 @@
+// Copyright Epic Games, Inc. All Rights Reserved.
+
+#include "InterchangeFactoryNodesModule.h"
+
+#include "CoreMinimal.h"
+#include "Modules/ModuleManager.h"
+
+class FInterchangeFactoryNodesModule : public IInterchangeFactoryNodesModule
+{
+	/** IModuleInterface implementation */
+	virtual void StartupModule() override;
+	virtual void ShutdownModule() override;
+};
+
+IMPLEMENT_MODULE(FInterchangeFactoryNodesModule, InterchangeFactoryNodes)
+
+
+
+void FInterchangeFactoryNodesModule::StartupModule()
+{
+
+}
+
+
+void FInterchangeFactoryNodesModule::ShutdownModule()
+{
+
+}
+
+
+
diff --git a/Engine/Plugins/Experimental/Interchange/Runtime/Source/FactoryNodes/Private/InterchangeMaterialFactoryNode.cpp b/Engine/Plugins/Experimental/Interchange/Runtime/Source/FactoryNodes/Private/InterchangeMaterialFactoryNode.cpp
new file mode 100644
--- /dev/null
+++ b/Engine/Plugins/Experimental/Interchange/Runtime/Source/FactoryNodes/Private/InterchangeMaterialFactoryNode.cpp
@@ -0,0 +1,269 @@
+// Copyright Epic Games, Inc. All Rights Reserved.
+
+#include "InterchangeMaterialFactoryNode.h"
+
+#if WITH_ENGINE
+
+#include "Materials/Material.h"
+#include "Materials/MaterialInstance.h"
+#include "Materials/MaterialInstanceConstant.h"
+#include "Materials/MaterialInterface.h"
+
+#endif
+
+FString UInterchangeBaseMaterialFactoryNode::GetMaterialFactoryNodeUidFromMaterialNodeUid(const FString& TranslatedNodeUid)
+{
+	FString NewUid = TEXT("Factory_") + TranslatedNodeUid;
+	return NewUid;
+}
+
+FString UInterchangeMaterialFactoryNode::GetTypeName() const
+{
+	const FString TypeName = TEXT("MaterialFactoryNode");
+	return TypeName;
+}
+
+UClass* UInterchangeMaterialFactoryNode::GetObjectClass() const
+{
+#if WITH_ENGINE
+	return UMaterial::StaticClass();
+#else
+	return nullptr;
+#endif
+}
+
+bool UInterchangeMaterialFactoryNode::GetBaseColorConnection(FString& ExpressionNodeUid, FString& OutputName) const
+{
+	return UInterchangeShaderPortsAPI::GetInputConnection(this, UE::Interchange::Materials::PBR::Parameters::BaseColor.ToString(), ExpressionNodeUid, OutputName);
+}
+
+bool UInterchangeMaterialFactoryNode::ConnectToBaseColor(const FString& AttributeValue)
+{
+	return UInterchangeShaderPortsAPI::ConnectDefaultOuputToInput(this, UE::Interchange::Materials::PBR::Parameters::BaseColor.ToString(), AttributeValue);
+}
+	
+bool UInterchangeMaterialFactoryNode::ConnectOutputToBaseColor(const FString& ExpressionNodeUid, const FString& OutputName)
+{
+	return UInterchangeShaderPortsAPI::ConnectOuputToInput(this, UE::Interchange::Materials::PBR::Parameters::BaseColor.ToString(), ExpressionNodeUid, OutputName);
+}
+
+bool UInterchangeMaterialFactoryNode::GetMetallicConnection(FString& ExpressionNodeUid, FString& OutputName) const
+{
+	return UInterchangeShaderPortsAPI::GetInputConnection(this, UE::Interchange::Materials::PBR::Parameters::Metallic.ToString(), ExpressionNodeUid, OutputName);
+}
+
+bool UInterchangeMaterialFactoryNode::ConnectToMetallic(const FString& AttributeValue)
+{
+	return UInterchangeShaderPortsAPI::ConnectDefaultOuputToInput(this, UE::Interchange::Materials::PBR::Parameters::Metallic.ToString(), AttributeValue);
+}
+
+bool UInterchangeMaterialFactoryNode::ConnectOutputToMetallic(const FString& ExpressionNodeUid, const FString& OutputName)
+{
+	return UInterchangeShaderPortsAPI::ConnectOuputToInput(this, UE::Interchange::Materials::PBR::Parameters::Metallic.ToString(), ExpressionNodeUid, OutputName);
+}
+
+bool UInterchangeMaterialFactoryNode::GetSpecularConnection(FString& ExpressionNodeUid, FString& OutputName) const
+{
+	return UInterchangeShaderPortsAPI::GetInputConnection(this, UE::Interchange::Materials::PBR::Parameters::Specular.ToString(), ExpressionNodeUid, OutputName);
+}
+
+bool UInterchangeMaterialFactoryNode::ConnectToSpecular(const FString& ExpressionNodeUid)
+{
+	return UInterchangeShaderPortsAPI::ConnectDefaultOuputToInput(this, UE::Interchange::Materials::PBR::Parameters::Specular.ToString(), ExpressionNodeUid);
+}
+
+bool UInterchangeMaterialFactoryNode::ConnectOutputToSpecular(const FString& ExpressionNodeUid, const FString& OutputName)
+{
+	return UInterchangeShaderPortsAPI::ConnectOuputToInput(this, UE::Interchange::Materials::PBR::Parameters::Specular.ToString(), ExpressionNodeUid, OutputName);
+}
+
+bool UInterchangeMaterialFactoryNode::GetRoughnessConnection(FString& ExpressionNodeUid, FString& OutputName) const
+{
+	return UInterchangeShaderPortsAPI::GetInputConnection(this, UE::Interchange::Materials::PBR::Parameters::Roughness.ToString(), ExpressionNodeUid, OutputName);
+}
+
+bool UInterchangeMaterialFactoryNode::ConnectToRoughness(const FString& ExpressionNodeUid)
+{
+	return UInterchangeShaderPortsAPI::ConnectDefaultOuputToInput(this, UE::Interchange::Materials::PBR::Parameters::Roughness.ToString(), ExpressionNodeUid);
+}
+
+bool UInterchangeMaterialFactoryNode::ConnectOutputToRoughness(const FString& ExpressionNodeUid, const FString& OutputName)
+{
+	return UInterchangeShaderPortsAPI::ConnectOuputToInput(this, UE::Interchange::Materials::PBR::Parameters::Roughness.ToString(), ExpressionNodeUid, OutputName);
+}
+
+bool UInterchangeMaterialFactoryNode::GetEmissiveColorConnection(FString& ExpressionNodeUid, FString& OutputName) const
+{
+	return UInterchangeShaderPortsAPI::GetInputConnection(this, UE::Interchange::Materials::Common::Parameters::EmissiveColor.ToString(), ExpressionNodeUid, OutputName);
+}
+
+bool UInterchangeMaterialFactoryNode::ConnectToEmissiveColor(const FString& ExpressionNodeUid)
+{
+	return UInterchangeShaderPortsAPI::ConnectDefaultOuputToInput(this, UE::Interchange::Materials::Common::Parameters::EmissiveColor.ToString(), ExpressionNodeUid);
+}
+
+bool UInterchangeMaterialFactoryNode::ConnectOutputToEmissiveColor(const FString& ExpressionNodeUid, const FString& OutputName)
+{
+	return UInterchangeShaderPortsAPI::ConnectOuputToInput(this, UE::Interchange::Materials::PBR::Parameters::EmissiveColor.ToString(), ExpressionNodeUid, OutputName);
+}
+
+bool UInterchangeMaterialFactoryNode::GetNormalConnection(FString& ExpressionNodeUid, FString& OutputName) const
+{
+	return UInterchangeShaderPortsAPI::GetInputConnection(this, UE::Interchange::Materials::Common::Parameters::Normal.ToString(), ExpressionNodeUid, OutputName);
+}
+
+bool UInterchangeMaterialFactoryNode::ConnectToNormal(const FString& ExpressionNodeUid)
+{
+	return UInterchangeShaderPortsAPI::ConnectDefaultOuputToInput(this, UE::Interchange::Materials::Common::Parameters::Normal.ToString(), ExpressionNodeUid);
+}
+
+bool UInterchangeMaterialFactoryNode::ConnectOutputToNormal(const FString& ExpressionNodeUid, const FString& OutputName)
+{
+	return UInterchangeShaderPortsAPI::ConnectOuputToInput(this, UE::Interchange::Materials::PBR::Parameters::Normal.ToString(), ExpressionNodeUid, OutputName);
+}
+
+bool UInterchangeMaterialFactoryNode::GetOpacityConnection(FString& ExpressionNodeUid, FString& OutputName) const
+{
+	return UInterchangeShaderPortsAPI::GetInputConnection(this, UE::Interchange::Materials::Common::Parameters::Opacity.ToString(), ExpressionNodeUid, OutputName);
+}
+
+bool UInterchangeMaterialFactoryNode::ConnectToOpacity(const FString& AttributeValue)
+{
+	return UInterchangeShaderPortsAPI::ConnectDefaultOuputToInput(this, UE::Interchange::Materials::Common::Parameters::Opacity.ToString(), AttributeValue);
+}
+
+bool UInterchangeMaterialFactoryNode::ConnectOutputToOpacity(const FString& ExpressionNodeUid, const FString& OutputName)
+{
+	return UInterchangeShaderPortsAPI::ConnectOuputToInput(this, UE::Interchange::Materials::PBR::Parameters::Opacity.ToString(), ExpressionNodeUid, OutputName);
+}
+
+bool UInterchangeMaterialFactoryNode::GetOcclusionConnection(FString& ExpressionNodeUid, FString& OutputName) const
+{
+	return UInterchangeShaderPortsAPI::GetInputConnection(this, UE::Interchange::Materials::Common::Parameters::Occlusion.ToString(), ExpressionNodeUid, OutputName);
+}
+
+bool UInterchangeMaterialFactoryNode::ConnectToOcclusion(const FString& AttributeValue)
+{
+	return UInterchangeShaderPortsAPI::ConnectDefaultOuputToInput(this, UE::Interchange::Materials::Common::Parameters::Occlusion.ToString(), AttributeValue);
+}
+
+bool UInterchangeMaterialFactoryNode::ConnectOutputToOcclusion(const FString& ExpressionNodeUid, const FString& OutputName)
+{
+	return UInterchangeShaderPortsAPI::ConnectOuputToInput(this, UE::Interchange::Materials::PBR::Parameters::Occlusion.ToString(), ExpressionNodeUid, OutputName);
+}
+
+bool UInterchangeMaterialFactoryNode::GetRefractionConnection(FString& ExpressionNodeUid, FString& OutputName) const
+{
+	return UInterchangeShaderPortsAPI::GetInputConnection(this, UE::Interchange::Materials::Common::Parameters::IndexOfRefraction.ToString(), ExpressionNodeUid, OutputName);
+}
+
+bool UInterchangeMaterialFactoryNode::ConnectToRefraction(const FString& AttributeValue)
+{
+	return UInterchangeShaderPortsAPI::ConnectDefaultOuputToInput(this, UE::Interchange::Materials::Common::Parameters::IndexOfRefraction.ToString(), AttributeValue);
+}
+
+bool UInterchangeMaterialFactoryNode::ConnectOutputToRefraction(const FString& ExpressionNodeUid, const FString& OutputName)
+{
+	return UInterchangeShaderPortsAPI::ConnectOuputToInput(this, UE::Interchange::Materials::PBR::Parameters::IndexOfRefraction.ToString(), ExpressionNodeUid, OutputName);
+}
+
+bool UInterchangeMaterialFactoryNode::GetCustomShadingModel(TEnumAsByte<EMaterialShadingModel>& AttributeValue) const
+{
+	IMPLEMENT_NODE_ATTRIBUTE_GETTER(ShadingModel, TEnumAsByte<EMaterialShadingModel>);
+}
+
+bool UInterchangeMaterialFactoryNode::SetCustomShadingModel(const TEnumAsByte<EMaterialShadingModel>& AttributeValue, bool bAddApplyDelegate)
+{
+	IMPLEMENT_NODE_ATTRIBUTE_SETTER(UInterchangeMaterialFactoryNode, ShadingModel, TEnumAsByte<EMaterialShadingModel>, UMaterial);
+}
+
+bool UInterchangeMaterialFactoryNode::GetCustomTranslucencyLightingMode(TEnumAsByte<ETranslucencyLightingMode>& AttributeValue) const
+{
+	IMPLEMENT_NODE_ATTRIBUTE_GETTER(TranslucencyLightingMode, TEnumAsByte<ETranslucencyLightingMode>);
+}
+
+bool UInterchangeMaterialFactoryNode::SetCustomTranslucencyLightingMode(const TEnumAsByte<ETranslucencyLightingMode>& AttributeValue, bool bAddApplyDelegate)
+{
+	IMPLEMENT_NODE_ATTRIBUTE_SETTER(UInterchangeMaterialFactoryNode, TranslucencyLightingMode, TEnumAsByte<ETranslucencyLightingMode>, UMaterial);
+}
+
+bool UInterchangeMaterialFactoryNode::GetCustomBlendMode(TEnumAsByte<EBlendMode>& AttributeValue) const
+{
+	IMPLEMENT_NODE_ATTRIBUTE_GETTER(BlendMode, TEnumAsByte<EBlendMode>);
+}
+
+bool UInterchangeMaterialFactoryNode::SetCustomBlendMode(const TEnumAsByte<EBlendMode>& AttributeValue, bool bAddApplyDelegate)
+{
+	IMPLEMENT_NODE_ATTRIBUTE_SETTER(UInterchangeMaterialFactoryNode, BlendMode, TEnumAsByte<EBlendMode>, UMaterial);
+}
+
+bool UInterchangeMaterialFactoryNode::GetCustomTwoSided(bool& AttributeValue) const
+{
+	IMPLEMENT_NODE_ATTRIBUTE_GETTER(TwoSided, bool);
+}
+
+bool UInterchangeMaterialFactoryNode::SetCustomTwoSided(const bool& AttributeValue, bool bAddApplyDelegate)
+{
+	IMPLEMENT_NODE_ATTRIBUTE_SETTER(UInterchangeMaterialFactoryNode, TwoSided, bool, UMaterial);
+}
+
+FString UInterchangeMaterialExpressionFactoryNode::GetTypeName() const
+{
+	const FString TypeName = TEXT("MaterialExpressionFactoryNode");
+	return TypeName;
+}
+
+bool UInterchangeMaterialExpressionFactoryNode::GetCustomExpressionClassName(FString& AttributeValue) const
+{
+	IMPLEMENT_NODE_ATTRIBUTE_GETTER(ExpressionClassName, FString);
+}
+
+bool UInterchangeMaterialExpressionFactoryNode::SetCustomExpressionClassName(const FString& AttributeValue)
+{
+	IMPLEMENT_NODE_ATTRIBUTE_SETTER_NODELEGATE(ExpressionClassName, FString);
+}
+
+FString UInterchangeMaterialInstanceFactoryNode::GetTypeName() const
+{
+	const FString TypeName = TEXT("MaterialInstanceFactoryNode");
+	return TypeName;
+}
+
+UClass* UInterchangeMaterialInstanceFactoryNode::GetObjectClass() const
+{
+#if WITH_ENGINE
+	FString InstanceClassName;
+	if (GetCustomInstanceClassName(InstanceClassName))
+	{
+		UClass* InstanceClass = FindObject<UClass>(nullptr, *InstanceClassName);
+		if (InstanceClass->IsChildOf<UMaterialInstance>())
+		{
+			return InstanceClass;
+		}
+	}
+
+	return UMaterialInstanceConstant::StaticClass();
+#else
+	return nullptr;
+#endif
+}
+
+bool UInterchangeMaterialInstanceFactoryNode::GetCustomInstanceClassName(FString& AttributeValue) const
+{
+	IMPLEMENT_NODE_ATTRIBUTE_GETTER(InstanceClassName, FString);
+}
+
+bool UInterchangeMaterialInstanceFactoryNode::SetCustomInstanceClassName(const FString& AttributeValue)
+{
+	IMPLEMENT_NODE_ATTRIBUTE_SETTER_NODELEGATE(InstanceClassName, FString);
+}
+
+bool UInterchangeMaterialInstanceFactoryNode::GetCustomParent(FString& AttributeValue) const
+{
+	IMPLEMENT_NODE_ATTRIBUTE_GETTER(Parent, FString);
+}
+
+bool UInterchangeMaterialInstanceFactoryNode::SetCustomParent(const FString& AttributeValue)
+{
+	IMPLEMENT_NODE_ATTRIBUTE_SETTER_NODELEGATE(Parent, FString);
+}
diff --git a/Engine/Plugins/Experimental/Interchange/Runtime/Source/FactoryNodes/Private/InterchangeSkeletalMeshFactoryNode.cpp b/Engine/Plugins/Experimental/Interchange/Runtime/Source/FactoryNodes/Private/InterchangeSkeletalMeshFactoryNode.cpp
new file mode 100644
--- /dev/null
+++ b/Engine/Plugins/Experimental/Interchange/Runtime/Source/FactoryNodes/Private/InterchangeSkeletalMeshFactoryNode.cpp
@@ -0,0 +1,211 @@
+// Copyright Epic Games, Inc. All Rights Reserved.
+
+#include "InterchangeSkeletalMeshFactoryNode.h"
+
+#if WITH_ENGINE
+#include "Engine/SkeletalMesh.h"
+#endif
+
+//Interchange namespace
+namespace UE
+{
+	namespace Interchange
+	{
+		const FString& FSkeletalMeshNodeStaticData::GetLodDependenciesBaseKey()
+		{
+			static FString LodDependencies_BaseKey = TEXT("Lod_Dependencies");
+			return LodDependencies_BaseKey;
+		}
+	}//ns Interchange
+}//ns UE
+
+UInterchangeSkeletalMeshFactoryNode::UInterchangeSkeletalMeshFactoryNode()
+{
+#if WITH_ENGINE
+	AssetClass = nullptr;
+#endif
+	LodDependencies.Initialize(Attributes, UE::Interchange::FSkeletalMeshNodeStaticData::GetLodDependenciesBaseKey());
+}
+
+void UInterchangeSkeletalMeshFactoryNode::InitializeSkeletalMeshNode(const FString& UniqueID, const FString& DisplayLabel, const FString& InAssetClass)
+{
+	bIsNodeClassInitialized = false;
+	InitializeNode(UniqueID, DisplayLabel, EInterchangeNodeContainerType::FactoryData);
+
+	FString OperationName = GetTypeName() + TEXT(".SetAssetClassName");
+	InterchangePrivateNodeBase::SetCustomAttribute<FString>(*Attributes, ClassNameAttributeKey, OperationName, InAssetClass);
+	FillAssetClassFromAttribute();
+}
+
+void UInterchangeSkeletalMeshFactoryNode::Serialize(FArchive& Ar)
+{
+	Super::Serialize(Ar);
+#if WITH_ENGINE
+	if (Ar.IsLoading())
+	{
+		//Make sure the class is properly set when we compile with engine, this will set the
+		//bIsNodeClassInitialized to true.
+		SetNodeClassFromClassAttribute();
+	}
+#endif //#if WITH_ENGINE
+}
+
+FString UInterchangeSkeletalMeshFactoryNode::GetTypeName() const
+{
+	const FString TypeName = TEXT("SkeletalMeshNode");
+	return TypeName;
+}
+
+FString UInterchangeSkeletalMeshFactoryNode::GetKeyDisplayName(const UE::Interchange::FAttributeKey& NodeAttributeKey) const
+{
+	FString KeyDisplayName = NodeAttributeKey.Key;
+	if (NodeAttributeKey.Key.Equals(UE::Interchange::FSkeletalMeshNodeStaticData::GetLodDependenciesBaseKey()))
+	{
+		KeyDisplayName = TEXT("LOD Dependencies Count");
+		return KeyDisplayName;
+	}
+	else if (NodeAttributeKey.Key.StartsWith(UE::Interchange::FSkeletalMeshNodeStaticData::GetLodDependenciesBaseKey()))
+	{
+		KeyDisplayName = TEXT("LOD Dependencies Index ");
+		const FString IndexKey = UE::Interchange::TArrayAttributeHelper<FString>::IndexKey();
+		int32 IndexPosition = NodeAttributeKey.Key.Find(IndexKey) + IndexKey.Len();
+		if (IndexPosition < NodeAttributeKey.Key.Len())
+		{
+			KeyDisplayName += NodeAttributeKey.Key.RightChop(IndexPosition);
+		}
+		return KeyDisplayName;
+	}
+	return Super::GetKeyDisplayName(NodeAttributeKey);
+}
+
+UClass* UInterchangeSkeletalMeshFactoryNode::GetObjectClass() const
+{
+	ensure(bIsNodeClassInitialized);
+#if WITH_ENGINE
+	return AssetClass.Get() != nullptr ? AssetClass.Get() : USkeletalMesh::StaticClass();
+#else
+	return nullptr;
+#endif
+}
+
+int32 UInterchangeSkeletalMeshFactoryNode::GetLodDataCount() const
+{
+	return LodDependencies.GetCount();
+}
+
+void UInterchangeSkeletalMeshFactoryNode::GetLodDataUniqueIds(TArray<FString>& OutLodDataUniqueIds) const
+{
+	LodDependencies.GetItems(OutLodDataUniqueIds);
+}
+
+bool UInterchangeSkeletalMeshFactoryNode::AddLodDataUniqueId(const FString& LodDataUniqueId)
+{
+	return LodDependencies.AddItem(LodDataUniqueId);
+}
+
+bool UInterchangeSkeletalMeshFactoryNode::RemoveLodDataUniqueId(const FString& LodDataUniqueId)
+{
+	return LodDependencies.RemoveItem(LodDataUniqueId);
+}
+
+bool UInterchangeSkeletalMeshFactoryNode::GetCustomSkeletonSoftObjectPath(FSoftObjectPath& AttributeValue) const
+{
+	IMPLEMENT_NODE_ATTRIBUTE_GETTER(SkeletonSoftObjectPath, FSoftObjectPath)
+}
+
+bool UInterchangeSkeletalMeshFactoryNode::SetCustomSkeletonSoftObjectPath(const FSoftObjectPath& AttributeValue)
+{
+	IMPLEMENT_NODE_ATTRIBUTE_SETTER_NODELEGATE(SkeletonSoftObjectPath, FSoftObjectPath)
+}
+
+bool UInterchangeSkeletalMeshFactoryNode::GetCustomImportMorphTarget(bool& AttributeValue) const
+{
+	IMPLEMENT_NODE_ATTRIBUTE_GETTER(ImportMorphTarget, bool)
+}
+
+bool UInterchangeSkeletalMeshFactoryNode::SetCustomImportMorphTarget(const bool& AttributeValue)
+{
+	IMPLEMENT_NODE_ATTRIBUTE_SETTER_NODELEGATE(ImportMorphTarget, bool)
+}
+
+bool UInterchangeSkeletalMeshFactoryNode::GetCustomCreatePhysicsAsset(bool& AttributeValue) const
+{
+	IMPLEMENT_NODE_ATTRIBUTE_GETTER(CreatePhysicsAsset, bool)
+}
+
+bool UInterchangeSkeletalMeshFactoryNode::SetCustomCreatePhysicsAsset(const bool& AttributeValue)
+{
+	IMPLEMENT_NODE_ATTRIBUTE_SETTER_NODELEGATE(CreatePhysicsAsset, bool)
+}
+
+bool UInterchangeSkeletalMeshFactoryNode::GetCustomPhysicAssetSoftObjectPath(FSoftObjectPath& AttributeValue) const
+{
+	IMPLEMENT_NODE_ATTRIBUTE_GETTER(PhysicAssetSoftObjectPath, FSoftObjectPath)
+}
+
+bool UInterchangeSkeletalMeshFactoryNode::SetCustomPhysicAssetSoftObjectPath(const FSoftObjectPath& AttributeValue)
+{
+	IMPLEMENT_NODE_ATTRIBUTE_SETTER_NODELEGATE(PhysicAssetSoftObjectPath, FSoftObjectPath)
+}
+
+bool UInterchangeSkeletalMeshFactoryNode::GetCustomVertexColorReplace(bool& AttributeValue) const
+{
+	IMPLEMENT_NODE_ATTRIBUTE_GETTER(VertexColorReplace, bool)
+}
+
+bool UInterchangeSkeletalMeshFactoryNode::SetCustomVertexColorReplace(const bool& AttributeValue)
+{
+	IMPLEMENT_NODE_ATTRIBUTE_SETTER_NODELEGATE(VertexColorReplace, bool)
+}
+
+bool UInterchangeSkeletalMeshFactoryNode::GetCustomVertexColorIgnore(bool& AttributeValue) const
+{
+	IMPLEMENT_NODE_ATTRIBUTE_GETTER(VertexColorIgnore, bool)
+}
+
+bool UInterchangeSkeletalMeshFactoryNode::SetCustomVertexColorIgnore(const bool& AttributeValue)
+{
+	IMPLEMENT_NODE_ATTRIBUTE_SETTER_NODELEGATE(VertexColorIgnore, bool)
+}
+
+bool UInterchangeSkeletalMeshFactoryNode::GetCustomVertexColorOverride(FColor& AttributeValue) const
+{
+	IMPLEMENT_NODE_ATTRIBUTE_GETTER(VertexColorOverride, FColor)
+}
+
+bool UInterchangeSkeletalMeshFactoryNode::SetCustomVertexColorOverride(const FColor& AttributeValue)
+{
+	IMPLEMENT_NODE_ATTRIBUTE_SETTER_NODELEGATE(VertexColorOverride, FColor)
+}
+
+void UInterchangeSkeletalMeshFactoryNode::FillAssetClassFromAttribute()
+{
+#if WITH_ENGINE
+	FString OperationName = GetTypeName() + TEXT(".GetAssetClassName");
+	FString ClassName;
+	InterchangePrivateNodeBase::GetCustomAttribute<FString>(*Attributes, ClassNameAttributeKey, OperationName, ClassName);
+	if (ClassName.Equals(USkeletalMesh::StaticClass()->GetName()))
+	{
+		AssetClass = USkeletalMesh::StaticClass();
+		bIsNodeClassInitialized = true;
+	}
+#endif
+}
+
+bool UInterchangeSkeletalMeshFactoryNode::SetNodeClassFromClassAttribute()
+{
+	if (!bIsNodeClassInitialized)
+	{
+		FillAssetClassFromAttribute();
+	}
+	return bIsNodeClassInitialized;
+}
+
+bool UInterchangeSkeletalMeshFactoryNode::IsEditorOnlyDataDefined()
+{
+#if WITH_EDITORONLY_DATA
+	return true;
+#else
+	return false;
+#endif
+}
diff --git a/Engine/Plugins/Experimental/Interchange/Runtime/Source/FactoryNodes/Private/InterchangeSkeletalMeshLodDataNode.cpp b/Engine/Plugins/Experimental/Interchange/Runtime/Source/FactoryNodes/Private/InterchangeSkeletalMeshLodDataNode.cpp
new file mode 100644
--- /dev/null
+++ b/Engine/Plugins/Experimental/Interchange/Runtime/Source/FactoryNodes/Private/InterchangeSkeletalMeshLodDataNode.cpp
@@ -0,0 +1,110 @@
+// Copyright Epic Games, Inc. All Rights Reserved.
+
+#include "InterchangeSkeletalMeshLodDataNode.h"
+
+//Interchange namespace
+namespace UE
+{
+	namespace Interchange
+	{
+		const FString& FSkeletalMeshNodeLodDataStaticData::GetMeshUidsBaseKey()
+		{
+			static FString MeshUids_BaseKey = TEXT("__MeshUids__Key");
+			return MeshUids_BaseKey;
+		}
+	}//ns Interchange
+}//ns UE
+
+UInterchangeSkeletalMeshLodDataNode::UInterchangeSkeletalMeshLodDataNode()
+{
+	MeshUids.Initialize(Attributes, UE::Interchange::FSkeletalMeshNodeLodDataStaticData::GetMeshUidsBaseKey());
+}
+
+/**
+	* Return the node type name of the class, we use this when reporting error
+	*/
+FString UInterchangeSkeletalMeshLodDataNode::GetTypeName() const
+{
+	const FString TypeName = TEXT("SkeletalMeshLodDataNode");
+	return TypeName;
+}
+
+FString UInterchangeSkeletalMeshLodDataNode::GetKeyDisplayName(const UE::Interchange::FAttributeKey& NodeAttributeKey) const
+{
+	FString KeyDisplayName = NodeAttributeKey.Key;
+	if (NodeAttributeKey.Key.Equals(UE::Interchange::FSkeletalMeshNodeLodDataStaticData::GetMeshUidsBaseKey()))
+	{
+		KeyDisplayName = TEXT("Mesh count");
+		return KeyDisplayName;
+	}
+	else if (NodeAttributeKey.Key.StartsWith(UE::Interchange::FSkeletalMeshNodeLodDataStaticData::GetMeshUidsBaseKey()))
+	{
+		KeyDisplayName = TEXT("Mesh index ");
+		const FString IndexKey = UE::Interchange::TArrayAttributeHelper<FString>::IndexKey();
+		int32 IndexPosition = NodeAttributeKey.Key.Find(IndexKey) + IndexKey.Len();
+		if (IndexPosition < NodeAttributeKey.Key.Len())
+		{
+			KeyDisplayName += NodeAttributeKey.Key.RightChop(IndexPosition);
+		}
+		return KeyDisplayName;
+	}
+	else if (NodeAttributeKey == Macro_CustomSkeletonUidKey)
+	{
+		KeyDisplayName = TEXT("Skeleton factory node");
+		return KeyDisplayName;
+	}
+	return Super::GetKeyDisplayName(NodeAttributeKey);
+}
+
+FString UInterchangeSkeletalMeshLodDataNode::GetAttributeCategory(const UE::Interchange::FAttributeKey& NodeAttributeKey) const
+{
+	if (NodeAttributeKey.Key.StartsWith(UE::Interchange::FSkeletalMeshNodeLodDataStaticData::GetMeshUidsBaseKey()))
+	{
+		return FString(TEXT("Meshes"));
+	}
+	return Super::GetAttributeCategory(NodeAttributeKey);
+}
+
+bool UInterchangeSkeletalMeshLodDataNode::GetCustomSkeletonUid(FString& AttributeValue) const
+{
+	IMPLEMENT_NODE_ATTRIBUTE_GETTER(SkeletonUid, FString);
+}
+
+bool UInterchangeSkeletalMeshLodDataNode::SetCustomSkeletonUid(const FString& AttributeValue)
+{
+	IMPLEMENT_NODE_ATTRIBUTE_SETTER_NODELEGATE(SkeletonUid, FString)
+}
+
+int32 UInterchangeSkeletalMeshLodDataNode::GetMeshUidsCount() const
+{
+	return MeshUids.GetCount();
+}
+
+void UInterchangeSkeletalMeshLodDataNode::GetMeshUids(TArray<FString>& OutMeshNames) const
+{
+	MeshUids.GetItems(OutMeshNames);
+}
+
+bool UInterchangeSkeletalMeshLodDataNode::AddMeshUid(const FString& MeshName)
+{
+	return MeshUids.AddItem(MeshName);
+}
+
+bool UInterchangeSkeletalMeshLodDataNode::RemoveMeshUid(const FString& MeshName)
+{
+	return MeshUids.RemoveItem(MeshName);
+}
+
+bool UInterchangeSkeletalMeshLodDataNode::RemoveAllMeshes()
+{
+	return MeshUids.RemoveAllItems();
+}
+
+bool UInterchangeSkeletalMeshLodDataNode::IsEditorOnlyDataDefined()
+{
+#if WITH_EDITORONLY_DATA
+	return true;
+#else
+	return false;
+#endif
+}
diff --git a/Engine/Plugins/Experimental/Interchange/Runtime/Source/FactoryNodes/Private/InterchangeStaticMeshFactoryNode.cpp b/Engine/Plugins/Experimental/Interchange/Runtime/Source/FactoryNodes/Private/InterchangeStaticMeshFactoryNode.cpp
new file mode 100644
--- /dev/null
+++ b/Engine/Plugins/Experimental/Interchange/Runtime/Source/FactoryNodes/Private/InterchangeStaticMeshFactoryNode.cpp
@@ -0,0 +1,170 @@
+// Copyright Epic Games, Inc. All Rights Reserved.
+
+#include "InterchangeStaticMeshFactoryNode.h"
+
+#if WITH_ENGINE
+#include "Engine/StaticMesh.h"
+#endif
+
+namespace UE
+{
+	namespace Interchange
+	{
+		const FString& FStaticMeshNodeStaticData::GetLodDependenciesBaseKey()
+		{
+			static FString LodDependencies_BaseKey = TEXT("Lod_Dependencies");
+			return LodDependencies_BaseKey;
+		}
+	} // namespace Interchange
+} // namespace UE
+
+
+UInterchangeStaticMeshFactoryNode::UInterchangeStaticMeshFactoryNode()
+{
+#if WITH_ENGINE
+	AssetClass = nullptr;
+#endif
+	LodDependencies.Initialize(Attributes, UE::Interchange::FStaticMeshNodeStaticData::GetLodDependenciesBaseKey());
+}
+
+void UInterchangeStaticMeshFactoryNode::InitializeStaticMeshNode(const FString& UniqueID, const FString& DisplayLabel, const FString& InAssetClass)
+{
+	bIsNodeClassInitialized = false;
+	InitializeNode(UniqueID, DisplayLabel, EInterchangeNodeContainerType::FactoryData);
+
+	FString OperationName = GetTypeName() + TEXT(".SetAssetClassName");
+	InterchangePrivateNodeBase::SetCustomAttribute<FString>(*Attributes, ClassNameAttributeKey, OperationName, InAssetClass);
+	FillAssetClassFromAttribute();
+}
+
+void UInterchangeStaticMeshFactoryNode::Serialize(FArchive& Ar)
+{
+	Super::Serialize(Ar);
+#if WITH_ENGINE
+	if (Ar.IsLoading())
+	{
+		// Make sure the class is properly set when we compile with engine, this will set the bIsNodeClassInitialized to true.
+		SetNodeClassFromClassAttribute();
+	}
+#endif //#if WITH_ENGINE
+}
+
+FString UInterchangeStaticMeshFactoryNode::GetTypeName() const
+{
+	const FString TypeName = TEXT("StaticMeshNode");
+	return TypeName;
+}
+
+FString UInterchangeStaticMeshFactoryNode::GetKeyDisplayName(const UE::Interchange::FAttributeKey& NodeAttributeKey) const
+{
+	FString KeyDisplayName = NodeAttributeKey.Key;
+	if (NodeAttributeKey.Key.Equals(UE::Interchange::FStaticMeshNodeStaticData::GetLodDependenciesBaseKey()))
+	{
+		KeyDisplayName = TEXT("LOD Dependencies Count");
+		return KeyDisplayName;
+	}
+	else if (NodeAttributeKey.Key.StartsWith(UE::Interchange::FStaticMeshNodeStaticData::GetLodDependenciesBaseKey()))
+	{
+		KeyDisplayName = TEXT("LOD Dependencies Index ");
+		const FString IndexKey = UE::Interchange::TArrayAttributeHelper<FString>::IndexKey();
+		int32 IndexPosition = NodeAttributeKey.Key.Find(IndexKey) + IndexKey.Len();
+		if (IndexPosition < NodeAttributeKey.Key.Len())
+		{
+			KeyDisplayName += NodeAttributeKey.Key.RightChop(IndexPosition);
+		}
+		return KeyDisplayName;
+	}
+	return Super::GetKeyDisplayName(NodeAttributeKey);
+}
+
+UClass* UInterchangeStaticMeshFactoryNode::GetObjectClass() const
+{
+	ensure(bIsNodeClassInitialized);
+#if WITH_ENGINE
+	return AssetClass.Get() != nullptr ? AssetClass.Get() : UStaticMesh::StaticClass();
+#else
+	return nullptr;
+#endif
+}
+
+int32 UInterchangeStaticMeshFactoryNode::GetLodDataCount() const
+{
+	return LodDependencies.GetCount();
+}
+
+void UInterchangeStaticMeshFactoryNode::GetLodDataUniqueIds(TArray<FString>& OutLodDataUniqueIds) const
+{
+	LodDependencies.GetItems(OutLodDataUniqueIds);
+}
+
+bool UInterchangeStaticMeshFactoryNode::AddLodDataUniqueId(const FString& LodDataUniqueId)
+{
+	return LodDependencies.AddItem(LodDataUniqueId);
+}
+
+bool UInterchangeStaticMeshFactoryNode::RemoveLodDataUniqueId(const FString& LodDataUniqueId)
+{
+	return LodDependencies.RemoveItem(LodDataUniqueId);
+}
+
+bool UInterchangeStaticMeshFactoryNode::GetCustomVertexColorReplace(bool& AttributeValue) const
+{
+	IMPLEMENT_NODE_ATTRIBUTE_GETTER(VertexColorReplace, bool)
+}
+
+bool UInterchangeStaticMeshFactoryNode::SetCustomVertexColorReplace(const bool& AttributeValue)
+{
+	IMPLEMENT_NODE_ATTRIBUTE_SETTER_NODELEGATE(VertexColorReplace, bool)
+}
+
+bool UInterchangeStaticMeshFactoryNode::GetCustomVertexColorIgnore(bool& AttributeValue) const
+{
+	IMPLEMENT_NODE_ATTRIBUTE_GETTER(VertexColorIgnore, bool)
+}
+
+bool UInterchangeStaticMeshFactoryNode::SetCustomVertexColorIgnore(const bool& AttributeValue)
+{
+	IMPLEMENT_NODE_ATTRIBUTE_SETTER_NODELEGATE(VertexColorIgnore, bool)
+}
+
+bool UInterchangeStaticMeshFactoryNode::GetCustomVertexColorOverride(FColor& AttributeValue) const
+{
+	IMPLEMENT_NODE_ATTRIBUTE_GETTER(VertexColorOverride, FColor)
+}
+
+bool UInterchangeStaticMeshFactoryNode::SetCustomVertexColorOverride(const FColor& AttributeValue)
+{
+	IMPLEMENT_NODE_ATTRIBUTE_SETTER_NODELEGATE(VertexColorOverride, FColor)
+}
+
+void UInterchangeStaticMeshFactoryNode::FillAssetClassFromAttribute()
+{
+#if WITH_ENGINE
+	FString OperationName = GetTypeName() + TEXT(".GetAssetClassName");
+	FString ClassName;
+	InterchangePrivateNodeBase::GetCustomAttribute<FString>(*Attributes, ClassNameAttributeKey, OperationName, ClassName);
+	if (ClassName.Equals(UStaticMesh::StaticClass()->GetName()))
+	{
+		AssetClass = UStaticMesh::StaticClass();
+		bIsNodeClassInitialized = true;
+	}
+#endif
+}
+
+bool UInterchangeStaticMeshFactoryNode::SetNodeClassFromClassAttribute()
+{
+	if (!bIsNodeClassInitialized)
+	{
+		FillAssetClassFromAttribute();
+	}
+	return bIsNodeClassInitialized;
+}
+
+bool UInterchangeStaticMeshFactoryNode::IsEditorOnlyDataDefined()
+{
+#if WITH_EDITORONLY_DATA
+	return true;
+#else
+	return false;
+#endif
+}
diff --git a/Engine/Plugins/Experimental/Interchange/Runtime/Source/FactoryNodes/Private/InterchangeStaticMeshLodDataNode.cpp b/Engine/Plugins/Experimental/Interchange/Runtime/Source/FactoryNodes/Private/InterchangeStaticMeshLodDataNode.cpp
new file mode 100644
--- /dev/null
+++ b/Engine/Plugins/Experimental/Interchange/Runtime/Source/FactoryNodes/Private/InterchangeStaticMeshLodDataNode.cpp
@@ -0,0 +1,96 @@
+// Copyright Epic Games, Inc. All Rights Reserved.
+
+#include "InterchangeStaticMeshLodDataNode.h"
+
+namespace UE
+{
+	namespace Interchange
+	{
+		const FString& FStaticMeshNodeLodDataStaticData::GetMeshUidsBaseKey()
+		{
+			static FString MeshUids_BaseKey = TEXT("__MeshUids__Key");
+			return MeshUids_BaseKey;
+		}
+	} // namespace Interchange
+} // namespace UE
+
+
+UInterchangeStaticMeshLodDataNode::UInterchangeStaticMeshLodDataNode()
+{
+	MeshUids.Initialize(Attributes, UE::Interchange::FStaticMeshNodeLodDataStaticData::GetMeshUidsBaseKey());
+}
+
+/**
+	* Return the node type name of the class, we use this when reporting error
+	*/
+FString UInterchangeStaticMeshLodDataNode::GetTypeName() const
+{
+	const FString TypeName = TEXT("StaticMeshLodDataNode");
+	return TypeName;
+}
+
+FString UInterchangeStaticMeshLodDataNode::GetKeyDisplayName(const UE::Interchange::FAttributeKey& NodeAttributeKey) const
+{
+	FString KeyDisplayName = NodeAttributeKey.Key;
+	if (NodeAttributeKey.Key.Equals(UE::Interchange::FStaticMeshNodeLodDataStaticData::GetMeshUidsBaseKey()))
+	{
+		KeyDisplayName = TEXT("Mesh count");
+		return KeyDisplayName;
+	}
+	else if (NodeAttributeKey.Key.StartsWith(UE::Interchange::FStaticMeshNodeLodDataStaticData::GetMeshUidsBaseKey()))
+	{
+		KeyDisplayName = TEXT("Mesh index ");
+		const FString IndexKey = UE::Interchange::TArrayAttributeHelper<FString>::IndexKey();
+		int32 IndexPosition = NodeAttributeKey.Key.Find(IndexKey) + IndexKey.Len();
+		if (IndexPosition < NodeAttributeKey.Key.Len())
+		{
+			KeyDisplayName += NodeAttributeKey.Key.RightChop(IndexPosition);
+		}
+		return KeyDisplayName;
+	}
+
+	return Super::GetKeyDisplayName(NodeAttributeKey);
+}
+
+FString UInterchangeStaticMeshLodDataNode::GetAttributeCategory(const UE::Interchange::FAttributeKey& NodeAttributeKey) const
+{
+	if (NodeAttributeKey.Key.StartsWith(UE::Interchange::FStaticMeshNodeLodDataStaticData::GetMeshUidsBaseKey()))
+	{
+		return FString(TEXT("Meshes"));
+	}
+	return Super::GetAttributeCategory(NodeAttributeKey);
+}
+
+int32 UInterchangeStaticMeshLodDataNode::GetMeshUidsCount() const
+{
+	return MeshUids.GetCount();
+}
+
+void UInterchangeStaticMeshLodDataNode::GetMeshUids(TArray<FString>& OutMeshNames) const
+{
+	MeshUids.GetItems(OutMeshNames);
+}
+
+bool UInterchangeStaticMeshLodDataNode::AddMeshUid(const FString& MeshName)
+{
+	return MeshUids.AddItem(MeshName);
+}
+
+bool UInterchangeStaticMeshLodDataNode::RemoveMeshUid(const FString& MeshName)
+{
+	return MeshUids.RemoveItem(MeshName);
+}
+
+bool UInterchangeStaticMeshLodDataNode::RemoveAllMeshes()
+{
+	return MeshUids.RemoveAllItems();
+}
+
+bool UInterchangeStaticMeshLodDataNode::IsEditorOnlyDataDefined()
+{
+#if WITH_EDITORONLY_DATA
+	return true;
+#else
+	return false;
+#endif
+}
diff --git a/Engine/Plugins/Experimental/Interchange/Runtime/Source/FactoryNodes/Public/InterchangeActorFactoryNode.h b/Engine/Plugins/Experimental/Interchange/Runtime/Source/FactoryNodes/Public/InterchangeActorFactoryNode.h
new file mode 100644
--- /dev/null
+++ b/Engine/Plugins/Experimental/Interchange/Runtime/Source/FactoryNodes/Public/InterchangeActorFactoryNode.h
@@ -0,0 +1,57 @@
+// Copyright Epic Games, Inc. All Rights Reserved.
+
+#pragma once
+
+#include "Nodes/InterchangeBaseNode.h"
+
+#include "InterchangeActorFactoryNode.generated.h"
+
+UCLASS(BlueprintType, Experimental)
+class INTERCHANGEFACTORYNODES_API UInterchangeActorFactoryNode : public UInterchangeBaseNode
+{
+	GENERATED_BODY()
+
+public:
+	virtual UClass* GetObjectClass() const override;
+
+	UFUNCTION(BlueprintCallable, Category = "Interchange | Node | ActorFactory")
+	bool GetCustomGlobalTransform(FTransform& AttributeValue) const
+	{
+		IMPLEMENT_NODE_ATTRIBUTE_GETTER(GlobalTransform, FTransform);
+	}
+
+	UFUNCTION(BlueprintCallable, Category = "Interchange | Node | ActorFactory")
+	bool SetCustomGlobalTransform(const FTransform& AttributeValue)
+	{
+		IMPLEMENT_NODE_ATTRIBUTE_SETTER_NODELEGATE(GlobalTransform, FTransform);
+	}
+
+	UFUNCTION(BlueprintCallable, Category = "Interchange | Node | ActorFactory")
+	bool GetCustomActorClassName(FString& AttributeValue) const
+	{
+		IMPLEMENT_NODE_ATTRIBUTE_GETTER(ActorClassName, FString);
+	}
+
+	UFUNCTION(BlueprintCallable, Category = "Interchange | Node | ActorFactory")
+	bool SetCustomActorClassName(const FString& AttributeValue)
+	{
+		IMPLEMENT_NODE_ATTRIBUTE_SETTER_NODELEGATE(ActorClassName, FString);
+	}
+
+	UFUNCTION(BlueprintCallable, Category = "Interchange | Node | ActorFactory")
+	bool GetCustomMobility(uint8& AttributeValue) const
+	{
+		IMPLEMENT_NODE_ATTRIBUTE_GETTER(Mobility, uint8);
+	}
+
+	UFUNCTION(BlueprintCallable, Category = "Interchange | Node | ActorFactory")
+	bool SetCustomMobility(const uint8& AttributeValue, bool bAddApplyDelegate = true)
+	{
+		IMPLEMENT_NODE_ATTRIBUTE_SETTER_NODELEGATE(Mobility, uint8);
+	}
+
+private:
+	const UE::Interchange::FAttributeKey Macro_CustomGlobalTransformKey = UE::Interchange::FAttributeKey(TEXT("GlobalTransform"));
+	const UE::Interchange::FAttributeKey Macro_CustomActorClassNameKey = UE::Interchange::FAttributeKey(TEXT("ActorClassName"));
+	const UE::Interchange::FAttributeKey Macro_CustomMobilityKey = UE::Interchange::FAttributeKey(TEXT("Mobility"));
+};
\ No newline at end of file
diff --git a/Engine/Plugins/Experimental/Interchange/Runtime/Source/FactoryNodes/Public/InterchangeCineCameraFactoryNode.h b/Engine/Plugins/Experimental/Interchange/Runtime/Source/FactoryNodes/Public/InterchangeCineCameraFactoryNode.h
new file mode 100644
--- /dev/null
+++ b/Engine/Plugins/Experimental/Interchange/Runtime/Source/FactoryNodes/Public/InterchangeCineCameraFactoryNode.h
@@ -0,0 +1,65 @@
+// Copyright Epic Games, Inc. All Rights Reserved.
+
+#pragma once
+
+#include "CoreMinimal.h"
+#include "InterchangeActorFactoryNode.h"
+
+#if WITH_ENGINE
+	#include "CineCameraComponent.h"
+#endif
+
+#include "InterchangeCineCameraFactoryNode.generated.h"
+
+UCLASS(BlueprintType)
+class INTERCHANGEFACTORYNODES_API UInterchangeCineCameraFactoryNode : public UInterchangeActorFactoryNode
+{
+	GENERATED_BODY()
+
+public:
+	UFUNCTION(BlueprintCallable, Category = "Interchange | Node | CameraFactory")
+	bool GetCustomFocalLength(float& AttributeValue) const
+	{
+		IMPLEMENT_NODE_ATTRIBUTE_GETTER(FocalLength, float);
+	}
+
+	UFUNCTION(BlueprintCallable, Category = "Interchange | Node | CameraFactory")
+	bool SetCustomFocalLength(const float& AttributeValue, bool bAddApplyDelegate = true)
+	{
+		IMPLEMENT_NODE_ATTRIBUTE_SETTER_WITH_CUSTOM_DELEGATE(UInterchangeCineCameraFactoryNode, FocalLength, float, UCineCameraComponent);
+	}
+
+	UFUNCTION(BlueprintCallable, Category = "Interchange | Node | CameraFactory")
+	bool GetCustomSensorWidth(float& AttributeValue) const
+	{
+		IMPLEMENT_NODE_ATTRIBUTE_GETTER(SensorWidth, float);
+	}
+
+	UFUNCTION(BlueprintCallable, Category = "Interchange | Node | CameraFactory")
+	bool SetCustomSensorWidth(const float& AttributeValue, bool bAddApplyDelegate = true)
+	{
+		IMPLEMENT_NODE_ATTRIBUTE_SETTER_WITH_CUSTOM_DELEGATE(UInterchangeCineCameraFactoryNode, SensorWidth, float, UCineCameraComponent);
+	}
+
+	UFUNCTION(BlueprintCallable, Category = "Interchange | Node | CameraFactory")
+	bool GetCustomSensorHeight(float& AttributeValue) const
+	{
+		IMPLEMENT_NODE_ATTRIBUTE_GETTER(SensorHeight, float);
+	}
+
+	UFUNCTION(BlueprintCallable, Category = "Interchange | Node | CameraFactory")
+	bool SetCustomSensorHeight(const float& AttributeValue, bool bAddApplyDelegate = true)
+	{
+		IMPLEMENT_NODE_ATTRIBUTE_SETTER_WITH_CUSTOM_DELEGATE(UInterchangeCineCameraFactoryNode, SensorHeight, float, UCineCameraComponent);
+	}
+
+private:
+	const UE::Interchange::FAttributeKey Macro_CustomFocalLengthKey = UE::Interchange::FAttributeKey(TEXT("FocalLength"));
+	const UE::Interchange::FAttributeKey Macro_CustomSensorWidthKey = UE::Interchange::FAttributeKey(TEXT("SensorWidth"));
+	const UE::Interchange::FAttributeKey Macro_CustomSensorHeightKey = UE::Interchange::FAttributeKey(TEXT("SensorHeight"));
+
+private:
+	IMPLEMENT_NODE_ATTRIBUTE_DELEGATE_BY_PROPERTYNAME(FocalLength, float, UCineCameraComponent, TEXT("CurrentFocalLength"));
+	IMPLEMENT_NODE_ATTRIBUTE_DELEGATE_BY_PROPERTYNAME(SensorWidth, float, UCineCameraComponent, TEXT("Filmback.SensorWidth"));
+	IMPLEMENT_NODE_ATTRIBUTE_DELEGATE_BY_PROPERTYNAME(SensorHeight, float, UCineCameraComponent, TEXT("Filmback.SensorHeight"));
+};
diff --git a/Engine/Plugins/Experimental/Interchange/Runtime/Source/FactoryNodes/Public/InterchangeFactoryNodesModule.h b/Engine/Plugins/Experimental/Interchange/Runtime/Source/FactoryNodes/Public/InterchangeFactoryNodesModule.h
new file mode 100644
--- /dev/null
+++ b/Engine/Plugins/Experimental/Interchange/Runtime/Source/FactoryNodes/Public/InterchangeFactoryNodesModule.h
@@ -0,0 +1,38 @@
+// Copyright Epic Games, Inc. All Rights Reserved.
+
+#pragma once
+
+#include "CoreMinimal.h"
+#include "Modules/ModuleInterface.h"
+#include "Modules/ModuleManager.h"
+
+/**
+ * The public interface to this module
+ */
+class IInterchangeFactoryNodesModule : public IModuleInterface
+{
+
+public:
+
+	/**
+	 * Singleton-like access to this module's interface.  This is just for convenience!
+	 * Beware of calling this during the shutdown phase, though.  Your module might have been unloaded already.
+	 *
+	 * @return Returns singleton instance, loading the module on demand if needed
+	 */
+	static inline IInterchangeFactoryNodesModule& Get()
+	{
+		return FModuleManager::LoadModuleChecked< IInterchangeFactoryNodesModule >( "InterchangeFactoryNodes" );
+	}
+
+	/**
+	 * Checks to see if this module is loaded and ready.  It is only valid to call Get() if IsAvailable() returns true.
+	 *
+	 * @return True if the module is loaded and ready to use
+	 */
+	static inline bool IsAvailable()
+	{
+		return FModuleManager::Get().IsModuleLoaded( "InterchangeFactoryNodes" );
+	}
+};
+
diff --git a/Engine/Plugins/Experimental/Interchange/Runtime/Source/FactoryNodes/Public/InterchangeMaterialFactoryNode.h b/Engine/Plugins/Experimental/Interchange/Runtime/Source/FactoryNodes/Public/InterchangeMaterialFactoryNode.h
new file mode 100644
--- /dev/null
+++ b/Engine/Plugins/Experimental/Interchange/Runtime/Source/FactoryNodes/Public/InterchangeMaterialFactoryNode.h
@@ -0,0 +1,200 @@
+// Copyright Epic Games, Inc. All Rights Reserved.
+
+#pragma once
+
+#include "CoreMinimal.h"
+#include "Nodes/InterchangeBaseNode.h"
+#include "InterchangeMaterialDefinitions.h"
+#include "InterchangeShaderGraphNode.h"
+
+#if WITH_ENGINE
+#include "Materials/MaterialInterface.h"
+#endif
+
+#include "InterchangeMaterialFactoryNode.generated.h"
+
+UCLASS(Abstract, Experimental)
+class INTERCHANGEFACTORYNODES_API UInterchangeBaseMaterialFactoryNode : public UInterchangeBaseNode
+{
+	GENERATED_BODY()
+
+public:
+	static FString GetMaterialFactoryNodeUidFromMaterialNodeUid(const FString& TranslatedNodeUid);
+};
+
+UCLASS(BlueprintType, Experimental)
+class INTERCHANGEFACTORYNODES_API UInterchangeMaterialFactoryNode : public UInterchangeBaseMaterialFactoryNode
+{
+	GENERATED_BODY()
+
+public:
+
+	virtual FString GetTypeName() const override;
+
+	UFUNCTION(BlueprintCallable, Category = "Interchange | Node | MaterialFactory")
+	virtual class UClass* GetObjectClass() const override;
+
+// Material Inputs
+public:
+	UFUNCTION(BlueprintCallable, Category = "Interchange | Node | MaterialFactory")
+	bool GetBaseColorConnection(FString& ExpressionNodeUid, FString& OutputName) const;
+
+	UFUNCTION(BlueprintCallable, Category = "Interchange | Node | MaterialFactory")
+	bool ConnectToBaseColor(const FString& AttributeValue);
+	
+	UFUNCTION(BlueprintCallable, Category = "Interchange | Node | MaterialFactory")
+	bool ConnectOutputToBaseColor(const FString& ExpressionNodeUid, const FString& OutputName);
+
+	UFUNCTION(BlueprintCallable, Category = "Interchange | Node | MaterialFactory")
+	bool GetMetallicConnection(FString& ExpressionNodeUid, FString& OutputName) const;
+
+	UFUNCTION(BlueprintCallable, Category = "Interchange | Node | MaterialFactory")
+	bool ConnectToMetallic(const FString& AttributeValue);
+
+	UFUNCTION(BlueprintCallable, Category = "Interchange | Node | MaterialFactory")
+	bool ConnectOutputToMetallic(const FString& ExpressionNodeUid, const FString& OutputName);
+
+	UFUNCTION(BlueprintCallable, Category = "Interchange | Node | MaterialFactory")
+	bool GetSpecularConnection(FString& ExpressionNodeUid, FString& OutputName) const;
+
+	UFUNCTION(BlueprintCallable, Category = "Interchange | Node | MaterialFactory")
+	bool ConnectToSpecular(const FString& ExpressionNodeUid);
+
+	UFUNCTION(BlueprintCallable, Category = "Interchange | Node | MaterialFactory")
+	bool ConnectOutputToSpecular(const FString& ExpressionNodeUid, const FString& OutputName);
+
+	UFUNCTION(BlueprintCallable, Category = "Interchange | Node | MaterialFactory")
+	bool GetRoughnessConnection(FString& ExpressionNodeUid, FString& OutputName) const;
+
+	UFUNCTION(BlueprintCallable, Category = "Interchange | Node | MaterialFactory")
+	bool ConnectToRoughness(const FString& ExpressionNodeUid);
+
+	UFUNCTION(BlueprintCallable, Category = "Interchange | Node | MaterialFactory")
+	bool ConnectOutputToRoughness(const FString& ExpressionNodeUid, const FString& OutputName);
+
+	UFUNCTION(BlueprintCallable, Category = "Interchange | Node | MaterialFactory")
+	bool GetEmissiveColorConnection(FString& ExpressionNodeUid, FString& OutputName) const;
+
+	UFUNCTION(BlueprintCallable, Category = "Interchange | Node | MaterialFactory")
+	bool ConnectToEmissiveColor(const FString& ExpressionNodeUid);
+
+	UFUNCTION(BlueprintCallable, Category = "Interchange | Node | MaterialFactory")
+	bool ConnectOutputToEmissiveColor(const FString& ExpressionNodeUid, const FString& OutputName);
+
+	UFUNCTION(BlueprintCallable, Category = "Interchange | Node | MaterialFactory")
+	bool GetNormalConnection(FString& ExpressionNodeUid, FString& OutputName) const;
+
+	UFUNCTION(BlueprintCallable, Category = "Interchange | Node | MaterialFactory")
+	bool ConnectToNormal(const FString& ExpressionNodeUid);
+
+	UFUNCTION(BlueprintCallable, Category = "Interchange | Node | MaterialFactory")
+	bool ConnectOutputToNormal(const FString& ExpressionNodeUid, const FString& OutputName);
+
+	UFUNCTION(BlueprintCallable, Category = "Interchange | Node | MaterialFactory")
+	bool GetOpacityConnection(FString& ExpressionNodeUid, FString& OutputName) const;
+
+	UFUNCTION(BlueprintCallable, Category = "Interchange | Node | MaterialFactory")
+	bool ConnectToOpacity(const FString& AttributeValue);
+
+	UFUNCTION(BlueprintCallable, Category = "Interchange | Node | MaterialFactory")
+	bool ConnectOutputToOpacity(const FString& ExpressionNodeUid, const FString& OutputName);
+
+	UFUNCTION(BlueprintCallable, Category = "Interchange | Node | MaterialFactory")
+	bool GetOcclusionConnection(FString& ExpressionNodeUid, FString& OutputName) const;
+
+	UFUNCTION(BlueprintCallable, Category = "Interchange | Node | MaterialFactory")
+	bool ConnectToOcclusion(const FString& AttributeValue);
+
+	UFUNCTION(BlueprintCallable, Category = "Interchange | Node | MaterialFactory")
+	bool ConnectOutputToOcclusion(const FString& ExpressionNodeUid, const FString& OutputName);
+
+	UFUNCTION(BlueprintCallable, Category = "Interchange | Node | MaterialFactory")
+	bool GetRefractionConnection(FString& ExpressionNodeUid, FString& OutputName) const;
+
+	UFUNCTION(BlueprintCallable, Category = "Interchange | Node | MaterialFactory")
+	bool ConnectToRefraction(const FString& AttributeValue);
+
+	UFUNCTION(BlueprintCallable, Category = "Interchange | Node | MaterialFactory")
+	bool ConnectOutputToRefraction(const FString& ExpressionNodeUid, const FString& OutputName);
+
+// Material parameters
+public:
+	UFUNCTION(BlueprintCallable, Category = "Interchange | Node | Material")
+	bool GetCustomShadingModel(TEnumAsByte<EMaterialShadingModel>& AttributeValue) const;
+ 
+	UFUNCTION(BlueprintCallable, Category = "Interchange | Node | Material")
+	bool SetCustomShadingModel(const TEnumAsByte<EMaterialShadingModel>& AttributeValue, bool bAddApplyDelegate = true);
+
+	UFUNCTION(BlueprintCallable, Category = "Interchange | Node | Material")
+	bool GetCustomTranslucencyLightingMode(TEnumAsByte<ETranslucencyLightingMode>& AttributeValue) const;
+ 
+	UFUNCTION(BlueprintCallable, Category = "Interchange | Node | Material")
+	bool SetCustomTranslucencyLightingMode(const TEnumAsByte<ETranslucencyLightingMode>& AttributeValue, bool bAddApplyDelegate = true);
+
+	UFUNCTION(BlueprintCallable, Category = "Interchange | Node | Material")
+	bool GetCustomBlendMode(TEnumAsByte<EBlendMode>& AttributeValue) const;
+ 
+	UFUNCTION(BlueprintCallable, Category = "Interchange | Node | Material")
+	bool SetCustomBlendMode(const TEnumAsByte<EBlendMode>& AttributeValue, bool bAddApplyDelegate = true);
+
+	UFUNCTION(BlueprintCallable, Category = "Interchange | Node | Material")
+	bool GetCustomTwoSided(bool& AttributeValue) const;
+ 
+	/** Sets if this shader graph should be rendered two sided or not. Defaults to off. */
+	UFUNCTION(BlueprintCallable, Category = "Interchange | Node | Material")
+	bool SetCustomTwoSided(const bool& AttributeValue, bool bAddApplyDelegate = true);
+
+private:
+	const UE::Interchange::FAttributeKey Macro_CustomShadingModelKey = UE::Interchange::FAttributeKey(TEXT("ShadingModel"));
+	const UE::Interchange::FAttributeKey Macro_CustomTranslucencyLightingModeKey = UE::Interchange::FAttributeKey(TEXT("TranslucencyLightingMode"));
+	const UE::Interchange::FAttributeKey Macro_CustomBlendModeKey = UE::Interchange::FAttributeKey(TEXT("BlendMode"));
+	const UE::Interchange::FAttributeKey Macro_CustomTwoSidedKey = UE::Interchange::FAttributeKey(TEXT("TwoSided"));
+};
+
+UCLASS(BlueprintType, Experimental)
+class INTERCHANGEFACTORYNODES_API UInterchangeMaterialExpressionFactoryNode : public UInterchangeBaseNode
+{
+	GENERATED_BODY()
+
+public:
+	virtual FString GetTypeName() const override;
+
+public:
+	UFUNCTION(BlueprintCallable, Category = "Interchange | Node | Material")
+	bool GetCustomExpressionClassName(FString& AttributeValue) const;
+
+	UFUNCTION(BlueprintCallable, Category = "Interchange | Node | Material")
+	bool SetCustomExpressionClassName(const FString& AttributeValue);
+
+private:
+	const UE::Interchange::FAttributeKey Macro_CustomExpressionClassNameKey = UE::Interchange::FAttributeKey(TEXT("ExpressionClassName"));
+
+};
+
+UCLASS(BlueprintType, Experimental)
+class INTERCHANGEFACTORYNODES_API UInterchangeMaterialInstanceFactoryNode : public UInterchangeBaseMaterialFactoryNode
+{
+	GENERATED_BODY()
+
+public:
+	virtual FString GetTypeName() const override;
+	virtual UClass* GetObjectClass() const override;
+
+public:
+	UFUNCTION(BlueprintCallable, Category = "Interchange | Node | MaterialInstanceFactory")
+	bool GetCustomInstanceClassName(FString& AttributeValue) const;
+
+	UFUNCTION(BlueprintCallable, Category = "Interchange | Node | MaterialInstanceFactory")
+	bool SetCustomInstanceClassName(const FString& AttributeValue);
+
+	UFUNCTION(BlueprintCallable, Category = "Interchange | Node | MaterialInstanceFactory")
+	bool GetCustomParent(FString& AttributeValue) const;
+
+	UFUNCTION(BlueprintCallable, Category = "Interchange | Node | MaterialInstanceFactory")
+	bool SetCustomParent(const FString& AttributeValue);
+
+private:
+	const UE::Interchange::FAttributeKey Macro_CustomInstanceClassNameKey = UE::Interchange::FAttributeKey(TEXT("InstanceClassName"));
+	const UE::Interchange::FAttributeKey Macro_CustomParentKey = UE::Interchange::FAttributeKey(TEXT("Parent"));
+	
+};
diff --git a/Engine/Plugins/Experimental/Interchange/Runtime/Source/FactoryNodes/Public/InterchangePhysicsAssetFactoryNode.h b/Engine/Plugins/Experimental/Interchange/Runtime/Source/FactoryNodes/Public/InterchangePhysicsAssetFactoryNode.h
new file mode 100644
--- /dev/null
+++ b/Engine/Plugins/Experimental/Interchange/Runtime/Source/FactoryNodes/Public/InterchangePhysicsAssetFactoryNode.h
@@ -0,0 +1,141 @@
+// Copyright Epic Games, Inc. All Rights Reserved.
+
+#pragma once
+
+#include "CoreMinimal.h"
+#include "Nodes/InterchangeBaseNode.h"
+
+#if WITH_ENGINE
+#include "Animation/Skeleton.h"
+#include "PhysicsEngine/PhysicsAsset.h"
+#endif
+
+#include "InterchangePhysicsAssetFactoryNode.generated.h"
+
+UCLASS(BlueprintType, Experimental)
+class INTERCHANGEFACTORYNODES_API UInterchangePhysicsAssetFactoryNode : public UInterchangeBaseNode
+{
+	GENERATED_BODY()
+
+public:
+	UInterchangePhysicsAssetFactoryNode()
+	:UInterchangeBaseNode()
+	{
+#if WITH_ENGINE
+		AssetClass = nullptr;
+#endif
+	}
+
+	/**
+	 * Initialize node data
+	 * @param: UniqueID - The uniqueId for this node
+	 * @param DisplayLabel - The name of the node
+	 * @param InAssetClass - The class the Skeleton factory will create for this node.
+	 */
+	UFUNCTION(BlueprintCallable, Category = "Interchange | Node | Skeleton")
+	void InitializePhysicsAssetNode(const FString& UniqueID, const FString& DisplayLabel, const FString& InAssetClass)
+	{
+		bIsNodeClassInitialized = false;
+		InitializeNode(UniqueID, DisplayLabel, EInterchangeNodeContainerType::FactoryData);
+		FString OperationName = GetTypeName() + TEXT(".SetAssetClassName");
+		InterchangePrivateNodeBase::SetCustomAttribute<FString>(*Attributes, ClassNameAttributeKey, OperationName, InAssetClass);
+		FillAssetClassFromAttribute();
+	}
+
+	virtual void Serialize(FArchive& Ar) override
+	{
+		Super::Serialize(Ar);
+#if WITH_ENGINE
+		if(Ar.IsLoading())
+		{
+			//Make sure the class is properly set when we compile with engine, this will set the
+			//bIsNodeClassInitialized to true.
+			SetNodeClassFromClassAttribute();
+		}
+#endif //#if WITH_ENGINE
+	}
+
+	/**
+	 * Return the node type name of the class, we use this when reporting error
+	 */
+	virtual FString GetTypeName() const override
+	{
+		const FString TypeName = TEXT("SkeletonNode");
+		return TypeName;
+	}
+
+	/** Get the class this node want to create */
+	UFUNCTION(BlueprintCallable, Category = "Interchange | Node | Skeleton")
+	virtual class UClass* GetObjectClass() const override
+	{
+		ensure(bIsNodeClassInitialized);
+#if WITH_ENGINE
+		return AssetClass.Get() != nullptr ? AssetClass.Get() : UPhysicsAsset::StaticClass();
+#else
+		return nullptr;
+#endif
+	}
+
+	virtual FGuid GetHash() const override
+	{
+		return Attributes->GetStorageHash();
+	}
+
+public:
+	/** Get skeletalMesh asset Uid used to create the data in the post pipeline step.*/
+	UFUNCTION(BlueprintCallable, Category = "Interchange | Node | Skeleton")
+	bool GetCustomSkeletalMeshUid(FString& AttributeValue) const
+	{
+		IMPLEMENT_NODE_ATTRIBUTE_GETTER(SkeletalMeshUid, FString);
+	}
+	/** Set skeletalMesh asset Uid used to create the data in the post pipeline step.*/
+	UFUNCTION(BlueprintCallable, Category = "Interchange | Node | Skeleton")
+	bool SetCustomSkeletalMeshUid(const FString& AttributeValue)
+	{
+		IMPLEMENT_NODE_ATTRIBUTE_SETTER_NODELEGATE(SkeletalMeshUid, FString)
+	}
+
+private:
+
+	void FillAssetClassFromAttribute()
+	{
+#if WITH_ENGINE
+		FString OperationName = GetTypeName() + TEXT(".GetAssetClassName");
+		FString ClassName;
+		InterchangePrivateNodeBase::GetCustomAttribute<FString>(*Attributes, ClassNameAttributeKey, OperationName, ClassName);
+		if (ClassName.Equals(UPhysicsAsset::StaticClass()->GetName()))
+		{
+			AssetClass = UPhysicsAsset::StaticClass();
+			bIsNodeClassInitialized = true;
+		}
+#endif
+	}
+
+	bool SetNodeClassFromClassAttribute()
+	{
+		if (!bIsNodeClassInitialized)
+		{
+			FillAssetClassFromAttribute();
+		}
+		return bIsNodeClassInitialized;
+	}
+
+	bool IsEditorOnlyDataDefined()
+	{
+#if WITH_EDITORONLY_DATA
+		return true;
+#else
+		return false;
+#endif
+	}
+
+	const UE::Interchange::FAttributeKey ClassNameAttributeKey = UE::Interchange::FBaseNodeStaticData::ClassTypeAttributeKey();
+	//SkeletalMesh asset Uid use to create the data in the post pipeline step
+	const UE::Interchange::FAttributeKey Macro_CustomSkeletalMeshUidKey = UE::Interchange::FAttributeKey(TEXT("SkeletalMeshUid"));
+	
+protected:
+#if WITH_ENGINE
+	TSubclassOf<UPhysicsAsset> AssetClass = nullptr;
+#endif
+	bool bIsNodeClassInitialized = false;
+};
diff --git a/Engine/Plugins/Experimental/Interchange/Runtime/Source/FactoryNodes/Public/InterchangeSkeletalMeshFactoryNode.h b/Engine/Plugins/Experimental/Interchange/Runtime/Source/FactoryNodes/Public/InterchangeSkeletalMeshFactoryNode.h
new file mode 100644
--- /dev/null
+++ b/Engine/Plugins/Experimental/Interchange/Runtime/Source/FactoryNodes/Public/InterchangeSkeletalMeshFactoryNode.h
@@ -0,0 +1,150 @@
+// Copyright Epic Games, Inc. All Rights Reserved.
+
+#pragma once
+
+#include "CoreMinimal.h"
+#include "Nodes/InterchangeBaseNode.h"
+
+#if WITH_ENGINE
+#include "Engine/SkeletalMesh.h"
+#endif
+
+#include "InterchangeSkeletalMeshFactoryNode.generated.h"
+
+//Interchange namespace
+namespace UE
+{
+	namespace Interchange
+	{
+
+		struct INTERCHANGEFACTORYNODES_API FSkeletalMeshNodeStaticData : public FBaseNodeStaticData
+		{
+			static const FString& GetLodDependenciesBaseKey();
+		};
+
+	}//ns Interchange
+}//ns UE
+
+UCLASS(BlueprintType, Experimental)
+class INTERCHANGEFACTORYNODES_API UInterchangeSkeletalMeshFactoryNode : public UInterchangeBaseNode
+{
+	GENERATED_BODY()
+
+public:
+	UInterchangeSkeletalMeshFactoryNode();
+
+	/**
+	 * Initialize node data
+	 * @param: UniqueID - The uniqueId for this node
+	 * @param DisplayLabel - The name of the node
+	 * @param InAssetClass - The class the SkeletalMesh factory will create for this node.
+	 */
+	UFUNCTION(BlueprintCallable, Category = "Interchange | Node | SkeletalMesh")
+	void InitializeSkeletalMeshNode(const FString& UniqueID, const FString& DisplayLabel, const FString& InAssetClass);
+
+	virtual void Serialize(FArchive& Ar) override;
+
+	/**
+	 * Return the node type name of the class, we use this when reporting error
+	 */
+	virtual FString GetTypeName() const override;
+
+	virtual FString GetKeyDisplayName(const UE::Interchange::FAttributeKey& NodeAttributeKey) const override;
+
+	/** Get the class this node want to create */
+	UFUNCTION(BlueprintCallable, Category = "Interchange | Node | SkeletalMesh")
+	virtual class UClass* GetObjectClass() const override;
+
+public:
+	/** Return The number of LOD this skeletalmesh has.*/
+	UFUNCTION(BlueprintCallable, Category = "Interchange | Node | SkeletalMesh")
+	int32 GetLodDataCount() const;
+
+	UFUNCTION(BlueprintCallable, Category = "Interchange | Node | SkeletalMesh")
+	void GetLodDataUniqueIds(TArray<FString>& OutLodDataUniqueIds) const;
+
+	UFUNCTION(BlueprintCallable, Category = "Interchange | Node | SkeletalMesh")
+	bool AddLodDataUniqueId(const FString& LodDataUniqueId);
+
+	UFUNCTION(BlueprintCallable, Category = "Interchange | Node | SkeletalMesh")
+	bool RemoveLodDataUniqueId(const FString& LodDataUniqueId);
+	
+	/** Query the skeletal mesh factory skeleton UObject. Return false if the attribute was not set.*/
+	UFUNCTION(BlueprintCallable, Category = "Interchange | Node | SkeletalMesh")
+	bool GetCustomSkeletonSoftObjectPath(FSoftObjectPath& AttributeValue) const;
+
+	/** Set the skeletal mesh factory skeleton UObject. Return false if the attribute cannot be set.*/
+	UFUNCTION(BlueprintCallable, Category = "Interchange | Node | SkeletalMesh")
+	bool SetCustomSkeletonSoftObjectPath(const FSoftObjectPath& AttributeValue);
+
+	/** Query weather the skeletal mesh factory should create the morph target. Return false if the attribute was not set.*/
+	UFUNCTION(BlueprintCallable, Category = "Interchange | Node | SkeletalMesh")
+	bool GetCustomImportMorphTarget(bool& AttributeValue) const;
+
+	/** Set weather the skeletal mesh factory should create the morph target. Return false if the attribute cannot be set.*/
+	UFUNCTION(BlueprintCallable, Category = "Interchange | Node | SkeletalMesh")
+	bool SetCustomImportMorphTarget(const bool& AttributeValue);
+
+	/** Query weather the skeletal mesh factory should create a physics asset. Return false if the attribute was not set.*/
+	UFUNCTION(BlueprintCallable, Category = "Interchange | Node | SkeletalMesh")
+	bool GetCustomCreatePhysicsAsset(bool& AttributeValue) const;
+
+	/** Set weather the skeletal mesh factory should create a physics asset. Return false if the attribute cannot be set.*/
+	UFUNCTION(BlueprintCallable, Category = "Interchange | Node | SkeletalMesh")
+	bool SetCustomCreatePhysicsAsset(const bool& AttributeValue);
+
+	/** Query a physics asset the skeletal mesh factory should use. Return false if the attribute was not set.*/
+	UFUNCTION(BlueprintCallable, Category = "Interchange | Node | SkeletalMesh")
+	bool GetCustomPhysicAssetSoftObjectPath(FSoftObjectPath& AttributeValue) const;
+
+	/** Set a physics asset the skeletal mesh factory should use. Return false if the attribute cannot be set.*/
+	UFUNCTION(BlueprintCallable, Category = "Interchange | Node | SkeletalMesh")
+	bool SetCustomPhysicAssetSoftObjectPath(const FSoftObjectPath& AttributeValue);
+
+	/** Query weather the skeletal mesh factory should replace the vertex color. Return false if the attribute was not set.*/
+	UFUNCTION(BlueprintCallable, Category = "Interchange | Node | SkeletalMesh")
+	bool GetCustomVertexColorReplace(bool& AttributeValue) const;
+
+	/** Set weather the skeletal mesh factory should replace the vertex color. Return false if the attribute cannot be set.*/
+	UFUNCTION(BlueprintCallable, Category = "Interchange | Node | SkeletalMesh")
+	bool SetCustomVertexColorReplace(const bool& AttributeValue);
+
+	/** Query weather the skeletal mesh factory should ignore the vertex color. Return false if the attribute was not set.*/
+	UFUNCTION(BlueprintCallable, Category = "Interchange | Node | SkeletalMesh")
+	bool GetCustomVertexColorIgnore(bool& AttributeValue) const;
+
+	/** Set weather the skeletal mesh factory should ignore the vertex color. Return false if the attribute cannot be set.*/
+	UFUNCTION(BlueprintCallable, Category = "Interchange | Node | SkeletalMesh")
+	bool SetCustomVertexColorIgnore(const bool& AttributeValue);
+
+	/** Query weather the skeletal mesh factory should override the vertex color. Return false if the attribute was not set.*/
+	UFUNCTION(BlueprintCallable, Category = "Interchange | Node | SkeletalMesh")
+	bool GetCustomVertexColorOverride(FColor& AttributeValue) const;
+
+	/** Set weather the skeletal mesh factory should override the vertex color. Return false if the attribute cannot be set.*/
+	UFUNCTION(BlueprintCallable, Category = "Interchange | Node | SkeletalMesh")
+	bool SetCustomVertexColorOverride(const FColor& AttributeValue);
+private:
+
+	void FillAssetClassFromAttribute();
+
+	bool SetNodeClassFromClassAttribute();
+
+	bool IsEditorOnlyDataDefined();
+
+	const UE::Interchange::FAttributeKey ClassNameAttributeKey = UE::Interchange::FBaseNodeStaticData::ClassTypeAttributeKey();
+	const UE::Interchange::FAttributeKey Macro_CustomImportMorphTargetKey = UE::Interchange::FAttributeKey(TEXT("ImportMorphTarget"));
+	const UE::Interchange::FAttributeKey Macro_CustomSkeletonSoftObjectPathKey = UE::Interchange::FAttributeKey(TEXT("SkeletonSoftObjectPath"));
+	const UE::Interchange::FAttributeKey Macro_CustomCreatePhysicsAssetKey = UE::Interchange::FAttributeKey(TEXT("CreatePhysicsAsset"));
+	const UE::Interchange::FAttributeKey Macro_CustomPhysicAssetSoftObjectPathKey = UE::Interchange::FAttributeKey(TEXT("PhysicAssetSoftObjectPath"));
+	const UE::Interchange::FAttributeKey Macro_CustomVertexColorReplaceKey = UE::Interchange::FAttributeKey(TEXT("VertexColorReplace"));
+	const UE::Interchange::FAttributeKey Macro_CustomVertexColorIgnoreKey = UE::Interchange::FAttributeKey(TEXT("VertexColorIgnore"));
+	const UE::Interchange::FAttributeKey Macro_CustomVertexColorOverrideKey = UE::Interchange::FAttributeKey(TEXT("VertexColorOverride"));
+
+	UE::Interchange::TArrayAttributeHelper<FString> LodDependencies;
+protected:
+#if WITH_ENGINE
+	TSubclassOf<USkeletalMesh> AssetClass = nullptr;
+#endif
+	bool bIsNodeClassInitialized = false;
+};
diff --git a/Engine/Plugins/Experimental/Interchange/Runtime/Source/FactoryNodes/Public/InterchangeSkeletalMeshLodDataNode.h b/Engine/Plugins/Experimental/Interchange/Runtime/Source/FactoryNodes/Public/InterchangeSkeletalMeshLodDataNode.h
new file mode 100644
--- /dev/null
+++ b/Engine/Plugins/Experimental/Interchange/Runtime/Source/FactoryNodes/Public/InterchangeSkeletalMeshLodDataNode.h
@@ -0,0 +1,78 @@
+// Copyright Epic Games, Inc. All Rights Reserved.
+
+#pragma once
+
+#include "CoreMinimal.h"
+#include "Nodes/InterchangeBaseNode.h"
+
+#include "InterchangeSkeletalMeshLodDataNode.generated.h"
+
+//Interchange namespace
+namespace UE
+{
+	namespace Interchange
+	{
+
+		struct INTERCHANGEFACTORYNODES_API FSkeletalMeshNodeLodDataStaticData : public FBaseNodeStaticData
+		{
+			static const FString& GetMeshUidsBaseKey();
+		};
+
+	}//ns Interchange
+}//ns UE
+
+UCLASS(BlueprintType, Experimental)
+class INTERCHANGEFACTORYNODES_API UInterchangeSkeletalMeshLodDataNode : public UInterchangeBaseNode
+{
+	GENERATED_BODY()
+
+public:
+	UInterchangeSkeletalMeshLodDataNode();
+
+	/**
+	 * Return the node type name of the class, we use this when reporting error
+	 */
+	virtual FString GetTypeName() const override;
+
+	virtual FString GetKeyDisplayName(const UE::Interchange::FAttributeKey& NodeAttributeKey) const override;
+
+	virtual FString GetAttributeCategory(const UE::Interchange::FAttributeKey& NodeAttributeKey) const override;
+
+public:
+	/** Query the LOD skeletal mesh factory skeleton reference. Return false if the attribute was not set.*/
+	UFUNCTION(BlueprintCallable, Category = "Interchange | Node | SkeletalMeshLodData")
+	bool GetCustomSkeletonUid(FString& AttributeValue) const;
+
+	/** Set the LOD skeletal mesh factory skeleton reference. Return false if the attribute cannot be set.*/
+	UFUNCTION(BlueprintCallable, Category = "Interchange | Node | SkeletalMeshLodData")
+	bool SetCustomSkeletonUid(const FString& AttributeValue);
+
+	/* Return the number of mesh geometry this LOD will be made of. Mesh uids can be either a scene or a mesh node. If its a scene it mean we want the mesh factory to bake the geo payload with the global transform of the scene node. */
+	UFUNCTION(BlueprintCallable, Category = "Interchange | Node | SkeletalMeshLodData")
+	int32 GetMeshUidsCount() const;
+
+	/* Query all mesh geometry this LOD will be made of. Mesh uids can be either a scene or a mesh node. If its a scene it mean we want the mesh factory to bake the geo payload with the global transform of the scene node. */
+	UFUNCTION(BlueprintCallable, Category = "Interchange | Node | SkeletalMeshLodData")
+	void GetMeshUids(TArray<FString>& OutMeshNames) const;
+
+	/* Add one mesh geometry use to create this LOD geometry. Mesh uids can be either a scene or a mesh node. If its a scene it mean we want the mesh factory to bake the geo payload with the global transform of the scene node. */
+	UFUNCTION(BlueprintCallable, Category = "Interchange | Node | SkeletalMeshLodData")
+	bool AddMeshUid(const FString& MeshName);
+
+	/* Remove one mesh geometry use to create this LOD geometry. Mesh uids can be either a scene or a mesh node. If its a scene it mean we want the mesh factory to bake the geo payload with the global transform of the scene node. */
+	UFUNCTION(BlueprintCallable, Category = "Interchange | Node | SkeletalMeshLodData")
+	bool RemoveMeshUid(const FString& MeshName);
+
+	/* Remove all mesh geometry use to create this LOD geometry. Mesh uids can be either a scene or a mesh node. If its a scene it mean we want the mesh factory to bake the geo payload with the global transform of the scene node. */
+	UFUNCTION(BlueprintCallable, Category = "Interchange | Node | SkeletalMeshLodData")
+	bool RemoveAllMeshes();
+
+private:
+
+	bool IsEditorOnlyDataDefined();
+
+	const UE::Interchange::FAttributeKey Macro_CustomSkeletonUidKey = UE::Interchange::FAttributeKey(TEXT("__SkeletonUid__Key"));
+
+	UE::Interchange::TArrayAttributeHelper<FString> MeshUids;
+protected:
+};
diff --git a/Engine/Plugins/Experimental/Interchange/Runtime/Source/FactoryNodes/Public/InterchangeSkeletonFactoryNode.h b/Engine/Plugins/Experimental/Interchange/Runtime/Source/FactoryNodes/Public/InterchangeSkeletonFactoryNode.h
new file mode 100644
--- /dev/null
+++ b/Engine/Plugins/Experimental/Interchange/Runtime/Source/FactoryNodes/Public/InterchangeSkeletonFactoryNode.h
@@ -0,0 +1,156 @@
+// Copyright Epic Games, Inc. All Rights Reserved.
+
+#pragma once
+
+#include "CoreMinimal.h"
+#include "Nodes/InterchangeBaseNode.h"
+
+#if WITH_ENGINE
+#include "Animation/Skeleton.h"
+#endif
+
+#include "InterchangeSkeletonFactoryNode.generated.h"
+
+UCLASS(BlueprintType, Experimental)
+class INTERCHANGEFACTORYNODES_API UInterchangeSkeletonFactoryNode : public UInterchangeBaseNode
+{
+	GENERATED_BODY()
+
+public:
+	UInterchangeSkeletonFactoryNode()
+	:UInterchangeBaseNode()
+	{
+#if WITH_ENGINE
+		AssetClass = nullptr;
+#endif
+	}
+
+	/**
+	 * Initialize node data
+	 * @param: UniqueID - The uniqueId for this node
+	 * @param DisplayLabel - The name of the node
+	 * @param InAssetClass - The class the Skeleton factory will create for this node.
+	 */
+	UFUNCTION(BlueprintCallable, Category = "Interchange | Node | Skeleton")
+	void InitializeSkeletonNode(const FString& UniqueID, const FString& DisplayLabel, const FString& InAssetClass)
+	{
+		bIsNodeClassInitialized = false;
+		InitializeNode(UniqueID, DisplayLabel, EInterchangeNodeContainerType::FactoryData);
+		FString OperationName = GetTypeName() + TEXT(".SetAssetClassName");
+		InterchangePrivateNodeBase::SetCustomAttribute<FString>(*Attributes, ClassNameAttributeKey, OperationName, InAssetClass);
+		FillAssetClassFromAttribute();
+	}
+
+	virtual void Serialize(FArchive& Ar) override
+	{
+		Super::Serialize(Ar);
+#if WITH_ENGINE
+		if(Ar.IsLoading())
+		{
+			//Make sure the class is properly set when we compile with engine, this will set the
+			//bIsNodeClassInitialized to true.
+			SetNodeClassFromClassAttribute();
+		}
+#endif //#if WITH_ENGINE
+	}
+
+	/**
+	 * Return the node type name of the class, we use this when reporting error
+	 */
+	virtual FString GetTypeName() const override
+	{
+		const FString TypeName = TEXT("SkeletonNode");
+		return TypeName;
+	}
+
+	/** Get the class this node want to create */
+	UFUNCTION(BlueprintCallable, Category = "Interchange | Node | Skeleton")
+	virtual class UClass* GetObjectClass() const override
+	{
+		ensure(bIsNodeClassInitialized);
+#if WITH_ENGINE
+		return AssetClass.Get() != nullptr ? AssetClass.Get() : USkeleton::StaticClass();
+#else
+		return nullptr;
+#endif
+	}
+
+	virtual FGuid GetHash() const override
+	{
+		return Attributes->GetStorageHash();
+	}
+
+public:
+	/** Return false if the Attribute was not set previously.*/
+	UFUNCTION(BlueprintCallable, Category = "Interchange | Node | Skeleton")
+	bool GetCustomRootJointUid(FString& AttributeValue) const
+	{
+		IMPLEMENT_NODE_ATTRIBUTE_GETTER(RootJointID, FString);
+	}
+
+	UFUNCTION(BlueprintCallable, Category = "Interchange | Node | Skeleton")
+	bool SetCustomRootJointUid(const FString& AttributeValue)
+	{
+		IMPLEMENT_NODE_ATTRIBUTE_SETTER_NODELEGATE(RootJointID, FString)
+	}
+
+	/** Query if this skeleton should replace joint transform with time zero evaluation instead of bind pose.*/
+	UFUNCTION(BlueprintCallable, Category = "Interchange | Node | Skeleton")
+	bool GetCustomUseTimeZeroForBindPose(bool& AttributeValue) const
+	{
+		IMPLEMENT_NODE_ATTRIBUTE_GETTER(UseTimeZeroForBindPose, bool);
+	}
+
+	/** If AttributeValue is true, force this skeleton bind pose to use time zero instead of bind pose.*/
+	UFUNCTION(BlueprintCallable, Category = "Interchange | Node | Skeleton")
+	bool SetCustomUseTimeZeroForBindPose(const bool& AttributeValue)
+	{
+		IMPLEMENT_NODE_ATTRIBUTE_SETTER_NODELEGATE(UseTimeZeroForBindPose, bool)
+	}
+
+private:
+
+	void FillAssetClassFromAttribute()
+	{
+#if WITH_ENGINE
+		FString OperationName = GetTypeName() + TEXT(".GetAssetClassName");
+		FString ClassName;
+		InterchangePrivateNodeBase::GetCustomAttribute<FString>(*Attributes, ClassNameAttributeKey, OperationName, ClassName);
+		if (ClassName.Equals(USkeleton::StaticClass()->GetName()))
+		{
+			AssetClass = USkeleton::StaticClass();
+			bIsNodeClassInitialized = true;
+		}
+#endif
+	}
+
+	bool SetNodeClassFromClassAttribute()
+	{
+		if (!bIsNodeClassInitialized)
+		{
+			FillAssetClassFromAttribute();
+		}
+		return bIsNodeClassInitialized;
+	}
+
+	bool IsEditorOnlyDataDefined()
+	{
+#if WITH_EDITORONLY_DATA
+		return true;
+#else
+		return false;
+#endif
+	}
+
+	const UE::Interchange::FAttributeKey ClassNameAttributeKey = UE::Interchange::FBaseNodeStaticData::ClassTypeAttributeKey();
+
+	//Skeleton
+	const UE::Interchange::FAttributeKey Macro_CustomRootJointIDKey = UE::Interchange::FAttributeKey(TEXT("RootJointID"));
+	const UE::Interchange::FAttributeKey Macro_CustomUseTimeZeroForBindPoseKey = UE::Interchange::FAttributeKey(TEXT("UseTimeZeroForBindPose"));
+	
+protected:
+#if WITH_ENGINE
+	TSubclassOf<USkeleton> AssetClass = nullptr;
+#endif
+	bool bIsNodeClassInitialized = false;
+};
diff --git a/Engine/Plugins/Experimental/Interchange/Runtime/Source/FactoryNodes/Public/InterchangeStaticMeshFactoryNode.h b/Engine/Plugins/Experimental/Interchange/Runtime/Source/FactoryNodes/Public/InterchangeStaticMeshFactoryNode.h
new file mode 100644
--- /dev/null
+++ b/Engine/Plugins/Experimental/Interchange/Runtime/Source/FactoryNodes/Public/InterchangeStaticMeshFactoryNode.h
@@ -0,0 +1,112 @@
+// Copyright Epic Games, Inc. All Rights Reserved.
+
+#pragma once
+
+#include "CoreMinimal.h"
+#include "Nodes/InterchangeBaseNode.h"
+
+#if WITH_ENGINE
+#include "Engine/StaticMesh.h"
+#endif
+
+#include "InterchangeStaticMeshFactoryNode.generated.h"
+
+
+namespace UE
+{
+	namespace Interchange
+	{
+		struct INTERCHANGEFACTORYNODES_API FStaticMeshNodeStaticData : public FBaseNodeStaticData
+		{
+			static const FString& GetLodDependenciesBaseKey();
+		};
+	} // namespace Interchange
+} // namespace UE
+
+
+UCLASS(BlueprintType, Experimental)
+class INTERCHANGEFACTORYNODES_API UInterchangeStaticMeshFactoryNode : public UInterchangeBaseNode
+{
+	GENERATED_BODY()
+
+public:
+	UInterchangeStaticMeshFactoryNode();
+
+	/**
+	 * Initialize node data
+	 * @param UniqueID - The uniqueId for this node
+	 * @param DisplayLabel - The name of the node
+	 * @param InAssetClass - The class the StaticMesh factory will create for this node.
+	 */
+	UFUNCTION(BlueprintCallable, Category = "Interchange | Node | StaticMesh")
+	void InitializeStaticMeshNode(const FString& UniqueID, const FString& DisplayLabel, const FString& InAssetClass);
+
+	virtual void Serialize(FArchive& Ar) override;
+
+	/**
+	 * Return the node type name of the class, we use this when reporting error
+	 */
+	virtual FString GetTypeName() const override;
+
+	virtual FString GetKeyDisplayName(const UE::Interchange::FAttributeKey& NodeAttributeKey) const override;
+
+	/** Get the class this node want to create */
+	UFUNCTION(BlueprintCallable, Category = "Interchange | Node | StaticMesh")
+	virtual class UClass* GetObjectClass() const override;
+
+public:
+	/** Return The number of LOD this static mesh has.*/
+	UFUNCTION(BlueprintCallable, Category = "Interchange | Node | StaticMesh")
+	int32 GetLodDataCount() const;
+
+	UFUNCTION(BlueprintCallable, Category = "Interchange | Node | StaticMesh")
+	void GetLodDataUniqueIds(TArray<FString>& OutLodDataUniqueIds) const;
+
+	UFUNCTION(BlueprintCallable, Category = "Interchange | Node | StaticMesh")
+	bool AddLodDataUniqueId(const FString& LodDataUniqueId);
+
+	UFUNCTION(BlueprintCallable, Category = "Interchange | Node | StaticMesh")
+	bool RemoveLodDataUniqueId(const FString& LodDataUniqueId);
+
+	/** Query weather the static mesh factory should replace the vertex color. Return false if the attribute was not set.*/
+	UFUNCTION(BlueprintCallable, Category = "Interchange | Node | StaticMesh")
+	bool GetCustomVertexColorReplace(bool& AttributeValue) const;
+
+	/** Set weather the static mesh factory should replace the vertex color. Return false if the attribute cannot be set.*/
+	UFUNCTION(BlueprintCallable, Category = "Interchange | Node | StaticMesh")
+	bool SetCustomVertexColorReplace(const bool& AttributeValue);
+
+	/** Query weather the static mesh factory should ignore the vertex color. Return false if the attribute was not set.*/
+	UFUNCTION(BlueprintCallable, Category = "Interchange | Node | StaticMesh")
+	bool GetCustomVertexColorIgnore(bool& AttributeValue) const;
+
+	/** Set weather the static mesh factory should ignore the vertex color. Return false if the attribute cannot be set.*/
+	UFUNCTION(BlueprintCallable, Category = "Interchange | Node | StaticMesh")
+	bool SetCustomVertexColorIgnore(const bool& AttributeValue);
+
+	/** Query weather the static mesh factory should override the vertex color. Return false if the attribute was not set.*/
+	UFUNCTION(BlueprintCallable, Category = "Interchange | Node | StaticMesh")
+	bool GetCustomVertexColorOverride(FColor& AttributeValue) const;
+
+	/** Set weather the static mesh factory should override the vertex color. Return false if the attribute cannot be set.*/
+	UFUNCTION(BlueprintCallable, Category = "Interchange | Node | StaticMesh")
+	bool SetCustomVertexColorOverride(const FColor& AttributeValue);
+
+private:
+
+	void FillAssetClassFromAttribute();
+	bool SetNodeClassFromClassAttribute();
+	bool IsEditorOnlyDataDefined();
+
+	const UE::Interchange::FAttributeKey ClassNameAttributeKey = UE::Interchange::FBaseNodeStaticData::ClassTypeAttributeKey();
+	const UE::Interchange::FAttributeKey Macro_CustomVertexColorReplaceKey = UE::Interchange::FAttributeKey(TEXT("VertexColorReplace"));
+	const UE::Interchange::FAttributeKey Macro_CustomVertexColorIgnoreKey = UE::Interchange::FAttributeKey(TEXT("VertexColorIgnore"));
+	const UE::Interchange::FAttributeKey Macro_CustomVertexColorOverrideKey = UE::Interchange::FAttributeKey(TEXT("VertexColorOverride"));
+
+	UE::Interchange::TArrayAttributeHelper<FString> LodDependencies;
+protected:
+#if WITH_ENGINE
+	TSubclassOf<UStaticMesh> AssetClass = nullptr;
+#endif
+	bool bIsNodeClassInitialized = false;
+};
diff --git a/Engine/Plugins/Experimental/Interchange/Runtime/Source/FactoryNodes/Public/InterchangeStaticMeshLodDataNode.h b/Engine/Plugins/Experimental/Interchange/Runtime/Source/FactoryNodes/Public/InterchangeStaticMeshLodDataNode.h
new file mode 100644
--- /dev/null
+++ b/Engine/Plugins/Experimental/Interchange/Runtime/Source/FactoryNodes/Public/InterchangeStaticMeshLodDataNode.h
@@ -0,0 +1,64 @@
+// Copyright Epic Games, Inc. All Rights Reserved.
+
+#pragma once
+
+#include "CoreMinimal.h"
+#include "Nodes/InterchangeBaseNode.h"
+
+#include "InterchangeStaticMeshLodDataNode.generated.h"
+
+
+namespace UE
+{
+	namespace Interchange
+	{
+		struct FStaticMeshNodeLodDataStaticData : public FBaseNodeStaticData
+		{
+			static const FString& GetMeshUidsBaseKey();
+		};
+	} // namespace Interchange
+} // namespace UE
+
+
+UCLASS(BlueprintType, Experimental)
+class INTERCHANGEFACTORYNODES_API UInterchangeStaticMeshLodDataNode : public UInterchangeBaseNode
+{
+	GENERATED_BODY()
+
+public:
+	UInterchangeStaticMeshLodDataNode();
+
+	/**
+	 * Return the node type name of the class, we use this when reporting error
+	 */
+	virtual FString GetTypeName() const override;
+
+	virtual FString GetKeyDisplayName(const UE::Interchange::FAttributeKey& NodeAttributeKey) const override;
+
+	virtual FString GetAttributeCategory(const UE::Interchange::FAttributeKey& NodeAttributeKey) const override;
+
+public:
+	/* Mesh Uids: It can be either a scene or a mesh node uid. If its a scene it mean we want the mesh factory to bake the geo payload with the global transform of the scene node. */
+
+	UFUNCTION(BlueprintCallable, Category = "Interchange | Node | StaticMeshLodData")
+	int32 GetMeshUidsCount() const;
+
+	UFUNCTION(BlueprintCallable, Category = "Interchange | Node | StaticMeshLodData")
+	void GetMeshUids(TArray<FString>& OutMeshNames) const;
+
+	UFUNCTION(BlueprintCallable, Category = "Interchange | Node | StaticMeshLodData")
+	bool AddMeshUid(const FString& MeshName);
+
+	UFUNCTION(BlueprintCallable, Category = "Interchange | Node | StaticMeshLodData")
+	bool RemoveMeshUid(const FString& MeshName);
+
+	UFUNCTION(BlueprintCallable, Category = "Interchange | Node | StaticMeshLodData")
+	bool RemoveAllMeshes();
+
+private:
+
+	bool IsEditorOnlyDataDefined();
+
+	UE::Interchange::TArrayAttributeHelper<FString> MeshUids;
+protected:
+};
diff --git a/Engine/Plugins/Experimental/Interchange/Runtime/Source/FactoryNodes/Public/InterchangeTexture2DArrayFactoryNode.h b/Engine/Plugins/Experimental/Interchange/Runtime/Source/FactoryNodes/Public/InterchangeTexture2DArrayFactoryNode.h
new file mode 100644
--- /dev/null
+++ b/Engine/Plugins/Experimental/Interchange/Runtime/Source/FactoryNodes/Public/InterchangeTexture2DArrayFactoryNode.h
@@ -0,0 +1,76 @@
+// Copyright Epic Games, Inc. All Rights Reserved.
+
+#pragma once
+
+#include "CoreMinimal.h"
+#include "InterchangeTextureFactoryNode.h"
+#include "Nodes/InterchangeBaseNode.h"
+
+#if WITH_ENGINE
+#include "Engine/Texture.h"
+#include "Engine/Texture2DArray.h"
+#endif
+
+
+#include "InterchangeTexture2DArrayFactoryNode.generated.h"
+
+UCLASS(BlueprintType, Experimental)
+class INTERCHANGEFACTORYNODES_API UInterchangeTexture2DArrayFactoryNode : public UInterchangeTextureFactoryNode
+{
+	GENERATED_BODY()
+
+public:
+
+	virtual bool SetCustomAddressX(const uint8 AttributeValue, bool bAddApplyDelegate = true) override
+	{
+#if WITH_EDITORONLY_DATA
+		IMPLEMENT_NODE_ATTRIBUTE_SETTER(UInterchangeTexture2DArrayFactoryNode, AddressX, uint8, UTexture2DArray)
+#else
+		IMPLEMENT_NODE_ATTRIBUTE_SETTER_NODELEGATE(AddressX, uint8)
+#endif
+	}
+
+	virtual bool SetCustomAddressY(const uint8 AttributeValue, bool bAddApplyDelegate = true) override
+	{
+#if WITH_EDITORONLY_DATA
+		IMPLEMENT_NODE_ATTRIBUTE_SETTER(UInterchangeTexture2DArrayFactoryNode, AddressY, uint8, UTexture2DArray)
+#else
+		IMPLEMENT_NODE_ATTRIBUTE_SETTER_NODELEGATE(AddressY, uint8)
+#endif
+	}
+
+	UFUNCTION(BlueprintCallable, Category = "Interchange | Node | Texture")
+	bool GetCustomAddressZ(uint8& AttributeValue) const
+	{
+		IMPLEMENT_NODE_ATTRIBUTE_GETTER(AddressZ, uint8);
+	}
+
+	UFUNCTION(BlueprintCallable, Category = "Interchange | Node | Texture")
+	bool SetCustomAddressZ(const uint8 AttributeValue, bool bAddApplyDelegate = true)
+	{
+#if WITH_EDITORONLY_DATA
+		IMPLEMENT_NODE_ATTRIBUTE_SETTER(UInterchangeTexture2DArrayFactoryNode, AddressZ, uint8, UTexture2DArray)
+#else
+		IMPLEMENT_NODE_ATTRIBUTE_SETTER_NODELEGATE(AddressZ, uint8)
+#endif
+	}
+
+
+private:
+#if WITH_ENGINE
+	virtual void FillAssetClassFromClassName(const FString& ClassName) override
+	{
+		if (ClassName.Equals(UTexture2DArray::StaticClass()->GetName()))
+		{
+			AssetClass = UTexture2DArray::StaticClass();
+			bIsTextureNodeClassInitialized = true;
+		}
+	}
+#endif
+
+	// Addressing
+	IMPLEMENT_NODE_ATTRIBUTE_KEY(AddressX);
+	IMPLEMENT_NODE_ATTRIBUTE_KEY(AddressY);
+	IMPLEMENT_NODE_ATTRIBUTE_KEY(AddressZ);
+
+};
\ No newline at end of file
diff --git a/Engine/Plugins/Experimental/Interchange/Runtime/Source/FactoryNodes/Public/InterchangeTextureCubeFactoryNode.h b/Engine/Plugins/Experimental/Interchange/Runtime/Source/FactoryNodes/Public/InterchangeTextureCubeFactoryNode.h
new file mode 100644
--- /dev/null
+++ b/Engine/Plugins/Experimental/Interchange/Runtime/Source/FactoryNodes/Public/InterchangeTextureCubeFactoryNode.h
@@ -0,0 +1,33 @@
+// Copyright Epic Games, Inc. All Rights Reserved.
+
+#pragma once
+
+#include "CoreMinimal.h"
+#include "InterchangeTextureFactoryNode.h"
+
+#if WITH_ENGINE
+#include "Engine/TextureCube.h"
+#endif
+
+
+#include "InterchangeTextureCubeFactoryNode.generated.h"
+
+UCLASS(BlueprintType, Experimental)
+class INTERCHANGEFACTORYNODES_API UInterchangeTextureCubeFactoryNode : public UInterchangeTextureFactoryNode
+{
+	GENERATED_BODY()
+
+private:
+
+#if WITH_ENGINE
+	virtual void FillAssetClassFromClassName(const FString& ClassName) override
+	{
+		if (ClassName.Equals(UTextureCube::StaticClass()->GetName()))
+		{
+			AssetClass = UTextureCube::StaticClass();
+			bIsTextureNodeClassInitialized = true;
+		}
+	}
+#endif
+
+};
\ No newline at end of file
diff --git a/Engine/Plugins/Experimental/Interchange/Runtime/Source/FactoryNodes/Public/InterchangeTextureFactoryNode.h b/Engine/Plugins/Experimental/Interchange/Runtime/Source/FactoryNodes/Public/InterchangeTextureFactoryNode.h
new file mode 100644
--- /dev/null
+++ b/Engine/Plugins/Experimental/Interchange/Runtime/Source/FactoryNodes/Public/InterchangeTextureFactoryNode.h
@@ -0,0 +1,855 @@
+// Copyright Epic Games, Inc. All Rights Reserved.
+
+#pragma once
+
+#include "CoreMinimal.h"
+#include "Nodes/InterchangeBaseNode.h"
+
+#if WITH_ENGINE
+#include "Engine/Texture.h"
+#include "Engine/Texture2D.h"
+#endif
+
+#include "InterchangeTextureFactoryNode.generated.h"
+
+
+UCLASS(BlueprintType, Experimental)
+class INTERCHANGEFACTORYNODES_API UInterchangeTextureFactoryNode : public UInterchangeBaseNode
+{
+	GENERATED_BODY()
+
+public:
+	UInterchangeTextureFactoryNode()
+		:UInterchangeBaseNode()
+	{
+#if WITH_ENGINE
+		AssetClass = nullptr;
+#endif
+	}
+
+	/**
+	 * Initialize node data
+	 * @param: UniqueID - The uniqueId for this node
+	 * @param DisplayLabel - The name of the node
+	 * @param InAssetClass - The class the texture factory will create for this node.
+	 */
+	UFUNCTION(BlueprintCallable, Category = "Interchange | Node | Texture")
+	void InitializeTextureNode(const FString& UniqueID, const FString& DisplayLabel, const FString& InAssetClass, const FString& InAssetName)
+	{
+		//Initialize
+		bIsTextureNodeClassInitialized = false;
+		InitializeNode(UniqueID, DisplayLabel, EInterchangeNodeContainerType::FactoryData);
+		
+		//Set the class
+		FString OperationName = GetTypeName() + TEXT(".SetAssetClassName");
+		InterchangePrivateNodeBase::SetCustomAttribute<FString>(*Attributes, ClassNameAttributeKey, OperationName, InAssetClass);
+		FillAssetClassFromAttribute();
+		
+		//Set the asset name
+		OperationName = GetTypeName() + TEXT(".SetAssetName");
+		InterchangePrivateNodeBase::SetCustomAttribute<FString>(*Attributes, AssetNameKey, OperationName, InAssetName);
+	}
+
+	virtual void Serialize(FArchive& Ar) override
+	{
+		Super::Serialize(Ar);
+#if WITH_ENGINE
+		if (Ar.IsLoading())
+		{
+			//Make sure the class is properly set when we compile with engine, this will set the
+			//bIsTextureNodeClassInitialized to true.
+			SetTextureNodeClassFromClassAttribute();
+		}
+#endif //#if WITH_ENGINE
+	}
+
+	/**
+	 * Return the node type name of the class, we use this when reporting error
+	 */
+	virtual FString GetTypeName() const override
+	{
+		const FString TypeName = TEXT("TextureFactoryNode");
+		return TypeName;
+	}
+
+	/** Get the class this node want to create */
+	UFUNCTION(BlueprintCallable, Category = "Interchange | Node | Texture")
+	virtual class UClass* GetObjectClass() const override
+	{
+		ensure(bIsTextureNodeClassInitialized);
+#if WITH_ENGINE
+		return AssetClass.Get() != nullptr ? AssetClass.Get() : UTexture::StaticClass();
+#else
+		return nullptr;
+#endif
+	}
+
+	virtual FGuid GetHash() const override
+	{
+		return Attributes->GetStorageHash();
+	}
+
+	static FString GetTextureFactoryNodeUidFromTextureNodeUid(const FString& TranslatedNodeUid)
+	{
+		FString NewUid = TEXT("Factory_") + TranslatedNodeUid;
+		return NewUid;
+	}
+
+	/**
+	 * Get the translated texture node unique ID.
+	 */
+	UFUNCTION(BlueprintCallable, Category = "Interchange | Node | Texture")
+	bool GetCustomTranslatedTextureNodeUid(FString& AttributeValue) const
+	{
+		IMPLEMENT_NODE_ATTRIBUTE_GETTER(TranslatedTextureNodeUid, FString);
+	}
+
+	/**
+	 * Set the translated texture node unique ID. This is the reference to the node that was create by the translator and this node is needed to get the texture payload.
+	 */
+	UFUNCTION(BlueprintCallable, Category = "Interchange | Node | Texture")
+	bool SetCustomTranslatedTextureNodeUid(const FString& AttributeValue)
+	{
+		IMPLEMENT_NODE_ATTRIBUTE_SETTER_NODELEGATE(TranslatedTextureNodeUid, FString);
+	}
+
+public:
+	/** Return false if the Attribute was not set previously.*/
+	UFUNCTION(BlueprintCallable, Category = "Interchange | Node | Texture")
+	bool GetCustomAdjustBrightness(float& AttributeValue) const
+	{
+		IMPLEMENT_NODE_ATTRIBUTE_GETTER(AdjustBrightness, float);
+	}
+
+	UFUNCTION(BlueprintCallable, Category = "Interchange | Node | Texture")
+	bool SetCustomAdjustBrightness(const float& AttributeValue, bool bAddApplyDelegate = true)
+	{
+#if WITH_EDITORONLY_DATA
+		IMPLEMENT_NODE_ATTRIBUTE_SETTER(UInterchangeTextureFactoryNode, AdjustBrightness, float, UTexture)
+#else
+		IMPLEMENT_NODE_ATTRIBUTE_SETTER_NODELEGATE(AdjustBrightness, float)
+#endif
+	}
+
+	/** Return false if the Attribute was not set previously.*/
+	UFUNCTION(BlueprintCallable, Category = "Interchange | Node | Texture")
+	bool GetCustomAdjustBrightnessCurve(float& AttributeValue) const
+	{
+		IMPLEMENT_NODE_ATTRIBUTE_GETTER(AdjustBrightnessCurve, float);
+	}
+
+	UFUNCTION(BlueprintCallable, Category = "Interchange | Node | Texture")
+	bool SetCustomAdjustBrightnessCurve(const float& AttributeValue, bool bAddApplyDelegate = true)
+	{
+#if WITH_EDITORONLY_DATA
+		IMPLEMENT_NODE_ATTRIBUTE_SETTER(UInterchangeTextureFactoryNode, AdjustBrightnessCurve, float, UTexture)
+#else
+		IMPLEMENT_NODE_ATTRIBUTE_SETTER_NODELEGATE(AdjustBrightnessCurve, float)
+#endif
+	}
+
+	/** Return false if the Attribute was not set previously.*/
+	UFUNCTION(BlueprintCallable, Category = "Interchange | Node | Texture")
+	bool GetCustomAdjustVibrance(float& AttributeValue) const
+	{
+		IMPLEMENT_NODE_ATTRIBUTE_GETTER(AdjustVibrance, float);
+	}
+
+	UFUNCTION(BlueprintCallable, Category = "Interchange | Node | Texture")
+	bool SetCustomAdjustVibrance(const float& AttributeValue, bool bAddApplyDelegate = true)
+	{
+#if WITH_EDITORONLY_DATA
+		IMPLEMENT_NODE_ATTRIBUTE_SETTER(UInterchangeTextureFactoryNode, AdjustVibrance, float, UTexture)
+#else
+		IMPLEMENT_NODE_ATTRIBUTE_SETTER_NODELEGATE(AdjustVibrance, float)
+#endif
+	}
+
+	/** Return false if the Attribute was not set previously.*/
+	UFUNCTION(BlueprintCallable, Category = "Interchange | Node | Texture")
+	bool GetCustomAdjustSaturation(float& AttributeValue) const
+	{
+		IMPLEMENT_NODE_ATTRIBUTE_GETTER(AdjustSaturation, float);
+	}
+
+	UFUNCTION(BlueprintCallable, Category = "Interchange | Node | Texture")
+	bool SetCustomAdjustSaturation(const float& AttributeValue, bool bAddApplyDelegate = true)
+	{
+#if WITH_EDITORONLY_DATA
+		IMPLEMENT_NODE_ATTRIBUTE_SETTER(UInterchangeTextureFactoryNode, AdjustSaturation, float, UTexture)
+#else
+		IMPLEMENT_NODE_ATTRIBUTE_SETTER_NODELEGATE(AdjustSaturation, float)
+#endif
+	}
+
+	/** Return false if the Attribute was not set previously.*/
+	UFUNCTION(BlueprintCallable, Category = "Interchange | Node | Texture")
+	bool GetCustomAdjustRGBCurve(float& AttributeValue) const
+	{
+		IMPLEMENT_NODE_ATTRIBUTE_GETTER(AdjustRGBCurve, float);
+	}
+
+	UFUNCTION(BlueprintCallable, Category = "Interchange | Node | Texture")
+	bool SetCustomAdjustRGBCurve(const float& AttributeValue, bool bAddApplyDelegate = true)
+	{
+#if WITH_EDITORONLY_DATA
+		IMPLEMENT_NODE_ATTRIBUTE_SETTER(UInterchangeTextureFactoryNode, AdjustRGBCurve, float, UTexture)
+#else
+		IMPLEMENT_NODE_ATTRIBUTE_SETTER_NODELEGATE(AdjustRGBCurve, float)
+#endif
+	}
+
+	/** Return false if the Attribute was not set previously.*/
+	UFUNCTION(BlueprintCallable, Category = "Interchange | Node | Texture")
+	bool GetCustomAdjustHue(float& AttributeValue) const
+	{
+		IMPLEMENT_NODE_ATTRIBUTE_GETTER(AdjustHue, float);
+	}
+
+	UFUNCTION(BlueprintCallable, Category = "Interchange | Node | Texture")
+	bool SetCustomAdjustHue(const float& AttributeValue, bool bAddApplyDelegate = true)
+	{
+#if WITH_EDITORONLY_DATA
+		IMPLEMENT_NODE_ATTRIBUTE_SETTER(UInterchangeTextureFactoryNode, AdjustHue, float, UTexture)
+#else
+		IMPLEMENT_NODE_ATTRIBUTE_SETTER_NODELEGATE(AdjustHue, float)
+#endif
+	}
+
+	/** Return false if the Attribute was not set previously.*/
+	UFUNCTION(BlueprintCallable, Category = "Interchange | Node | Texture")
+	bool GetCustomAdjustMinAlpha(float& AttributeValue) const
+	{
+		IMPLEMENT_NODE_ATTRIBUTE_GETTER(AdjustMinAlpha, float);
+	}
+
+	UFUNCTION(BlueprintCallable, Category = "Interchange | Node | Texture")
+	bool SetCustomAdjustMinAlpha(const float& AttributeValue, bool bAddApplyDelegate = true)
+	{
+#if WITH_EDITORONLY_DATA
+		IMPLEMENT_NODE_ATTRIBUTE_SETTER(UInterchangeTextureFactoryNode, AdjustMinAlpha, float, UTexture)
+#else
+		IMPLEMENT_NODE_ATTRIBUTE_SETTER_NODELEGATE(AdjustMinAlpha, float)
+#endif
+	}
+
+	/** Return false if the Attribute was not set previously.*/
+	UFUNCTION(BlueprintCallable, Category = "Interchange | Node | Texture")
+	bool GetCustomAdjustMaxAlpha(float& AttributeValue) const
+	{
+		IMPLEMENT_NODE_ATTRIBUTE_GETTER(AdjustMaxAlpha, float);
+	}
+
+	UFUNCTION(BlueprintCallable, Category = "Interchange | Node | Texture")
+	bool SetCustomAdjustMaxAlpha(const float& AttributeValue, bool bAddApplyDelegate = true)
+	{
+#if WITH_EDITORONLY_DATA
+		IMPLEMENT_NODE_ATTRIBUTE_SETTER(UInterchangeTextureFactoryNode, AdjustMaxAlpha, float, UTexture)
+#else
+		IMPLEMENT_NODE_ATTRIBUTE_SETTER_NODELEGATE(AdjustMaxAlpha, float)
+#endif
+	}
+
+	/** Return false if the Attribute was not set previously.*/
+	UFUNCTION(BlueprintCallable, Category = "Interchange | Node | Texture")
+	bool GetCustombChromaKeyTexture(bool& AttributeValue) const
+	{
+		IMPLEMENT_NODE_ATTRIBUTE_GETTER(bChromaKeyTexture, bool);
+	}
+
+	UFUNCTION(BlueprintCallable, Category = "Interchange | Node | Texture")
+	bool SetCustombChromaKeyTexture(const bool& AttributeValue, bool bAddApplyDelegate = true)
+	{
+#if WITH_EDITORONLY_DATA
+		IMPLEMENT_NODE_ATTRIBUTE_SETTER(UInterchangeTextureFactoryNode, bChromaKeyTexture, bool, UTexture)
+#else
+		IMPLEMENT_NODE_ATTRIBUTE_SETTER_NODELEGATE(bChromaKeyTexture, bool)
+#endif
+	}
+
+	/** Return false if the Attribute was not set previously.*/
+	UFUNCTION(BlueprintCallable, Category = "Interchange | Node | Texture")
+	bool GetCustomChromaKeyThreshold(float& AttributeValue) const
+	{
+		IMPLEMENT_NODE_ATTRIBUTE_GETTER(ChromaKeyThreshold, float);
+	}
+
+	UFUNCTION(BlueprintCallable, Category = "Interchange | Node | Texture")
+	bool SetCustomChromaKeyThreshold(const float& AttributeValue, bool bAddApplyDelegate = true)
+	{
+#if WITH_EDITORONLY_DATA
+		IMPLEMENT_NODE_ATTRIBUTE_SETTER(UInterchangeTextureFactoryNode, ChromaKeyThreshold, float, UTexture)
+#else
+		IMPLEMENT_NODE_ATTRIBUTE_SETTER_NODELEGATE(ChromaKeyThreshold, float)
+#endif
+	}
+
+	/** Return false if the Attribute was not set previously.*/
+	UFUNCTION(BlueprintCallable, Category = "Interchange | Node | Texture")
+	bool GetCustomChromaKeyColor(FColor& AttributeValue) const
+	{
+		IMPLEMENT_NODE_ATTRIBUTE_GETTER(ChromaKeyColor, FColor);
+	}
+
+	UFUNCTION(BlueprintCallable, Category = "Interchange | Node | Texture")
+	bool SetCustomChromaKeyColor(const FColor& AttributeValue, bool bAddApplyDelegate = true)
+	{
+#if WITH_EDITORONLY_DATA
+		IMPLEMENT_NODE_ATTRIBUTE_SETTER(UInterchangeTextureFactoryNode, ChromaKeyColor, FColor, UTexture)
+#else
+		IMPLEMENT_NODE_ATTRIBUTE_SETTER_NODELEGATE(ChromaKeyColor, FColor)
+#endif
+	}
+
+	//////////////////////////////////////////////////////////////////////////
+	//Texture Compression Begin
+
+	/** Return false if the Attribute was not set previously.*/
+	UFUNCTION(BlueprintCallable, Category = "Interchange | Node | Texture")
+	bool GetCustomCompressionNoAlpha(bool& AttributeValue) const
+	{
+		IMPLEMENT_NODE_ATTRIBUTE_GETTER(CompressionNoAlpha, bool);
+	}
+
+	UFUNCTION(BlueprintCallable, Category = "Interchange | Node | Texture")
+	bool SetCustomCompressionNoAlpha(const bool& AttributeValue, bool bAddApplyDelegate = true)
+	{
+#if WITH_EDITORONLY_DATA
+		IMPLEMENT_NODE_ATTRIBUTE_SETTER(UInterchangeTextureFactoryNode, CompressionNoAlpha, bool, UTexture)
+#else
+		IMPLEMENT_NODE_ATTRIBUTE_SETTER_NODELEGATE(CompressionNoAlpha, bool)
+#endif
+	}
+
+	/** Return false if the Attribute was not set previously.*/
+	UFUNCTION(BlueprintCallable, Category = "Interchange | Node | Texture")
+	bool GetCustomDeferCompression(bool& AttributeValue) const
+	{
+		IMPLEMENT_NODE_ATTRIBUTE_GETTER(DeferCompression, bool);
+	}
+
+	UFUNCTION(BlueprintCallable, Category = "Interchange | Node | Texture")
+	bool SetCustomDeferCompression(const bool& AttributeValue, bool bAddApplyDelegate = true)
+	{
+#if WITH_EDITORONLY_DATA
+		IMPLEMENT_NODE_ATTRIBUTE_SETTER(UInterchangeTextureFactoryNode, DeferCompression, bool, UTexture)
+#else
+		IMPLEMENT_NODE_ATTRIBUTE_SETTER_NODELEGATE(DeferCompression, bool)
+#endif
+	}
+
+	/** Return false if the Attribute was not set previously.*/
+	UFUNCTION(BlueprintCallable, Category = "Interchange | Node | Texture")
+	bool GetCustomLossyCompressionAmount(uint8& AttributeValue) const
+	{
+		IMPLEMENT_NODE_ATTRIBUTE_GETTER(LossyCompressionAmount, uint8);
+	}
+
+	UFUNCTION(BlueprintCallable, Category = "Interchange | Node | Texture")
+	bool SetCustomLossyCompressionAmount(const uint8& AttributeValue, bool bAddApplyDelegate = true)
+	{
+#if WITH_EDITORONLY_DATA
+		IMPLEMENT_NODE_ATTRIBUTE_SETTER(UInterchangeTextureFactoryNode, LossyCompressionAmount, uint8, UTexture)
+#else
+		IMPLEMENT_NODE_ATTRIBUTE_SETTER_NODELEGATE(LossyCompressionAmount, uint8)
+#endif
+	}
+
+	/** Return false if the Attribute was not set previously.*/
+	UFUNCTION(BlueprintCallable, Category = "Interchange | Node | Texture")
+	bool GetCustomMaxTextureSize(int32& AttributeValue) const
+	{
+		IMPLEMENT_NODE_ATTRIBUTE_GETTER(MaxTextureSize, int32);
+	}
+
+	UFUNCTION(BlueprintCallable, Category = "Interchange | Node | Texture")
+	bool SetCustomMaxTextureSize(const int32& AttributeValue, bool bAddApplyDelegate = true)
+	{
+#if WITH_EDITORONLY_DATA
+		IMPLEMENT_NODE_ATTRIBUTE_SETTER(UInterchangeTextureFactoryNode, MaxTextureSize, int32, UTexture)
+#else
+		IMPLEMENT_NODE_ATTRIBUTE_SETTER_NODELEGATE(MaxTextureSize, int32)
+#endif
+	}
+
+	/** Return false if the Attribute was not set previously.*/
+	UFUNCTION(BlueprintCallable, Category = "Interchange | Node | Texture")
+	bool GetCustomCompressionQuality(uint8& AttributeValue) const
+	{
+		IMPLEMENT_NODE_ATTRIBUTE_GETTER(CompressionQuality, uint8);
+	}
+
+	UFUNCTION(BlueprintCallable, Category = "Interchange | Node | Texture")
+	bool SetCustomCompressionQuality(const uint8& AttributeValue, bool bAddApplyDelegate = true)
+	{
+#if WITH_EDITORONLY_DATA
+		IMPLEMENT_NODE_ATTRIBUTE_SETTER(UInterchangeTextureFactoryNode, CompressionQuality, uint8, UTexture)
+#else
+		IMPLEMENT_NODE_ATTRIBUTE_SETTER_NODELEGATE(CompressionQuality, uint8)
+#endif
+	}
+
+	/** Return false if the Attribute was not set previously.*/
+	UFUNCTION(BlueprintCallable, Category = "Interchange | Node | Texture")
+	bool GetCustomCompressionSettings(uint8& AttributeValue) const
+	{
+		IMPLEMENT_NODE_ATTRIBUTE_GETTER(CompressionSettings, uint8);
+	}
+
+	UFUNCTION(BlueprintCallable, Category = "Interchange | Node | Texture")
+	bool SetCustomCompressionSettings(const uint8& AttributeValue, bool bAddApplyDelegate = true)
+	{
+		IMPLEMENT_NODE_ATTRIBUTE_SETTER(UInterchangeTextureFactoryNode, CompressionSettings, uint8, UTexture)
+	}
+
+	//////////////////////////////////////////////////////////////////////////
+	//Texture general
+
+	/** Return false if the Attribute was not set previously.*/
+	UFUNCTION(BlueprintCallable, Category = "Interchange | Node | Texture")
+	bool GetCustombDitherMipMapAlpha(bool& AttributeValue) const
+	{
+		IMPLEMENT_NODE_ATTRIBUTE_GETTER(bDitherMipMapAlpha, bool);
+	}
+
+	UFUNCTION(BlueprintCallable, Category = "Interchange | Node | Texture")
+	bool SetCustombDitherMipMapAlpha(const bool& AttributeValue, bool bAddApplyDelegate = true)
+	{
+#if WITH_EDITORONLY_DATA
+		IMPLEMENT_NODE_ATTRIBUTE_SETTER(UInterchangeTextureFactoryNode, bDitherMipMapAlpha, bool, UTexture)
+#else
+		IMPLEMENT_NODE_ATTRIBUTE_SETTER_NODELEGATE(bDitherMipMapAlpha, bool)
+#endif
+	}
+
+	/** Return false if the Attribute was not set previously.*/
+	UFUNCTION(BlueprintCallable, Category = "Interchange | Node | Texture")
+	bool GetCustomAlphaCoverageThresholds(FVector4& AttributeValue) const
+	{
+		IMPLEMENT_NODE_ATTRIBUTE_GETTER(AlphaCoverageThresholds, FVector4);
+	}
+
+	UFUNCTION(BlueprintCallable, Category = "Interchange | Node | Texture")
+	bool SetCustomAlphaCoverageThresholds(const FVector4& AttributeValue, bool bAddApplyDelegate = true)
+	{
+#if WITH_EDITORONLY_DATA
+		IMPLEMENT_NODE_ATTRIBUTE_SETTER(UInterchangeTextureFactoryNode, AlphaCoverageThresholds, FVector4, UTexture)
+#else
+		IMPLEMENT_NODE_ATTRIBUTE_SETTER_NODELEGATE(AlphaCoverageThresholds, FVector4)
+#endif
+	}
+
+	/** Return false if the Attribute was not set previously.*/
+	UFUNCTION(BlueprintCallable, Category = "Interchange | Node | Texture")
+	bool GetCustombFlipGreenChannel(bool& AttributeValue) const
+	{
+		IMPLEMENT_NODE_ATTRIBUTE_GETTER(bFlipGreenChannel, bool);
+	}
+
+	UFUNCTION(BlueprintCallable, Category = "Interchange | Node | Texture")
+	bool SetCustombFlipGreenChannel(const bool& AttributeValue, bool bAddApplyDelegate = true)
+	{
+#if WITH_EDITORONLY_DATA
+		IMPLEMENT_NODE_ATTRIBUTE_SETTER(UInterchangeTextureFactoryNode, bFlipGreenChannel, bool, UTexture)
+#else
+		IMPLEMENT_NODE_ATTRIBUTE_SETTER_NODELEGATE(bFlipGreenChannel, bool)
+#endif
+	}
+
+	/** Return false if the Attribute was not set previously.*/
+	UFUNCTION(BlueprintCallable, Category = "Interchange | Node | Texture")
+	bool GetCustomPowerOfTwoMode(uint8& AttributeValue) const
+	{
+		IMPLEMENT_NODE_ATTRIBUTE_GETTER(PowerOfTwoMode, uint8);
+	}
+
+	UFUNCTION(BlueprintCallable, Category = "Interchange | Node | Texture")
+	bool SetCustomPowerOfTwoMode(const uint8& AttributeValue, bool bAddApplyDelegate = true)
+	{
+#if WITH_EDITORONLY_DATA
+		IMPLEMENT_NODE_ATTRIBUTE_SETTER(UInterchangeTextureFactoryNode, PowerOfTwoMode, uint8, UTexture)
+#else
+		IMPLEMENT_NODE_ATTRIBUTE_SETTER_NODELEGATE(PowerOfTwoMode, uint8)
+#endif
+	}
+
+	/** Return false if the Attribute was not set previously.*/
+	UFUNCTION(BlueprintCallable, Category = "Interchange | Node | Texture")
+	bool GetCustomPaddingColor(FColor& AttributeValue) const
+	{
+		IMPLEMENT_NODE_ATTRIBUTE_GETTER(PaddingColor, FColor);
+	}
+
+	UFUNCTION(BlueprintCallable, Category = "Interchange | Node | Texture")
+	bool SetCustomPaddingColor(const FColor& AttributeValue, bool bAddApplyDelegate = true)
+	{
+#if WITH_EDITORONLY_DATA
+		IMPLEMENT_NODE_ATTRIBUTE_SETTER(UInterchangeTextureFactoryNode, PaddingColor, FColor, UTexture)
+#else
+		IMPLEMENT_NODE_ATTRIBUTE_SETTER_NODELEGATE(PaddingColor, FColor)
+#endif
+	}
+
+	/** Return false if the Attribute was not set previously.*/
+	UFUNCTION(BlueprintCallable, Category = "Interchange | Node | Texture")
+	bool GetCustomFilter(uint8& AttributeValue) const
+	{
+		IMPLEMENT_NODE_ATTRIBUTE_GETTER(Filter, uint8);
+	}
+
+	UFUNCTION(BlueprintCallable, Category = "Interchange | Node | Texture")
+	bool SetCustomFilter(const uint8& AttributeValue, bool bAddApplyDelegate = true)
+	{
+		IMPLEMENT_NODE_ATTRIBUTE_SETTER(UInterchangeTextureFactoryNode, Filter, uint8, UTexture)
+	}
+
+	/** Return false if the Attribute was not set previously.*/
+	UFUNCTION(BlueprintCallable, Category = "Interchange | Node | Texture")
+	bool GetCustomMipLoadOptions(uint8& AttributeValue) const
+	{
+		IMPLEMENT_NODE_ATTRIBUTE_GETTER(MipLoadOptions, uint8);
+	}
+
+	UFUNCTION(BlueprintCallable, Category = "Interchange | Node | Texture")
+	bool SetCustomMipLoadOptions(const uint8& AttributeValue, bool bAddApplyDelegate = true)
+	{
+		IMPLEMENT_NODE_ATTRIBUTE_SETTER(UInterchangeTextureFactoryNode, MipLoadOptions, uint8, UTexture)
+	}
+
+	/** Return false if the Attribute was not set previously.*/
+	UFUNCTION(BlueprintCallable, Category = "Interchange | Node | Texture")
+	bool GetCustomSRGB(bool& AttributeValue) const
+	{
+		IMPLEMENT_NODE_ATTRIBUTE_GETTER(SRGB, bool);
+	}
+
+	UFUNCTION(BlueprintCallable, Category = "Interchange | Node | Texture")
+	bool SetCustomSRGB(const bool& AttributeValue, bool bAddApplyDelegate = true)
+	{
+		IMPLEMENT_NODE_ATTRIBUTE_SETTER(UInterchangeTextureFactoryNode, SRGB, bool, UTexture)
+	}
+
+	/** Return false if the Attribute was not set previously.*/
+	UFUNCTION(BlueprintCallable, Category = "Interchange | Node | Texture")
+	bool GetCustombUseLegacyGamma(bool& AttributeValue) const
+	{
+		IMPLEMENT_NODE_ATTRIBUTE_GETTER(bUseLegacyGamma, bool);
+	}
+
+	UFUNCTION(BlueprintCallable, Category = "Interchange | Node | Texture")
+	bool SetCustombUseLegacyGamma(const bool& AttributeValue, bool bAddApplyDelegate = true)
+	{
+#if WITH_EDITORONLY_DATA
+		IMPLEMENT_NODE_ATTRIBUTE_SETTER(UInterchangeTextureFactoryNode, bUseLegacyGamma, bool, UTexture)
+#else
+		IMPLEMENT_NODE_ATTRIBUTE_SETTER_NODELEGATE(bUseLegacyGamma, bool)
+#endif
+	}
+
+	/** Return false if the Attribute was not set previously.*/
+	UFUNCTION(BlueprintCallable, Category = "Interchange | Node | Texture")
+	bool GetCustomVirtualTextureStreaming(bool& AttributeValue) const
+	{
+		IMPLEMENT_NODE_ATTRIBUTE_GETTER(VirtualTextureStreaming, bool);
+	}
+
+	UFUNCTION(BlueprintCallable, Category = "Interchange | Node | Texture")
+	bool SetCustomVirtualTextureStreaming(const bool& AttributeValue, bool bAddApplyDelegate = true)
+	{
+		IMPLEMENT_NODE_ATTRIBUTE_SETTER(UInterchangeTextureFactoryNode, VirtualTextureStreaming, bool, UTexture)
+	}
+
+	/** Return false if the Attribute was not set previously.*/
+	UFUNCTION(BlueprintCallable, Category = "Interchange | Node | Texture")
+	virtual bool GetCustomAddressX(uint8& AttributeValue) const
+	{
+		IMPLEMENT_NODE_ATTRIBUTE_GETTER(AddressX, uint8);
+	}
+
+	UFUNCTION(BlueprintCallable, Category = "Interchange | Node | Texture")
+	virtual bool SetCustomAddressX(const uint8 AttributeValue, bool bAddApplyDelegate = true)
+	{
+		IMPLEMENT_NODE_ATTRIBUTE_SETTER(UInterchangeTextureFactoryNode, AddressX, uint8, UTexture2D)
+	}
+
+	/** Return false if the Attribute was not set previously.*/
+	UFUNCTION(BlueprintCallable, Category = "Interchange | Node | Texture")
+	virtual bool GetCustomAddressY(uint8& AttributeValue) const
+	{
+		IMPLEMENT_NODE_ATTRIBUTE_GETTER(AddressY, uint8);
+	}
+
+	UFUNCTION(BlueprintCallable, Category = "Interchange | Node | Texture")
+	virtual bool SetCustomAddressY(const uint8 AttributeValue, bool bAddApplyDelegate = true)
+	{
+		IMPLEMENT_NODE_ATTRIBUTE_SETTER(UInterchangeTextureFactoryNode, AddressY, uint8, UTexture2D)
+	}
+
+	//////////////////////////////////////////////////////////////////////////
+	//Level of Detail
+
+	/** Return false if the Attribute was not set previously.*/
+	UFUNCTION(BlueprintCallable, Category = "Interchange | Node | Texture")
+	bool GetCustombPreserveBorder(bool& AttributeValue) const
+	{
+		IMPLEMENT_NODE_ATTRIBUTE_GETTER(bPreserveBorder, bool);
+	}
+
+	UFUNCTION(BlueprintCallable, Category = "Interchange | Node | Texture")
+	bool SetCustombPreserveBorder(const bool& AttributeValue, bool bAddApplyDelegate = true)
+	{
+#if WITH_EDITORONLY_DATA
+		IMPLEMENT_NODE_ATTRIBUTE_SETTER(UInterchangeTextureFactoryNode, bPreserveBorder, bool, UTexture)
+#else
+		IMPLEMENT_NODE_ATTRIBUTE_SETTER_NODELEGATE(bPreserveBorder, bool)
+#endif
+	}
+
+	/** Return false if the Attribute was not set previously.*/
+	UFUNCTION(BlueprintCallable, Category = "Interchange | Node | Texture")
+	bool GetCustomMipGenSettings(uint8& AttributeValue) const
+	{
+		IMPLEMENT_NODE_ATTRIBUTE_GETTER(MipGenSettings, uint8);
+	}
+
+	UFUNCTION(BlueprintCallable, Category = "Interchange | Node | Texture")
+	bool SetCustomMipGenSettings(const uint8& AttributeValue, bool bAddApplyDelegate = true)
+	{
+#if WITH_EDITORONLY_DATA
+		IMPLEMENT_NODE_ATTRIBUTE_SETTER(UInterchangeTextureFactoryNode, MipGenSettings, uint8, UTexture)
+#else
+		IMPLEMENT_NODE_ATTRIBUTE_SETTER_NODELEGATE(MipGenSettings, uint8)
+#endif
+	}
+
+	/** Return false if the Attribute was not set previously.*/
+	UFUNCTION(BlueprintCallable, Category = "Interchange | Node | Texture")
+	bool GetCustomLODBias(int32& AttributeValue) const
+	{
+		IMPLEMENT_NODE_ATTRIBUTE_GETTER(LODBias, int32);
+	}
+
+	UFUNCTION(BlueprintCallable, Category = "Interchange | Node | Texture")
+	bool SetCustomLODBias(const int32& AttributeValue, bool bAddApplyDelegate = true)
+	{
+		IMPLEMENT_NODE_ATTRIBUTE_SETTER(UInterchangeTextureFactoryNode, LODBias, int32, UTexture)
+	}
+
+	/** Return false if the Attribute was not set previously.*/
+	UFUNCTION(BlueprintCallable, Category = "Interchange | Node | Texture")
+	bool GetCustomLODGroup(uint8& AttributeValue) const
+	{
+		IMPLEMENT_NODE_ATTRIBUTE_GETTER(LODGroup, uint8);
+	}
+
+	UFUNCTION(BlueprintCallable, Category = "Interchange | Node | Texture")
+	bool SetCustomLODGroup(const uint8& AttributeValue, bool bAddApplyDelegate = true)
+	{
+		IMPLEMENT_NODE_ATTRIBUTE_SETTER(UInterchangeTextureFactoryNode, LODGroup, uint8, UTexture)
+	}
+
+	//TODO support per platform data in the FAttributeStorage, so we can set different value per platform at the pipeline stage, We set only the default value for now
+	/** Return false if the Attribute was not set previously.*/
+	UFUNCTION(BlueprintCallable, Category = "Interchange | Node | Texture")
+	bool GetCustomDownscale(float& AttributeValue) const
+	{
+		IMPLEMENT_NODE_ATTRIBUTE_GETTER(Downscale, float);
+	}
+
+	UFUNCTION(BlueprintCallable, Category = "Interchange | Node | Texture")
+	bool SetCustomDownscale(const float& AttributeValue, bool bAddApplyDelegate = true)
+	{
+		IMPLEMENT_NODE_ATTRIBUTE_SETTER_WITH_CUSTOM_DELEGATE(UInterchangeTextureFactoryNode, Downscale, float, UTexture)
+	}
+
+	/** Return false if the Attribute was not set previously.*/
+	UFUNCTION(BlueprintCallable, Category = "Interchange | Node | Texture")
+	bool GetCustomDownscaleOptions(uint8& AttributeValue) const
+	{
+		IMPLEMENT_NODE_ATTRIBUTE_GETTER(DownscaleOptions, uint8);
+	}
+
+	UFUNCTION(BlueprintCallable, Category = "Interchange | Node | Texture")
+	bool SetCustomDownscaleOptions(const uint8& AttributeValue, bool bAddApplyDelegate = true)
+	{
+		IMPLEMENT_NODE_ATTRIBUTE_SETTER(UInterchangeTextureFactoryNode, DownscaleOptions, uint8, UTexture)
+	}
+
+	//////////////////////////////////////////////////////////////////////////
+	//Compositing
+
+	/** Return false if the Attribute was not set previously.*/
+	UFUNCTION(BlueprintCallable, Category = "Interchange | Node | Texture")
+	bool GetCustomCompositeTextureMode(uint8& AttributeValue) const
+	{
+		IMPLEMENT_NODE_ATTRIBUTE_GETTER(CompositeTextureMode, uint8);
+	}
+
+	UFUNCTION(BlueprintCallable, Category = "Interchange | Node | Texture")
+	bool SetCustomCompositeTextureMode(const uint8& AttributeValue, bool bAddApplyDelegate = true)
+	{
+#if WITH_EDITORONLY_DATA
+		IMPLEMENT_NODE_ATTRIBUTE_SETTER(UInterchangeTextureFactoryNode, CompositeTextureMode, uint8, UTexture)
+#else
+		IMPLEMENT_NODE_ATTRIBUTE_SETTER_NODELEGATE(CompositeTextureMode, uint8)
+#endif
+	}
+
+	/** Return false if the Attribute was not set previously.*/
+	UFUNCTION(BlueprintCallable, Category = "Interchange | Node | Texture")
+	bool GetCustomCompositePower(float& AttributeValue) const
+	{
+		IMPLEMENT_NODE_ATTRIBUTE_GETTER(CompositePower, float);
+	}
+
+	UFUNCTION(BlueprintCallable, Category = "Interchange | Node | Texture")
+	bool SetCustomCompositePower(const float& AttributeValue, bool bAddApplyDelegate = true)
+	{
+#if WITH_EDITORONLY_DATA
+		IMPLEMENT_NODE_ATTRIBUTE_SETTER(UInterchangeTextureFactoryNode, CompositePower, float, UTexture)
+#else
+		IMPLEMENT_NODE_ATTRIBUTE_SETTER_NODELEGATE(CompositePower, float)
+#endif
+	}
+
+private:
+
+	void FillAssetClassFromAttribute()
+	{
+#if WITH_ENGINE
+		FString OperationName = GetTypeName() + TEXT(".GetAssetClassName");
+		FString ClassName;
+		InterchangePrivateNodeBase::GetCustomAttribute<FString>(*Attributes, ClassNameAttributeKey, OperationName, ClassName);
+		FillAssetClassFromClassName(ClassName);
+#endif
+	}
+
+#if WITH_ENGINE
+	virtual void FillAssetClassFromClassName(const FString& ClassName)
+	{
+		if (ClassName.Equals(UTexture2D::StaticClass()->GetName()))
+		{
+			AssetClass = UTexture2D::StaticClass();
+			bIsTextureNodeClassInitialized = true;
+		}
+		else if (ClassName.Equals(UTexture::StaticClass()->GetName()))
+		{
+			AssetClass = UTexture::StaticClass();
+			bIsTextureNodeClassInitialized = true;
+		}
+	}
+#endif
+
+	bool SetTextureNodeClassFromClassAttribute()
+	{
+		if (!bIsTextureNodeClassInitialized)
+		{
+			FillAssetClassFromAttribute();
+		}
+		return bIsTextureNodeClassInitialized;
+	}
+
+	const UE::Interchange::FAttributeKey ClassNameAttributeKey = UE::Interchange::FBaseNodeStaticData::ClassTypeAttributeKey();
+	const UE::Interchange::FAttributeKey AssetNameKey = UE::Interchange::FBaseNodeStaticData::AssetNameKey();
+
+	//Translated texture node unique id
+	const UE::Interchange::FAttributeKey Macro_CustomTranslatedTextureNodeUidKey = UE::Interchange::FAttributeKey(TEXT("TranslatedTextureNodeUid"));
+	
+	//Texture Adjustments
+	const UE::Interchange::FAttributeKey Macro_CustomAdjustBrightnessKey = UE::Interchange::FAttributeKey(TEXT("AdjustBrightness"));
+	const UE::Interchange::FAttributeKey Macro_CustomAdjustBrightnessCurveKey = UE::Interchange::FAttributeKey(TEXT("AdjustBrightnessCurve"));
+	const UE::Interchange::FAttributeKey Macro_CustomAdjustVibranceKey = UE::Interchange::FAttributeKey(TEXT("AdjustVibrance"));
+	const UE::Interchange::FAttributeKey Macro_CustomAdjustSaturationKey = UE::Interchange::FAttributeKey(TEXT("AdjustSaturation"));
+	const UE::Interchange::FAttributeKey Macro_CustomAdjustRGBCurveKey = UE::Interchange::FAttributeKey(TEXT("AdjustRGBCurve"));
+	const UE::Interchange::FAttributeKey Macro_CustomAdjustHueKey = UE::Interchange::FAttributeKey(TEXT("AdjustHue"));
+	const UE::Interchange::FAttributeKey Macro_CustomAdjustMinAlphaKey = UE::Interchange::FAttributeKey(TEXT("AdjustMinAlpha"));
+	const UE::Interchange::FAttributeKey Macro_CustomAdjustMaxAlphaKey = UE::Interchange::FAttributeKey(TEXT("AdjustMaxAlpha"));
+	const UE::Interchange::FAttributeKey Macro_CustombChromaKeyTextureKey = UE::Interchange::FAttributeKey(TEXT("bChromaKeyTexture"));
+	const UE::Interchange::FAttributeKey Macro_CustomChromaKeyThresholdKey = UE::Interchange::FAttributeKey(TEXT("ChromaKeyThreshold"));
+	const UE::Interchange::FAttributeKey Macro_CustomChromaKeyColorKey = UE::Interchange::FAttributeKey(TEXT("ChromaKeyColor"));
+
+	//Texture Compression
+	const UE::Interchange::FAttributeKey Macro_CustomCompressionNoAlphaKey = UE::Interchange::FAttributeKey(TEXT("CompressionNoAlpha"));
+	const UE::Interchange::FAttributeKey Macro_CustomDeferCompressionKey = UE::Interchange::FAttributeKey(TEXT("DeferCompression"));
+	const UE::Interchange::FAttributeKey Macro_CustomLossyCompressionAmountKey = UE::Interchange::FAttributeKey(TEXT("LossyCompressionAmount"));
+	const UE::Interchange::FAttributeKey Macro_CustomMaxTextureSizeKey = UE::Interchange::FAttributeKey(TEXT("MaxTextureSize"));
+	const UE::Interchange::FAttributeKey Macro_CustomCompressionQualityKey = UE::Interchange::FAttributeKey(TEXT("CompressionQuality"));
+	const UE::Interchange::FAttributeKey Macro_CustomCompressionSettingsKey = UE::Interchange::FAttributeKey(TEXT("CompressionSettings"));
+
+	//Texture general
+	const UE::Interchange::FAttributeKey Macro_CustombDitherMipMapAlphaKey = UE::Interchange::FAttributeKey(TEXT("bDitherMipMapAlpha"));
+	const UE::Interchange::FAttributeKey Macro_CustomAlphaCoverageThresholdsKey = UE::Interchange::FAttributeKey(TEXT("AlphaCoverageThresholds"));
+	const UE::Interchange::FAttributeKey Macro_CustombFlipGreenChannelKey = UE::Interchange::FAttributeKey(TEXT("bFlipGreenChannel"));
+	const UE::Interchange::FAttributeKey Macro_CustombForcePVRTC4Key = UE::Interchange::FAttributeKey(TEXT("bForcePVRTC4"));
+	const UE::Interchange::FAttributeKey Macro_CustomPowerOfTwoModeKey = UE::Interchange::FAttributeKey(TEXT("PowerOfTwoMode"));
+	const UE::Interchange::FAttributeKey Macro_CustomPaddingColorKey = UE::Interchange::FAttributeKey(TEXT("PaddingColor"));
+	const UE::Interchange::FAttributeKey Macro_CustomFilterKey = UE::Interchange::FAttributeKey(TEXT("Filter"));
+	const UE::Interchange::FAttributeKey Macro_CustomMipLoadOptionsKey = UE::Interchange::FAttributeKey(TEXT("MipLoadOptions"));
+	const UE::Interchange::FAttributeKey Macro_CustomSRGBKey = UE::Interchange::FAttributeKey(TEXT("SRGB"));
+	const UE::Interchange::FAttributeKey Macro_CustombUseLegacyGammaKey = UE::Interchange::FAttributeKey(TEXT("bUseLegacyGamma"));
+	const UE::Interchange::FAttributeKey Macro_CustomVirtualTextureStreamingKey = UE::Interchange::FAttributeKey(TEXT("VirtualTextureStreaming"));
+	const UE::Interchange::FAttributeKey Macro_CustomAddressXKey = UE::Interchange::FAttributeKey(TEXT("AddressX"));
+	const UE::Interchange::FAttributeKey Macro_CustomAddressYKey = UE::Interchange::FAttributeKey(TEXT("AddressY"));
+
+	//Level of Detail
+	const UE::Interchange::FAttributeKey Macro_CustombPreserveBorderKey = UE::Interchange::FAttributeKey(TEXT("bPreserveBorder"));
+	const UE::Interchange::FAttributeKey Macro_CustomMipGenSettingsKey = UE::Interchange::FAttributeKey(TEXT("MipGenSettings"));
+	const UE::Interchange::FAttributeKey Macro_CustomLODBiasKey = UE::Interchange::FAttributeKey(TEXT("LODBias"));
+	const UE::Interchange::FAttributeKey Macro_CustomLODGroupKey = UE::Interchange::FAttributeKey(TEXT("LODGroup"));
+	//TODO support per platform data in the FAttributeStorage, so we can set different value per platform at the pipeline stage, We set only the default value for now
+	const UE::Interchange::FAttributeKey Macro_CustomDownscaleKey = UE::Interchange::FAttributeKey(TEXT("Downscale"));
+	const UE::Interchange::FAttributeKey Macro_CustomDownscaleOptionsKey = UE::Interchange::FAttributeKey(TEXT("DownscaleOptions"));
+
+	//Compositing
+	const UE::Interchange::FAttributeKey Macro_CustomCompositeTextureModeKey = UE::Interchange::FAttributeKey(TEXT("CompositeTextureMode"));
+	const UE::Interchange::FAttributeKey Macro_CustomCompositePowerKey = UE::Interchange::FAttributeKey(TEXT("CompositePower"));
+
+	//TODO support per platform data in the FAttributeStorage, so we can set different value per platform at the pipeline stage, We set only the default value for now
+	//IMPLEMENT_NODE_ATTRIBUTE_APPLY_UOBJECT(Downscale, float, UTexture, );
+#if WITH_ENGINE
+	bool ApplyCustomDownscaleToAsset(UObject* Asset) const
+	{
+		if (!Asset)
+		{
+			return false;
+		}
+		UTexture* TypedObject = Cast<UTexture>(Asset);
+		if (!TypedObject)
+		{
+			return false;
+		}
+		float ValueData;
+		if (GetCustomDownscale(ValueData))
+		{
+			TypedObject->Downscale.Default = ValueData;
+			return true;
+		}
+		return false;
+	}
+
+	bool FillCustomDownscaleFromAsset(UObject* Asset)
+	{
+		if (!Asset)
+		{
+			return false;
+		}
+		UTexture* TypedObject = Cast<UTexture>(Asset);
+		if (!TypedObject)
+		{
+			return false;
+		}
+		if (SetCustomDownscale(TypedObject->Downscale.Default, false))
+		{
+			return true;
+		}
+		return false;
+	}
+#endif //WITH_ENGINE
+
+protected:
+#if WITH_ENGINE
+	TSubclassOf<UTexture> AssetClass = nullptr;
+#endif
+
+	bool bIsTextureNodeClassInitialized = false;
+};
diff --git a/Engine/Plugins/Experimental/Interchange/Runtime/Source/FactoryNodes/Public/InterchangeTextureLightProfileFactoryNode.h b/Engine/Plugins/Experimental/Interchange/Runtime/Source/FactoryNodes/Public/InterchangeTextureLightProfileFactoryNode.h
new file mode 100644
--- /dev/null
+++ b/Engine/Plugins/Experimental/Interchange/Runtime/Source/FactoryNodes/Public/InterchangeTextureLightProfileFactoryNode.h
@@ -0,0 +1,74 @@
+// Copyright Epic Games, Inc. All Rights Reserved.
+
+#pragma once
+
+#include "CoreMinimal.h"
+#include "InterchangeTextureFactoryNode.h"
+#include "Nodes/InterchangeBaseNode.h"
+
+#if WITH_ENGINE
+#include "Engine/Texture.h"
+#include "Engine/TextureLightProfile.h"
+#endif
+
+
+#include "InterchangeTextureLightProfileFactoryNode.generated.h"
+
+UCLASS(BlueprintType, Experimental)
+class INTERCHANGEFACTORYNODES_API UInterchangeTextureLightProfileFactoryNode : public UInterchangeTextureFactoryNode
+{
+	GENERATED_BODY()
+
+public:
+
+	UFUNCTION(BlueprintCallable, Category = "Interchange | Node | Texture")
+	bool GetCustomBrightness(float& AttributeValue) const
+	{
+		IMPLEMENT_NODE_ATTRIBUTE_GETTER(Brightness, float);
+	}
+
+	UFUNCTION(BlueprintCallable, Category = "Interchange | Node | Texture")
+	bool SetCustomBrightness(const float AttributeValue, bool bAddApplyDelegate = true)
+	{
+#if WITH_EDITORONLY_DATA
+		IMPLEMENT_NODE_ATTRIBUTE_SETTER(UInterchangeTextureLightProfileFactoryNode, Brightness, float, UTextureLightProfile)
+#else
+		IMPLEMENT_NODE_ATTRIBUTE_SETTER_NODELEGATE(Brightness, float)
+#endif
+	}
+
+
+	UFUNCTION(BlueprintCallable, Category = "Interchange | Node | Texture")
+	bool GetCustomTextureMultiplier(float& AttributeValue) const
+	{
+		IMPLEMENT_NODE_ATTRIBUTE_GETTER(TextureMultiplier, float);
+	}
+
+	UFUNCTION(BlueprintCallable, Category = "Interchange | Node | Texture")
+	bool SetCustomTextureMultiplier(const float AttributeValue, bool bAddApplyDelegate = true)
+	{
+#if WITH_EDITORONLY_DATA
+		IMPLEMENT_NODE_ATTRIBUTE_SETTER(UInterchangeTextureLightProfileFactoryNode, TextureMultiplier, float, UTextureLightProfile)
+#else
+		IMPLEMENT_NODE_ATTRIBUTE_SETTER_NODELEGATE(TextureMultiplier, float)
+#endif
+	}
+
+
+private:
+#if WITH_ENGINE
+	virtual void FillAssetClassFromClassName(const FString& ClassName) override
+	{
+		if (ClassName.Equals(UTextureLightProfile::StaticClass()->GetName()))
+		{
+			AssetClass = UTextureLightProfile::StaticClass();
+			bIsTextureNodeClassInitialized = true;
+		}
+	}
+#endif
+
+	// Addressing
+	IMPLEMENT_NODE_ATTRIBUTE_KEY(Brightness);
+	IMPLEMENT_NODE_ATTRIBUTE_KEY(TextureMultiplier);
+
+};
\ No newline at end of file
diff --git a/Engine/Plugins/Experimental/Interchange/Runtime/Source/Import/InterchangeImport.Build.cs b/Engine/Plugins/Experimental/Interchange/Runtime/Source/Import/InterchangeImport.Build.cs
--- a/Engine/Plugins/Experimental/Interchange/Runtime/Source/Import/InterchangeImport.Build.cs
+++ b/Engine/Plugins/Experimental/Interchange/Runtime/Source/Import/InterchangeImport.Build.cs
@@ -15,6 +15,7 @@ namespace UnrealBuildTool.Rules
 					"InterchangeCore",
 					"InterchangeDispatcher",
 					"InterchangeEngine",
+					"InterchangeFactoryNodes",
 					"InterchangeNodes",
 					"MeshDescription",
 					"StaticMeshDescription",
diff --git a/Engine/Plugins/Experimental/Interchange/Runtime/Source/Nodes/Private/InterchangeActorFactoryNode.cpp b/Engine/Plugins/Experimental/Interchange/Runtime/Source/Nodes/Private/InterchangeActorFactoryNode.cpp
deleted file mode 100644
diff --git a/Engine/Plugins/Experimental/Interchange/Runtime/Source/Nodes/Private/InterchangeMaterialFactoryNode.cpp b/Engine/Plugins/Experimental/Interchange/Runtime/Source/Nodes/Private/InterchangeMaterialFactoryNode.cpp
deleted file mode 100644
diff --git a/Engine/Plugins/Experimental/Interchange/Runtime/Source/Nodes/Private/InterchangeSkeletalMeshFactoryNode.cpp b/Engine/Plugins/Experimental/Interchange/Runtime/Source/Nodes/Private/InterchangeSkeletalMeshFactoryNode.cpp
deleted file mode 100644
diff --git a/Engine/Plugins/Experimental/Interchange/Runtime/Source/Nodes/Private/InterchangeSkeletalMeshLodDataNode.cpp b/Engine/Plugins/Experimental/Interchange/Runtime/Source/Nodes/Private/InterchangeSkeletalMeshLodDataNode.cpp
deleted file mode 100644
diff --git a/Engine/Plugins/Experimental/Interchange/Runtime/Source/Nodes/Private/InterchangeStaticMeshFactoryNode.cpp b/Engine/Plugins/Experimental/Interchange/Runtime/Source/Nodes/Private/InterchangeStaticMeshFactoryNode.cpp
deleted file mode 100644
diff --git a/Engine/Plugins/Experimental/Interchange/Runtime/Source/Nodes/Private/InterchangeStaticMeshLodDataNode.cpp b/Engine/Plugins/Experimental/Interchange/Runtime/Source/Nodes/Private/InterchangeStaticMeshLodDataNode.cpp
deleted file mode 100644
diff --git a/Engine/Plugins/Experimental/Interchange/Runtime/Source/Nodes/Public/InterchangeActorFactoryNode.h b/Engine/Plugins/Experimental/Interchange/Runtime/Source/Nodes/Public/InterchangeActorFactoryNode.h
deleted file mode 100644
diff --git a/Engine/Plugins/Experimental/Interchange/Runtime/Source/Nodes/Public/InterchangeCineCameraFactoryNode.h b/Engine/Plugins/Experimental/Interchange/Runtime/Source/Nodes/Public/InterchangeCineCameraFactoryNode.h
deleted file mode 100644
diff --git a/Engine/Plugins/Experimental/Interchange/Runtime/Source/Nodes/Public/InterchangeMaterialFactoryNode.h b/Engine/Plugins/Experimental/Interchange/Runtime/Source/Nodes/Public/InterchangeMaterialFactoryNode.h
deleted file mode 100644
diff --git a/Engine/Plugins/Experimental/Interchange/Runtime/Source/Nodes/Public/InterchangePhysicsAssetFactoryNode.h b/Engine/Plugins/Experimental/Interchange/Runtime/Source/Nodes/Public/InterchangePhysicsAssetFactoryNode.h
deleted file mode 100644
diff --git a/Engine/Plugins/Experimental/Interchange/Runtime/Source/Nodes/Public/InterchangeSkeletalMeshFactoryNode.h b/Engine/Plugins/Experimental/Interchange/Runtime/Source/Nodes/Public/InterchangeSkeletalMeshFactoryNode.h
deleted file mode 100644
diff --git a/Engine/Plugins/Experimental/Interchange/Runtime/Source/Nodes/Public/InterchangeSkeletalMeshLodDataNode.h b/Engine/Plugins/Experimental/Interchange/Runtime/Source/Nodes/Public/InterchangeSkeletalMeshLodDataNode.h
deleted file mode 100644
diff --git a/Engine/Plugins/Experimental/Interchange/Runtime/Source/Nodes/Public/InterchangeSkeletonFactoryNode.h b/Engine/Plugins/Experimental/Interchange/Runtime/Source/Nodes/Public/InterchangeSkeletonFactoryNode.h
deleted file mode 100644
diff --git a/Engine/Plugins/Experimental/Interchange/Runtime/Source/Nodes/Public/InterchangeStaticMeshFactoryNode.h b/Engine/Plugins/Experimental/Interchange/Runtime/Source/Nodes/Public/InterchangeStaticMeshFactoryNode.h
deleted file mode 100644
diff --git a/Engine/Plugins/Experimental/Interchange/Runtime/Source/Nodes/Public/InterchangeStaticMeshLodDataNode.h b/Engine/Plugins/Experimental/Interchange/Runtime/Source/Nodes/Public/InterchangeStaticMeshLodDataNode.h
deleted file mode 100644
diff --git a/Engine/Plugins/Experimental/Interchange/Runtime/Source/Nodes/Public/InterchangeTexture2DArrayFactoryNode.h b/Engine/Plugins/Experimental/Interchange/Runtime/Source/Nodes/Public/InterchangeTexture2DArrayFactoryNode.h
deleted file mode 100644
diff --git a/Engine/Plugins/Experimental/Interchange/Runtime/Source/Nodes/Public/InterchangeTextureCubeFactoryNode.h b/Engine/Plugins/Experimental/Interchange/Runtime/Source/Nodes/Public/InterchangeTextureCubeFactoryNode.h
deleted file mode 100644
diff --git a/Engine/Plugins/Experimental/Interchange/Runtime/Source/Nodes/Public/InterchangeTextureFactoryNode.h b/Engine/Plugins/Experimental/Interchange/Runtime/Source/Nodes/Public/InterchangeTextureFactoryNode.h
deleted file mode 100644
diff --git a/Engine/Plugins/Experimental/Interchange/Runtime/Source/Nodes/Public/InterchangeTextureLightProfileFactoryNode.h b/Engine/Plugins/Experimental/Interchange/Runtime/Source/Nodes/Public/InterchangeTextureLightProfileFactoryNode.h
deleted file mode 100644
diff --git a/Engine/Plugins/Experimental/Interchange/Runtime/Source/Pipelines/InterchangePipelines.Build.cs b/Engine/Plugins/Experimental/Interchange/Runtime/Source/Pipelines/InterchangePipelines.Build.cs
--- a/Engine/Plugins/Experimental/Interchange/Runtime/Source/Pipelines/InterchangePipelines.Build.cs
+++ b/Engine/Plugins/Experimental/Interchange/Runtime/Source/Pipelines/InterchangePipelines.Build.cs
@@ -14,6 +14,7 @@ namespace UnrealBuildTool.Rules
 					"Engine",
 					"InterchangeCore",
 					"InterchangeEngine",
+					"InterchangeFactoryNodes",
 					"InterchangeNodes"
 				}
 			);
