Translator for .obj meshes (+ material / texture support)

diff --git a/Engine/Plugins/Experimental/Interchange/Runtime/Source/Import/Private/InterchangeImportModule.cpp b/Engine/Plugins/Experimental/Interchange/Runtime/Source/Import/Private/InterchangeImportModule.cpp
--- a/Engine/Plugins/Experimental/Interchange/Runtime/Source/Import/Private/InterchangeImportModule.cpp
+++ b/Engine/Plugins/Experimental/Interchange/Runtime/Source/Import/Private/InterchangeImportModule.cpp
@@ -7,6 +7,7 @@
 #include "InterchangeImportLog.h"
 #include "InterchangeManager.h"
 #include "Material/InterchangeMaterialFactory.h"
+#include "Mesh/InterchangeOBJTranslator.h"
 #include "Mesh/InterchangePhysicsAssetFactory.h"
 #include "Mesh/InterchangeSkeletalMeshFactory.h"
 #include "Mesh/InterchangeSkeletonFactory.h"
@@ -50,10 +51,12 @@ void FInterchangeImportModule::StartupModule()
 		//Register the translators
 		//Scenes
 		InterchangeManager.RegisterTranslator(UInterchangeFbxTranslator::StaticClass()); //Do not submit uncommented until we replace completly fbx importer (staticmesh + skeletalMesh + animation)
+		InterchangeManager.RegisterTranslator(UInterchangeOBJTranslator::StaticClass());
 
 		//Textures
 
 		// UDIM must be registered before the other texture translators
+		// @todo: this seems a bit fragile; can we do better?
 		InterchangeManager.RegisterTranslator(UInterchangeUDIMTranslator::StaticClass());
 		InterchangeManager.RegisterTextureOnlyTranslatorClass(UInterchangeUDIMTranslator::StaticClass());
 
diff --git a/Engine/Plugins/Experimental/Interchange/Runtime/Source/Import/Private/Mesh/InterchangeOBJTranslator.cpp b/Engine/Plugins/Experimental/Interchange/Runtime/Source/Import/Private/Mesh/InterchangeOBJTranslator.cpp
new file mode 100644
--- /dev/null
+++ b/Engine/Plugins/Experimental/Interchange/Runtime/Source/Import/Private/Mesh/InterchangeOBJTranslator.cpp
@@ -0,0 +1,1016 @@
+// Copyright Epic Games, Inc. All Rights Reserved. 
+
+#include "Mesh/InterchangeOBJTranslator.h"
+
+#include "Algo/BinarySearch.h"
+#include "Algo/Find.h"
+#include "Async/Async.h"
+#include "Async/Future.h"
+#include "InterchangeImportLog.h"
+#include "InterchangeManager.h"
+#include "InterchangeMaterialDefinitions.h"
+#include "InterchangeMeshNode.h"
+#include "InterchangeShaderGraphNode.h"
+#include "InterchangeTexture2DNode.h"
+#include "InterchangeTextureNode.h"
+#include "Misc/ConfigCacheIni.h"
+#include "Misc/FileHelper.h"
+#include "Misc/Paths.h"
+#include "Misc/ScopedSlowTask.h"
+#include "Modules/ModuleManager.h"
+#include "Nodes/InterchangeBaseNodeContainer.h"
+#include "StaticMeshAttributes.h"
+#include "UObject/GCObjectScopeGuard.h"
+
+
+/**
+ * This is a binary representation of the .obj file 
+ */
+struct FObjData
+{
+	/** Filename of the .obj file being parsed */
+	FString ObjFilename;
+
+	/** .obj element arrays */
+	TArray<FVector3f> Positions;
+	TArray<FVector2f> UVs;
+	TArray<FVector3f> Normals;
+
+	/** Defines the data for a single face vertex */
+	struct FVertexData
+	{
+		FVertexData(int32 InVertexIndex, int32 InUVIndex, int32 InNormalIndex)
+			: VertexIndex(InVertexIndex), UVIndex(InUVIndex), NormalIndex(InNormalIndex)
+		{}
+
+		/** .obj face vertices have three attributes: position index, UV index, normal index */
+		int32 VertexIndex;
+		int32 UVIndex;
+		int32 NormalIndex;
+	};
+
+	/** Defines the data for a single face */
+	struct FFaceData
+	{
+		/** .obj faces consist of an array of face vertices */
+		TArray<FVertexData, TInlineAllocator<4>> Vertices;
+
+		/** Return the FVertexData of this face which contains the given vertex index */
+		const FVertexData& GetVertexDataContainingVertexIndex(int32 VertexIndex) const
+		{
+			const FVertexData* VertexData = Algo::FindBy(Vertices, VertexIndex, &FVertexData::VertexIndex);
+			check(VertexData);
+			return *VertexData;
+		}
+	};
+
+	/** Defines the data for a group */
+	struct FGroupData
+	{
+		FString MaterialName;
+		TArray<FFaceData> Faces;
+	};
+
+	/** All the groups defined by the .obj */
+	TMap<FString, FGroupData> Groups;
+
+	/** Defines the data for a material */
+	struct FMaterialData
+	{
+		// These are the defaults for unspecified .mtl attributes
+		int32 IlluminationModel = 0;
+		FVector3f DiffuseColor = FVector3f(0.8f, 0.8f, 0.8f);
+		FVector3f AmbientColor = FVector3f(0.2f, 0.2f, 0.2f);
+		FVector3f SpecularColor = FVector3f(1.0f, 1.0f, 1.0f);
+		FString DiffuseTexture;
+		FString AmbientTexture;
+		FString SpecularTexture;
+		FVector3f TransmissionFilter = FVector3f(1.0f, 1.0f, 1.0f);
+		float RefractiveIndex = 1.0f;
+		float Transparency = 1.0f;
+	};
+
+	/** All the materials defined by the .obj */
+	TMap<FString, FMaterialData> Materials;
+
+	/** This is the current group name being parsed */
+	FString CurrentGroup;
+
+	/** This is the current material name being parsed */
+	FString MaterialBeingDefined;
+
+	/** This is the current material name in use */
+	FString CurrentMaterial;
+
+	/** Generates a MeshDescription for the named group */
+	FMeshDescription MakeMeshDescriptionForGroup(const FString& GroupName) const;
+
+	/** Returns a bounding box fitting the vertices in the named group */
+	FBox GetGroupBoundingBox(const FString& GroupName) const;
+
+	/** Returns the vertex count for the named group */
+	int32 GetGroupVertexCount(const FString& GroupName) const;
+
+	/** Returns the poly count for the named group */
+	int32 GetGroupPolygonCount(const FString& GroupName) const;
+
+private:
+	TArray<int32> GetVertexIndicesUsedByGroup(const FGroupData& GroupData) const;
+	TArray<int32> GetUVIndicesUsedByGroup(const FGroupData& GroupData) const;
+};
+
+
+static FVector3f PositionToUEBasis(const FVector3f& InVector)
+{
+	return FVector3f(InVector.X, -InVector.Y, InVector.Z);
+}
+
+
+static FVector2f UVToUEBasis(const FVector2f& InVector)
+{
+	return FVector2f(InVector.X, 1.0f - InVector.Y);
+}
+
+
+TArray<int32> FObjData::GetVertexIndicesUsedByGroup(const FGroupData& GroupData) const
+{
+	TSet<int32> VertexIndexSet;
+
+	for (const FFaceData& FaceData : GroupData.Faces)
+	{
+		for (const FVertexData& VertexData : FaceData.Vertices)
+		{
+			VertexIndexSet.Add(VertexData.VertexIndex);
+		}
+	}
+
+	return VertexIndexSet.Array();
+}
+
+
+TArray<int32> FObjData::GetUVIndicesUsedByGroup(const FGroupData& GroupData) const
+{
+	TSet<int32> UVIndexSet;
+
+	for (const FFaceData& FaceData : GroupData.Faces)
+	{
+		for (const FVertexData& VertexData : FaceData.Vertices)
+		{
+			UVIndexSet.Add(VertexData.UVIndex);
+		}
+	}
+
+	return UVIndexSet.Array();
+}
+
+
+FBox FObjData::GetGroupBoundingBox(const FString& GroupName) const
+{
+	FBox Box;
+
+	const FGroupData* GroupDataPtr = Groups.Find(GroupName);
+	if (!GroupDataPtr)
+	{
+		// If group name not found, return an empty bounding box
+		return Box;
+	}
+
+	const FGroupData& GroupData = *GroupDataPtr;
+
+	for (int32 PositionIndex : GetVertexIndicesUsedByGroup(GroupData))
+	{
+		Box += PositionToUEBasis(Positions[PositionIndex]);
+	}
+
+	return Box;
+}
+
+
+int32 FObjData::GetGroupVertexCount(const FString& GroupName) const
+{
+	const FGroupData* GroupDataPtr = Groups.Find(GroupName);
+	if (!GroupDataPtr)
+	{
+		return 0;
+	}
+
+	const FGroupData& GroupData = *GroupDataPtr;
+	return GetVertexIndicesUsedByGroup(GroupData).Num();
+}
+
+
+int32 FObjData::GetGroupPolygonCount(const FString& GroupName) const
+{
+	const FGroupData* GroupDataPtr = Groups.Find(GroupName);
+	if (!GroupDataPtr)
+	{
+		return 0;
+	}
+
+	const FGroupData& GroupData = *GroupDataPtr;
+	return GroupData.Faces.Num();
+}
+
+
+FMeshDescription FObjData::MakeMeshDescriptionForGroup(const FString& GroupName) const
+{
+	FMeshDescription MeshDescription;
+
+	const FGroupData* GroupDataPtr = Groups.Find(GroupName);
+	if (!GroupDataPtr)
+	{
+		// If group name not found, return an empty mesh description
+		return MeshDescription;
+	}
+
+	FStaticMeshAttributes Attributes(MeshDescription);
+	Attributes.Register();
+
+	MeshDescription.SuspendVertexInstanceIndexing();
+	MeshDescription.SuspendEdgeIndexing();
+	MeshDescription.SuspendPolygonIndexing();
+	MeshDescription.SuspendPolygonGroupIndexing();
+	MeshDescription.SuspendUVIndexing();
+
+	const FGroupData& GroupData = *GroupDataPtr;
+
+	// We prefer that the MeshDescription vertex and UV element buffers be in the same order as the .obj ones.
+	// (there are certain use cases for which this is convenient, in order to be able to match imported data to the source)
+	// So get the list of vertex and UV indices used by the faces in this group, and sort and compact them to contiguous indices,
+	// preserving their original ordering.
+
+	TArray<int32> VertexIndexMapping = GetVertexIndicesUsedByGroup(GroupData);
+	TArray<int32> UVIndexMapping = GetUVIndicesUsedByGroup(GroupData);
+	VertexIndexMapping.Sort();
+	UVIndexMapping.Sort();
+	
+	// Create vertices and initialize positions
+	// Note that we create a contiguous list of vertex indices from 0...n-1, referencing potentially sparse obj vertices
+
+	MeshDescription.ReserveNewVertices(VertexIndexMapping.Num());
+	for (int32 ObjVertexIndex : VertexIndexMapping)
+	{
+		FVertexID VertexIndex = MeshDescription.CreateVertex();
+		Attributes.GetVertexPositions()[VertexIndex] = PositionToUEBasis(Positions[ObjVertexIndex]);
+	}
+
+	// Create UVs and initialize values
+
+	const int32 UVChannel = 0;
+	MeshDescription.ReserveNewUVs(UVIndexMapping.Num());
+	for (int32 ObjUVIndex : UVIndexMapping)
+	{
+		FUVID UVIndex = MeshDescription.CreateUV(UVChannel);
+		Attributes.GetUVCoordinates(UVChannel)[UVIndex] = UVs[ObjUVIndex];
+	}
+
+	// Create polygon group
+
+	FPolygonGroupID PolygonGroupIndex = MeshDescription.CreatePolygonGroup();
+	Attributes.GetPolygonGroupMaterialSlotNames()[PolygonGroupIndex] = FName(GroupData.MaterialName);
+
+	// Create faces.
+	// n-gons are preserved
+
+	MeshDescription.ReserveNewTriangles(GroupData.Faces.Num());
+	MeshDescription.ReserveNewPolygons(GroupData.Faces.Num());
+
+	TArray<FVertexInstanceID, TInlineAllocator<8>> VertexInstanceIDs;
+	for (const FFaceData& FaceData : GroupData.Faces)
+	{
+		VertexInstanceIDs.Reset();
+		MeshDescription.ReserveNewVertexInstances(FaceData.Vertices.Num());
+
+		for (const FVertexData& VertexData : FaceData.Vertices)
+		{
+			FVertexID VertexID = Algo::BinarySearch(VertexIndexMapping, VertexData.VertexIndex);
+			FVertexInstanceID VertexInstanceID = MeshDescription.CreateVertexInstance(VertexID);
+			VertexInstanceIDs.Add(VertexInstanceID);
+
+			if (VertexData.NormalIndex != INDEX_NONE)
+			{
+				Attributes.GetVertexInstanceNormals()[VertexInstanceID] = PositionToUEBasis(Normals[VertexData.NormalIndex]);
+			}
+
+			if (VertexData.UVIndex != INDEX_NONE)
+			{
+				Attributes.GetVertexInstanceUVs()[VertexInstanceID] = UVToUEBasis(UVs[VertexData.UVIndex]);
+			}
+		}
+
+		int32 PolygonIndex = MeshDescription.CreatePolygon(PolygonGroupIndex, VertexInstanceIDs);
+	}
+
+	// Determine edge hardnesses
+	// Iterate through all the edges, looking at the adjacent faces.
+	// If the face vertex normals aren't all the same, we consider the edge to be hard.
+	// Internal edges within polygons are always counted as soft.
+
+	// Allow triangles adjacent to edges to be indexed
+	// We need this when getting edge connected polygons.
+	MeshDescription.ResumeEdgeIndexing();
+
+	for (const FEdgeID& EdgeID : MeshDescription.Edges().GetElementIDs())
+	{
+		// Don't look at internal edges (edges within a polygon which break it into triangles)
+		if (!MeshDescription.IsEdgeInternal(EdgeID))
+		{
+			bool bAllNormalsEqual = true;
+			TArray<FPolygonID, TInlineAllocator<2>> EdgePolygonIDs = MeshDescription.GetEdgeConnectedPolygons<TInlineAllocator<2>>(EdgeID);
+
+			for (const FVertexID& EdgeVertexID : MeshDescription.GetEdgeVertices(EdgeID))
+			{
+				if (bAllNormalsEqual)
+				{
+					// This is the vertex index as it appears in the .obj
+					int32 ObjVertexIndex = VertexIndexMapping[EdgeVertexID];
+
+					// For the vertex we are considering, look at the first face adjacent to the edge, and find the vertex data which includes it
+					// This is a baseline we will use to compare all the other adjacent faces
+					const FVertexData& VertexData = GroupData.Faces[EdgePolygonIDs[0]].GetVertexDataContainingVertexIndex(ObjVertexIndex);
+
+					for (int32 Index = 1; Index < EdgePolygonIDs.Num(); Index++)
+					{
+						// For all other adjacent faces, find the vertex data which includes the vertex we are considering.
+						// If the normal index is not the same as the baseline, we know this must be a hard edge.
+						const FVertexData& VertexDataToCompare = GroupData.Faces[EdgePolygonIDs[Index]].GetVertexDataContainingVertexIndex(ObjVertexIndex);
+
+						if (VertexData.NormalIndex != VertexDataToCompare.NormalIndex)
+						{
+							bAllNormalsEqual = false;
+							break;
+						}
+					}
+				}
+			}
+
+			if (!bAllNormalsEqual)
+			{
+				Attributes.GetEdgeHardnesses()[EdgeID] = true;
+			}
+
+		}
+	}
+
+	MeshDescription.ResumeVertexInstanceIndexing();
+	MeshDescription.ResumePolygonIndexing();
+	MeshDescription.ResumePolygonGroupIndexing();
+	MeshDescription.ResumeUVIndexing();
+
+	return MeshDescription;
+}
+
+
+namespace UE {
+namespace Interchange {
+
+namespace ObjParser
+{
+	using FFunctionType = TFunction<bool(FObjData&, FStringView)>;
+	using FKeywordMap = TMap<FString, FFunctionType>;
+
+
+	// Determine if the given character can be counted as a token terminator
+	static bool IsTerminatorChar(TCHAR CharToCheck, TCHAR PossibleTerminator = '\0')
+	{
+		return CharToCheck == TEXT(' ') || CharToCheck == TEXT('\t') || CharToCheck == PossibleTerminator;
+	}
+
+
+	// Skip over any leading whitespace, modifying the passed stringview to start at the first non-whitespace character
+	static void SkipWhitespace(FStringView& Line)
+	{
+		int32 Index = 0;
+		while (Index < Line.Len() && (Line[Index] == TEXT(' ') || Line[Index] == TEXT('\t')))
+		{
+			Index++;
+		}
+
+		Line.RightChopInline(Index);
+	}
+
+
+	// Return the next token on the line (delimited by the optionally specified terminator character).
+	// Modify the passed stringview to start at the beginning of the next token.
+	static FStringView GetToken(FStringView& Line, TCHAR Terminator = TEXT(' '))
+	{
+		SkipWhitespace(Line);
+
+		for (int32 Index = 0; Index < Line.Len(); Index++)
+		{
+			if (IsTerminatorChar(Line[Index], Terminator))
+			{
+				FStringView Result = Line.Left(Index);
+				Line.RightChopInline(Index + 1);
+				SkipWhitespace(Line);
+				return Result;
+			}
+		}
+
+		FStringView Result = Line;
+		Line = FStringView();
+		return Result;
+	}
+
+
+	// Parse a floating-point literal from the start of the given line, returning true if it was successful.
+	// The token containing the floating-point literal will be consumed and removed from the line.
+	static bool GetFloat(FStringView& Line, float& Result, TCHAR Terminator = TEXT(' '))
+	{
+		FStringView Token = GetToken(Line, Terminator);
+		if (Token.IsEmpty())
+		{
+			return false;
+		}
+
+		// @todo: have to convert this to a string in order to lex it - this is horrible
+		LexFromString(Result, *FString(Token));
+
+		return true;
+	}
+
+
+	static bool GetVector2(FStringView& Line, FVector2f& Result)
+	{
+		if (!GetFloat(Line, Result.X)) { return false; }
+		if (!GetFloat(Line, Result.Y)) { return false; }
+		return true;
+	}
+
+
+	static bool GetVector3(FStringView& Line, FVector3f& Result)
+	{
+		if (!GetFloat(Line, Result.X)) { return false; }
+		if (!GetFloat(Line, Result.Y)) { return false; }
+		if (!GetFloat(Line, Result.Z)) { return false; }
+		return true;
+	}
+
+
+	// Parse an integer literal from the start of the given line, returning true if it was successful.
+	// The token containing the integer literal will be consumed and removed from the line.
+	static bool GetInteger(FStringView& Line, int32& Result, TCHAR Terminator = TEXT(' '))
+	{
+		FStringView Token = GetToken(Line, Terminator);
+		if (Token.IsEmpty())
+		{
+			return false;
+		}
+
+		// @todo: have to convert this to a string in order to lex it - this is horrible
+		LexFromString(Result, *FString(Token));
+
+		return true;
+	}
+
+
+	static bool ParseComment(FObjData& ObjData, FStringView Line)
+	{
+		// Nothing else to do
+		return true;
+	}
+
+
+	static bool ParseVertexPosition(FObjData& ObjData, FStringView Line)
+	{
+		FVector3f Position;
+		bool bSuccess = GetVector3(Line, Position);
+
+		ObjData.Positions.Add(Position * 100.0f);
+
+		if (!Line.IsEmpty())
+		{
+			// @todo: emit warning through Interchange results container
+			// @todo: get new API working through InterchangeManager (keyed multiple times on translator, pipeline, factory objects)
+			// We have a pointer to the Translator inside the FObjData, so we can get the results container from there.
+			UE_LOG(LogTemp, Warning, TEXT("Unexpected extra arguments on v keyword"));
+		}
+
+		return bSuccess;
+	}
+
+
+	static bool ParseTextureCoordinate(FObjData& ObjData, FStringView Line)
+	{
+		FVector2f UV;
+		bool bSuccess = GetVector2(Line, UV);
+
+		ObjData.UVs.Add(UV);
+
+		if (!Line.IsEmpty())
+		{
+			UE_LOG(LogTemp, Warning, TEXT("Unexpected extra arguments on vt keyword"));
+		}
+
+		return bSuccess;
+	}
+
+
+	static bool ParseNormalVector(FObjData& ObjData, FStringView Line)
+	{
+		FVector3f Normal;
+		bool bSuccess = GetVector3(Line, Normal);
+
+		ObjData.Normals.Add(Normal);
+
+		if (!Line.IsEmpty())
+		{
+			UE_LOG(LogTemp, Warning, TEXT("Unexpected extra arguments on vn keyword"));
+		}
+
+		return bSuccess;
+	}
+
+
+	static bool ParseFace(FObjData& ObjData, FStringView Line)
+	{
+		FObjData::FFaceData FaceData;
+
+		while (!Line.IsEmpty())
+		{
+			FStringView FaceToken = GetToken(Line);
+
+			auto GetFaceIndex = [](FStringView& Token, int32 NumElements)
+			{
+				int32 Index = INDEX_NONE;
+				if (GetInteger(Token, Index, TEXT('/')))
+				{
+					if (Index < 0)
+					{
+						// Negative indices refer to vertices declared relatively in the file
+						Index += NumElements;
+					}
+					else
+					{
+						// Indices start at one, so adjust for zero-base
+						Index--;
+					}
+				}
+
+				return Index;
+			};
+
+			int32 VertexIndex = GetFaceIndex(FaceToken, ObjData.Positions.Num());
+			int32 UVIndex = GetFaceIndex(FaceToken, ObjData.UVs.Num());
+			int32 NormalIndex = GetFaceIndex(FaceToken, ObjData.Normals.Num());
+
+			FaceData.Vertices.Emplace(VertexIndex, UVIndex, NormalIndex);
+		}
+
+		if (FaceData.Vertices.Num() > 2)
+		{
+			FObjData::FGroupData& GroupData = ObjData.Groups.FindOrAdd(ObjData.CurrentGroup);
+			GroupData.Faces.Emplace(FaceData);
+			GroupData.MaterialName = ObjData.CurrentMaterial;
+		}
+		else
+		{
+			UE_LOG(LogTemp, Warning, TEXT("Too few vertices on f keyword"));
+			return false;
+		}
+
+		return true;
+	}
+
+
+	static bool ParseUseMaterial(FObjData& ObjData, FStringView Line)
+	{
+		FStringView MaterialName = GetToken(Line);
+
+		if (MaterialName.IsEmpty())
+		{
+			UE_LOG(LogTemp, Warning, TEXT("Missing material name on usemtl keyword"));
+		}
+
+		if (!Line.IsEmpty())
+		{
+			UE_LOG(LogTemp, Warning, TEXT("Unexpected extra arguments on newmtl keyword"));
+		}
+
+		ObjData.CurrentMaterial = FString(MaterialName);
+
+		return true;
+	}
+
+
+	static bool ParseGroup(FObjData& ObjData, FStringView Line)
+	{
+		FStringView GroupName = GetToken(Line);
+		if (GroupName.IsEmpty())
+		{
+			UE_LOG(LogTemp, Warning, TEXT("Missing group name on g keyword"));
+		}
+
+		if (!Line.IsEmpty())
+		{
+			UE_LOG(LogTemp, Warning, TEXT("Unexpected extra arguments on g keyword"));
+		}
+
+		ObjData.CurrentGroup = GroupName;
+		return true;
+	}
+
+
+	static bool ParseLine(FObjData& ObjData, FStringView Line, const FKeywordMap& KeywordMap)
+	{
+		// Dispatch by keyword to appropriate parsing function
+
+		FString Keyword = FString(GetToken(Line));
+		if (Keyword.IsEmpty())
+		{
+			return true;
+		}
+
+		if (const FFunctionType* Handler = KeywordMap.Find(Keyword))
+		{
+			return (*Handler)(ObjData, Line);
+		}
+		else
+		{
+			// Unknown keyword, report it and move on
+			UE_LOG(LogTemp, Warning, TEXT("Unknown keyword: %s"), *Keyword);
+			return true;
+		}
+	}
+
+
+	static bool ParseFile(FObjData& ObjData, const TCHAR* Filename, const FKeywordMap& KeywordMap)
+	{
+		// Read the file line by line, passing each line to ParseLine
+
+		bool bSuccess = true;
+		FFileHelper::LoadFileToStringWithLineVisitor(Filename,
+			[&ObjData, &KeywordMap, &bSuccess](FStringView Line)
+			{
+				bool bResult = ObjParser::ParseLine(ObjData, Line, KeywordMap);
+				if (bResult == false)
+				{
+					bSuccess = false;
+				}
+			}
+		);
+
+		return bSuccess;
+	}
+
+
+	static bool ParseNewMaterial(FObjData& ObjData, FStringView Line)
+	{
+		FStringView MaterialName = GetToken(Line);
+
+		if (MaterialName.IsEmpty())
+		{
+			UE_LOG(LogTemp, Warning, TEXT("Missing material name on newmtl keyword"));
+		}
+
+		if (!Line.IsEmpty())
+		{
+			UE_LOG(LogTemp, Warning, TEXT("Unexpected extra arguments on newmtl keyword"));
+		}
+
+		ObjData.Materials.Add(FString(MaterialName));
+		ObjData.MaterialBeingDefined = MaterialName;
+		return true;
+	}
+
+
+	static inline bool GetMaterialProperty(FStringView& Line, int32& Result) { return GetInteger(Line, Result); }
+	static inline bool GetMaterialProperty(FStringView& Line, float& Result) { return GetFloat(Line, Result); }
+	static inline bool GetMaterialProperty(FStringView& Line, FVector3f& Result) { return GetVector3(Line, Result); }
+
+
+	template <typename PropertyType, PropertyType FObjData::FMaterialData::*Property>
+	static bool ParseMaterialProperty(FObjData& ObjData, FStringView Line)
+	{
+		PropertyType Value = PropertyType();
+		if (GetMaterialProperty(Line, Value))
+		{
+			if (FObjData::FMaterialData* MaterialData = ObjData.Materials.Find(ObjData.MaterialBeingDefined))
+			{
+				MaterialData->*Property = Value;
+			}
+			else
+			{
+				UE_LOG(LogTemp, Warning, TEXT("Missing newmtl keyword"));
+				return false;
+			}
+		}
+
+		if (!Line.IsEmpty())
+		{
+			UE_LOG(LogTemp, Warning, TEXT("Unexpected extra arguments on illum keyword"));
+		}
+
+		return true;
+	}
+
+
+	template <FString FObjData::FMaterialData::* Property>
+	static bool ParseMaterialPath(FObjData& ObjData, FStringView Line)
+	{
+		if (Line.IsEmpty())
+		{
+			UE_LOG(LogTemp, Warning, TEXT("Missing texture filename on map_Kd keyword"));
+		}
+
+		if (FObjData::FMaterialData* MaterialData = ObjData.Materials.Find(ObjData.MaterialBeingDefined))
+		{
+			MaterialData->*Property = FPaths::GetPath(ObjData.ObjFilename) / FString(Line);
+		}
+		else
+		{
+			UE_LOG(LogTemp, Warning, TEXT("Missing newmtl keyword"));
+			return false;
+		}
+
+		return true;
+	}
+
+
+	static bool ParseMaterialLib(FObjData& ObjData, FStringView Line)
+	{
+		static FKeywordMap MtlKeywordMap =
+		{
+			{ TEXT("#"),       ParseComment },
+			{ TEXT("newmtl"),  ParseNewMaterial },
+			{ TEXT("illum"),   ParseMaterialProperty<int32,     &FObjData::FMaterialData::IlluminationModel> },
+			{ TEXT("Kd"),      ParseMaterialProperty<FVector3f, &FObjData::FMaterialData::DiffuseColor> },
+			{ TEXT("Ka"),      ParseMaterialProperty<FVector3f, &FObjData::FMaterialData::AmbientColor> },
+			{ TEXT("Ks"),      ParseMaterialProperty<FVector3f, &FObjData::FMaterialData::SpecularColor> },
+			{ TEXT("Tf"),      ParseMaterialProperty<FVector3f, &FObjData::FMaterialData::TransmissionFilter> },
+			{ TEXT("Ni"),      ParseMaterialProperty<float,     &FObjData::FMaterialData::RefractiveIndex> },
+			{ TEXT("d"),       ParseMaterialProperty<float,     &FObjData::FMaterialData::Transparency> },
+			{ TEXT("map_Kd"),  ParseMaterialPath<&FObjData::FMaterialData::DiffuseTexture> },
+			{ TEXT("map_Ka"),  ParseMaterialPath<&FObjData::FMaterialData::AmbientTexture> },
+			{ TEXT("map_Ks"),  ParseMaterialPath<&FObjData::FMaterialData::SpecularTexture> }
+		};
+
+		// ZBrush doesn't surround filenames with spaces in quotes, so just take the entire remainder of the line as the filename,
+		// instead of getting the next token.
+		FStringView MtlFilename = Line; // (instead of GetToken(Line);)
+		if (MtlFilename.IsEmpty())
+		{
+			UE_LOG(LogTemp, Warning, TEXT("Missing filename on mtllib keyword"));
+		}
+
+		// Parse the file referenced by the .obj
+		// This will set the current material name
+		FString FileToParse = FPaths::GetPath(ObjData.ObjFilename) / FString(MtlFilename);
+		bool bSuccess = ParseFile(ObjData, *FileToParse, MtlKeywordMap);
+
+		// After parsing, clear the current material name
+		ObjData.MaterialBeingDefined.Empty();
+
+		return bSuccess;
+	}
+
+
+} // namespace ObjParser
+} // namespace Interchange
+} // namespace UE
+
+
+/**
+ * Build and return a UID name for a mesh node.
+ * @todo: move this to be a static method on the MeshNode class itself?
+ */
+static FString MakeMeshNodeUid(const FString& Name)
+{
+	return FString(TEXT("\\Mesh\\")) + (Name.IsEmpty() ? FString(TEXT("Null")) : Name);
+}
+
+
+/**
+ * Build and return a UID name for a shader graph node.
+ * @todo: move this to be a static method on the ShaderGraphNode class itself?
+ */
+static FString MakeShaderGraphNodeUid(const FString& Name)
+{
+	return FString(TEXT("\\Material\\")) + (Name.IsEmpty() ? FString(TEXT("Null")) : Name);
+}
+
+
+/**
+ * Build and return a UID name for a texture node.
+ * @todo: move this to be a static method on the TextureNode class itself?
+ */
+static FString MakeTextureNodeUid(const FString& Name)
+{
+	return FString(TEXT("\\Texture\\")) + (Name.IsEmpty() ? FString(TEXT("Null")) : Name);
+}
+
+
+UInterchangeOBJTranslator::UInterchangeOBJTranslator()
+	: ObjDataPtr(MakePimpl<FObjData>())
+{
+}
+
+
+UInterchangeOBJTranslator::~UInterchangeOBJTranslator()
+{
+}
+
+
+bool UInterchangeOBJTranslator::CanImportSourceData(const UInterchangeSourceData* InSourceData) const
+{
+	const bool bIncludeDot = false;
+	FString Extension = FPaths::GetExtension(InSourceData->GetFilename(), bIncludeDot);
+	return Extension == FString(TEXT("obj"));
+}
+
+
+bool UInterchangeOBJTranslator::Translate(UInterchangeBaseNodeContainer& BaseNodeContainer) const
+{
+	// Get filename from SourceData
+
+	FString Filename = GetSourceData()->GetFilename();
+	if (!FPaths::FileExists(Filename))
+	{
+		return false;
+	}
+
+	ObjDataPtr->ObjFilename = Filename;
+
+
+	// @TODO: API design
+	// The Translate() method is const, implying that the translator object should not be modified.
+	// However, in this particular case, the translator needs to build some data which is required when the payload is requested.
+	// In the case of FBX, this happens in an adjacent process.
+	// In our case, we need to put it somewhere, and we use a separate object to hold it; however this is not really in the spirit of
+	// 'const', we are effectively using the loophole that an owned object need not be const even if the owner is const.
+	//
+	// We need to think about the semantics of this a bit better.
+	// Is the translator *really* to be considered a const object?  Should it be allowed to hold state?
+	// Would it be better to shift the payload ownership to the nodes which the translator creates?
+	// Do we maybe want to decouple translators from payloads (which arguably makes more sense).
+
+	using namespace UE::Interchange;
+
+	static ObjParser::FKeywordMap ObjKeywordMap =
+	{
+		{ TEXT("#"),      ObjParser::ParseComment },
+		{ TEXT("v"),      ObjParser::ParseVertexPosition },
+		{ TEXT("vt"),     ObjParser::ParseTextureCoordinate },
+		{ TEXT("vn"),     ObjParser::ParseNormalVector },
+		{ TEXT("f"),      ObjParser::ParseFace },
+		{ TEXT("g"),      ObjParser::ParseGroup },
+		{ TEXT("mtllib"), ObjParser::ParseMaterialLib },
+		{ TEXT("usemtl"), ObjParser::ParseUseMaterial },
+//		{ TEXT("s"),      ObjParser::ParseSmoothing },
+	};
+
+	bool bSuccess = ObjParser::ParseFile(*ObjDataPtr.Get(), *Filename, ObjKeywordMap);
+
+	if (bSuccess)
+	{
+		// Add mesh nodes to the container for each mesh group
+
+		for (const TPair<FString, FObjData::FGroupData>& Group : ObjDataPtr->Groups)
+		{
+			const FString& GroupName = Group.Key;
+			FString NodeUid = MakeMeshNodeUid(GroupName);
+
+			UInterchangeMeshNode* MeshNode = NewObject<UInterchangeMeshNode>(&BaseNodeContainer);
+			MeshNode->InitializeNode(NodeUid, GroupName, EInterchangeNodeContainerType::TranslatedAsset);
+			BaseNodeContainer.AddNode(MeshNode);
+
+			// The payload key is the group name.
+			// A MeshDescription is generated for each group of faces.
+			MeshNode->SetPayLoadKey(GroupName);
+
+			MeshNode->SetCustomBoundingBox(ObjDataPtr->GetGroupBoundingBox(GroupName));
+			MeshNode->SetCustomVertexCount(ObjDataPtr->GetGroupVertexCount(GroupName));
+			MeshNode->SetCustomPolygonCount(ObjDataPtr->GetGroupPolygonCount(GroupName));
+			MeshNode->SetCustomHasVertexNormal(true);
+			MeshNode->SetCustomHasVertexBinormal(false);
+			MeshNode->SetCustomHasVertexTangent(false);
+			MeshNode->SetCustomHasSmoothGroup(false);
+			MeshNode->SetCustomHasVertexColor(false);
+
+			if (!Group.Value.MaterialName.IsEmpty())
+			{
+				MeshNode->SetMaterialDependencyUid(MakeShaderGraphNodeUid(Group.Value.MaterialName));
+			}
+		}
+
+		// Add material nodes to the container
+
+		for (const TPair<FString, FObjData::FMaterialData>& Material : ObjDataPtr->Materials)
+		{
+			const FString& MaterialName = Material.Key;
+			const FObjData::FMaterialData& MaterialData = Material.Value;
+			FString NodeUid = MakeShaderGraphNodeUid(MaterialName);
+
+			UInterchangeShaderGraphNode* ShaderGraphNode = NewObject<UInterchangeShaderGraphNode>(&BaseNodeContainer);
+			ShaderGraphNode->InitializeNode(NodeUid, MaterialName, EInterchangeNodeContainerType::TranslatedAsset);
+			BaseNodeContainer.AddNode(ShaderGraphNode);
+
+			if (MaterialData.IlluminationModel >= 2)
+			{
+				AddMaterialNodes(BaseNodeContainer, ShaderGraphNode, Materials::Phong::Parameters::SpecularColor.ToString(), MaterialData.SpecularColor, MaterialData.SpecularTexture);
+			}
+
+			if (MaterialData.IlluminationModel >= 1)
+			{
+				AddMaterialNodes(BaseNodeContainer, ShaderGraphNode, Materials::Phong::Parameters::EmissiveColor.ToString(), MaterialData.AmbientColor, MaterialData.AmbientTexture);
+			}
+
+			// Always add diffuse colour
+			AddMaterialNodes(BaseNodeContainer, ShaderGraphNode, Materials::Phong::Parameters::DiffuseColor.ToString(), MaterialData.DiffuseColor, MaterialData.DiffuseTexture);
+		}
+	}
+
+	return bSuccess;
+}
+
+
+TFuture<TOptional<UE::Interchange::FStaticMeshPayloadData>> UInterchangeOBJTranslator::GetStaticMeshPayloadData(const FString& PayloadKey) const
+{
+	return Async(EAsyncExecution::TaskGraph, [this, &PayloadKey]
+		{
+			using namespace UE::Interchange;
+
+			FStaticMeshPayloadData Payload;
+			Payload.MeshDescription = ObjDataPtr->MakeMeshDescriptionForGroup(PayloadKey);
+
+			return TOptional<FStaticMeshPayloadData>(Payload);
+		}
+	);
+}
+
+
+TOptional<UE::Interchange::FImportImage> UInterchangeOBJTranslator::GetTexturePayloadData(const UInterchangeSourceData* InSourceData, const FString& PayLoadKey) const
+{
+	// @TODO: API design
+	// This method is copied verbatim from InterchangeFBXTranslator.
+	// We need a better way of allowing a translator to delegate translation to other translator.
+	// Here, we have a texture source referenced by the .obj which is ideally handled by the existing texture translators.
+	// In this case, we can just implement the texture payload interface here and use a temporary texture translator instance
+	// to generate the payload, but what if the translator needed to do non-trivial work during the Translate() phase?
+	UInterchangeSourceData* PayloadSourceData = UInterchangeManager::GetInterchangeManager().CreateSourceData(PayLoadKey);
+	FGCObjectScopeGuard ScopedSourceData(PayloadSourceData);
+	if (!PayloadSourceData)
+	{
+		return TOptional<UE::Interchange::FImportImage>();
+	}
+
+	UInterchangeTranslatorBase* SourceTranslator = UInterchangeManager::GetInterchangeManager().GetTranslatorForSourceData(PayloadSourceData);
+	FGCObjectScopeGuard ScopedSourceTranslator(SourceTranslator);
+	const IInterchangeTexturePayloadInterface* TextureTranslator = Cast<IInterchangeTexturePayloadInterface>(SourceTranslator);
+	if (!ensure(TextureTranslator))
+	{
+		return TOptional<UE::Interchange::FImportImage>();
+	}
+
+	return TextureTranslator->GetTexturePayloadData(PayloadSourceData, PayLoadKey);
+}
+
+
+
+void UInterchangeOBJTranslator::AddMaterialNodes(UInterchangeBaseNodeContainer& BaseNodeContainer, UInterchangeShaderGraphNode* ShaderGraphNode, const FString& InputType, const FVector3f& Color, const FString& TexturePath) const
+{
+	if (!TexturePath.IsEmpty())
+	{
+		// @todo: better way to handle textures than implementing the TexturePayload interface in the OBJTranslator
+		// and creating a temporary texture translator there.
+
+		using namespace UE::Interchange::Materials::Standard::Nodes;
+
+		FString TextureName = FString(TEXT("TEX_")) + FPaths::GetBaseFilename(TexturePath);
+
+		UInterchangeShaderNode* TextureSampleShader = NewObject<UInterchangeShaderNode>(&BaseNodeContainer);
+		TextureSampleShader->SetCustomShaderType(TextureSample::Name.ToString());
+
+		FString TextureSampleShaderUid = ShaderGraphNode->GetUniqueID() + TEXT("\\Textures\\") + TextureName;
+		TextureSampleShader->InitializeNode(TextureSampleShaderUid, TextureName, EInterchangeNodeContainerType::TranslatedAsset);
+
+		BaseNodeContainer.AddNode(TextureSampleShader);
+		BaseNodeContainer.SetNodeParentUid(TextureSampleShaderUid, ShaderGraphNode->GetUniqueID());
+
+		FString TextureNodeUid = MakeTextureNodeUid(TextureName);
+		UInterchangeTexture2DNode* TextureNode = Cast<UInterchangeTexture2DNode>(BaseNodeContainer.GetNode(TextureNodeUid));
+		if (!TextureNode)
+		{
+			TextureNode = NewObject<UInterchangeTexture2DNode>(&BaseNodeContainer);
+			TextureNode->InitializeNode(TextureNodeUid, TextureName, EInterchangeNodeContainerType::TranslatedAsset);
+			BaseNodeContainer.AddNode(TextureNode);
+
+			FString NormalizedTexturePath(TexturePath);
+			FPaths::NormalizeFilename(NormalizedTexturePath);
+			TextureNode->SetPayLoadKey(NormalizedTexturePath);
+		}
+
+		TextureSampleShader->AddStringAttribute(UInterchangeShaderPortsAPI::MakeInputValueKey(TextureSample::Inputs::Texture.ToString()), TextureNodeUid);
+		TextureSampleShader->AddFloatAttribute(UInterchangeShaderPortsAPI::MakeInputValueKey(TextureSample::Inputs::UTiling.ToString()), 1.0f);
+		TextureSampleShader->AddFloatAttribute(UInterchangeShaderPortsAPI::MakeInputValueKey(TextureSample::Inputs::VTiling.ToString()), 1.0f);
+
+		// @todo: this is not currently working correctly
+		// Material graph complains that input is not LinearColor - seems to need an FColor->FLinearColor conversion node adding
+		UInterchangeShaderPortsAPI::ConnectDefaultOuputToInput(ShaderGraphNode, InputType, TextureSampleShaderUid);
+	}
+	else
+	{
+		ShaderGraphNode->AddLinearColorAttribute(UInterchangeShaderPortsAPI::MakeInputValueKey(InputType), FLinearColor(Color));
+	}
+
+}
\ No newline at end of file
diff --git a/Engine/Plugins/Experimental/Interchange/Runtime/Source/Import/Public/Mesh/InterchangeOBJTranslator.h b/Engine/Plugins/Experimental/Interchange/Runtime/Source/Import/Public/Mesh/InterchangeOBJTranslator.h
new file mode 100644
--- /dev/null
+++ b/Engine/Plugins/Experimental/Interchange/Runtime/Source/Import/Public/Mesh/InterchangeOBJTranslator.h
@@ -0,0 +1,79 @@
+// Copyright Epic Games, Inc. All Rights Reserved.
+
+#pragma once
+
+#include "CoreMinimal.h"
+#include "UObject/Object.h"
+#include "UObject/ObjectMacros.h"
+#include "InterchangeTranslatorBase.h"
+#include "Nodes/InterchangeBaseNodeContainer.h"
+#include "Mesh/InterchangeStaticMeshPayload.h"
+#include "Mesh/InterchangeStaticMeshPayloadInterface.h"
+#include "Texture/InterchangeTexturePayloadInterface.h"
+
+#include "InterchangeOBJTranslator.generated.h"
+
+class UInterchangeShaderGraphNode;
+struct FObjData;
+
+
+UCLASS(BlueprintType, Experimental)
+class INTERCHANGEIMPORT_API UInterchangeOBJTranslator : public UInterchangeTranslatorBase,
+	                                                    public IInterchangeStaticMeshPayloadInterface,
+	                                                    public IInterchangeTexturePayloadInterface
+{
+	GENERATED_BODY()
+
+public:
+
+	UInterchangeOBJTranslator();
+	virtual ~UInterchangeOBJTranslator();
+
+
+	/*
+	 * return true if the translator can translate the specified source data.
+	 */
+	virtual bool CanImportSourceData(const UInterchangeSourceData* InSourceData) const override;
+
+	/**
+	 * Translate the associated source data into a node hold by the specified nodes container.
+	 *
+	 * @param BaseNodeContainer - The unreal objects descriptions container where to put the translated source data.
+	 * @return true if the translator can translate the source data, false otherwise.
+	 */
+	virtual bool Translate(UInterchangeBaseNodeContainer& BaseNodeContainer) const override;
+
+
+	/* IInterchangeStaticMeshPayloadInterface Begin */
+
+	/**
+	 * Once the translation is done, the import process need a way to retrieve payload data.
+	 * This payload will be use by the factories to create the asset.
+	 *
+	 * @param PayloadKey - The key to retrieve the a particular payload contain into the specified source data.
+	 * @return a PayloadData containing the the data ask with the key.
+	 */
+	virtual TFuture<TOptional<UE::Interchange::FStaticMeshPayloadData>> GetStaticMeshPayloadData(const FString& PayloadKey) const override;
+
+	/* IInterchangeStaticMeshPayloadInterface End */
+
+	//////////////////////////////////////////////////////////////////////////
+	/* IInterchangeTexturePayloadInterface Begin */
+
+	/**
+	 * Once the translation is done, the import process need a way to retrieve payload data.
+	 * This payload will be use by the factories to create the asset.
+	 *
+	 * @param InSourceData - The source data containing the data to translate
+	 * @param PayloadKey - The key to retrieve the a particular payload contain into the specified source data.
+	 * @return a PayloadData containing the imported data. The TOptional will not be set if there is an error.
+	 */
+	virtual TOptional<UE::Interchange::FImportImage> GetTexturePayloadData(const UInterchangeSourceData* InSourceData, const FString& PayLoadKey) const override;
+
+	/* IInterchangeTexturePayloadInterface End */
+
+private:
+	void AddMaterialNodes(UInterchangeBaseNodeContainer& BaseNodeContainer, UInterchangeShaderGraphNode* ShaderGraphNode, const FString& InputType, const FVector3f& Color, const FString& TexturePath) const;
+
+	TPimplPtr<FObjData> ObjDataPtr;
+};
