Fix depending on libstdfbx.so as a runtime depends

diff --git a/Engine/Plugins/Experimental/Interchange/Runtime/Interchange.uplugin b/Engine/Plugins/Experimental/Interchange/Runtime/Interchange.uplugin
--- a/Engine/Plugins/Experimental/Interchange/Runtime/Interchange.uplugin
+++ b/Engine/Plugins/Experimental/Interchange/Runtime/Interchange.uplugin
@@ -53,7 +53,7 @@
 		},
 		{
 			"Name": "InterchangeFbxParser",
-			"Type": "Runtime",
+			"Type": "Editor",
 			"LoadingPhase": "Default",
 			"PlatformAllowList": [
 				"Win64",
@@ -64,4 +64,4 @@
 	],
 	"IsExperimentalVersion": false,
 	"SupportedPrograms": [ "InterchangeWorker" ]
-}
\ No newline at end of file
+}
