Interchange GLTF Translator
- New translator leveraging our GLTFCore module.
- Imports meshes, materials and textures.

diff --git a/Engine/Plugins/Editor/GLTFImporter/Source/GLTFCore/Private/GLTF/MaterialUtilities.h b/Engine/Plugins/Editor/GLTFImporter/Source/GLTFCore/Private/GLTF/MaterialUtilities.h
--- a/Engine/Plugins/Editor/GLTFImporter/Source/GLTFCore/Private/GLTF/MaterialUtilities.h
+++ b/Engine/Plugins/Editor/GLTFImporter/Source/GLTFCore/Private/GLTF/MaterialUtilities.h
@@ -24,7 +24,7 @@ namespace GLTF
 			{
 				OutMap.TextureIndex = TexIndex;
 				uint32 TexCoord     = GetUnsignedInt(TexObj, TEXT("texCoord"), 0);
-				check(TexCoord < 2);
+
 				OutMap.TexCoord = TexCoord;
 				if (InScaleName && TexObj.HasTypedField<EJson::Number>(InScaleName))
 				{
diff --git a/Engine/Plugins/Editor/GLTFImporter/Source/GLTFCore/Private/GLTFMeshFactory.cpp b/Engine/Plugins/Editor/GLTFImporter/Source/GLTFCore/Private/GLTFMeshFactory.cpp
--- a/Engine/Plugins/Editor/GLTFImporter/Source/GLTFCore/Private/GLTFMeshFactory.cpp
+++ b/Engine/Plugins/Editor/GLTFImporter/Source/GLTFCore/Private/GLTFMeshFactory.cpp
@@ -163,6 +163,8 @@ namespace GLTF
 		const int32 NumUVs = FMath::Max(1, GetNumUVs(Mesh));
 
 		FStaticMeshAttributes StaticMeshAttributes(*MeshDescription);
+		StaticMeshAttributes.Register();
+
 		TVertexAttributesRef<FVector3f> VertexPositions = StaticMeshAttributes.GetVertexPositions();
 		TEdgeAttributesRef<bool>  EdgeHardnesses = StaticMeshAttributes.GetEdgeHardnesses();
 		TPolygonGroupAttributesRef<FName> PolygonGroupImportedMaterialSlotNames = StaticMeshAttributes.GetPolygonGroupMaterialSlotNames();
diff --git a/Engine/Plugins/Experimental/Interchange/Runtime/Interchange.uplugin b/Engine/Plugins/Experimental/Interchange/Runtime/Interchange.uplugin
--- a/Engine/Plugins/Experimental/Interchange/Runtime/Interchange.uplugin
+++ b/Engine/Plugins/Experimental/Interchange/Runtime/Interchange.uplugin
@@ -62,6 +62,12 @@
 			]
 		}
 	],
+	"Plugins": [
+		{
+			"Name": "GLTFImporter",
+			"Enabled": true
+		}
+	],
 	"IsExperimentalVersion": false,
 	"SupportedPrograms": [ "InterchangeWorker" ]
 }
diff --git a/Engine/Plugins/Experimental/Interchange/Runtime/Source/Import/InterchangeImport.Build.cs b/Engine/Plugins/Experimental/Interchange/Runtime/Source/Import/InterchangeImport.Build.cs
--- a/Engine/Plugins/Experimental/Interchange/Runtime/Source/Import/InterchangeImport.Build.cs
+++ b/Engine/Plugins/Experimental/Interchange/Runtime/Source/Import/InterchangeImport.Build.cs
@@ -27,6 +27,7 @@ namespace UnrealBuildTool.Rules
 				{
 					"AssetRegistry",
 					"CinematicCamera",
+					"GLTFCore",
 					"IESFile",
 					"ImageWrapper",
 					"InterchangeDispatcher",
diff --git a/Engine/Plugins/Experimental/Interchange/Runtime/Source/Import/Private/Gltf/InterchangeGltfTranslator.cpp b/Engine/Plugins/Experimental/Interchange/Runtime/Source/Import/Private/Gltf/InterchangeGltfTranslator.cpp
new file mode 100644
--- /dev/null
+++ b/Engine/Plugins/Experimental/Interchange/Runtime/Source/Import/Private/Gltf/InterchangeGltfTranslator.cpp
@@ -0,0 +1,504 @@
+// Copyright Epic Games, Inc. All Rights Reserved. 
+
+#include "Gltf/InterchangeGltfTranslator.h"
+
+#include "GLTFAsset.h"
+#include "GLTFMeshFactory.h"
+#include "GLTFReader.h"
+
+#include "InterchangeCameraNode.h"
+#include "InterchangeLightNode.h"
+#include "InterchangeManager.h"
+#include "InterchangeMaterialDefinitions.h"
+#include "InterchangeMeshNode.h"
+#include "InterchangeSceneNode.h"
+#include "InterchangeShaderGraphNode.h"
+#include "InterchangeTexture2DNode.h"
+
+#include "Algo/Find.h"
+#include "StaticMeshAttributes.h"
+#include "UObject/GCObjectScopeGuard.h"
+
+void UInterchangeGltfTranslator::HandleGltfNode( UInterchangeBaseNodeContainer& NodeContainer, const GLTF::FNode& GltfNode, const FString& ParentNodeUid ) const
+{
+	const FString NodeUid = ParentNodeUid + TEXT("\\") + GltfNode.Name;
+
+	UInterchangeSceneNode* ParentSceneNode = Cast< UInterchangeSceneNode >( NodeContainer.GetNode( ParentNodeUid ) );
+
+	UInterchangeSceneNode* InterchangeSceneNode = NewObject< UInterchangeSceneNode >( &NodeContainer );
+	InterchangeSceneNode->InitializeNode( NodeUid, GltfNode.Name, EInterchangeNodeContainerType::TranslatedScene );
+	NodeContainer.AddNode( InterchangeSceneNode );
+
+	FTransform Transform = GltfNode.Transform;
+
+	constexpr float MetersToCentimeters = 100.f;
+	Transform.SetTranslation( Transform.GetTranslation() * MetersToCentimeters );
+
+	if ( ParentSceneNode )
+	{
+		FTransform ParentTransform;
+		ParentSceneNode->GetCustomGlobalTransform( ParentTransform );
+
+		Transform = Transform * ParentTransform;
+	}
+
+	switch ( GltfNode.Type )
+	{
+		case GLTF::FNode::EType::Mesh:
+		{
+			if ( GltfAsset.Meshes.IsValidIndex( GltfNode.MeshIndex ) )
+			{
+				const FString MeshNodeUid = TEXT("\\Mesh\\") + GltfAsset.Meshes[ GltfNode.MeshIndex ].Name;
+				InterchangeSceneNode->SetCustomAssetInstanceUid( MeshNodeUid );
+			}
+			break;
+		}
+
+		case GLTF::FNode::EType::Camera:
+		{
+			Transform.ConcatenateRotation(FRotator(0, -90, 0).Quaternion());
+
+			if ( GltfAsset.Cameras.IsValidIndex( GltfNode.CameraIndex ) )
+			{
+				const FString CameraNodeUid = TEXT("\\Camera\\") + GltfAsset.Cameras[ GltfNode.CameraIndex ].Name;
+				InterchangeSceneNode->SetCustomAssetInstanceUid( CameraNodeUid );
+			}
+			break;
+		}
+
+		case GLTF::FNode::EType::Light:
+		{
+			Transform.ConcatenateRotation(FRotator(0, -90, 0).Quaternion());
+
+			if ( GltfAsset.Lights.IsValidIndex( GltfNode.LightIndex ) )
+			{
+				const FString LightNodeUid = TEXT("\\Light\\") + GltfAsset.Lights[ GltfNode.LightIndex ].Name;
+				InterchangeSceneNode->SetCustomAssetInstanceUid( LightNodeUid );
+			}
+		}
+
+		case GLTF::FNode::EType::Transform:
+		default:
+		{
+			break;
+		}
+	}
+
+	InterchangeSceneNode->SetCustomGlobalTransform( Transform );
+
+	if ( !ParentNodeUid.IsEmpty() )
+	{
+		NodeContainer.SetNodeParentUid( NodeUid, ParentNodeUid );
+	}
+
+	for ( const int32 ChildIndex : GltfNode.Children )
+	{
+		if ( GltfAsset.Nodes.IsValidIndex( ChildIndex ) )
+		{
+			HandleGltfNode( NodeContainer, GltfAsset.Nodes[ ChildIndex ], NodeUid );
+		}
+	}
+}
+
+void UInterchangeGltfTranslator::HandleGltfMaterialParameter( UInterchangeBaseNodeContainer& NodeContainer, const GLTF::FTextureMap& TextureMap, UInterchangeShaderGraphNode& ShaderGraphNode,
+		const FString& MapName, const TVariant< FLinearColor, float >& MapFactor, const FString& OutputChannel ) const
+{
+	using namespace UE::Interchange::Materials;
+
+	if ( GltfAsset.Textures.IsValidIndex( TextureMap.TextureIndex ) )
+	{
+		const FString NodeName = MapName;
+		const FString NodeUid = ShaderGraphNode.GetUniqueID() + TEXT("_") + NodeName;
+
+		UInterchangeShaderNode* ColorNode = NewObject< UInterchangeShaderNode >( &NodeContainer );
+		ColorNode->InitializeNode( NodeUid, NodeName, EInterchangeNodeContainerType::TranslatedAsset );
+		NodeContainer.AddNode( ColorNode );
+		NodeContainer.SetNodeParentUid( NodeUid, ShaderGraphNode.GetUniqueID() );
+
+		ColorNode->SetCustomShaderType( Standard::Nodes::TextureSample::Name.ToString() );
+		const FString TextureUid = TEXT("\\Texture\\") + GltfAsset.Textures[ TextureMap.TextureIndex ].Source.URI;
+		ColorNode->AddStringAttribute( UInterchangeShaderPortsAPI::MakeInputValueKey( Standard::Nodes::TextureSample::Inputs::Texture.ToString() ), TextureUid );
+
+		bool bNeedsFactorNode = false;
+
+		if ( MapFactor.IsType< float >() )
+		{
+			bNeedsFactorNode = !FMath::IsNearlyEqual( MapFactor.Get< float >(), 1.f );
+		}
+		else if ( MapFactor.IsType< FLinearColor >() )
+		{
+			bNeedsFactorNode = !MapFactor.Get< FLinearColor >().Equals( FLinearColor::White );
+		}
+
+		if ( bNeedsFactorNode )
+		{
+			const FString FactorNodeUid = NodeUid + TEXT("_Factor");
+			UInterchangeShaderNode* FactorNode = NewObject< UInterchangeShaderNode >( &NodeContainer );
+			FactorNode->InitializeNode( FactorNodeUid, NodeName + TEXT("_Factor"), EInterchangeNodeContainerType::TranslatedAsset );
+			NodeContainer.AddNode( FactorNode );
+			NodeContainer.SetNodeParentUid( FactorNodeUid, ShaderGraphNode.GetUniqueID() );
+
+			FactorNode->SetCustomShaderType( Standard::Nodes::Multiply::Name.ToString() );
+
+			if ( MapFactor.IsType< float >() )
+			{
+				FactorNode->AddFloatAttribute( UInterchangeShaderPortsAPI::MakeInputValueKey( Standard::Nodes::Multiply::Inputs::B.ToString() ), MapFactor.Get< float >() );
+			}
+			else if ( MapFactor.IsType< FLinearColor >() )
+			{
+				FactorNode->AddLinearColorAttribute( UInterchangeShaderPortsAPI::MakeInputValueKey( Standard::Nodes::Multiply::Inputs::B.ToString() ), MapFactor.Get< FLinearColor >() );
+			}
+
+			UInterchangeShaderPortsAPI::ConnectOuputToInput( FactorNode, Standard::Nodes::Multiply::Inputs::A.ToString(), NodeUid, OutputChannel );
+			UInterchangeShaderPortsAPI::ConnectDefaultOuputToInput( &ShaderGraphNode, MapName, FactorNodeUid );
+		}
+		else
+		{
+			UInterchangeShaderPortsAPI::ConnectOuputToInput( &ShaderGraphNode, MapName, NodeUid, OutputChannel );
+		}
+	}
+	else
+	{
+		if ( MapFactor.IsType< FLinearColor >() )
+		{
+			ShaderGraphNode.AddLinearColorAttribute( UInterchangeShaderPortsAPI::MakeInputValueKey( MapName ), MapFactor.Get< FLinearColor >() );
+		}
+		else if ( MapFactor.IsType< float >() )
+		{
+			ShaderGraphNode.AddFloatAttribute( UInterchangeShaderPortsAPI::MakeInputValueKey( MapName ), MapFactor.Get< float >() );
+		}
+	}
+}
+
+void UInterchangeGltfTranslator::HandleGltfMaterial( UInterchangeBaseNodeContainer& NodeContainer, const GLTF::FMaterial& GltfMaterial, UInterchangeShaderGraphNode& ShaderGraphNode ) const
+{
+	using namespace UE::Interchange::Materials;
+
+	ShaderGraphNode.SetCustomTwoSided( GltfMaterial.bIsDoubleSided );
+
+	if ( GltfMaterial.ShadingModel == GLTF::FMaterial::EShadingModel::MetallicRoughness )
+	{
+		// Base Color
+		{
+			TVariant< FLinearColor, float > BaseColorFactor;
+			BaseColorFactor.Set< FLinearColor >( FLinearColor( GltfMaterial.BaseColorFactor ) );
+
+			HandleGltfMaterialParameter( NodeContainer, GltfMaterial.BaseColor, ShaderGraphNode, PBR::Parameters::BaseColor.ToString(),
+				BaseColorFactor, Standard::Nodes::TextureSample::Outputs::RGB.ToString() );
+		}
+
+		// Metallic
+		{
+			TVariant< FLinearColor, float > MetallicFactor;
+			MetallicFactor.Set< float >( GltfMaterial.MetallicRoughness.MetallicFactor );
+
+			HandleGltfMaterialParameter( NodeContainer, GltfMaterial.MetallicRoughness.Map, ShaderGraphNode, PBR::Parameters::Metallic.ToString(),
+				MetallicFactor, Standard::Nodes::TextureSample::Outputs::B.ToString() );
+		}
+
+		// Roughness
+		{
+			TVariant< FLinearColor, float > RoughnessFactor;
+			RoughnessFactor.Set< float >( GltfMaterial.MetallicRoughness.RoughnessFactor );
+
+			HandleGltfMaterialParameter( NodeContainer, GltfMaterial.MetallicRoughness.Map, ShaderGraphNode, PBR::Parameters::Roughness.ToString(),
+				RoughnessFactor, Standard::Nodes::TextureSample::Outputs::G.ToString() );
+		}
+
+		// Specular
+		if (GltfMaterial.bHasSpecular)
+		{
+			TVariant< FLinearColor, float > SpecularFactor;
+			SpecularFactor.Set< float >( GltfMaterial.Specular.SpecularFactor );
+
+			HandleGltfMaterialParameter( NodeContainer, GltfMaterial.Specular.SpecularMap, ShaderGraphNode, PBR::Parameters::Specular.ToString(),
+				SpecularFactor, Standard::Nodes::TextureSample::Outputs::RGB.ToString() );
+		}
+	}
+	else if ( GltfMaterial.ShadingModel == GLTF::FMaterial::EShadingModel::SpecularGlossiness )
+	{
+		// Diffuse Color
+		{
+			TVariant< FLinearColor, float > DiffuseColorFactor;
+			DiffuseColorFactor.Set< FLinearColor >( FLinearColor( GltfMaterial.BaseColorFactor ) );
+
+			HandleGltfMaterialParameter( NodeContainer, GltfMaterial.BaseColor, ShaderGraphNode, Phong::Parameters::DiffuseColor.ToString(),
+				DiffuseColorFactor, Standard::Nodes::TextureSample::Outputs::RGB.ToString() );
+		}
+
+		// Specular Color
+		{
+			TVariant< FLinearColor, float > SpecularColorFactor;
+			SpecularColorFactor.Set< FLinearColor >( FLinearColor( GltfMaterial.SpecularGlossiness.SpecularFactor ) );
+
+			HandleGltfMaterialParameter( NodeContainer, GltfMaterial.SpecularGlossiness.Map, ShaderGraphNode, Phong::Parameters::SpecularColor.ToString(),
+				SpecularColorFactor, Standard::Nodes::TextureSample::Outputs::RGB.ToString() );
+		}
+
+		// Glossiness
+		{
+			TVariant< FLinearColor, float > GlossinessFactor;
+			GlossinessFactor.Set< float >( GltfMaterial.SpecularGlossiness.GlossinessFactor );
+
+			HandleGltfMaterialParameter( NodeContainer, GltfMaterial.SpecularGlossiness.Map, ShaderGraphNode, Phong::Parameters::Shininess.ToString(),
+				GlossinessFactor, Standard::Nodes::TextureSample::Outputs::A.ToString() );
+		}
+	}
+
+	// Additional maps
+	{
+		// Normal
+		if ( GltfMaterial.Normal.TextureIndex != INDEX_NONE )
+		{
+			TVariant< FLinearColor, float > NormalFactor;
+			NormalFactor.Set< float >( GltfMaterial.NormalScale );
+
+			HandleGltfMaterialParameter( NodeContainer, GltfMaterial.Normal, ShaderGraphNode, Common::Parameters::Normal.ToString(),
+				NormalFactor, Standard::Nodes::TextureSample::Outputs::RGB.ToString() );
+		}
+
+		// Emissive
+		if ( GltfMaterial.Emissive.TextureIndex != INDEX_NONE || !GltfMaterial.EmissiveFactor.IsNearlyZero() )
+		{
+			TVariant< FLinearColor, float > EmissiveFactor;
+			EmissiveFactor.Set< FLinearColor >( FLinearColor( GltfMaterial.EmissiveFactor ) );
+
+			HandleGltfMaterialParameter( NodeContainer, GltfMaterial.Emissive, ShaderGraphNode, Common::Parameters::EmissiveColor.ToString(),
+				EmissiveFactor, Standard::Nodes::TextureSample::Outputs::RGB.ToString() );
+		}
+
+		// Occlusion
+		if ( GltfMaterial.Occlusion .TextureIndex != INDEX_NONE )
+		{
+			TVariant< FLinearColor, float > OcclusionFactor;
+			OcclusionFactor.Set< float >( GltfMaterial.OcclusionStrength );
+
+			HandleGltfMaterialParameter( NodeContainer, GltfMaterial.Occlusion, ShaderGraphNode, PBR::Parameters::Occlusion.ToString(),
+				OcclusionFactor, Standard::Nodes::TextureSample::Outputs::RGB.ToString() );
+		}
+
+		// Opacity (use the base color alpha channel)
+		if ( GltfMaterial.AlphaMode != GLTF::FMaterial::EAlphaMode::Opaque )
+		{
+			TVariant< FLinearColor, float > OpacityFactor;
+			OpacityFactor.Set< float >( GltfMaterial.BaseColorFactor.W );
+
+			HandleGltfMaterialParameter( NodeContainer, GltfMaterial.BaseColor, ShaderGraphNode, PBR::Parameters::Opacity.ToString(),
+				OpacityFactor, Standard::Nodes::TextureSample::Outputs::A.ToString() );
+
+			ShaderGraphNode.AddFloatAttribute( UInterchangeShaderPortsAPI::MakeInputValueKey( PBR::Parameters::IndexOfRefraction.ToString() ), GltfMaterial.IOR );
+		}
+	}
+}
+
+bool UInterchangeGltfTranslator::CanImportSourceData( const UInterchangeSourceData* InSourceData ) const
+{
+	const bool bIncludeDot = false;
+	const FString Extension = FPaths::GetExtension( InSourceData->GetFilename(), bIncludeDot );
+
+	TStaticArray<FString, 2> GltfExtensions;
+	GltfExtensions[0] = TEXT("gltf;GL Transmission Format");
+	GltfExtensions[1] = TEXT("glb;GL Transmission Format (Binary)");
+
+	const bool bExtensionMatches =
+		Algo::FindByPredicate( GltfExtensions,
+		[ &Extension ]( const FString& GltfExtension )
+		{
+			return GltfExtension.StartsWith( Extension );
+		}) != nullptr;
+
+	return bExtensionMatches;
+}
+
+bool UInterchangeGltfTranslator::Translate( UInterchangeBaseNodeContainer& NodeContainer ) const
+{
+	FString Filename = GetSourceData()->GetFilename();
+	if ( !FPaths::FileExists( Filename ) )
+	{
+		return false;
+	}
+
+	GLTF::FFileReader GltfFileReader;
+
+	const bool bLoadImageData = false;
+	const bool bLoadMetaData = false;
+	GltfFileReader.ReadFile( Filename, bLoadImageData, bLoadMetaData, const_cast< UInterchangeGltfTranslator* >( this )->GltfAsset );
+
+	// Textures
+	{
+		int32 TextureIndex = 0;
+		for ( const GLTF::FTexture& GltfTexture : GltfAsset.Textures )
+		{
+			UInterchangeTexture2DNode* TextureNode = NewObject< UInterchangeTexture2DNode >( &NodeContainer );
+			FString TextureNodeUid = TEXT("\\Texture\\") + GltfTexture.Source.URI;
+			TextureNode->InitializeNode( TextureNodeUid, GltfTexture.Source.URI, EInterchangeNodeContainerType::TranslatedAsset );
+			TextureNode->SetPayLoadKey( LexToString( TextureIndex++ ) );
+			NodeContainer.AddNode( TextureNode );
+		}
+	}
+
+	// Materials
+	{
+		for ( const GLTF::FMaterial& GltfMaterial : GltfAsset.Materials )
+		{
+			UInterchangeShaderGraphNode* ShaderGraphNode = NewObject< UInterchangeShaderGraphNode >( &NodeContainer );
+			FString ShaderGraphNodeUid = TEXT("\\Material\\") + GltfMaterial.Name;
+			ShaderGraphNode->InitializeNode( ShaderGraphNodeUid, GltfMaterial.Name, EInterchangeNodeContainerType::TranslatedAsset );
+			NodeContainer.AddNode( ShaderGraphNode );
+
+			HandleGltfMaterial( NodeContainer, GltfMaterial, *ShaderGraphNode );
+		}
+	}
+
+	// Meshes
+	{
+		int32 MeshIndex = 0;
+		for ( const GLTF::FMesh& GltfMesh : GltfAsset.Meshes )
+		{
+			UInterchangeMeshNode* MeshNode = NewObject< UInterchangeMeshNode >( &NodeContainer );
+			FString MeshNodeUid = TEXT("\\Mesh\\") + GltfMesh.Name;
+			MeshNode->InitializeNode( MeshNodeUid, GltfMesh.Name, EInterchangeNodeContainerType::TranslatedAsset );
+			MeshNode->SetPayLoadKey( LexToString( MeshIndex++ ) );
+			NodeContainer.AddNode( MeshNode );
+
+			// Assign materials
+			for ( const GLTF::FPrimitive& Primitive : GltfMesh.Primitives )
+			{
+				if ( GltfAsset.Materials.IsValidIndex( Primitive.MaterialIndex ) )
+				{
+					const FString ShaderGraphNodeUid = TEXT("\\Material\\") + GltfAsset.Materials[ Primitive.MaterialIndex ].Name;
+					MeshNode->SetMaterialDependencyUid( ShaderGraphNodeUid );
+				}
+			}
+		}
+	}
+
+	// Cameras
+	{
+		for ( const GLTF::FCamera& GltfCamera : GltfAsset.Cameras )
+		{
+			UInterchangeCameraNode* CameraNode = NewObject< UInterchangeCameraNode >( &NodeContainer );
+			FString CameraNodeUid = TEXT("\\Camera\\") + GltfCamera.Name;
+			CameraNode->InitializeNode( CameraNodeUid, GltfCamera.Name, EInterchangeNodeContainerType::TranslatedAsset );
+			NodeContainer.AddNode( CameraNode );
+		}
+	}
+
+	// Lights
+	{
+		for ( const GLTF::FLight& GltfLight : GltfAsset.Lights )
+		{
+			UInterchangeLightNode* LightNode = NewObject< UInterchangeLightNode >( &NodeContainer );
+			FString LightNodeUid = TEXT("\\Light\\") + GltfLight.Name;
+			LightNode->InitializeNode( LightNodeUid, GltfLight.Name, EInterchangeNodeContainerType::TranslatedAsset );
+			NodeContainer.AddNode( LightNode );
+		}
+	}
+
+	// Scenes
+	{
+		int32 SceneIndex = 0;
+		for ( const GLTF::FScene& GltfScene : GltfAsset.Scenes )
+		{
+			UInterchangeSceneNode* SceneNode = NewObject< UInterchangeSceneNode >( &NodeContainer );
+
+			FString SceneName = GltfScene.Name;
+			if (SceneName.IsEmpty())
+			{
+				SceneName = TEXT("Scene");
+				if ( GltfAsset.Scenes.Num() > 1 )
+				{
+					SceneName += TEXT("_") + LexToString( SceneIndex );
+				}
+			}
+
+			FString SceneNodeUid = TEXT("\\Scene\\") + SceneName;
+			SceneNode->InitializeNode( SceneNodeUid, SceneName, EInterchangeNodeContainerType::TranslatedScene );
+			NodeContainer.AddNode( SceneNode );
+
+			for ( const int32 NodeIndex : GltfScene.Nodes )
+			{
+				if ( GltfAsset.Nodes.IsValidIndex( NodeIndex ) )
+				{
+					HandleGltfNode( NodeContainer, GltfAsset.Nodes[ NodeIndex ], SceneNodeUid );
+				}
+			}
+
+			++SceneIndex;
+		}
+	}
+
+	return true;
+}
+
+TFuture< TOptional< UE::Interchange::FStaticMeshPayloadData > > UInterchangeGltfTranslator::GetStaticMeshPayloadData( const FString& PayLoadKey ) const
+{
+	TPromise< TOptional< UE::Interchange::FStaticMeshPayloadData > > MeshPayloadDataPromise;
+
+	int32 MeshIndex = 0;
+	LexFromString( MeshIndex, *PayLoadKey );
+
+	if ( !GltfAsset.Meshes.IsValidIndex( MeshIndex ) )
+	{
+		MeshPayloadDataPromise.SetValue(TOptional< UE::Interchange::FStaticMeshPayloadData >());
+		return MeshPayloadDataPromise.GetFuture();
+	}
+
+	UE::Interchange::FStaticMeshPayloadData StaticMeshPayloadData;
+
+	const GLTF::FMesh& GltfMesh = GltfAsset.Meshes[ MeshIndex ];
+	GLTF::FMeshFactory MeshFactory;
+	MeshFactory.SetUniformScale( 100.f ); // GLTF is in meters while UE is in centimeters
+	MeshFactory.FillMeshDescription( GltfMesh, &StaticMeshPayloadData.MeshDescription );
+
+	// Patch polygon groups material slot names to match Interchange expectations (rename material slots from indices to material names)
+	{
+		FStaticMeshAttributes StaticMeshAttributes( StaticMeshPayloadData.MeshDescription );
+
+		for ( int32 MaterialSlotIndex = 0; MaterialSlotIndex < StaticMeshAttributes.GetPolygonGroupMaterialSlotNames().GetNumElements(); ++MaterialSlotIndex )
+		{
+			int32 MaterialIndex = 0;
+			LexFromString( MaterialIndex, *StaticMeshAttributes.GetPolygonGroupMaterialSlotNames()[ MaterialSlotIndex ].ToString() );
+
+			if ( GltfAsset.Materials.IsValidIndex( MaterialIndex ) )
+			{
+				StaticMeshAttributes.GetPolygonGroupMaterialSlotNames()[ MaterialSlotIndex ] = *GltfAsset.Materials[ MaterialIndex ].Name;
+			}
+		}
+	}
+
+	MeshPayloadDataPromise.SetValue( StaticMeshPayloadData );
+
+	return MeshPayloadDataPromise.GetFuture();
+}
+
+TOptional< UE::Interchange::FImportImage > UInterchangeGltfTranslator::GetTexturePayloadData( const UInterchangeSourceData* InSourceData, const FString& PayLoadKey ) const
+{
+	int32 TextureIndex = 0;
+	LexFromString( TextureIndex, *PayLoadKey );
+
+	if ( !GltfAsset.Textures.IsValidIndex( TextureIndex ) )
+	{
+		return TOptional< UE::Interchange::FImportImage >();
+	}
+
+	GLTF::FTexture GltfTexture = GltfAsset.Textures[ TextureIndex ];
+
+	UInterchangeSourceData* PayloadSourceData = UInterchangeManager::GetInterchangeManager().CreateSourceData( GltfTexture.Source.FilePath );
+	FGCObjectScopeGuard ScopedSourceData( PayloadSourceData );
+	
+	if ( !PayloadSourceData )
+	{
+		return TOptional<UE::Interchange::FImportImage>();
+	}
+
+	UInterchangeTranslatorBase* SourceTranslator = UInterchangeManager::GetInterchangeManager().GetTranslatorForSourceData( PayloadSourceData );
+	FGCObjectScopeGuard ScopedSourceTranslator( SourceTranslator );
+	const IInterchangeTexturePayloadInterface* TextureTranslator = Cast< IInterchangeTexturePayloadInterface >( SourceTranslator );
+	if ( !ensure( TextureTranslator ) )
+	{
+		return TOptional<UE::Interchange::FImportImage>();
+	}
+
+	return TextureTranslator->GetTexturePayloadData( PayloadSourceData, GltfTexture.Source.FilePath );
+}
\ No newline at end of file
diff --git a/Engine/Plugins/Experimental/Interchange/Runtime/Source/Import/Private/InterchangeImportModule.cpp b/Engine/Plugins/Experimental/Interchange/Runtime/Source/Import/Private/InterchangeImportModule.cpp
--- a/Engine/Plugins/Experimental/Interchange/Runtime/Source/Import/Private/InterchangeImportModule.cpp
+++ b/Engine/Plugins/Experimental/Interchange/Runtime/Source/Import/Private/InterchangeImportModule.cpp
@@ -4,6 +4,7 @@
 
 #include "CoreMinimal.h"
 #include "Fbx/InterchangeFbxTranslator.h"
+#include "Gltf/InterchangeGltfTranslator.h"
 #include "InterchangeImportLog.h"
 #include "InterchangeManager.h"
 #include "Material/InterchangeMaterialFactory.h"
@@ -50,7 +51,8 @@ void FInterchangeImportModule::StartupModule()
 
 		//Register the translators
 		//Scenes
-		InterchangeManager.RegisterTranslator(UInterchangeFbxTranslator::StaticClass()); //Do not submit uncommented until we replace completly fbx importer (staticmesh + skeletalMesh + animation)
+		InterchangeManager.RegisterTranslator(UInterchangeFbxTranslator::StaticClass());
+		InterchangeManager.RegisterTranslator(UInterchangeGltfTranslator::StaticClass());
 		InterchangeManager.RegisterTranslator(UInterchangeOBJTranslator::StaticClass());
 
 		//Textures
diff --git a/Engine/Plugins/Experimental/Interchange/Runtime/Source/Import/Private/Texture/InterchangeJPGTranslator.cpp b/Engine/Plugins/Experimental/Interchange/Runtime/Source/Import/Private/Texture/InterchangeJPGTranslator.cpp
--- a/Engine/Plugins/Experimental/Interchange/Runtime/Source/Import/Private/Texture/InterchangeJPGTranslator.cpp
+++ b/Engine/Plugins/Experimental/Interchange/Runtime/Source/Import/Private/Texture/InterchangeJPGTranslator.cpp
@@ -1,6 +1,8 @@
 // Copyright Epic Games, Inc. All Rights Reserved. 
 #include "Texture/InterchangeJPGTranslator.h"
 
+#include "Algo/Find.h"
+#include "Containers/StaticArray.h"
 #include "Engine/Texture.h"
 #include "Engine/Texture2D.h"
 #include "IImageWrapper.h"
@@ -18,8 +20,19 @@
 bool UInterchangeJPGTranslator::CanImportSourceData(const UInterchangeSourceData* InSourceData) const
 {
 	FString Extension = FPaths::GetExtension(InSourceData->GetFilename());
-	FString JPGExtension = (TEXT("jpg;Texture"));
-	return JPGExtension.StartsWith(Extension);
+
+	TStaticArray<FString, 2> JPGExtensions;
+	JPGExtensions[0] = TEXT("jpg;Texture");
+	JPGExtensions[1] = TEXT("jpeg;Texture");
+
+	const bool bExtensionMatches =
+		Algo::FindByPredicate(JPGExtensions,
+		[&Extension](const FString& JPGExtension)
+		{
+			return JPGExtension.StartsWith(Extension);
+		}) != nullptr;
+
+	return bExtensionMatches;
 }
 
 bool UInterchangeJPGTranslator::Translate(UInterchangeBaseNodeContainer& BaseNodeContainer) const
diff --git a/Engine/Plugins/Experimental/Interchange/Runtime/Source/Import/Public/Gltf/InterchangeGltfTranslator.h b/Engine/Plugins/Experimental/Interchange/Runtime/Source/Import/Public/Gltf/InterchangeGltfTranslator.h
new file mode 100644
--- /dev/null
+++ b/Engine/Plugins/Experimental/Interchange/Runtime/Source/Import/Public/Gltf/InterchangeGltfTranslator.h
@@ -0,0 +1,53 @@
+// Copyright Epic Games, Inc. All Rights Reserved.
+#pragma once
+
+#include "GLTFAsset.h"
+#include "InterchangeTranslatorBase.h"
+#include "Mesh/InterchangeStaticMeshPayload.h"
+#include "Mesh/InterchangeStaticMeshPayloadInterface.h"
+#include "Nodes/InterchangeBaseNodeContainer.h"
+#include "Texture/InterchangeTexturePayloadInterface.h"
+#include "UObject/Object.h"
+#include "UObject/ObjectMacros.h"
+
+#include "InterchangeGltfTranslator.generated.h"
+
+class UInterchangeShaderGraphNode;
+
+/* Gltf translator class support import of texture, material, static mesh, skeletal mesh, */
+
+UCLASS(BlueprintType, Experimental)
+class UInterchangeGltfTranslator : public UInterchangeTranslatorBase,
+	public IInterchangeStaticMeshPayloadInterface, public IInterchangeTexturePayloadInterface
+{
+	GENERATED_BODY()
+
+public:
+	/** Begin UInterchangeTranslatorBase API*/
+	virtual bool CanImportSourceData( const UInterchangeSourceData* InSourceData ) const override;
+	virtual bool Translate( UInterchangeBaseNodeContainer& BaseNodeContainer ) const override;
+	/** End UInterchangeTranslatorBase API*/
+
+	/* IInterchangeStaticMeshPayloadInterface Begin */
+
+	virtual TFuture< TOptional< UE::Interchange::FStaticMeshPayloadData > > GetStaticMeshPayloadData( const FString& PayLoadKey ) const override;
+
+	/* IInterchangeStaticMeshPayloadInterface End */
+
+	/* IInterchangeTexturePayloadInterface Begin */
+
+	virtual TOptional< UE::Interchange::FImportImage > GetTexturePayloadData( const UInterchangeSourceData* InSourceData, const FString& PayLoadKey ) const override;
+
+	/* IInterchangeTexturePayloadInterface End */
+
+protected:
+	void HandleGltfNode( UInterchangeBaseNodeContainer& NodeContainer, const GLTF::FNode& GltfNode, const FString& ParentNodeUid ) const;
+	void HandleGltfMaterial( UInterchangeBaseNodeContainer& NodeContainer, const GLTF::FMaterial& GltfMaterial, UInterchangeShaderGraphNode& ShaderGraphNode ) const;
+	void HandleGltfMaterialParameter( UInterchangeBaseNodeContainer& NodeContainer, const GLTF::FTextureMap& TextureMap, UInterchangeShaderGraphNode& ShaderGraphNode,
+		const FString& MapName, const TVariant< FLinearColor, float >& MapFactor, const FString& OutputChannel ) const;
+
+private:
+	GLTF::FAsset GltfAsset;
+};
+
+
