Interchange: Improved conversion from Shader Graph to Material Instance
- Smarter processing of Lerp and Multiply nodes. If no textures are involved, we will compute the resulting value.
- Added a divide by 100 in the Phong material for the Shininess parameter before passing it to Roughness since we expect it to be between 2 and 100.
- Changed GLTF specular-glossiness to plug its glossiness in the roughness input with a OneMinus instead of going through shininess since its a PBR shader and not a real Phong shader.

diff --git a/Engine/Plugins/Experimental/Interchange/Runtime/Source/Import/Private/Gltf/InterchangeGltfTranslator.cpp b/Engine/Plugins/Experimental/Interchange/Runtime/Source/Import/Private/Gltf/InterchangeGltfTranslator.cpp
--- a/Engine/Plugins/Experimental/Interchange/Runtime/Source/Import/Private/Gltf/InterchangeGltfTranslator.cpp
+++ b/Engine/Plugins/Experimental/Interchange/Runtime/Source/Import/Private/Gltf/InterchangeGltfTranslator.cpp
@@ -101,10 +101,30 @@ void UInterchangeGltfTranslator::HandleGltfNode( UInterchangeBaseNodeContainer&
 }
 
 void UInterchangeGltfTranslator::HandleGltfMaterialParameter( UInterchangeBaseNodeContainer& NodeContainer, const GLTF::FTextureMap& TextureMap, UInterchangeShaderGraphNode& ShaderGraphNode,
-		const FString& MapName, const TVariant< FLinearColor, float >& MapFactor, const FString& OutputChannel ) const
+		const FString& MapName, const TVariant< FLinearColor, float >& MapFactor, const FString& OutputChannel, const bool bInverse ) const
 {
 	using namespace UE::Interchange::Materials;
 
+	UInterchangeShaderNode* NodeToConnectTo = &ShaderGraphNode;
+	FString InputToConnectTo = MapName;
+
+	if (bInverse)
+	{
+		const FString OneMinusNodeName = MapName + TEXT("OneMinus");
+		const FString OneMinusNodeUid = ShaderGraphNode.GetUniqueID() + TEXT("_") + OneMinusNodeName;
+		UInterchangeShaderNode* OneMinusNode = NewObject< UInterchangeShaderNode >( &NodeContainer );
+		OneMinusNode->InitializeNode( OneMinusNodeUid, OneMinusNodeName, EInterchangeNodeContainerType::TranslatedAsset );
+		NodeContainer.AddNode( OneMinusNode );
+		NodeContainer.SetNodeParentUid( OneMinusNodeUid, ShaderGraphNode.GetUniqueID() );
+
+		OneMinusNode->SetCustomShaderType(Standard::Nodes::OneMinus::Name.ToString());
+
+		UInterchangeShaderPortsAPI::ConnectDefaultOuputToInput(NodeToConnectTo, InputToConnectTo, OneMinusNode->GetUniqueID());
+
+		NodeToConnectTo = OneMinusNode;
+		InputToConnectTo = Standard::Nodes::OneMinus::Inputs::Input.ToString();
+	}
+
 	if ( GltfAsset.Textures.IsValidIndex( TextureMap.TextureIndex ) )
 	{
 		const FString NodeName = MapName;
@@ -150,22 +170,22 @@ void UInterchangeGltfTranslator::HandleGltfMaterialParameter( UInterchangeBaseNo
 			}
 
 			UInterchangeShaderPortsAPI::ConnectOuputToInput( FactorNode, Standard::Nodes::Multiply::Inputs::A.ToString(), NodeUid, OutputChannel );
-			UInterchangeShaderPortsAPI::ConnectDefaultOuputToInput( &ShaderGraphNode, MapName, FactorNodeUid );
+			UInterchangeShaderPortsAPI::ConnectDefaultOuputToInput( NodeToConnectTo, InputToConnectTo, FactorNodeUid );
 		}
 		else
 		{
-			UInterchangeShaderPortsAPI::ConnectOuputToInput( &ShaderGraphNode, MapName, NodeUid, OutputChannel );
+			UInterchangeShaderPortsAPI::ConnectOuputToInput( NodeToConnectTo, InputToConnectTo, NodeUid, OutputChannel );
 		}
 	}
 	else
 	{
 		if ( MapFactor.IsType< FLinearColor >() )
 		{
-			ShaderGraphNode.AddLinearColorAttribute( UInterchangeShaderPortsAPI::MakeInputValueKey( MapName ), MapFactor.Get< FLinearColor >() );
+			NodeToConnectTo->AddLinearColorAttribute( UInterchangeShaderPortsAPI::MakeInputValueKey( InputToConnectTo ), MapFactor.Get< FLinearColor >() );
 		}
 		else if ( MapFactor.IsType< float >() )
 		{
-			ShaderGraphNode.AddFloatAttribute( UInterchangeShaderPortsAPI::MakeInputValueKey( MapName ), MapFactor.Get< float >() );
+			NodeToConnectTo->AddFloatAttribute( UInterchangeShaderPortsAPI::MakeInputValueKey( InputToConnectTo ), MapFactor.Get< float >() );
 		}
 	}
 }
@@ -240,8 +260,9 @@ void UInterchangeGltfTranslator::HandleGltfMaterial( UInterchangeBaseNodeContain
 			TVariant< FLinearColor, float > GlossinessFactor;
 			GlossinessFactor.Set< float >( GltfMaterial.SpecularGlossiness.GlossinessFactor );
 
-			HandleGltfMaterialParameter( NodeContainer, GltfMaterial.SpecularGlossiness.Map, ShaderGraphNode, Phong::Parameters::Shininess.ToString(),
-				GlossinessFactor, Standard::Nodes::TextureSample::Outputs::A.ToString() );
+			const bool bInverse = true;
+			HandleGltfMaterialParameter( NodeContainer, GltfMaterial.SpecularGlossiness.Map, ShaderGraphNode, PBR::Parameters::Roughness.ToString(),
+				GlossinessFactor, Standard::Nodes::TextureSample::Outputs::A.ToString(), bInverse );
 		}
 	}
 
diff --git a/Engine/Plugins/Experimental/Interchange/Runtime/Source/Import/Public/Gltf/InterchangeGltfTranslator.h b/Engine/Plugins/Experimental/Interchange/Runtime/Source/Import/Public/Gltf/InterchangeGltfTranslator.h
--- a/Engine/Plugins/Experimental/Interchange/Runtime/Source/Import/Public/Gltf/InterchangeGltfTranslator.h
+++ b/Engine/Plugins/Experimental/Interchange/Runtime/Source/Import/Public/Gltf/InterchangeGltfTranslator.h
@@ -45,7 +45,7 @@ protected:
 	void HandleGltfNode( UInterchangeBaseNodeContainer& NodeContainer, const GLTF::FNode& GltfNode, const FString& ParentNodeUid ) const;
 	void HandleGltfMaterial( UInterchangeBaseNodeContainer& NodeContainer, const GLTF::FMaterial& GltfMaterial, UInterchangeShaderGraphNode& ShaderGraphNode ) const;
 	void HandleGltfMaterialParameter( UInterchangeBaseNodeContainer& NodeContainer, const GLTF::FTextureMap& TextureMap, UInterchangeShaderGraphNode& ShaderGraphNode,
-		const FString& MapName, const TVariant< FLinearColor, float >& MapFactor, const FString& OutputChannel ) const;
+		const FString& MapName, const TVariant< FLinearColor, float >& MapFactor, const FString& OutputChannel, const bool bInverse = false ) const;
 
 private:
 	GLTF::FAsset GltfAsset;
diff --git a/Engine/Plugins/Experimental/Interchange/Runtime/Source/Nodes/Public/InterchangeMaterialDefinitions.h b/Engine/Plugins/Experimental/Interchange/Runtime/Source/Nodes/Public/InterchangeMaterialDefinitions.h
--- a/Engine/Plugins/Experimental/Interchange/Runtime/Source/Nodes/Public/InterchangeMaterialDefinitions.h
+++ b/Engine/Plugins/Experimental/Interchange/Runtime/Source/Nodes/Public/InterchangeMaterialDefinitions.h
@@ -98,7 +98,7 @@ namespace Interchange
 				using namespace Lambert::Parameters;
 
 				const FName SpecularColor = TEXT("SpecularColor"); // Type: linear color
-				const FName Shininess = TEXT("Shininess"); // Type: float
+				const FName Shininess = TEXT("Shininess"); // Type: float, this is the specular exponent, expected range: 2-100
 			}
 		}
 
diff --git a/Engine/Plugins/Experimental/Interchange/Runtime/Source/Pipelines/Private/InterchangeGenericMaterialPipeline.cpp b/Engine/Plugins/Experimental/Interchange/Runtime/Source/Pipelines/Private/InterchangeGenericMaterialPipeline.cpp
--- a/Engine/Plugins/Experimental/Interchange/Runtime/Source/Pipelines/Private/InterchangeGenericMaterialPipeline.cpp
+++ b/Engine/Plugins/Experimental/Interchange/Runtime/Source/Pipelines/Private/InterchangeGenericMaterialPipeline.cpp
@@ -15,6 +15,7 @@
 #include "InterchangePipelineLog.h"
 #include "InterchangeSourceData.h"
 
+#include "Materials/MaterialExpressionDivide.h"
 #include "Materials/MaterialExpressionLinearInterpolate.h"
 #include "Materials/MaterialExpressionMaterialFunctionCall.h"
 #include "Materials/MaterialExpressionOneMinus.h"
@@ -87,8 +88,6 @@ void UInterchangeGenericMaterialPipeline::ExecutePreImportPipeline(UInterchangeB
 			}
 		}
 	}
-
-	return;
 }
 
 UInterchangeBaseMaterialFactoryNode* UInterchangeGenericMaterialPipeline::CreateBaseMaterialFactoryNode(const UInterchangeBaseNode* MaterialNode, TSubclassOf<UInterchangeBaseMaterialFactoryNode> NodeType)
@@ -184,7 +183,7 @@ bool UInterchangeGenericMaterialPipeline::HandlePhongModel(const UInterchangeSha
 
 			if (DiffuseExpressionFactoryNode.Get<0>())
 			{
-				UInterchangeShaderPortsAPI::ConnectOuputToInput(FunctionCallExpression, TEXT("DiffuseColor"),
+				UInterchangeShaderPortsAPI::ConnectOuputToInput(FunctionCallExpression, Parameters::DiffuseColor.ToString(),
 					DiffuseExpressionFactoryNode.Get<0>()->GetUniqueID(), DiffuseExpressionFactoryNode.Get<1>());
 			}
 		}
@@ -196,7 +195,7 @@ bool UInterchangeGenericMaterialPipeline::HandlePhongModel(const UInterchangeSha
 
 			if (SpecularExpressionFactoryNode.Get<0>())
 			{
-				UInterchangeShaderPortsAPI::ConnectOuputToInput(FunctionCallExpression, TEXT("SpecularColor"),
+				UInterchangeShaderPortsAPI::ConnectOuputToInput(FunctionCallExpression, Parameters::SpecularColor.ToString(),
 					SpecularExpressionFactoryNode.Get<0>()->GetUniqueID(), SpecularExpressionFactoryNode.Get<1>());
 			}
 		}
@@ -211,11 +210,24 @@ bool UInterchangeGenericMaterialPipeline::HandlePhongModel(const UInterchangeSha
 
 				if (ShininessExpressionFactoryNode.Get<0>())
 				{
+					UInterchangeMaterialExpressionFactoryNode* DivideShininessNode =
+						CreateExpressionNode(TEXT("DivideShininess"), ShininessExpressionFactoryNode.Get<0>()->GetUniqueID(), UMaterialExpressionDivide::StaticClass());
+
+					const float ShininessScale = 100.f; // Divide shininess by 100 to bring it into a 0-1 range for roughness.
+					const FString ShininessScaleParameterName = GET_MEMBER_NAME_CHECKED(UMaterialExpressionDivide, ConstB).ToString();
+					DivideShininessNode->AddFloatAttribute(ShininessScaleParameterName, ShininessScale);
+					DivideShininessNode->AddApplyAndFillDelegates<float>(ShininessScaleParameterName, UMaterialExpressionDivide::StaticClass(),  GET_MEMBER_NAME_CHECKED(UMaterialExpressionDivide, ConstB));
+
+					// Connect Shininess to Divide
+					UInterchangeShaderPortsAPI::ConnectOuputToInput(DivideShininessNode, GET_MEMBER_NAME_CHECKED(UMaterialExpressionDivide, A).ToString(),
+						ShininessExpressionFactoryNode.Get<0>()->GetUniqueID(), ShininessExpressionFactoryNode.Get<1>());
+
 					UInterchangeMaterialExpressionFactoryNode* InverseShininessNode =
 						CreateExpressionNode(TEXT("InverseShininess"), ShininessExpressionFactoryNode.Get<0>()->GetUniqueID(), UMaterialExpressionOneMinus::StaticClass());
 
-					UInterchangeShaderPortsAPI::ConnectOuputToInput(InverseShininessNode, GET_MEMBER_NAME_CHECKED(UMaterialExpressionOneMinus, Input).ToString(),
-						ShininessExpressionFactoryNode.Get<0>()->GetUniqueID(), ShininessExpressionFactoryNode.Get<1>());
+					// Connect Divide to Inverse
+					UInterchangeShaderPortsAPI::ConnectDefaultOuputToInput(InverseShininessNode, GET_MEMBER_NAME_CHECKED(UMaterialExpressionOneMinus, Input).ToString(),
+						DivideShininessNode->GetUniqueID());
 
 					MaterialFactoryNode->ConnectToRoughness(InverseShininessNode->GetUniqueID());
 				}
@@ -685,6 +697,12 @@ UInterchangeMaterialExpressionFactoryNode* UInterchangeGenericMaterialPipeline::
 
 TTuple<UInterchangeMaterialExpressionFactoryNode*, FString> UInterchangeGenericMaterialPipeline::CreateMaterialExpressionForInput(UInterchangeMaterialFactoryNode* MaterialFactoryNode, const UInterchangeShaderNode* ShaderNode, const FString& InputName, const FString& ParentUid)
 {
+	// Make sure we don't create an expression for an input if it already has one
+	if (UInterchangeShaderPortsAPI::HasInput(MaterialFactoryNode, *InputName))
+	{
+		TTuple<UInterchangeMaterialExpressionFactoryNode*, FString>{};
+	}
+
 	// If we have a connection
 	// - Create material expression for the connected shader node
 	//
@@ -724,12 +742,10 @@ UInterchangeMaterialFactoryNode* UInterchangeGenericMaterialPipeline::CreateMate
 
 	if (!HandlePhongModel(ShaderGraphNode, MaterialFactoryNode))
 	{
-		if (!HandleLambertModel(ShaderGraphNode, MaterialFactoryNode))
-		{
-			HandlePBRModel(ShaderGraphNode, MaterialFactoryNode);
-		}
+		HandleLambertModel(ShaderGraphNode, MaterialFactoryNode);
 	}
 
+	HandlePBRModel(ShaderGraphNode, MaterialFactoryNode); // Always process the PBR parameters. If they were already assigned from Phong or Lambert, they will be ignored.
 	HandleCommonParameters(ShaderGraphNode, MaterialFactoryNode);
 
 	return MaterialFactoryNode;
@@ -826,7 +842,7 @@ UInterchangeMaterialInstanceFactoryNode* UInterchangeGenericMaterialPipeline::Cr
 	return MaterialInstanceFactoryNode;
 }
 
-TVariant<FString, FLinearColor, float> UInterchangeGenericMaterialPipeline::VisitShaderNode(const UInterchangeShaderNode* ShaderNode)
+TVariant<FString, FLinearColor, float> UInterchangeGenericMaterialPipeline::VisitShaderNode(const UInterchangeShaderNode* ShaderNode) const
 {
 	using namespace UE::Interchange::Materials::Standard::Nodes;
 
@@ -839,6 +855,14 @@ TVariant<FString, FLinearColor, float> UInterchangeGenericMaterialPipeline::Visi
 		{
 			return VisitTextureSampleNode(ShaderNode);
 		}
+		else if (*ShaderType == Lerp::Name)
+		{
+			return VisitLerpNode(ShaderNode);
+		}
+		else if (*ShaderType == Multiply::Name)
+		{
+			return VisitMultiplyNode(ShaderNode);
+		}
 	}
 
 	{
@@ -848,42 +872,152 @@ TVariant<FString, FLinearColor, float> UInterchangeGenericMaterialPipeline::Visi
 		if (Inputs.Num() > 0)
 		{
 			const FString& InputName = Inputs[0];
-			FString ConnectedShaderNodeUid;
-			FString OutputName;
-			if (UInterchangeShaderPortsAPI::GetInputConnection(ShaderNode, InputName, ConnectedShaderNodeUid, OutputName))
+			Result = VisitShaderInput(ShaderNode, InputName);
+		}
+	}
+
+	return Result;
+}
+
+TVariant<FString, FLinearColor, float> UInterchangeGenericMaterialPipeline::VisitShaderInput(const UInterchangeShaderNode* ShaderNode, const FString& InputName) const
+{
+	TVariant<FString, FLinearColor, float> Result;
+
+	FString ConnectedShaderNodeUid;
+	FString OutputName;
+	if (UInterchangeShaderPortsAPI::GetInputConnection(ShaderNode, InputName, ConnectedShaderNodeUid, OutputName))
+	{
+		if (UInterchangeShaderNode* ConnectedShaderNode = Cast<UInterchangeShaderNode>(BaseNodeContainer->GetNode(ConnectedShaderNodeUid)))
+		{
+			Result = VisitShaderNode(ConnectedShaderNode);
+		}
+	}
+	else
+	{
+		switch(UInterchangeShaderPortsAPI::GetInputType(ShaderNode, InputName))
+		{
+		case UE::Interchange::EAttributeTypes::Float:
 			{
-				if (UInterchangeShaderNode* ConnectedShaderNode = Cast<UInterchangeShaderNode>(BaseNodeContainer->GetNode(ConnectedShaderNodeUid)))
-				{
-					Result = VisitShaderNode(ConnectedShaderNode);
-				}
+				float InputValue = 0.f;
+				ShaderNode->GetFloatAttribute(UInterchangeShaderPortsAPI::MakeInputValueKey(InputName), InputValue);
+				Result.Set<float>(InputValue);
 			}
-			else
+			break;
+		case UE::Interchange::EAttributeTypes::LinearColor:
 			{
-				switch(UInterchangeShaderPortsAPI::GetInputType(ShaderNode, InputName))
-				{
-				case UE::Interchange::EAttributeTypes::Float:
-					{
-						float InputValue = 0.f;
-						ShaderNode->GetFloatAttribute(UInterchangeShaderPortsAPI::MakeInputValueKey(InputName), InputValue);
-						Result.Set<float>(InputValue);
-					}
-					break;
-				case UE::Interchange::EAttributeTypes::LinearColor:
-					{
-						FLinearColor InputValue = FLinearColor::White;
-						ShaderNode->GetLinearColorAttribute(UInterchangeShaderPortsAPI::MakeInputValueKey(InputName), InputValue);
-						Result.Set<FLinearColor>(InputValue);
-					}
-					break;
-				}
+				FLinearColor InputValue = FLinearColor::White;
+				ShaderNode->GetLinearColorAttribute(UInterchangeShaderPortsAPI::MakeInputValueKey(InputName), InputValue);
+				Result.Set<FLinearColor>(InputValue);
 			}
+			break;
 		}
 	}
 
 	return Result;
 }
 
-TVariant<FString, FLinearColor, float> UInterchangeGenericMaterialPipeline::VisitTextureSampleNode(const UInterchangeShaderNode* ShaderNode)
+TVariant<FString, FLinearColor, float> UInterchangeGenericMaterialPipeline::VisitLerpNode(const UInterchangeShaderNode* ShaderNode) const
+{
+	using namespace UE::Interchange::Materials::Standard::Nodes;
+
+	TVariant<FString, FLinearColor, float> ResultA = VisitShaderInput(ShaderNode, Lerp::Inputs::A.ToString());
+	TVariant<FString, FLinearColor, float> ResultB = VisitShaderInput(ShaderNode, Lerp::Inputs::B.ToString());
+
+	TVariant<FString, FLinearColor, float> ResultFactor = VisitShaderInput(ShaderNode, Lerp::Inputs::Factor.ToString());
+
+	bool bResultAIsStrongest = true;
+
+	if (ResultFactor.IsType<float>())
+	{
+		const float Factor = ResultFactor.Get<float>();
+		bResultAIsStrongest = (Factor <= 0.5f);
+
+		// Bake the lerp into a single value
+		if (!ResultA.IsType<FString>() && !ResultB.IsType<FString>())
+		{
+			if (ResultA.IsType<float>() && ResultB.IsType<float>())
+			{
+				const float ValueA = ResultA.Get<float>();
+				const float ValueB = ResultB.Get<float>();
+
+				TVariant<FString, FLinearColor, float> Result;
+				Result.Set<float>(FMath::Lerp(ValueA, ValueB, Factor));
+				return Result;
+			}
+			else if (ResultA.IsType<FLinearColor>() && ResultB.IsType<FLinearColor>())
+			{
+				const FLinearColor ValueA = ResultA.Get<FLinearColor>();
+				const FLinearColor ValueB = ResultB.Get<FLinearColor>();
+
+				TVariant<FString, FLinearColor, float> Result;
+				Result.Set<FLinearColor>(FMath::Lerp(ValueA, ValueB, Factor));
+				return Result;
+			}
+		}
+	}
+
+	if (bResultAIsStrongest)
+	{
+		return ResultA;
+	}
+	else
+	{
+		return ResultB;
+	}
+}
+
+TVariant<FString, FLinearColor, float> UInterchangeGenericMaterialPipeline::VisitMultiplyNode(const UInterchangeShaderNode* ShaderNode) const
+{
+	using namespace UE::Interchange::Materials::Standard::Nodes;
+
+	TVariant<FString, FLinearColor, float> ResultA = VisitShaderInput(ShaderNode, Lerp::Inputs::A.ToString());
+	TVariant<FString, FLinearColor, float> ResultB = VisitShaderInput(ShaderNode, Lerp::Inputs::B.ToString());
+
+	// Bake the multiply into a single value if possible
+	if (!ResultA.IsType<FString>() && !ResultB.IsType<FString>())
+	{
+		if (ResultA.IsType<float>() && ResultB.IsType<float>())
+		{
+			const float ValueA = ResultA.Get<float>();
+			const float ValueB = ResultB.Get<float>();
+
+			TVariant<FString, FLinearColor, float> Result;
+			Result.Set<float>(ValueA * ValueB);
+			return Result;
+		}
+		else if (ResultA.IsType<FLinearColor>() && ResultB.IsType<FLinearColor>())
+		{
+			const FLinearColor ValueA = ResultA.Get<FLinearColor>();
+			const FLinearColor ValueB = ResultB.Get<FLinearColor>();
+
+			TVariant<FString, FLinearColor, float> Result;
+			Result.Set<FLinearColor>(ValueA * ValueB);
+			return Result;
+		}
+		else if (ResultA.IsType<FLinearColor>() && ResultB.IsType<float>())
+		{
+			const FLinearColor ValueA = ResultA.Get<FLinearColor>();
+			const float ValueB = ResultB.Get<float>();
+
+			TVariant<FString, FLinearColor, float> Result;
+			Result.Set<FLinearColor>(ValueA * ValueB);
+			return Result;
+		}
+		else if (ResultA.IsType<float>() && ResultB.IsType<FLinearColor>())
+		{
+			const float ValueA = ResultA.Get<float>();
+			const FLinearColor ValueB = ResultB.Get<FLinearColor>();
+
+			TVariant<FString, FLinearColor, float> Result;
+			Result.Set<FLinearColor>(ValueA * ValueB);
+			return Result;
+		}
+	}
+
+	return ResultA;
+}
+
+TVariant<FString, FLinearColor, float> UInterchangeGenericMaterialPipeline::VisitTextureSampleNode(const UInterchangeShaderNode* ShaderNode) const
 {
 	using namespace UE::Interchange::Materials::Standard::Nodes;
 
diff --git a/Engine/Plugins/Experimental/Interchange/Runtime/Source/Pipelines/Public/InterchangeGenericMaterialPipeline.h b/Engine/Plugins/Experimental/Interchange/Runtime/Source/Pipelines/Public/InterchangeGenericMaterialPipeline.h
--- a/Engine/Plugins/Experimental/Interchange/Runtime/Source/Pipelines/Public/InterchangeGenericMaterialPipeline.h
+++ b/Engine/Plugins/Experimental/Interchange/Runtime/Source/Pipelines/Public/InterchangeGenericMaterialPipeline.h
@@ -89,8 +89,17 @@ private:
 	 * Only its first input is visited as it's assumed that it's the most impactful.
 	 * The goal is to simplify a branch of a node graph to a single value, to be used for material instancing.
 	 */
-	TVariant<FString, FLinearColor, float> VisitShaderNode(const UInterchangeShaderNode* ShaderNode);
-	TVariant<FString, FLinearColor, float> VisitTextureSampleNode(const UInterchangeShaderNode* ShaderNode);
+	TVariant<FString, FLinearColor, float> VisitShaderNode(const UInterchangeShaderNode* ShaderNode) const;
+	TVariant<FString, FLinearColor, float> VisitShaderInput(const UInterchangeShaderNode* ShaderNode, const FString& InputName) const;
+
+	/** 
+	 * Returns the strongest value in a lerp.
+	 * If we're lerping between scalars or colors, the lerp result will get computed and returned.
+	 * If we're lerping between textures, the strongest one is returned based on the lerp factor.
+	 */
+	TVariant<FString, FLinearColor, float> VisitLerpNode(const UInterchangeShaderNode* ShaderNode) const;
+	TVariant<FString, FLinearColor, float> VisitMultiplyNode(const UInterchangeShaderNode* ShaderNode) const;
+	TVariant<FString, FLinearColor, float> VisitTextureSampleNode(const UInterchangeShaderNode* ShaderNode) const;
 
 private:
 	bool bParsingForNormalInput;
