Interchange
- Moved the texture processing from the generic assets pipeline to its own generic texture pipeline.

diff --git a/Engine/Plugins/Experimental/Interchange/Runtime/Source/Pipelines/Private/InterchangeGenericAssetsPipeline.cpp b/Engine/Plugins/Experimental/Interchange/Runtime/Source/Pipelines/Private/InterchangeGenericAssetsPipeline.cpp
--- a/Engine/Plugins/Experimental/Interchange/Runtime/Source/Pipelines/Private/InterchangeGenericAssetsPipeline.cpp
+++ b/Engine/Plugins/Experimental/Interchange/Runtime/Source/Pipelines/Private/InterchangeGenericAssetsPipeline.cpp
@@ -4,6 +4,8 @@
 
 #include "Animation/Skeleton.h"
 #include "CoreMinimal.h"
+#include "InterchangeGenericMaterialPipeline.h"
+#include "InterchangeGenericTexturePipeline.h"
 #include "InterchangeMaterialFactoryNode.h"
 #include "InterchangeMeshNode.h"
 #include "InterchangePipelineLog.h"
@@ -15,8 +17,6 @@
 #include "InterchangeSourceData.h"
 #include "InterchangeStaticMeshFactoryNode.h"
 #include "InterchangeStaticMeshLodDataNode.h"
-#include "InterchangeTextureFactoryNode.h"
-#include "InterchangeTextureNode.h"
 #include "Misc/Paths.h"
 #include "Nodes/InterchangeBaseNode.h"
 #include "Nodes/InterchangeBaseNodeContainer.h"
@@ -27,6 +27,7 @@
 UInterchangeGenericAssetsPipeline::UInterchangeGenericAssetsPipeline()
 {
 	MaterialPipeline = CreateDefaultSubobject<UInterchangeGenericMaterialPipeline>("MaterialPipeline");
+	TexturePipeline = CreateDefaultSubobject<UInterchangeGenericTexturePipeline>("TexturePipeline");
 }
 
 void UInterchangeGenericAssetsPipeline::PreDialogCleanup(const FName PipelineStackName)
@@ -52,30 +53,10 @@ void UInterchangeGenericAssetsPipeline::ExecutePreImportPipeline(UInterchangeBas
 	}
 
 	PipelineMeshesUtilities = UInterchangePipelineMeshesUtilities::CreateInterchangePipelineMeshesUtilities(BaseNodeContainer);
-	
-	//Find all translated node we need for this pipeline
-	BaseNodeContainer->IterateNodes([this](const FString& NodeUid, UInterchangeBaseNode* Node)
-	{
-		switch(Node->GetNodeContainerType())
-		{
-			case EInterchangeNodeContainerType::TranslatedAsset:
-			{
-				if (UInterchangeTextureNode* TextureNode = Cast<UInterchangeTextureNode>(Node))
-				{
-					TextureNodes.Add(TextureNode);
-				}
-			}
-			break;
-		}
-	});
-	
-	if (bImportTextures)
+
+	if (TexturePipeline)
 	{
-		//import textures
-		for (const UInterchangeTextureNode* TextureNode : TextureNodes)
-		{
-			HandleCreationOfTextureFactoryNode(TextureNode);
-		}
+		TexturePipeline->ExecutePreImportPipeline(InBaseNodeContainer, InSourceDatas);
 	}
 
 	if (MaterialPipeline)
@@ -244,7 +225,10 @@ void UInterchangeGenericAssetsPipeline::ExecutePostImportPipeline(const UInterch
 	//Finish the physics asset import, it need the skeletal mesh render data to create the physics collision geometry
 	PostImportPhysicsAssetImport(CreatedAsset, Node);
 
-	PostImportTextureAssetImport(CreatedAsset, bIsAReimport);
+	if (TexturePipeline)
+	{
+		TexturePipeline->ExecutePostImportPipeline(InBaseNodeContainer, NodeKey, CreatedAsset, bIsAReimport);
+	}
 }
 
 void UInterchangeGenericAssetsPipeline::ImplementUseSourceNameForAssetOption()
diff --git a/Engine/Plugins/Experimental/Interchange/Runtime/Source/Pipelines/Private/InterchangeGenericTexturePipeline.cpp b/Engine/Plugins/Experimental/Interchange/Runtime/Source/Pipelines/Private/InterchangeGenericTexturePipeline.cpp
--- a/Engine/Plugins/Experimental/Interchange/Runtime/Source/Pipelines/Private/InterchangeGenericTexturePipeline.cpp
+++ b/Engine/Plugins/Experimental/Interchange/Runtime/Source/Pipelines/Private/InterchangeGenericTexturePipeline.cpp
@@ -1,11 +1,11 @@
 // Copyright Epic Games, Inc. All Rights Reserved. 
 
-#include "InterchangeGenericAssetsPipeline.h"
+#include "InterchangeGenericTexturePipeline.h"
 
-#include "CoreMinimal.h"
 #include "Engine/Texture.h"
 #include "Engine/TextureCube.h"
 #include "Engine/TextureLightProfile.h"
+#include "InterchangePipelineLog.h"
 #include "InterchangeTexture2DArrayFactoryNode.h"
 #include "InterchangeTexture2DArrayFactoryNode.h"
 #include "InterchangeTexture2DArrayNode.h"
@@ -17,7 +17,6 @@
 #include "InterchangeTextureCubeNode.h"
 #include "InterchangeTextureCubeNode.h"
 #include "InterchangeTextureFactoryNode.h"
-#include "InterchangeTextureFactoryNode.h"
 #include "InterchangeTextureLightProfileFactoryNode.h"
 #include "InterchangeTextureLightProfileFactoryNode.h"
 #include "InterchangeTextureLightProfileNode.h"
@@ -101,7 +100,65 @@ namespace UE::Interchange::Private
 #endif
 }
 
-UInterchangeTextureFactoryNode* UInterchangeGenericAssetsPipeline::HandleCreationOfTextureFactoryNode(const UInterchangeTextureNode* TextureNode)
+void UInterchangeGenericTexturePipeline::ExecutePreImportPipeline(UInterchangeBaseNodeContainer* InBaseNodeContainer, const TArray<UInterchangeSourceData*>& InSourceDatas)
+{
+	if (!InBaseNodeContainer)
+	{
+		UE_LOG(LogInterchangePipeline, Warning, TEXT("UInterchangeGenericTexturePipeline: Cannot execute pre-import pipeline because InBaseNodeContrainer is null"));
+		return;
+	}
+
+	BaseNodeContainer = InBaseNodeContainer;
+	SourceDatas.Empty(InSourceDatas.Num());
+	for (const UInterchangeSourceData* SourceData : InSourceDatas)
+	{
+		SourceDatas.Add(SourceData);
+	}
+	
+	//Find all translated node we need for this pipeline
+	BaseNodeContainer->IterateNodes([this](const FString& NodeUid, UInterchangeBaseNode* Node)
+	{
+		switch(Node->GetNodeContainerType())
+		{
+			case EInterchangeNodeContainerType::TranslatedAsset:
+			{
+				if (UInterchangeTextureNode* TextureNode = Cast<UInterchangeTextureNode>(Node))
+				{
+					TextureNodes.Add(TextureNode);
+				}
+			}
+			break;
+		}
+	});
+
+	if (bImportTextures)
+	{
+		for (const UInterchangeTextureNode* TextureNode : TextureNodes)
+		{
+			HandleCreationOfTextureFactoryNode(TextureNode);
+		}
+	}
+}
+
+void UInterchangeGenericTexturePipeline::ExecutePostImportPipeline(const UInterchangeBaseNodeContainer* InBaseNodeContainer, const FString& NodeKey, UObject* CreatedAsset, bool bIsAReimport)
+{
+	//We do not use the provided base container since ExecutePreImportPipeline cache it
+	//We just make sure the same one is pass in parameter
+	if (!InBaseNodeContainer || !ensure(BaseNodeContainer == InBaseNodeContainer) || !CreatedAsset)
+	{
+		return;
+	}
+
+	UInterchangeBaseNode* Node = BaseNodeContainer->GetNode(NodeKey);
+	if (!Node)
+	{
+		return;
+	}
+
+	PostImportTextureAssetImport(CreatedAsset, bIsAReimport);
+}
+
+UInterchangeTextureFactoryNode* UInterchangeGenericTexturePipeline::HandleCreationOfTextureFactoryNode(const UInterchangeTextureNode* TextureNode)
 {
 	UClass* FactoryClass = UE::Interchange::Private::GetDefaultFactoryClassFromTextureNodeClass(TextureNode->GetClass());
 
@@ -122,7 +179,7 @@ UInterchangeTextureFactoryNode* UInterchangeGenericAssetsPipeline::HandleCreatio
 	return CreateTextureFactoryNode(TextureNode, FactoryClass);
 }
 
-UInterchangeTextureFactoryNode* UInterchangeGenericAssetsPipeline::CreateTextureFactoryNode(const UInterchangeTextureNode* TextureNode, const TSubclassOf<UInterchangeTextureFactoryNode>& FactorySubclass)
+UInterchangeTextureFactoryNode* UInterchangeGenericTexturePipeline::CreateTextureFactoryNode(const UInterchangeTextureNode* TextureNode, const TSubclassOf<UInterchangeTextureFactoryNode>& FactorySubclass)
 {
 	FString DisplayLabel = TextureNode->GetDisplayLabel();
 	FString NodeUid = UInterchangeTextureFactoryNode::GetTextureFactoryNodeUidFromTextureNodeUid(TextureNode->GetUniqueID());
@@ -168,7 +225,7 @@ UInterchangeTextureFactoryNode* UInterchangeGenericAssetsPipeline::CreateTexture
 	return TextureFactoryNode;
 }
 
-void UInterchangeGenericAssetsPipeline::PostImportTextureAssetImport(UObject* CreatedAsset, bool bIsAReimport)
+void UInterchangeGenericTexturePipeline::PostImportTextureAssetImport(UObject* CreatedAsset, bool bIsAReimport)
 {
 #if WITH_EDITOR
 	if (!bIsAReimport && bDetectNormalMapTexture)
diff --git a/Engine/Plugins/Experimental/Interchange/Runtime/Source/Pipelines/Public/InterchangeGenericAssetsPipeline.h b/Engine/Plugins/Experimental/Interchange/Runtime/Source/Pipelines/Public/InterchangeGenericAssetsPipeline.h
--- a/Engine/Plugins/Experimental/Interchange/Runtime/Source/Pipelines/Public/InterchangeGenericAssetsPipeline.h
+++ b/Engine/Plugins/Experimental/Interchange/Runtime/Source/Pipelines/Public/InterchangeGenericAssetsPipeline.h
@@ -3,7 +3,6 @@
 #pragma once
 
 #include "CoreMinimal.h"
-#include "InterchangeGenericMaterialPipeline.h"
 #include "InterchangePipelineBase.h"
 #include "InterchangeSourceData.h"
 #include "Nodes/InterchangeBaseNode.h"
@@ -13,6 +12,8 @@
 
 #include "InterchangeGenericAssetsPipeline.generated.h"
 
+class UInterchangeGenericMaterialPipeline;
+class UInterchangeGenericTexturePipeline;
 class UInterchangeMeshNode;
 class UInterchangePipelineMeshesUtilities;
 class UInterchangeResult;
@@ -22,12 +23,6 @@ class UInterchangeSkeletalMeshLodDataNode;
 class UInterchangeSkeletonFactoryNode;
 class UInterchangeStaticMeshFactoryNode;
 class UInterchangeStaticMeshLodDataNode;
-class UInterchangeTexture2DArrayFactoryNode;
-class UInterchangeTexture2DArrayNode;
-class UInterchangeTextureCubeFactoryNode;
-class UInterchangeTextureCubeNode;
-class UInterchangeTextureFactoryNode;
-class UInterchangeTextureNode;
 template<class T> class TSubclassOf;
 
 /** Force mesh type, if user want to import all meshes as one type*/
@@ -182,26 +177,8 @@ public:
 	//////	TEXTURES_CATEGORY Properties //////
 
 
-	/** If enabled, imports the texture assets found in the sources. */
-	UPROPERTY(EditAnywhere, BlueprintReadWrite, Category = "Textures")
-	bool bImportTextures = true;
-
-#if WITH_EDITORONLY_DATA
-	/** 
-	 * If enable, after a new import a test will be run to see if the texture is a normal map
-	 * If the texture is a normal map the SRG, CompressionSettings and LODGroup settings will be adjusted.
-	 */
-	UPROPERTY(EditAnywhere, BlueprintReadWrite, Category = "Textures")
-	bool bDetectNormalMapTexture = true;
-
-	/** If enabled, the texture's green channel will be inverted for normal maps. */
-	UPROPERTY(EditAnywhere, Category = "Textures")
-	bool bFlipNormalMapGreenChannel = false;
-
-	/** Specify the files type that should be imported as long/lat cubemap */
-	UPROPERTY(EditAnywhere, Category = "Textures")
-	TSet<FString> FileExtensionsToImportAsLongLatCubemap = {"hdr"};
-#endif
+	UPROPERTY(VisibleAnywhere, Instanced, Category = "Textures")
+	TObjectPtr<UInterchangeGenericTexturePipeline> TexturePipeline;
 
 	// END Pre import pipeline properties
 	//////////////////////////////////////////////////////////////////////////
@@ -225,13 +202,6 @@ private:
 
 	UInterchangeBaseNodeContainer* BaseNodeContainer;
 	TArray<const UInterchangeSourceData*> SourceDatas;
-
-	
-	/** Texture translated assets nodes */
-	TArray<UInterchangeTextureNode*> TextureNodes;
-
-	/** Texture factory assets nodes */
-	TArray<UInterchangeTextureFactoryNode*> TextureFactoryNodes;
 	
 	/************************************************************************/
 	/* Skeletal mesh API BEGIN                                              */
@@ -314,17 +284,6 @@ private:
 	/** Specialize for skeletalmesh */
 	void ImplementUseSourceNameForAssetOptionSkeletalMesh(const int32 MeshesAndAnimsImportedNodeCount);
 
-	/************************************************************************/
-	/* Texture API BEGIN                                              */
-
-	UInterchangeTextureFactoryNode* HandleCreationOfTextureFactoryNode(const UInterchangeTextureNode* TextureNode);
-
-	UInterchangeTextureFactoryNode* CreateTextureFactoryNode(const UInterchangeTextureNode* TextureNode, const TSubclassOf<UInterchangeTextureFactoryNode>& FactorySubclass);
-
-	void PostImportTextureAssetImport(UObject* CreatedAsset, bool bIsAReimport);
-
-	/* Texture API END                                                */
-	/************************************************************************/
 };
 
 
diff --git a/Engine/Plugins/Experimental/Interchange/Runtime/Source/Pipelines/Public/InterchangeGenericTexturePipeline.h b/Engine/Plugins/Experimental/Interchange/Runtime/Source/Pipelines/Public/InterchangeGenericTexturePipeline.h
new file mode 100644
--- /dev/null
+++ b/Engine/Plugins/Experimental/Interchange/Runtime/Source/Pipelines/Public/InterchangeGenericTexturePipeline.h
@@ -0,0 +1,67 @@
+// Copyright Epic Games, Inc. All Rights Reserved.
+
+#pragma once
+
+#include "InterchangePipelineBase.h"
+#include "InterchangeSourceData.h"
+#include "Nodes/InterchangeBaseNodeContainer.h"
+#include "UObject/Object.h"
+#include "UObject/ObjectMacros.h"
+
+#include "InterchangeGenericTexturePipeline.generated.h"
+
+class UInterchangeTextureFactoryNode;
+class UInterchangeTextureNode;
+
+UCLASS(BlueprintType, Experimental)
+class INTERCHANGEPIPELINES_API UInterchangeGenericTexturePipeline : public UInterchangePipelineBase
+{
+	GENERATED_BODY()
+
+public:
+	/** If enabled, imports the texture assets found in the sources. */
+	UPROPERTY(EditAnywhere, BlueprintReadWrite, Category = "Textures")
+	bool bImportTextures = true;
+
+#if WITH_EDITORONLY_DATA
+	/** 
+	 * If enable, after a new import a test will be run to see if the texture is a normal map
+	 * If the texture is a normal map the SRG, CompressionSettings and LODGroup settings will be adjusted.
+	 */
+	UPROPERTY(EditAnywhere, BlueprintReadWrite, Category = "Textures")
+	bool bDetectNormalMapTexture = true;
+
+	/** If enabled, the texture's green channel will be inverted for normal maps. */
+	UPROPERTY(EditAnywhere, Category = "Textures")
+	bool bFlipNormalMapGreenChannel = false;
+
+	/** Specify the files type that should be imported as long/lat cubemap */
+	UPROPERTY(EditAnywhere, Category = "Textures")
+	TSet<FString> FileExtensionsToImportAsLongLatCubemap = {"hdr"};
+#endif
+
+public:
+	virtual void ExecutePreImportPipeline(UInterchangeBaseNodeContainer* InBaseNodeContainer, const TArray<UInterchangeSourceData*>& InSourceDatas) override;
+	virtual void ExecutePostImportPipeline(const UInterchangeBaseNodeContainer* InBaseNodeContainer, const FString& NodeKey, UObject* CreatedAsset, bool bIsAReimport) override;
+
+protected:
+	UInterchangeTextureFactoryNode* HandleCreationOfTextureFactoryNode(const UInterchangeTextureNode* TextureNode);
+
+	UInterchangeTextureFactoryNode* CreateTextureFactoryNode(const UInterchangeTextureNode* TextureNode, const TSubclassOf<UInterchangeTextureFactoryNode>& FactorySubclass);
+
+	void PostImportTextureAssetImport(UObject* CreatedAsset, bool bIsAReimport);
+
+private:
+	UPROPERTY()
+	TObjectPtr<UInterchangeBaseNodeContainer> BaseNodeContainer;
+
+	TArray<const UInterchangeSourceData*> SourceDatas;
+	
+	/** Texture translated assets nodes */
+	TArray<UInterchangeTextureNode*> TextureNodes;
+
+	/** Texture factory assets nodes */
+	TArray<UInterchangeTextureFactoryNode*> TextureFactoryNodes;
+};
+
+
