Interchange: Fixed compilation warning by adding missing return statement.

diff --git a/Engine/Plugins/Experimental/Interchange/Runtime/Source/Pipelines/Private/InterchangeGenericMaterialPipeline.cpp b/Engine/Plugins/Experimental/Interchange/Runtime/Source/Pipelines/Private/InterchangeGenericMaterialPipeline.cpp
--- a/Engine/Plugins/Experimental/Interchange/Runtime/Source/Pipelines/Private/InterchangeGenericMaterialPipeline.cpp
+++ b/Engine/Plugins/Experimental/Interchange/Runtime/Source/Pipelines/Private/InterchangeGenericMaterialPipeline.cpp
@@ -700,7 +700,7 @@ TTuple<UInterchangeMaterialExpressionFactoryNode*, FString> UInterchangeGenericM
 	// Make sure we don't create an expression for an input if it already has one
 	if (UInterchangeShaderPortsAPI::HasInput(MaterialFactoryNode, *InputName))
 	{
-		TTuple<UInterchangeMaterialExpressionFactoryNode*, FString>{};
+		return TTuple<UInterchangeMaterialExpressionFactoryNode*, FString>{};
 	}
 
 	// If we have a connection
