Scene node will now only serialize the local transform, all SetGlobalTransform functions has been removed.
Fix scene import when there is no specified factory
The scene node specialtype attribute can now be a "Transform" type. This allow the generic scene pipeline to found group transform node.

#rnx

diff --git a/Engine/Plugins/Experimental/Interchange/Runtime/Source/Dispatcher/Public/InterchangeCommands.h b/Engine/Plugins/Experimental/Interchange/Runtime/Source/Dispatcher/Public/InterchangeCommands.h
--- a/Engine/Plugins/Experimental/Interchange/Runtime/Source/Dispatcher/Public/InterchangeCommands.h
+++ b/Engine/Plugins/Experimental/Interchange/Runtime/Source/Dispatcher/Public/InterchangeCommands.h
@@ -17,7 +17,7 @@ namespace UE
 		{
 		public:
 			//Major version should be updated when there is an existing API that has been change.
-			static int32 GetMajor() { return 3; }
+			static int32 GetMajor() { return 4; }
 			//Minor version must be updated when there is an addition to the API.
 			static int32 GetMinor() { return 0; }
 			//Patch version should be update if there is some bug fixes in the private code.
diff --git a/Engine/Plugins/Experimental/Interchange/Runtime/Source/Import/Private/Gltf/InterchangeGltfTranslator.cpp b/Engine/Plugins/Experimental/Interchange/Runtime/Source/Import/Private/Gltf/InterchangeGltfTranslator.cpp
--- a/Engine/Plugins/Experimental/Interchange/Runtime/Source/Import/Private/Gltf/InterchangeGltfTranslator.cpp
+++ b/Engine/Plugins/Experimental/Interchange/Runtime/Source/Import/Private/Gltf/InterchangeGltfTranslator.cpp
@@ -34,14 +34,6 @@ void UInterchangeGltfTranslator::HandleGltfNode( UInterchangeBaseNodeContainer&
 	constexpr float MetersToCentimeters = 100.f;
 	Transform.SetTranslation( Transform.GetTranslation() * MetersToCentimeters );
 
-	if ( ParentSceneNode )
-	{
-		FTransform ParentTransform;
-		ParentSceneNode->GetCustomGlobalTransform( ParentTransform );
-
-		Transform = Transform * ParentTransform;
-	}
-
 	switch ( GltfNode.Type )
 	{
 		case GLTF::FNode::EType::Mesh:
@@ -84,7 +76,7 @@ void UInterchangeGltfTranslator::HandleGltfNode( UInterchangeBaseNodeContainer&
 		}
 	}
 
-	InterchangeSceneNode->SetCustomGlobalTransform( Transform );
+	InterchangeSceneNode->SetCustomLocalTransform(&NodeContainer, Transform );
 
 	if ( !ParentNodeUid.IsEmpty() )
 	{
diff --git a/Engine/Plugins/Experimental/Interchange/Runtime/Source/Import/Private/Mesh/InterchangeSkeletalMeshFactory.cpp b/Engine/Plugins/Experimental/Interchange/Runtime/Source/Import/Private/Mesh/InterchangeSkeletalMeshFactory.cpp
--- a/Engine/Plugins/Experimental/Interchange/Runtime/Source/Import/Private/Mesh/InterchangeSkeletalMeshFactory.cpp
+++ b/Engine/Plugins/Experimental/Interchange/Runtime/Source/Import/Private/Mesh/InterchangeSkeletalMeshFactory.cpp
@@ -431,11 +431,11 @@ namespace UE
 					}
 					
 					FTransform JointBindPoseGlobalTransform;
-					if (!JointNode->GetCustomBindPoseGlobalTransform(JointBindPoseGlobalTransform))
+					if (!JointNode->GetCustomBindPoseGlobalTransform(NodeContainer, JointBindPoseGlobalTransform))
 					{
 						//If there is no bind pose we will fall back on the CustomGlobalTransform of the link.
 						//We ensure here because any scenenode should have a valid CustomGlobalTransform.
-						if (!ensure(JointNode->GetCustomGlobalTransform(JointBindPoseGlobalTransform)))
+						if (!ensure(JointNode->GetCustomGlobalTransform(NodeContainer, JointBindPoseGlobalTransform)))
 						{
 							//No value to convert from, skip this joint.
 							continue;
@@ -443,7 +443,7 @@ namespace UE
 					}
 
 					FTransform JointTimeZeroGlobalTransform;
-					if (!JointNode->GetCustomTimeZeroGlobalTransform(JointTimeZeroGlobalTransform))
+					if (!JointNode->GetCustomTimeZeroGlobalTransform(NodeContainer, JointTimeZeroGlobalTransform))
 					{
 						//If there is no time zero global transform we cannot set the bind pose to time zero.
 						//We must skip this joint.
@@ -1246,9 +1246,9 @@ UObject* UInterchangeSkeletalMeshFactory::CreateAsset(const FCreateAssetParams&
 					MeshReference.MeshNode = Cast<UInterchangeMeshNode>(Arguments.NodeContainer->GetNode(MeshDependencyUid));
 					//Cache the scene node global matrix, we will use this matrix to bake the vertices, add the node geometric mesh offset to this matrix to bake it properly
 					FTransform SceneNodeGlobalTransform;
-					if (!bUseTimeZeroAsBindPose || !MeshReference.SceneNode->GetCustomTimeZeroGlobalTransform(SceneNodeGlobalTransform))
+					if (!bUseTimeZeroAsBindPose || !MeshReference.SceneNode->GetCustomTimeZeroGlobalTransform(Arguments.NodeContainer, SceneNodeGlobalTransform))
 					{
-						ensure(MeshReference.SceneNode->GetCustomGlobalTransform(SceneNodeGlobalTransform));
+						ensure(MeshReference.SceneNode->GetCustomGlobalTransform(Arguments.NodeContainer, SceneNodeGlobalTransform));
 					}
 					FTransform SceneNodeGeometricTransform;
 					if(MeshReference.SceneNode->GetCustomGeometricTransform(SceneNodeGeometricTransform))
diff --git a/Engine/Plugins/Experimental/Interchange/Runtime/Source/Import/Private/Mesh/InterchangeStaticMeshFactory.cpp b/Engine/Plugins/Experimental/Interchange/Runtime/Source/Import/Private/Mesh/InterchangeStaticMeshFactory.cpp
--- a/Engine/Plugins/Experimental/Interchange/Runtime/Source/Import/Private/Mesh/InterchangeStaticMeshFactory.cpp
+++ b/Engine/Plugins/Experimental/Interchange/Runtime/Source/Import/Private/Mesh/InterchangeStaticMeshFactory.cpp
@@ -209,7 +209,7 @@ UObject* UInterchangeStaticMeshFactory::CreateAsset(const FCreateAssetParams& Ar
 
 					// Cache the scene node global matrix, we will use this matrix to bake the vertices
 					FTransform SceneNodeGlobalTransform;
-					if (MeshReference.SceneNode->GetCustomGlobalTransform(SceneNodeGlobalTransform))
+					if (MeshReference.SceneNode->GetCustomGlobalTransform(Arguments.NodeContainer, SceneNodeGlobalTransform))
 					{
 						MeshReference.SceneGlobalTransform = SceneNodeGlobalTransform;
 					}
diff --git a/Engine/Plugins/Experimental/Interchange/Runtime/Source/Nodes/Private/InterchangeSceneNode.cpp b/Engine/Plugins/Experimental/Interchange/Runtime/Source/Nodes/Private/InterchangeSceneNode.cpp
--- a/Engine/Plugins/Experimental/Interchange/Runtime/Source/Nodes/Private/InterchangeSceneNode.cpp
+++ b/Engine/Plugins/Experimental/Interchange/Runtime/Source/Nodes/Private/InterchangeSceneNode.cpp
@@ -1,6 +1,7 @@
 // Copyright Epic Games, Inc. All Rights Reserved.
 
 #include "InterchangeSceneNode.h"
+#include "Nodes/InterchangeBaseNodeContainer.h"
 
 //Interchange namespace
 namespace UE
@@ -19,6 +20,12 @@ namespace UE
 			static FString MaterialDependencyUids_BaseKey = TEXT("__MaterialDependencyUidsBaseKey__");
 			return MaterialDependencyUids_BaseKey;
 		}
+		
+		const FString& FSceneNodeStaticData::GetTransformSpecializeTypeString()
+		{
+			static FString TransformSpecializeTypeString = TEXT("Transform");
+			return TransformSpecializeTypeString;
+		}
 
 		const FString& FSceneNodeStaticData::GetJointSpecializeTypeString()
 		{
@@ -98,7 +105,6 @@ FString UInterchangeSceneNode::GetAttributeCategory(const UE::Interchange::FAttr
 		return FString(TEXT("MaterialDependencies"));
 	}
 	else if (NodeAttributeKey == Macro_CustomLocalTransformKey
-		|| NodeAttributeKey == Macro_CustomGlobalTransformKey
 		|| NodeAttributeKey == Macro_CustomAssetInstanceUidKey)
 	{
 		return FString(TEXT("Scene"));
@@ -187,19 +193,15 @@ bool UInterchangeSceneNode::GetCustomLocalTransform(FTransform& AttributeValue)
 	IMPLEMENT_NODE_ATTRIBUTE_GETTER(LocalTransform, FTransform);
 }
 
-bool UInterchangeSceneNode::SetCustomLocalTransform(const FTransform& AttributeValue)
+bool UInterchangeSceneNode::SetCustomLocalTransform(const UInterchangeBaseNodeContainer* BaseNodeContainer, const FTransform& AttributeValue)
 {
+	ResetGlobalTransformCachesOfNodeAndAllChildren(BaseNodeContainer, this);
 	IMPLEMENT_NODE_ATTRIBUTE_SETTER_NODELEGATE(LocalTransform, FTransform);
 }
 
-bool UInterchangeSceneNode::GetCustomGlobalTransform(FTransform& AttributeValue) const
-{
-	IMPLEMENT_NODE_ATTRIBUTE_GETTER(GlobalTransform, FTransform);
-}
-
-bool UInterchangeSceneNode::SetCustomGlobalTransform(const FTransform& AttributeValue)
+bool UInterchangeSceneNode::GetCustomGlobalTransform(const UInterchangeBaseNodeContainer* BaseNodeContainer, FTransform& AttributeValue, bool bForceRecache /*= false*/) const
 {
-	IMPLEMENT_NODE_ATTRIBUTE_SETTER_NODELEGATE(GlobalTransform, FTransform);
+	return GetGlobalTransformInternal(Macro_CustomLocalTransformKey, CacheGlobalTransform, BaseNodeContainer, AttributeValue, bForceRecache);
 }
 
 bool UInterchangeSceneNode::GetCustomBindPoseLocalTransform(FTransform& AttributeValue) const
@@ -207,19 +209,15 @@ bool UInterchangeSceneNode::GetCustomBindPoseLocalTransform(FTransform& Attribut
 	IMPLEMENT_NODE_ATTRIBUTE_GETTER(BindPoseLocalTransform, FTransform);
 }
 
-bool UInterchangeSceneNode::SetCustomBindPoseLocalTransform(const FTransform& AttributeValue)
+bool UInterchangeSceneNode::SetCustomBindPoseLocalTransform(const UInterchangeBaseNodeContainer* BaseNodeContainer, const FTransform& AttributeValue)
 {
+	ResetGlobalTransformCachesOfNodeAndAllChildren(BaseNodeContainer, this);
 	IMPLEMENT_NODE_ATTRIBUTE_SETTER_NODELEGATE(BindPoseLocalTransform, FTransform);
 }
 
-bool UInterchangeSceneNode::GetCustomBindPoseGlobalTransform(FTransform& AttributeValue) const
+bool UInterchangeSceneNode::GetCustomBindPoseGlobalTransform(const UInterchangeBaseNodeContainer* BaseNodeContainer, FTransform& AttributeValue, bool bForceRecache /*= false*/) const
 {
-	IMPLEMENT_NODE_ATTRIBUTE_GETTER(BindPoseGlobalTransform, FTransform);
-}
-
-bool UInterchangeSceneNode::SetCustomBindPoseGlobalTransform(const FTransform& AttributeValue)
-{
-	IMPLEMENT_NODE_ATTRIBUTE_SETTER_NODELEGATE(BindPoseGlobalTransform, FTransform);
+	return GetGlobalTransformInternal(Macro_CustomBindPoseLocalTransformKey, CacheBindPoseGlobalTransform, BaseNodeContainer, AttributeValue, bForceRecache);
 }
 
 bool UInterchangeSceneNode::GetCustomTimeZeroLocalTransform(FTransform& AttributeValue) const
@@ -227,19 +225,15 @@ bool UInterchangeSceneNode::GetCustomTimeZeroLocalTransform(FTransform& Attribut
 	IMPLEMENT_NODE_ATTRIBUTE_GETTER(TimeZeroLocalTransform, FTransform);
 }
 
-bool UInterchangeSceneNode::SetCustomTimeZeroLocalTransform(const FTransform& AttributeValue)
+bool UInterchangeSceneNode::SetCustomTimeZeroLocalTransform(const UInterchangeBaseNodeContainer* BaseNodeContainer, const FTransform& AttributeValue)
 {
+	ResetGlobalTransformCachesOfNodeAndAllChildren(BaseNodeContainer, this);
 	IMPLEMENT_NODE_ATTRIBUTE_SETTER_NODELEGATE(TimeZeroLocalTransform, FTransform);
 }
 
-bool UInterchangeSceneNode::GetCustomTimeZeroGlobalTransform(FTransform& AttributeValue) const
-{
-	IMPLEMENT_NODE_ATTRIBUTE_GETTER(TimeZeroGlobalTransform, FTransform);
-}
-
-bool UInterchangeSceneNode::SetCustomTimeZeroGlobalTransform(const FTransform& AttributeValue)
+bool UInterchangeSceneNode::GetCustomTimeZeroGlobalTransform(const UInterchangeBaseNodeContainer* BaseNodeContainer, FTransform& AttributeValue, bool bForceRecache /*= false*/) const
 {
-	IMPLEMENT_NODE_ATTRIBUTE_SETTER_NODELEGATE(TimeZeroGlobalTransform, FTransform);
+	return GetGlobalTransformInternal(Macro_CustomTimeZeroLocalTransformKey, CacheTimeZeroGlobalTransform, BaseNodeContainer, AttributeValue, bForceRecache);
 }
 
 bool UInterchangeSceneNode::GetCustomGeometricTransform(FTransform& AttributeValue) const
@@ -262,3 +256,87 @@ bool UInterchangeSceneNode::SetCustomAssetInstanceUid(const FString& AttributeVa
 	IMPLEMENT_NODE_ATTRIBUTE_SETTER_NODELEGATE(AssetInstanceUid, FString);
 }
 
+void UInterchangeSceneNode::ResetAllGlobalTransformCaches(const UInterchangeBaseNodeContainer* BaseNodeContainer)
+{
+	BaseNodeContainer->IterateNodes([](const FString& NodeUid, UInterchangeBaseNode* Node)
+		{
+			if (UInterchangeSceneNode* SceneNode = Cast<UInterchangeSceneNode>(Node))
+			{
+				SceneNode->CacheGlobalTransform.Reset();
+				SceneNode->CacheBindPoseGlobalTransform.Reset();
+				SceneNode->CacheTimeZeroGlobalTransform.Reset();
+			}
+		});
+}
+
+void UInterchangeSceneNode::ResetGlobalTransformCachesOfNodeAndAllChildren(const UInterchangeBaseNodeContainer* BaseNodeContainer, const UInterchangeBaseNode* ParentNode)
+{
+	check(ParentNode);
+	if (const UInterchangeSceneNode* SceneNode = Cast<UInterchangeSceneNode>(ParentNode))
+	{
+		SceneNode->CacheGlobalTransform.Reset();
+		SceneNode->CacheBindPoseGlobalTransform.Reset();
+		SceneNode->CacheTimeZeroGlobalTransform.Reset();
+	}
+	TArray<FString> ChildrenUids = BaseNodeContainer->GetNodeChildrenUids(ParentNode->GetUniqueID());
+	for (const FString& ChildUid : ChildrenUids)
+	{
+		if (const UInterchangeBaseNode* ChildNode = BaseNodeContainer->GetNode(ChildUid))
+		{
+			ResetGlobalTransformCachesOfNodeAndAllChildren(BaseNodeContainer, ChildNode);
+		}
+	}
+}
+
+bool UInterchangeSceneNode::GetGlobalTransformInternal(const UE::Interchange::FAttributeKey LocalTransformKey, TOptional<FTransform>& CacheTransform, const UInterchangeBaseNodeContainer* BaseNodeContainer, FTransform& AttributeValue, bool bForceRecache) const
+{
+	if (!Attributes->ContainAttribute(LocalTransformKey))
+	{
+		return false;
+	}
+	if (bForceRecache)
+	{
+		CacheTransform.Reset();
+	}
+	if (!CacheTransform.IsSet())
+	{
+		FTransform LocalTransform;
+		UE::Interchange::FAttributeStorage::TAttributeHandle<FTransform> AttributeHandle = GetAttributeHandle<FTransform>(LocalTransformKey);
+		if (AttributeHandle.IsValid() && AttributeHandle.Get(LocalTransform) == UE::Interchange::EAttributeStorageResult::Operation_Success)
+		{
+			//Compute the Global
+			if (Attributes->ContainAttribute(UE::Interchange::FBaseNodeStaticData::ParentIDKey()))
+			{
+				FTransform GlobalParent;
+				if (const UInterchangeSceneNode* ParentSceneNode = Cast<UInterchangeSceneNode>(BaseNodeContainer->GetNode(GetParentUid())))
+				{
+					if (LocalTransformKey == Macro_CustomLocalTransformKey)
+					{
+						ParentSceneNode->GetCustomGlobalTransform(BaseNodeContainer, GlobalParent, bForceRecache);
+					}
+					else if (LocalTransformKey == Macro_CustomBindPoseLocalTransformKey)
+					{
+						ParentSceneNode->GetCustomBindPoseGlobalTransform(BaseNodeContainer, GlobalParent, bForceRecache);
+					}
+					else if (LocalTransformKey == Macro_CustomTimeZeroLocalTransformKey)
+					{
+						ParentSceneNode->GetCustomTimeZeroGlobalTransform(BaseNodeContainer, GlobalParent, bForceRecache);
+					}
+				}
+				CacheTransform = LocalTransform * GlobalParent;
+			}
+			else
+			{
+				CacheTransform = LocalTransform;
+			}
+		}
+		else
+		{
+			CacheTransform = FTransform::Identity;
+		}
+	}
+	//The cache is always valid here
+	check(CacheTransform.IsSet());
+	AttributeValue = CacheTransform.GetValue();
+	return true;
+}
diff --git a/Engine/Plugins/Experimental/Interchange/Runtime/Source/Nodes/Public/InterchangeSceneNode.h b/Engine/Plugins/Experimental/Interchange/Runtime/Source/Nodes/Public/InterchangeSceneNode.h
--- a/Engine/Plugins/Experimental/Interchange/Runtime/Source/Nodes/Public/InterchangeSceneNode.h
+++ b/Engine/Plugins/Experimental/Interchange/Runtime/Source/Nodes/Public/InterchangeSceneNode.h
@@ -3,10 +3,12 @@
 #pragma once
 
 #include "CoreMinimal.h"
+#include "Misc/Optional.h"
 #include "Nodes/InterchangeBaseNode.h"
 
 #include "InterchangeSceneNode.generated.h"
 
+class UInterchangeBaseNodeContainer;
 //Interchange namespace
 namespace UE
 {
@@ -17,6 +19,7 @@ namespace UE
 		{
 			static const FString& GetNodeSpecializeTypeBaseKey();
 			static const FString& GetMaterialDependencyUidsBaseKey();
+			static const FString& GetTransformSpecializeTypeString();
 			static const FString& GetJointSpecializeTypeString();
 			static const FString& GetLodGroupSpecializeTypeString();
 		};
@@ -89,15 +92,11 @@ public:
 
 	/** Store the default scene node local transform. */
 	UFUNCTION(BlueprintCallable, Category = "Interchange | Node | Scene")
-	bool SetCustomLocalTransform(const FTransform& AttributeValue);
+	bool SetCustomLocalTransform(const UInterchangeBaseNodeContainer* BaseNodeContainer, const FTransform& AttributeValue);
 
-	/** Return the default scene node global transform. */
+	/** Return the default scene node global transform. This value is computed with all parent local transform. */
 	UFUNCTION(BlueprintCallable, Category = "Interchange | Node | Scene")
-	bool GetCustomGlobalTransform(FTransform& AttributeValue) const;
-
-	/** Store the default scene node global transform. */
-	UFUNCTION(BlueprintCallable, Category = "Interchange | Node | Scene")
-	bool SetCustomGlobalTransform(const FTransform& AttributeValue);
+	bool GetCustomGlobalTransform(const UInterchangeBaseNodeContainer* BaseNodeContainer, FTransform& AttributeValue, bool bForceRecache = false) const;
 
 	//Bind pose transform is the transform of the joint when the binding with the mesh was done.
 	//This attribute should be set only if we have a joint.
@@ -108,15 +107,11 @@ public:
 
 	/** Store the bind pose scene node local transform. */
 	UFUNCTION(BlueprintCallable, Category = "Interchange | Node | Joint")
-	bool SetCustomBindPoseLocalTransform(const FTransform& AttributeValue);
-
-	/** Return the bind pose scene node global transform. */
-	UFUNCTION(BlueprintCallable, Category = "Interchange | Node | Joint")
-	bool GetCustomBindPoseGlobalTransform(FTransform& AttributeValue) const;
+	bool SetCustomBindPoseLocalTransform(const UInterchangeBaseNodeContainer* BaseNodeContainer, const FTransform& AttributeValue);
 
-	/** Store the bind pose scene node global transform. */
+	/** Return the bind pose scene node global transform. This value is computed with all parent bind pose local transform. */
 	UFUNCTION(BlueprintCallable, Category = "Interchange | Node | Joint")
-	bool SetCustomBindPoseGlobalTransform(const FTransform& AttributeValue);
+	bool GetCustomBindPoseGlobalTransform(const UInterchangeBaseNodeContainer* BaseNodeContainer, FTransform& AttributeValue, bool bForceRecache = false) const;
 
 	//Time zero transform is the transform of the node at time zero.
 	//This is useful when there is no bind pose or when we import rigid mesh.
@@ -127,15 +122,11 @@ public:
 
 	/** Store the time zero scene node local transform. */
 	UFUNCTION(BlueprintCallable, Category = "Interchange | Node | Joint")
-	bool SetCustomTimeZeroLocalTransform(const FTransform& AttributeValue);
-
-	/** Return the time zero scene node global transform. */
-	UFUNCTION(BlueprintCallable, Category = "Interchange | Node | Joint")
-	bool GetCustomTimeZeroGlobalTransform(FTransform& AttributeValue) const;
+	bool SetCustomTimeZeroLocalTransform(const UInterchangeBaseNodeContainer* BaseNodeContainer, const FTransform& AttributeValue);
 
-	/** Store the time zero scene node global transform. */
+	/** Return the time zero scene node global transform. This value is computed with all parent timezero local transform. */
 	UFUNCTION(BlueprintCallable, Category = "Interchange | Node | Joint")
-	bool SetCustomTimeZeroGlobalTransform(const FTransform& AttributeValue);
+	bool GetCustomTimeZeroGlobalTransform(const UInterchangeBaseNodeContainer* BaseNodeContainer, FTransform& AttributeValue, bool bForceRecache = false) const;
 
 	/** Return the geometric offset. Any mesh attach to this scene node will be offset using this transform. */
 	UFUNCTION(BlueprintCallable, Category = "Interchange | Node | Joint")
@@ -153,17 +144,27 @@ public:
 	UFUNCTION(BlueprintCallable, Category = "Interchange | Node | Scene")
 	bool SetCustomAssetInstanceUid(const FString& AttributeValue);
 
+	/** This static function make sure all the global transform caches are reset for all the UInterchangeSceneNode nodes in the UInterchangeBaseNodeContainer */
+	static void ResetAllGlobalTransformCaches(const UInterchangeBaseNodeContainer* BaseNodeContainer);
+
+	/** This static function make sure all the global transform caches are reset for all the UInterchangeSceneNode nodes children in the UInterchangeBaseNodeContainer */
+	static void ResetGlobalTransformCachesOfNodeAndAllChildren(const UInterchangeBaseNodeContainer* BaseNodeContainer, const UInterchangeBaseNode* ParentNode);
+
 private:
+
+	bool GetGlobalTransformInternal(const UE::Interchange::FAttributeKey LocalTransformKey, TOptional<FTransform>& CacheTransform, const UInterchangeBaseNodeContainer* BaseNodeContainer, FTransform& AttributeValue, bool bForceRecache) const;
+
 	//Scene Attribute Keys
 	const UE::Interchange::FAttributeKey Macro_CustomLocalTransformKey = UE::Interchange::FAttributeKey(TEXT("LocalTransform"));
-	const UE::Interchange::FAttributeKey Macro_CustomGlobalTransformKey = UE::Interchange::FAttributeKey(TEXT("GlobalTransform"));
 	const UE::Interchange::FAttributeKey Macro_CustomBindPoseLocalTransformKey = UE::Interchange::FAttributeKey(TEXT("BindPoseLocalTransform"));
-	const UE::Interchange::FAttributeKey Macro_CustomBindPoseGlobalTransformKey = UE::Interchange::FAttributeKey(TEXT("BindPoseGlobalTransform"));
 	const UE::Interchange::FAttributeKey Macro_CustomTimeZeroLocalTransformKey = UE::Interchange::FAttributeKey(TEXT("TimeZeroLocalTransform"));
-	const UE::Interchange::FAttributeKey Macro_CustomTimeZeroGlobalTransformKey = UE::Interchange::FAttributeKey(TEXT("TimeZeroGlobalTransform"));
 	const UE::Interchange::FAttributeKey Macro_CustomGeometricTransformKey = UE::Interchange::FAttributeKey(TEXT("GeometricTransform"));
 	const UE::Interchange::FAttributeKey Macro_CustomAssetInstanceUidKey = UE::Interchange::FAttributeKey(TEXT("AssetInstanceUid"));
 
 	UE::Interchange::TArrayAttributeHelper<FString> NodeSpecializeTypes;
 	UE::Interchange::TArrayAttributeHelper<FString> MaterialDependencyUids;
+
+	mutable TOptional<FTransform> CacheGlobalTransform;
+	mutable TOptional<FTransform> CacheBindPoseGlobalTransform;
+	mutable TOptional<FTransform> CacheTimeZeroGlobalTransform;
 };
diff --git a/Engine/Plugins/Experimental/Interchange/Runtime/Source/Parsers/Fbx/Private/FbxScene.cpp b/Engine/Plugins/Experimental/Interchange/Runtime/Source/Parsers/Fbx/Private/FbxScene.cpp
--- a/Engine/Plugins/Experimental/Interchange/Runtime/Source/Parsers/Fbx/Private/FbxScene.cpp
+++ b/Engine/Plugins/Experimental/Interchange/Runtime/Source/Parsers/Fbx/Private/FbxScene.cpp
@@ -71,6 +71,24 @@ namespace UE
 				CreateAssetNodeReference(UnrealSceneNode, NodeAttribute, NodeContainer, UInterchangeLightNode::StaticAssetTypeName());
 			}
 
+			bool DoesTheParentHierarchyContainJoints(FbxNode* Node)
+			{
+				if (!Node)
+				{
+					return false;
+				}
+				int32 AttributeCount = Node->GetNodeAttributeCount();
+				for (int32 AttributeIndex = 0; AttributeIndex < AttributeCount; ++AttributeIndex)
+				{
+					FbxNodeAttribute* NodeAttribute = Node->GetNodeAttributeByIndex(AttributeIndex);
+					if (NodeAttribute->GetAttributeType() == FbxNodeAttribute::eSkeleton)
+					{
+						return true;
+					}
+				}
+				return DoesTheParentHierarchyContainJoints(Node->GetParent());
+			}
+
 			void FFbxScene::AddHierarchyRecursively(UInterchangeSceneNode* UnrealParentNode, FbxNode* Node, FbxScene* SDKScene, UInterchangeBaseNodeContainer& NodeContainer)
 			{
 				FString NodeName = FFbxHelper::GetFbxObjectName(Node);
@@ -113,18 +131,17 @@ namespace UE
 				{
 					FbxAMatrix GlobalFbxMatrix = Node->EvaluateGlobalTransform();
 					FTransform GlobalTransform = GetConvertedTransform(GlobalFbxMatrix);
-					UnrealNode->SetCustomGlobalTransform(GlobalTransform);
 					if (FbxNode* ParentNode = Node->GetParent())
 					{
 						FbxAMatrix GlobalFbxParentMatrix = ParentNode->EvaluateGlobalTransform();
 						FbxAMatrix	LocalFbxMatrix = GlobalFbxParentMatrix.Inverse() * GlobalFbxMatrix;
 						FTransform LocalTransform = GetConvertedTransform(LocalFbxMatrix);
-						UnrealNode->SetCustomLocalTransform(LocalTransform);
+						UnrealNode->SetCustomLocalTransform(&NodeContainer, LocalTransform);
 					}
 					else
 					{
 						//No parent, set the same matrix has the global
-						UnrealNode->SetCustomLocalTransform(GlobalTransform);
+						UnrealNode->SetCustomLocalTransform(&NodeContainer, GlobalTransform);
 					}
 				}
 
@@ -156,8 +173,15 @@ namespace UE
 							break;
 
 						case FbxNodeAttribute::eNull:
+						{
+							if (!DoesTheParentHierarchyContainJoints(Node->GetParent()))
+							{
+								//eNull node will be set has a transform and a joint specialized type
+								UnrealNode->AddSpecializedType(FSceneNodeStaticData::GetTransformSpecializeTypeString());
+							}
+						}
+						//No break since the eNull act has a skeleton if possible
 						case FbxNodeAttribute::eSkeleton:
-							//eNull node will be set has a skeleton node
 						{
 							//Add the joint specialized type
 							UnrealNode->AddSpecializedType(FSceneNodeStaticData::GetJointSpecializeTypeString());
@@ -166,7 +190,6 @@ namespace UE
 							if (FFbxMesh::GetGlobalJointBindPoseTransform(SDKScene, Node, GlobalBindPoseJointMatrix))
 							{
 								FTransform GlobalBindPoseJointTransform = GetConvertedTransform(GlobalBindPoseJointMatrix);
-								UnrealNode->SetCustomBindPoseGlobalTransform(GlobalBindPoseJointTransform);
 								//We grab the fbx parent node to compute the local transform
 								if (FbxNode* ParentNode = Node->GetParent())
 								{
@@ -174,12 +197,12 @@ namespace UE
 									FFbxMesh::GetGlobalJointBindPoseTransform(SDKScene, ParentNode, GlobalFbxParentMatrix);
 									FbxAMatrix	LocalFbxMatrix = GlobalFbxParentMatrix.Inverse() * GlobalBindPoseJointMatrix;
 									FTransform LocalBindPoseJointTransform = GetConvertedTransform(LocalFbxMatrix);
-									UnrealNode->SetCustomBindPoseLocalTransform(LocalBindPoseJointTransform);
+									UnrealNode->SetCustomBindPoseLocalTransform(&NodeContainer, LocalBindPoseJointTransform);
 								}
 								else
 								{
 									//No parent, set the same matrix has the global
-									UnrealNode->SetCustomBindPoseLocalTransform(GlobalBindPoseJointTransform);
+									UnrealNode->SetCustomBindPoseLocalTransform(&NodeContainer, GlobalBindPoseJointTransform);
 								}
 							}
 
@@ -188,18 +211,17 @@ namespace UE
 								//Set the global node transform
 								FbxAMatrix GlobalFbxMatrix = Node->EvaluateGlobalTransform(FBXSDK_TIME_ZERO);
 								FTransform GlobalTransform = GetConvertedTransform(GlobalFbxMatrix);
-								UnrealNode->SetCustomTimeZeroGlobalTransform(GlobalTransform);
 								if (FbxNode* ParentNode = Node->GetParent())
 								{
 									FbxAMatrix GlobalFbxParentMatrix = ParentNode->EvaluateGlobalTransform(FBXSDK_TIME_ZERO);
 									FbxAMatrix	LocalFbxMatrix = GlobalFbxParentMatrix.Inverse() * GlobalFbxMatrix;
 									FTransform LocalTransform = GetConvertedTransform(LocalFbxMatrix);
-									UnrealNode->SetCustomTimeZeroLocalTransform(LocalTransform);
+									UnrealNode->SetCustomTimeZeroLocalTransform(&NodeContainer, LocalTransform);
 								}
 								else
 								{
 									//No parent, set the same matrix has the global
-									UnrealNode->SetCustomTimeZeroLocalTransform(GlobalTransform);
+									UnrealNode->SetCustomTimeZeroLocalTransform(&NodeContainer, GlobalTransform);
 								}
 							}
 
diff --git a/Engine/Plugins/Experimental/Interchange/Runtime/Source/Pipelines/Private/InterchangeGenericScenesPipeline.cpp b/Engine/Plugins/Experimental/Interchange/Runtime/Source/Pipelines/Private/InterchangeGenericScenesPipeline.cpp
--- a/Engine/Plugins/Experimental/Interchange/Runtime/Source/Pipelines/Private/InterchangeGenericScenesPipeline.cpp
+++ b/Engine/Plugins/Experimental/Interchange/Runtime/Source/Pipelines/Private/InterchangeGenericScenesPipeline.cpp
@@ -46,15 +46,24 @@ void UInterchangeGenericLevelPipeline::ExecutePreImportPipeline(UInterchangeBase
 
 	for (const UInterchangeSceneNode* SceneNode : SceneNodes)
 	{
-		//Ignore specialized types for now as they are used for bones hierarchies and other asset internal data
-		if (SceneNode && SceneNode->GetSpecializedTypeCount() == 0)
+		if (SceneNode)
 		{
-			CreateActorFactoryNode(SceneNode, InBaseNodeContainer);
+			if (SceneNode->GetSpecializedTypeCount() > 0)
+			{
+				TArray<FString> SpecializeTypes;
+				SceneNode->GetSpecializedTypes(SpecializeTypes);
+				if (!SpecializeTypes.Contains(UE::Interchange::FSceneNodeStaticData::GetTransformSpecializeTypeString()))
+				{
+					//Skip any scene node that have specialized types but not the "Transform" type.
+					continue;
+				}
+			}
+			CreateActorFactoryNode(InBaseNodeContainer, SceneNode, InBaseNodeContainer);
 		}
 	}
 }
 
-void UInterchangeGenericLevelPipeline::CreateActorFactoryNode(const UInterchangeSceneNode* SceneNode, UInterchangeBaseNodeContainer* FactoryNodeContainer)
+void UInterchangeGenericLevelPipeline::CreateActorFactoryNode(UInterchangeBaseNodeContainer* InBaseNodeContainer, const UInterchangeSceneNode* SceneNode, UInterchangeBaseNodeContainer* FactoryNodeContainer)
 {
 	if (!SceneNode)
 	{
@@ -94,7 +103,7 @@ void UInterchangeGenericLevelPipeline::CreateActorFactoryNode(const UInterchange
 	ActorFactoryNode->AddTargetNodeUid(SceneNode->GetUniqueID());
 
 	FTransform GlobalTransform;
-	if (SceneNode->GetCustomGlobalTransform(GlobalTransform))
+	if (SceneNode->GetCustomGlobalTransform(InBaseNodeContainer, GlobalTransform))
 	{
 		ActorFactoryNode->SetCustomGlobalTransform(GlobalTransform);
 	}
diff --git a/Engine/Plugins/Experimental/Interchange/Runtime/Source/Pipelines/Public/InterchangeGenericScenesPipeline.h b/Engine/Plugins/Experimental/Interchange/Runtime/Source/Pipelines/Public/InterchangeGenericScenesPipeline.h
--- a/Engine/Plugins/Experimental/Interchange/Runtime/Source/Pipelines/Public/InterchangeGenericScenesPipeline.h
+++ b/Engine/Plugins/Experimental/Interchange/Runtime/Source/Pipelines/Public/InterchangeGenericScenesPipeline.h
@@ -31,7 +31,7 @@ protected:
 		return true;
 	}
 
-	void CreateActorFactoryNode(const UInterchangeSceneNode* SceneNode, UInterchangeBaseNodeContainer* FactoryNodeContainer);
+	void CreateActorFactoryNode(UInterchangeBaseNodeContainer* InBaseNodeContainer, const UInterchangeSceneNode* SceneNode, UInterchangeBaseNodeContainer* FactoryNodeContainer);
 	
 };
 
diff --git a/Engine/Source/Developer/AssetTools/Private/AssetTools.cpp b/Engine/Source/Developer/AssetTools/Private/AssetTools.cpp
--- a/Engine/Source/Developer/AssetTools/Private/AssetTools.cpp
+++ b/Engine/Source/Developer/AssetTools/Private/AssetTools.cpp
@@ -1263,7 +1263,7 @@ void UAssetToolsImpl::ExpandDirectories(const TArray<FString>& Files, const FStr
 		}
 	}
 }
-TArray<UObject*> UAssetToolsImpl::ImportAssets(const TArray<FString>& Files, const FString& DestinationPath, UFactory* ChosenFactory, bool bSyncToBrowser /* = true */, TArray<TPair<FString, FString>>* FilesAndDestinations /* = nullptr */, bool bAllowAsyncImport /* = false */) const
+TArray<UObject*> UAssetToolsImpl::ImportAssets(const TArray<FString>& Files, const FString& DestinationPath, UFactory* ChosenFactory, bool bSyncToBrowser /* = true */, TArray<TPair<FString, FString>>* FilesAndDestinations /* = nullptr */, bool bAllowAsyncImport /* = false */, bool bSceneImport /*= false*/) const
 {
 	const bool bForceOverrideExisting = false;
 
@@ -1274,6 +1274,7 @@ TArray<UObject*> UAssetToolsImpl::ImportAssets(const TArray<FString>& Files, con
 	Params.bSyncToBrowser = bSyncToBrowser;
 	Params.SpecifiedFactory = TStrongObjectPtr<UFactory>(ChosenFactory);
 	Params.bAllowAsyncImport = bAllowAsyncImport;
+	Params.bSceneImport = bSceneImport;
 
 	return ImportAssetsInternal(Files, DestinationPath, FilesAndDestinations, Params);
 }
@@ -1982,7 +1983,7 @@ TArray<UObject*> UAssetToolsImpl::ImportAssetsInternal(const TArray<FString>& Fi
 						}
 					};
 
-				if (Params.SpecifiedFactory && Params.SpecifiedFactory->GetClass()->IsChildOf(USceneImportFactory::StaticClass()))
+				if (Params.bSceneImport || (Params.SpecifiedFactory && Params.SpecifiedFactory->GetClass()->IsChildOf(USceneImportFactory::StaticClass())))
 				{
 					TPair<UE::Interchange::FAssetImportResultRef, UE::Interchange::FSceneImportResultRef> InterchangeResults =
 						InterchangeManager.ImportSceneAsync(DestinationPath, ScopedSourceData.GetSourceData(), ImportAssetParameters);
diff --git a/Engine/Source/Developer/AssetTools/Private/AssetTools.h b/Engine/Source/Developer/AssetTools/Private/AssetTools.h
--- a/Engine/Source/Developer/AssetTools/Private/AssetTools.h
+++ b/Engine/Source/Developer/AssetTools/Private/AssetTools.h
@@ -30,6 +30,7 @@ struct FAssetImportParams
 		, bForceOverrideExisting(false)
 		, bAutomated(false)
 		, bAllowAsyncImport(false)
+		, bSceneImport(false)
 	{}
 
 	/** Factory to use for importing files */
@@ -46,6 +47,8 @@ struct FAssetImportParams
 	bool bAutomated : 1;
 	/** Temporary setting to enable the async import. (We will figure a better solution in a later release) */
 	bool bAllowAsyncImport : 1;
+	/** Whether or not this is a scene import */
+	bool bSceneImport : 1;
 };
 
 
@@ -87,7 +90,7 @@ public:
 	virtual void RenameReferencingSoftObjectPaths(const TArray<UPackage *> PackagesToCheck, const TMap<FSoftObjectPath, FSoftObjectPath>& AssetRedirectorMap) override;
 	virtual TArray<UObject*> ImportAssetsWithDialog(const FString& DestinationPath) override;
 	virtual void ImportAssetsWithDialogAsync(const FString& DestinationPath) override;
-	virtual TArray<UObject*> ImportAssets(const TArray<FString>& Files, const FString& DestinationPath, UFactory* ChosenFactory, bool bSyncToBrowser = true, TArray<TPair<FString, FString>>* FilesAndDestinations = nullptr, bool bAllowAsyncImport = false) const override;
+	virtual TArray<UObject*> ImportAssets(const TArray<FString>& Files, const FString& DestinationPath, UFactory* ChosenFactory, bool bSyncToBrowser = true, TArray<TPair<FString, FString>>* FilesAndDestinations = nullptr, bool bAllowAsyncImport = false, bool bSceneImport = false) const override;
 	virtual TArray<UObject*> ImportAssetsAutomated(const UAutomatedAssetImportData* ImportData) override;
 	virtual void ImportAssetTasks(const TArray<UAssetImportTask*>& ImportTasks) override;
 	virtual void ExportAssets(const TArray<FString>& AssetsToExport, const FString& ExportPath) override;
diff --git a/Engine/Source/Developer/AssetTools/Public/IAssetTools.h b/Engine/Source/Developer/AssetTools/Public/IAssetTools.h
--- a/Engine/Source/Developer/AssetTools/Public/IAssetTools.h
+++ b/Engine/Source/Developer/AssetTools/Public/IAssetTools.h
@@ -297,7 +297,7 @@ public:
 	 * @param bAllowAsyncImport	This allow the import code to use a async importer if enabled. (Note doing so will ignore the ChosenFactory arguments (If you want a async import prefer the InterchangeManager api))
 	 * @return list of successfully imported assets
 	 */
-	virtual TArray<UObject*> ImportAssets(const TArray<FString>& Files, const FString& DestinationPath, UFactory* ChosenFactory = NULL, bool bSyncToBrowser = true, TArray<TPair<FString, FString>>* FilesAndDestinations = nullptr, bool bAllowAsyncImport = false) const = 0;
+	virtual TArray<UObject*> ImportAssets(const TArray<FString>& Files, const FString& DestinationPath, UFactory* ChosenFactory = NULL, bool bSyncToBrowser = true, TArray<TPair<FString, FString>>* FilesAndDestinations = nullptr, bool bAllowAsyncImport = false, bool bSceneImport = false) const = 0;
 
 	/**
 	 * Imports assets using data specified completely up front.  Does not ever ask any questions of the user or show any modal error messages
diff --git a/Engine/Source/Editor/UnrealEd/Private/FileHelpers.cpp b/Engine/Source/Editor/UnrealEd/Private/FileHelpers.cpp
--- a/Engine/Source/Editor/UnrealEd/Private/FileHelpers.cpp
+++ b/Engine/Source/Editor/UnrealEd/Private/FileHelpers.cpp
@@ -1478,7 +1478,9 @@ void FEditorFileUtils::Import(const FString& InFilename)
 		Files.Add(InFilename);
 
 		const bool bSyncToBrowser = bImportsAssets;
-		AssetToolsModule.Get().ImportAssets(Files, Path, SceneFactory, bSyncToBrowser, nullptr, true);
+		constexpr bool bAllowAsyncImport = true;
+		constexpr bool bSceneImport = true;
+		AssetToolsModule.Get().ImportAssets(Files, Path, SceneFactory, bSyncToBrowser, nullptr, bAllowAsyncImport, bSceneImport);
 	}
 	else
 	{
