Changes to transform basis and scale for .obj vertices.

diff --git a/Engine/Plugins/Experimental/Interchange/Runtime/Source/Import/Private/Mesh/InterchangeOBJTranslator.cpp b/Engine/Plugins/Experimental/Interchange/Runtime/Source/Import/Private/Mesh/InterchangeOBJTranslator.cpp
--- a/Engine/Plugins/Experimental/Interchange/Runtime/Source/Import/Private/Mesh/InterchangeOBJTranslator.cpp
+++ b/Engine/Plugins/Experimental/Interchange/Runtime/Source/Import/Private/Mesh/InterchangeOBJTranslator.cpp
@@ -120,9 +120,16 @@ private:
 };
 
 
+static FVector3f NormalToUEBasis(const FVector3f& InVector)
+{
+	return FVector3f(InVector.X, InVector.Z, InVector.Y);
+}
+
+
 static FVector3f PositionToUEBasis(const FVector3f& InVector)
 {
-	return FVector3f(InVector.X, -InVector.Y, InVector.Z);
+	// Convert from meters to centimeters
+	return FVector3f(InVector.X, InVector.Z, InVector.Y) * 100.0f;
 }
 
 
@@ -289,7 +296,7 @@ FMeshDescription FObjData::MakeMeshDescriptionForGroup(const FString& GroupName)
 
 			if (VertexData.NormalIndex != INDEX_NONE)
 			{
-				Attributes.GetVertexInstanceNormals()[VertexInstanceID] = PositionToUEBasis(Normals[VertexData.NormalIndex]);
+				Attributes.GetVertexInstanceNormals()[VertexInstanceID] = NormalToUEBasis(Normals[VertexData.NormalIndex]);
 			}
 
 			if (VertexData.UVIndex != INDEX_NONE)
@@ -476,7 +483,7 @@ namespace ObjParser
 		FVector3f Position;
 		bool bSuccess = GetVector3(Line, Position);
 
-		ObjData.Positions.Add(Position * 100.0f);
+		ObjData.Positions.Add(Position);
 
 		if (!Line.IsEmpty())
 		{
@@ -1001,8 +1008,6 @@ void UInterchangeOBJTranslator::AddMaterialNodes(UInterchangeBaseNodeContainer&
 		TextureSampleShader->AddFloatAttribute(UInterchangeShaderPortsAPI::MakeInputValueKey(TextureSample::Inputs::UTiling.ToString()), 1.0f);
 		TextureSampleShader->AddFloatAttribute(UInterchangeShaderPortsAPI::MakeInputValueKey(TextureSample::Inputs::VTiling.ToString()), 1.0f);
 
-		// @todo: this is not currently working correctly
-		// Material graph complains that input is not LinearColor - seems to need an FColor->FLinearColor conversion node adding
 		UInterchangeShaderPortsAPI::ConnectDefaultOuputToInput(ShaderGraphNode, InputType, TextureSampleShaderUid);
 	}
 	else
