Move SkeletalMesh and StaticMesh pipeline into there own subobject

#rnx

diff --git a/Engine/Plugins/Experimental/Interchange/Runtime/Source/Pipelines/Private/InterchangeGenericAssetsPipeline.cpp b/Engine/Plugins/Experimental/Interchange/Runtime/Source/Pipelines/Private/InterchangeGenericAssetsPipeline.cpp
--- a/Engine/Plugins/Experimental/Interchange/Runtime/Source/Pipelines/Private/InterchangeGenericAssetsPipeline.cpp
+++ b/Engine/Plugins/Experimental/Interchange/Runtime/Source/Pipelines/Private/InterchangeGenericAssetsPipeline.cpp
@@ -5,6 +5,7 @@
 #include "Animation/Skeleton.h"
 #include "CoreMinimal.h"
 #include "InterchangeGenericMaterialPipeline.h"
+#include "InterchangeGenericMeshPipeline.h"
 #include "InterchangeGenericTexturePipeline.h"
 #include "InterchangeMaterialFactoryNode.h"
 #include "InterchangeMeshNode.h"
@@ -27,13 +28,27 @@
 UInterchangeGenericAssetsPipeline::UInterchangeGenericAssetsPipeline()
 {
 	MaterialPipeline = CreateDefaultSubobject<UInterchangeGenericMaterialPipeline>("MaterialPipeline");
+	MeshPipeline = CreateDefaultSubobject<UInterchangeGenericMeshPipeline>("MeshPipeline");
 	TexturePipeline = CreateDefaultSubobject<UInterchangeGenericTexturePipeline>("TexturePipeline");
 }
 
 void UInterchangeGenericAssetsPipeline::PreDialogCleanup(const FName PipelineStackName)
 {
-	//We always clean the pipeline skeleton when showing the dialog
-	Skeleton = nullptr;
+	if (TexturePipeline)
+	{
+		TexturePipeline->PreDialogCleanup(PipelineStackName);
+	}
+
+	if (MaterialPipeline)
+	{
+		MaterialPipeline->PreDialogCleanup(PipelineStackName);
+	}
+	
+	if (MeshPipeline)
+	{
+		MeshPipeline->PreDialogCleanup(PipelineStackName);
+	}
+	
 	SaveSettings(PipelineStackName);
 }
 
@@ -52,149 +67,23 @@ void UInterchangeGenericAssetsPipeline::ExecutePreImportPipeline(UInterchangeBas
 		SourceDatas.Add(SourceData);
 	}
 
-	PipelineMeshesUtilities = UInterchangePipelineMeshesUtilities::CreateInterchangePipelineMeshesUtilities(BaseNodeContainer);
-
 	if (TexturePipeline)
 	{
-		TexturePipeline->ExecutePreImportPipeline(InBaseNodeContainer, InSourceDatas);
+		TexturePipeline->ScriptedExecutePreImportPipeline(InBaseNodeContainer, InSourceDatas);
 	}
 
 	if (MaterialPipeline)
 	{
-		MaterialPipeline->ExecutePreImportPipeline(InBaseNodeContainer, InSourceDatas);
+		MaterialPipeline->ScriptedExecutePreImportPipeline(InBaseNodeContainer, InSourceDatas);
 	}
 
-	//Create skeletalmesh factory nodes
-	ExecutePreImportPipelineSkeletalMesh();
-
-	if (bImportStaticMeshes && (ForceAllMeshAsType == EInterchangeForceMeshType::IFMT_None || ForceAllMeshAsType == EInterchangeForceMeshType::IFMT_StaticMesh))
+	if (MeshPipeline)
 	{
-		const bool bConvertSkeletalMeshToStaticMesh = (ForceAllMeshAsType == EInterchangeForceMeshType::IFMT_StaticMesh);
-		if (bCombineStaticMeshes)
-		{
-			// Combine all the static meshes
-
-			bool bFoundMeshes = false;
-			if (bBakeMeshes)
-			{
-				// If baking transforms, get all the static mesh instance nodes, and group them by LOD
-				TArray<FString> MeshUids;
-				PipelineMeshesUtilities->GetAllStaticMeshInstance(MeshUids, bConvertSkeletalMeshToStaticMesh);
-
-				TMap<int32, TArray<FString>> MeshUidsPerLodIndex;
-
-				for (const FString& MeshUid : MeshUids)
-				{
-					const FInterchangeMeshInstance& MeshInstance = PipelineMeshesUtilities->GetMeshInstanceByUid(MeshUid);
-					for (const auto& LodIndexAndSceneNodeContainer : MeshInstance.SceneNodePerLodIndex)
-					{
-						const int32 LodIndex = LodIndexAndSceneNodeContainer.Key;
-						const FInterchangeLodSceneNodeContainer& SceneNodeContainer = LodIndexAndSceneNodeContainer.Value;
-
-						TArray<FString>& TranslatedNodes = MeshUidsPerLodIndex.FindOrAdd(LodIndex);
-						for (const UInterchangeSceneNode* SceneNode : SceneNodeContainer.SceneNodes)
-						{
-							TranslatedNodes.Add(SceneNode->GetUniqueID());
-						}
-					}
-				}
-
-				// If we got some instances, create a static mesh factory node
-				if (MeshUidsPerLodIndex.Num() > 0)
-				{
-					UInterchangeStaticMeshFactoryNode* StaticMeshFactoryNode = CreateStaticMeshFactoryNode(MeshUidsPerLodIndex);
-					StaticMeshFactoryNodes.Add(StaticMeshFactoryNode);
-					bFoundMeshes = true;
-				}
-			}
-
-			if (!bFoundMeshes)
-			{
-				// If we haven't yet managed to build a factory node, look at static mesh geometry directly.
-				TArray<FString> MeshUids;
-				PipelineMeshesUtilities->GetAllStaticMeshGeometry(MeshUids, bConvertSkeletalMeshToStaticMesh);
-
-				TMap<int32, TArray<FString>> MeshUidsPerLodIndex;
-
-				for (const FString& MeshUid : MeshUids)
-				{
-					// MeshGeometry cannot have Lod since LODs are defined in the scene node
-					const FInterchangeMeshGeometry& MeshGeometry = PipelineMeshesUtilities->GetMeshGeometryByUid(MeshUid);
-					const int32 LodIndex = 0;
-					TArray<FString>& TranslatedNodes = MeshUidsPerLodIndex.FindOrAdd(LodIndex);
-					TranslatedNodes.Add(MeshGeometry.MeshUid);
-				}
-
-				if (MeshUidsPerLodIndex.Num() > 0)
-				{
-					UInterchangeStaticMeshFactoryNode* StaticMeshFactoryNode = CreateStaticMeshFactoryNode(MeshUidsPerLodIndex);
-					StaticMeshFactoryNodes.Add(StaticMeshFactoryNode);
-				}
-			}
-
-		}
-		else
-		{
-			// Do not combine static meshes
-
-			bool bFoundMeshes = false;
-			if (bBakeMeshes)
-			{
-				TArray<FString> MeshUids;
-				PipelineMeshesUtilities->GetAllStaticMeshInstance(MeshUids, bConvertSkeletalMeshToStaticMesh);
-				
-				for (const FString& MeshUid : MeshUids)
-				{
-					// @TODO: check this is correct. What happens when we explicitly don't combine scenes with a LOD group node and multiple transform nodes?
-					TMap<int32, TArray<FString>> MeshUidsPerLodIndex;
-
-					const FInterchangeMeshInstance& MeshInstance = PipelineMeshesUtilities->GetMeshInstanceByUid(MeshUid);
-					for (const auto& LodIndexAndSceneNodeContainer : MeshInstance.SceneNodePerLodIndex)
-					{
-						const int32 LodIndex = LodIndexAndSceneNodeContainer.Key;
-						const FInterchangeLodSceneNodeContainer& SceneNodeContainer = LodIndexAndSceneNodeContainer.Value;
-
-						TArray<FString>& TranslatedNodes = MeshUidsPerLodIndex.FindOrAdd(LodIndex);
-						for (const UInterchangeSceneNode* SceneNode : SceneNodeContainer.SceneNodes)
-						{
-							TranslatedNodes.Add(SceneNode->GetUniqueID());
-						}
-					}
-
-					if (MeshUidsPerLodIndex.Num() > 0)
-					{
-						UInterchangeStaticMeshFactoryNode* StaticMeshFactoryNode = CreateStaticMeshFactoryNode(MeshUidsPerLodIndex);
-						StaticMeshFactoryNodes.Add(StaticMeshFactoryNode);
-						bFoundMeshes = true;
-					}
-				}
-			}
-
-			if (!bFoundMeshes)
-			{
-				TArray<FString> MeshUids;
-				PipelineMeshesUtilities->GetAllStaticMeshGeometry(MeshUids, bConvertSkeletalMeshToStaticMesh);
-
-				for (const FString& MeshUid : MeshUids)
-				{
-					TMap<int32, TArray<FString>> MeshUidsPerLodIndex;
-
-					const FInterchangeMeshGeometry& MeshGeometry = PipelineMeshesUtilities->GetMeshGeometryByUid(MeshUid);
-					const int32 LodIndex = 0;
-					TArray<FString>& TranslatedNodes = MeshUidsPerLodIndex.FindOrAdd(LodIndex);
-					TranslatedNodes.Add(MeshGeometry.MeshUid);
-
-					if (MeshUidsPerLodIndex.Num() > 0)
-					{
-						UInterchangeStaticMeshFactoryNode* StaticMeshFactoryNode = CreateStaticMeshFactoryNode(MeshUidsPerLodIndex);
-						StaticMeshFactoryNodes.Add(StaticMeshFactoryNode);
-					}
-				}
-			}
-		}
+		MeshPipeline->ScriptedExecutePreImportPipeline(InBaseNodeContainer, InSourceDatas);
 	}
 
 	ImplementUseSourceNameForAssetOption();
+
 	//Make sure all factory nodes have the specified strategy
 	BaseNodeContainer->IterateNodes([ReimportStrategyClosure = ReimportStrategy](const FString& NodeUid, UInterchangeBaseNode* Node)
 		{
@@ -207,27 +96,19 @@ void UInterchangeGenericAssetsPipeline::ExecutePreImportPipeline(UInterchangeBas
 
 void UInterchangeGenericAssetsPipeline::ExecutePostImportPipeline(const UInterchangeBaseNodeContainer* InBaseNodeContainer, const FString& NodeKey, UObject* CreatedAsset, bool bIsAReimport)
 {
-	//We do not use the provided base container since ExecutePreImportPipeline cache it
-	//We just make sure the same one is pass in parameter
-	if (!InBaseNodeContainer || !ensure(BaseNodeContainer == InBaseNodeContainer) || !CreatedAsset)
+	if (TexturePipeline)
 	{
-		return;
+		TexturePipeline->ScriptedExecutePostImportPipeline(InBaseNodeContainer, NodeKey, CreatedAsset, bIsAReimport);
 	}
 
-	UInterchangeBaseNode* Node = BaseNodeContainer->GetNode(NodeKey);
-	if (!Node)
+	if (MaterialPipeline)
 	{
-		return;
+		MaterialPipeline->ScriptedExecutePostImportPipeline(InBaseNodeContainer, NodeKey, CreatedAsset, bIsAReimport);
 	}
 
-	PostImportSkeletalMesh(CreatedAsset, Node);
-
-	//Finish the physics asset import, it need the skeletal mesh render data to create the physics collision geometry
-	PostImportPhysicsAssetImport(CreatedAsset, Node);
-
-	if (TexturePipeline)
+	if (MeshPipeline)
 	{
-		TexturePipeline->ExecutePostImportPipeline(InBaseNodeContainer, NodeKey, CreatedAsset, bIsAReimport);
+		MeshPipeline->ScriptedExecutePostImportPipeline(InBaseNodeContainer, NodeKey, CreatedAsset, bIsAReimport);
 	}
 }
 
@@ -248,7 +129,7 @@ void UInterchangeGenericAssetsPipeline::ImplementUseSourceNameForAssetOption()
 
 	//Skeletalmesh must always rename created skeleton and physics asset to (SKNAME_Skeleton of SKNAME_PhysicsAsset). So the pipeline option
 	// bUseSourceNameForAsset will change or not the SKNAME.
-	ImplementUseSourceNameForAssetOptionSkeletalMesh(MeshesAndAnimsImportedNodeCount);
+	MeshPipeline->ImplementUseSourceNameForAssetOptionSkeletalMesh(MeshesAndAnimsImportedNodeCount, bUseSourceNameForAsset);
 
 	if (bUseSourceNameForAsset && MeshesAndAnimsImportedNodeCount == 1)
 	{
@@ -261,207 +142,3 @@ void UInterchangeGenericAssetsPipeline::ImplementUseSourceNameForAssetOption()
 	}
 
 }
-
-bool UInterchangeGenericAssetsPipeline::MakeMeshFactoryNodeUidAndDisplayLabel(const TMap<int32, TArray<FString>>& MeshUidsPerLodIndex, int32 LodIndex, FString& NewNodeUid, FString& DisplayLabel)
-{
-	int32 SceneNodeCount = 0;
-
-	if (!ensure(LodIndex >= 0 && MeshUidsPerLodIndex.Num() > LodIndex))
-	{
-		return false;
-	}
-
-	for (const TPair<int32, TArray<FString>>& LodIndexAndMeshUids : MeshUidsPerLodIndex)
-	{
-		if (LodIndex == LodIndexAndMeshUids.Key)
-		{
-			const TArray<FString>& Uids = LodIndexAndMeshUids.Value;
-			if (Uids.Num() > 0)
-			{
-				const FString& Uid = Uids[0];
-				UInterchangeBaseNode* Node = BaseNodeContainer->GetNode(Uid);
-
-				if (UInterchangeMeshNode* MeshNode = Cast<UInterchangeMeshNode>(Node))
-				{
-					DisplayLabel = Node->GetDisplayLabel();
-					NewNodeUid = Uid;
-					return true;
-				}
-
-				if (UInterchangeSceneNode* SceneNode = Cast<UInterchangeSceneNode>(Node))
-				{
-					FString RefMeshUid;
-					if (SceneNode->GetCustomAssetInstanceUid(RefMeshUid))
-					{
-						UInterchangeBaseNode* MeshNode = BaseNodeContainer->GetNode(RefMeshUid);
-						if (MeshNode)
-						{
-							DisplayLabel = MeshNode->GetDisplayLabel();
-							if (Uids.Num() == 1)
-							{
-								// If we are instancing one scene node, we want to use it to name the mesh
-								DisplayLabel = SceneNode->GetDisplayLabel();
-							}
-
-							// Use the first scene node uid this LOD references, adding backslash since this uid is not asset typed (\\Mesh\\) like a mesh node
-							// @TODO: change this so that scene nodes get a type prefix of their own?
-							NewNodeUid = TEXT("\\") + Uid;
-							return true;
-						}
-					}
-				}
-			}
-
-			// We found the lod but there is no valid Mesh node to return the Uid
-			break;
-		}
-	}
-
-	return false;
-}
-
-UInterchangeStaticMeshFactoryNode* UInterchangeGenericAssetsPipeline::CreateStaticMeshFactoryNode(const TMap<int32, TArray<FString>>& MeshUidsPerLodIndex)
-{
-	if (MeshUidsPerLodIndex.Num() == 0)
-	{
-		return nullptr;
-	}
-
-	// Create the static mesh factory node, name it according to the first mesh node compositing the meshes
-	FString StaticMeshUid_MeshNamePart;
-	FString DisplayLabel;
-	const int32 BaseLodIndex = 0;
-	if (!MakeMeshFactoryNodeUidAndDisplayLabel(MeshUidsPerLodIndex, BaseLodIndex, StaticMeshUid_MeshNamePart, DisplayLabel))
-	{
-		// Log an error
-		return nullptr;
-	}
-
-	const FString StaticMeshUid = TEXT("\\StaticMesh") + StaticMeshUid_MeshNamePart;
-	UInterchangeStaticMeshFactoryNode* StaticMeshFactoryNode = NewObject<UInterchangeStaticMeshFactoryNode>(BaseNodeContainer, NAME_None);
-	if (!ensure(StaticMeshFactoryNode))
-	{
-		return nullptr;
-	}
-
-	StaticMeshFactoryNode->InitializeStaticMeshNode(StaticMeshUid, DisplayLabel, UStaticMesh::StaticClass()->GetName());
-	BaseNodeContainer->AddNode(StaticMeshFactoryNode);
-
-	AddLodDataToStaticMesh(StaticMeshFactoryNode, MeshUidsPerLodIndex);
-
-	switch (VertexColorImportOption)
-	{
-		case EInterchangeVertexColorImportOption::IVCIO_Replace:
-		{
-			StaticMeshFactoryNode->SetCustomVertexColorReplace(true);
-		}
-		break;
-		case EInterchangeVertexColorImportOption::IVCIO_Ignore:
-		{
-			StaticMeshFactoryNode->SetCustomVertexColorIgnore(true);
-		}
-		break;
-		case EInterchangeVertexColorImportOption::IVCIO_Override:
-		{
-			StaticMeshFactoryNode->SetCustomVertexColorOverride(VertexOverrideColor);
-		}
-		break;
-	}
-	return StaticMeshFactoryNode;
-}
-
-
-UInterchangeStaticMeshLodDataNode* UInterchangeGenericAssetsPipeline::CreateStaticMeshLodDataNode(const FString& NodeName, const FString& NodeUniqueID)
-{
-	FString DisplayLabel(NodeName);
-	FString NodeUID(NodeUniqueID);
-	UInterchangeStaticMeshLodDataNode* StaticMeshLodDataNode = NewObject<UInterchangeStaticMeshLodDataNode>(BaseNodeContainer, NAME_None);
-	if (!ensure(StaticMeshLodDataNode))
-	{
-		// @TODO: log error
-		return nullptr;
-	}
-
-	StaticMeshLodDataNode->InitializeNode(NodeUID, DisplayLabel, EInterchangeNodeContainerType::FactoryData);
-	BaseNodeContainer->AddNode(StaticMeshLodDataNode);
-	return StaticMeshLodDataNode;
-}
-
-
-void UInterchangeGenericAssetsPipeline::AddLodDataToStaticMesh(UInterchangeStaticMeshFactoryNode* StaticMeshFactoryNode, const TMap<int32, TArray<FString>>& NodeUidsPerLodIndex)
-{
-	const FString StaticMeshFactoryUid = StaticMeshFactoryNode->GetUniqueID();
-
-	for (const TPair<int32, TArray<FString>>& LodIndexAndNodeUids : NodeUidsPerLodIndex)
-	{
-		const int32 LodIndex = LodIndexAndNodeUids.Key;
-		if (!bImportLods && LodIndex > 0)
-		{
-			// If the pipeline should not import lods, skip any lod over base lod
-			continue;
-		}
-
-		const TArray<FString>& NodeUids = LodIndexAndNodeUids.Value;
-
-		// Create a lod data node with all the meshes for this LOD
-		const FString StaticMeshLodDataName = TEXT("LodData") + FString::FromInt(LodIndex);
-		const FString LODDataPrefix = TEXT("\\LodData") + (LodIndex > 0 ? FString::FromInt(LodIndex) : TEXT(""));
-		const FString StaticMeshLodDataUniqueID = LODDataPrefix + StaticMeshFactoryUid;
-
-		// Create LodData node if it doesn't already exist
-		UInterchangeStaticMeshLodDataNode* LodDataNode = Cast<UInterchangeStaticMeshLodDataNode>(BaseNodeContainer->GetNode(StaticMeshLodDataUniqueID));
-		if (!LodDataNode)
-		{
-			// Add the data for the LOD (all the mesh node fbx path, so we can find them when we will create the payload data)
-			LodDataNode = CreateStaticMeshLodDataNode(StaticMeshLodDataName, StaticMeshLodDataUniqueID);
-			LodDataNode->SetParentUid(StaticMeshFactoryUid);
-			StaticMeshFactoryNode->AddLodDataUniqueId(StaticMeshLodDataUniqueID);
-		}
-
-		for (const FString& NodeUid : NodeUids)
-		{
-			TArray<FString> MaterialDependencies;
-			if (const UInterchangeSceneNode* SceneNode = Cast<UInterchangeSceneNode>(BaseNodeContainer->GetNode(NodeUid)))
-			{
-				FString MeshDependency;
-				SceneNode->GetCustomAssetInstanceUid(MeshDependency);
-				if (BaseNodeContainer->IsNodeUidValid(MeshDependency))
-				{
-					const UInterchangeMeshNode* MeshDependencyNode = Cast<UInterchangeMeshNode>(BaseNodeContainer->GetNode(MeshDependency));
-					StaticMeshFactoryNode->AddTargetNodeUid(MeshDependency);
-					MeshDependencyNode->AddTargetNodeUid(StaticMeshFactoryNode->GetUniqueID());
-					MeshDependencyNode->GetMaterialDependencies(MaterialDependencies);
-				}
-				else
-				{
-					SceneNode->GetMaterialDependencyUids(MaterialDependencies);
-				}
-			}
-			else if (const UInterchangeMeshNode* MeshNode = Cast<UInterchangeMeshNode>(BaseNodeContainer->GetNode(NodeUid)))
-			{
-				StaticMeshFactoryNode->AddTargetNodeUid(NodeUid);
-				MeshNode->AddTargetNodeUid(StaticMeshFactoryNode->GetUniqueID());
-				MeshNode->GetMaterialDependencies(MaterialDependencies);
-			}
-
-			const int32 MaterialCount = MaterialDependencies.Num();
-			for (int32 MaterialIndex = 0; MaterialIndex < MaterialCount; ++MaterialIndex)
-			{
-				const FString MaterialFactoryNodeUid = UInterchangeMaterialFactoryNode::GetMaterialFactoryNodeUidFromMaterialNodeUid(MaterialDependencies[MaterialIndex]);
-				if (BaseNodeContainer->IsNodeUidValid(MaterialFactoryNodeUid))
-				{
-					BaseNodeContainer->GetNode(MaterialFactoryNodeUid)->SetEnabled(true);
-					// Create a factory dependency so Material asset are imported before the skeletal mesh asset
-					TArray<FString> FactoryDependencies;
-					StaticMeshFactoryNode->GetFactoryDependencies(FactoryDependencies);
-					if (!FactoryDependencies.Contains(MaterialFactoryNodeUid))
-					{
-						StaticMeshFactoryNode->AddFactoryDependencyUid(MaterialFactoryNodeUid);
-					}
-				}
-			}
-
-			LodDataNode->AddMeshUid(NodeUid);
-		}
-	}
-}
diff --git a/Engine/Plugins/Experimental/Interchange/Runtime/Source/Pipelines/Private/InterchangeGenericMeshPipeline.cpp b/Engine/Plugins/Experimental/Interchange/Runtime/Source/Pipelines/Private/InterchangeGenericMeshPipeline.cpp
new file mode 100644
--- /dev/null
+++ b/Engine/Plugins/Experimental/Interchange/Runtime/Source/Pipelines/Private/InterchangeGenericMeshPipeline.cpp
@@ -0,0 +1,66 @@
+// Copyright Epic Games, Inc. All Rights Reserved. 
+#include "InterchangeGenericMeshPipeline.h"
+
+#include "CoreMinimal.h"
+#include "InterchangeMeshNode.h"
+#include "InterchangePipelineLog.h"
+#include "InterchangePipelineMeshesUtilities.h"
+#include "InterchangeSkeletalMeshFactoryNode.h"
+#include "InterchangeStaticMeshFactoryNode.h"
+#include "InterchangeSourceData.h"
+#include "Misc/Paths.h"
+#include "Nodes/InterchangeBaseNode.h"
+#include "Nodes/InterchangeBaseNodeContainer.h"
+
+void UInterchangeGenericMeshPipeline::ExecutePreImportPipeline(UInterchangeBaseNodeContainer* InBaseNodeContainer, const TArray<UInterchangeSourceData*>& InSourceDatas)
+{
+	if (!InBaseNodeContainer)
+	{
+		UE_LOG(LogInterchangePipeline, Warning, TEXT("UInterchangeGenericMeshPipeline: Cannot execute pre-import pipeline because InBaseNodeContrainer is null"));
+		return;
+	}
+
+	BaseNodeContainer = InBaseNodeContainer;
+	SourceDatas.Empty(InSourceDatas.Num());
+	for (const UInterchangeSourceData* SourceData : InSourceDatas)
+	{
+		SourceDatas.Add(SourceData);
+	}
+
+	PipelineMeshesUtilities = UInterchangePipelineMeshesUtilities::CreateInterchangePipelineMeshesUtilities(BaseNodeContainer);
+
+	//Create skeletalmesh factory nodes
+	ExecutePreImportPipelineSkeletalMesh();
+
+	//Create staticmesh factory nodes
+	ExecutePreImportPipelineStaticMesh();
+}
+
+void UInterchangeGenericMeshPipeline::ExecutePostImportPipeline(const UInterchangeBaseNodeContainer* InBaseNodeContainer, const FString& NodeKey, UObject* CreatedAsset, bool bIsAReimport)
+{
+	//We do not use the provided base container since ExecutePreImportPipeline cache it
+	//We just make sure the same one is pass in parameter
+	if (!InBaseNodeContainer || !ensure(BaseNodeContainer == InBaseNodeContainer) || !CreatedAsset)
+	{
+		return;
+	}
+
+	UInterchangeBaseNode* Node = BaseNodeContainer->GetNode(NodeKey);
+	if (!Node)
+	{
+		return;
+	}
+
+	PostImportSkeletalMesh(CreatedAsset, Node);
+
+	//Finish the physics asset import, it need the skeletal mesh render data to create the physics collision geometry
+	PostImportPhysicsAssetImport(CreatedAsset, Node);
+}
+
+void UInterchangeGenericMeshPipeline::PreDialogCleanup(const FName PipelineStackName)
+{
+	//We always clean the pipeline skeleton when showing the dialog
+	Skeleton = nullptr;
+}
+
+
diff --git a/Engine/Plugins/Experimental/Interchange/Runtime/Source/Pipelines/Private/InterchangeGenericSkeletalMeshPipeline.cpp b/Engine/Plugins/Experimental/Interchange/Runtime/Source/Pipelines/Private/InterchangeGenericSkeletalMeshPipeline.cpp
--- a/Engine/Plugins/Experimental/Interchange/Runtime/Source/Pipelines/Private/InterchangeGenericSkeletalMeshPipeline.cpp
+++ b/Engine/Plugins/Experimental/Interchange/Runtime/Source/Pipelines/Private/InterchangeGenericSkeletalMeshPipeline.cpp
@@ -1,5 +1,5 @@
 // Copyright Epic Games, Inc. All Rights Reserved. 
-#include "InterchangeGenericAssetsPipeline.h"
+#include "InterchangeGenericMeshPipeline.h"
 
 #include "Animation/Skeleton.h"
 #include "AssetRegistry/AssetRegistryModule.h"
@@ -240,7 +240,7 @@ namespace UE::Interchange::SkeletalMeshGenericPipeline
 	}
 }
 
-void UInterchangeGenericAssetsPipeline::ExecutePreImportPipelineSkeletalMesh()
+void UInterchangeGenericMeshPipeline::ExecutePreImportPipelineSkeletalMesh()
 {
 	if (!bImportSkeletalMeshes)
 	{
@@ -416,7 +416,7 @@ void UInterchangeGenericAssetsPipeline::ExecutePreImportPipelineSkeletalMesh()
 }
 
 
-UInterchangeSkeletonFactoryNode* UInterchangeGenericAssetsPipeline::CreateSkeletonFactoryNode(const FString& RootJointUid)
+UInterchangeSkeletonFactoryNode* UInterchangeGenericMeshPipeline::CreateSkeletonFactoryNode(const FString& RootJointUid)
 {
 	const UInterchangeBaseNode* RootJointNode = BaseNodeContainer->GetNode(RootJointUid);
 	if (!RootJointNode)
@@ -466,7 +466,7 @@ UInterchangeSkeletonFactoryNode* UInterchangeGenericAssetsPipeline::CreateSkelet
 	return SkeletonFactoryNode;
 }
 
-UInterchangeSkeletalMeshFactoryNode* UInterchangeGenericAssetsPipeline::CreateSkeletalMeshFactoryNode(const FString& RootJointUid, const TMap<int32, TArray<FString>>& MeshUidsPerLodIndex)
+UInterchangeSkeletalMeshFactoryNode* UInterchangeGenericMeshPipeline::CreateSkeletalMeshFactoryNode(const FString& RootJointUid, const TMap<int32, TArray<FString>>& MeshUidsPerLodIndex)
 {
 	//Get the skeleton factory node
 	const UInterchangeBaseNode* RootJointNode = BaseNodeContainer->GetNode(RootJointUid);
@@ -631,7 +631,7 @@ UInterchangeSkeletalMeshFactoryNode* UInterchangeGenericAssetsPipeline::CreateSk
 	return SkeletalMeshFactoryNode;
 }
 
-UInterchangeSkeletalMeshLodDataNode* UInterchangeGenericAssetsPipeline::CreateSkeletalMeshLodDataNode(const FString& NodeName, const FString& NodeUniqueID)
+UInterchangeSkeletalMeshLodDataNode* UInterchangeGenericMeshPipeline::CreateSkeletalMeshLodDataNode(const FString& NodeName, const FString& NodeUniqueID)
 {
 	FString DisplayLabel(NodeName);
 	FString NodeUID(NodeUniqueID);
@@ -647,7 +647,7 @@ UInterchangeSkeletalMeshLodDataNode* UInterchangeGenericAssetsPipeline::CreateSk
 	return SkeletalMeshLodDataNode;
 }
 
-void UInterchangeGenericAssetsPipeline::AddLodDataToSkeletalMesh(const UInterchangeSkeletonFactoryNode* SkeletonFactoryNode, UInterchangeSkeletalMeshFactoryNode* SkeletalMeshFactoryNode, const TMap<int32, TArray<FString>>& NodeUidsPerLodIndex)
+void UInterchangeGenericMeshPipeline::AddLodDataToSkeletalMesh(const UInterchangeSkeletonFactoryNode* SkeletonFactoryNode, UInterchangeSkeletalMeshFactoryNode* SkeletalMeshFactoryNode, const TMap<int32, TArray<FString>>& NodeUidsPerLodIndex)
 {
 	const FString SkeletalMeshUid = SkeletalMeshFactoryNode->GetUniqueID();
 	const FString SkeletonUid = SkeletonFactoryNode->GetUniqueID();
@@ -727,7 +727,7 @@ void UInterchangeGenericAssetsPipeline::AddLodDataToSkeletalMesh(const UIntercha
 	}
 }
 
-void UInterchangeGenericAssetsPipeline::PostImportSkeletalMesh(UObject* CreatedAsset, UInterchangeBaseNode* Node)
+void UInterchangeGenericMeshPipeline::PostImportSkeletalMesh(UObject* CreatedAsset, UInterchangeBaseNode* Node)
 {
 	if (!BaseNodeContainer)
 	{
@@ -747,7 +747,7 @@ void UInterchangeGenericAssetsPipeline::PostImportSkeletalMesh(UObject* CreatedA
 	}
 }
 
-void UInterchangeGenericAssetsPipeline::PostImportPhysicsAssetImport(UObject* CreatedAsset, UInterchangeBaseNode* Node)
+void UInterchangeGenericMeshPipeline::PostImportPhysicsAssetImport(UObject* CreatedAsset, UInterchangeBaseNode* Node)
 {
 #if WITH_EDITOR
 	if (!bCreatePhysicsAsset || !BaseNodeContainer)
@@ -801,7 +801,7 @@ void UInterchangeGenericAssetsPipeline::PostImportPhysicsAssetImport(UObject* Cr
 #endif //WITH_EDITOR
 }
 
-void UInterchangeGenericAssetsPipeline::ImplementUseSourceNameForAssetOptionSkeletalMesh(const int32 MeshesAndAnimsImportedNodeCount)
+void UInterchangeGenericMeshPipeline::ImplementUseSourceNameForAssetOptionSkeletalMesh(const int32 MeshesAndAnimsImportedNodeCount, const bool bUseSourceNameForAsset)
 {
 	const UClass* SkeletalMeshFactoryNodeClass = UInterchangeSkeletalMeshFactoryNode::StaticClass();
 	TArray<FString> SkeletalMeshNodeUids;
diff --git a/Engine/Plugins/Experimental/Interchange/Runtime/Source/Pipelines/Private/InterchangeGenericStaticMeshPipeline.cpp b/Engine/Plugins/Experimental/Interchange/Runtime/Source/Pipelines/Private/InterchangeGenericStaticMeshPipeline.cpp
new file mode 100644
--- /dev/null
+++ b/Engine/Plugins/Experimental/Interchange/Runtime/Source/Pipelines/Private/InterchangeGenericStaticMeshPipeline.cpp
@@ -0,0 +1,354 @@
+// Copyright Epic Games, Inc. All Rights Reserved. 
+#include "InterchangeGenericMeshPipeline.h"
+
+#include "AssetRegistry/AssetRegistryModule.h"
+#include "Async/Async.h"
+#include "CoreMinimal.h"
+#include "InterchangeMaterialFactoryNode.h"
+#include "InterchangeMeshNode.h"
+#include "InterchangePipelineLog.h"
+#include "InterchangePipelineMeshesUtilities.h"
+#include "InterchangeSceneNode.h"
+#include "InterchangeStaticMeshFactoryNode.h"
+#include "InterchangeStaticMeshLodDataNode.h"
+#include "InterchangeSourceData.h"
+#include "Misc/Paths.h"
+#include "Nodes/InterchangeBaseNode.h"
+#include "Nodes/InterchangeBaseNodeContainer.h"
+#include "Tasks/Task.h"
+#include "UObject/Object.h"
+#include "UObject/ObjectMacros.h"
+
+void UInterchangeGenericMeshPipeline::ExecutePreImportPipelineStaticMesh()
+{
+	if (bImportStaticMeshes && (ForceAllMeshAsType == EInterchangeForceMeshType::IFMT_None || ForceAllMeshAsType == EInterchangeForceMeshType::IFMT_StaticMesh))
+	{
+		const bool bConvertSkeletalMeshToStaticMesh = (ForceAllMeshAsType == EInterchangeForceMeshType::IFMT_StaticMesh);
+		if (bCombineStaticMeshes)
+		{
+			// Combine all the static meshes
+
+			bool bFoundMeshes = false;
+			if (bBakeMeshes)
+			{
+				// If baking transforms, get all the static mesh instance nodes, and group them by LOD
+				TArray<FString> MeshUids;
+				PipelineMeshesUtilities->GetAllStaticMeshInstance(MeshUids, bConvertSkeletalMeshToStaticMesh);
+
+				TMap<int32, TArray<FString>> MeshUidsPerLodIndex;
+
+				for (const FString& MeshUid : MeshUids)
+				{
+					const FInterchangeMeshInstance& MeshInstance = PipelineMeshesUtilities->GetMeshInstanceByUid(MeshUid);
+					for (const auto& LodIndexAndSceneNodeContainer : MeshInstance.SceneNodePerLodIndex)
+					{
+						const int32 LodIndex = LodIndexAndSceneNodeContainer.Key;
+						const FInterchangeLodSceneNodeContainer& SceneNodeContainer = LodIndexAndSceneNodeContainer.Value;
+
+						TArray<FString>& TranslatedNodes = MeshUidsPerLodIndex.FindOrAdd(LodIndex);
+						for (const UInterchangeSceneNode* SceneNode : SceneNodeContainer.SceneNodes)
+						{
+							TranslatedNodes.Add(SceneNode->GetUniqueID());
+						}
+					}
+				}
+
+				// If we got some instances, create a static mesh factory node
+				if (MeshUidsPerLodIndex.Num() > 0)
+				{
+					UInterchangeStaticMeshFactoryNode* StaticMeshFactoryNode = CreateStaticMeshFactoryNode(MeshUidsPerLodIndex);
+					StaticMeshFactoryNodes.Add(StaticMeshFactoryNode);
+					bFoundMeshes = true;
+				}
+			}
+
+			if (!bFoundMeshes)
+			{
+				// If we haven't yet managed to build a factory node, look at static mesh geometry directly.
+				TArray<FString> MeshUids;
+				PipelineMeshesUtilities->GetAllStaticMeshGeometry(MeshUids, bConvertSkeletalMeshToStaticMesh);
+
+				TMap<int32, TArray<FString>> MeshUidsPerLodIndex;
+
+				for (const FString& MeshUid : MeshUids)
+				{
+					// MeshGeometry cannot have Lod since LODs are defined in the scene node
+					const FInterchangeMeshGeometry& MeshGeometry = PipelineMeshesUtilities->GetMeshGeometryByUid(MeshUid);
+					const int32 LodIndex = 0;
+					TArray<FString>& TranslatedNodes = MeshUidsPerLodIndex.FindOrAdd(LodIndex);
+					TranslatedNodes.Add(MeshGeometry.MeshUid);
+				}
+
+				if (MeshUidsPerLodIndex.Num() > 0)
+				{
+					UInterchangeStaticMeshFactoryNode* StaticMeshFactoryNode = CreateStaticMeshFactoryNode(MeshUidsPerLodIndex);
+					StaticMeshFactoryNodes.Add(StaticMeshFactoryNode);
+				}
+			}
+
+		}
+		else
+		{
+			// Do not combine static meshes
+
+			bool bFoundMeshes = false;
+			if (bBakeMeshes)
+			{
+				TArray<FString> MeshUids;
+				PipelineMeshesUtilities->GetAllStaticMeshInstance(MeshUids, bConvertSkeletalMeshToStaticMesh);
+
+				for (const FString& MeshUid : MeshUids)
+				{
+					// @TODO: check this is correct. What happens when we explicitly don't combine scenes with a LOD group node and multiple transform nodes?
+					TMap<int32, TArray<FString>> MeshUidsPerLodIndex;
+
+					const FInterchangeMeshInstance& MeshInstance = PipelineMeshesUtilities->GetMeshInstanceByUid(MeshUid);
+					for (const auto& LodIndexAndSceneNodeContainer : MeshInstance.SceneNodePerLodIndex)
+					{
+						const int32 LodIndex = LodIndexAndSceneNodeContainer.Key;
+						const FInterchangeLodSceneNodeContainer& SceneNodeContainer = LodIndexAndSceneNodeContainer.Value;
+
+						TArray<FString>& TranslatedNodes = MeshUidsPerLodIndex.FindOrAdd(LodIndex);
+						for (const UInterchangeSceneNode* SceneNode : SceneNodeContainer.SceneNodes)
+						{
+							TranslatedNodes.Add(SceneNode->GetUniqueID());
+						}
+					}
+
+					if (MeshUidsPerLodIndex.Num() > 0)
+					{
+						UInterchangeStaticMeshFactoryNode* StaticMeshFactoryNode = CreateStaticMeshFactoryNode(MeshUidsPerLodIndex);
+						StaticMeshFactoryNodes.Add(StaticMeshFactoryNode);
+						bFoundMeshes = true;
+					}
+				}
+			}
+
+			if (!bFoundMeshes)
+			{
+				TArray<FString> MeshUids;
+				PipelineMeshesUtilities->GetAllStaticMeshGeometry(MeshUids, bConvertSkeletalMeshToStaticMesh);
+
+				for (const FString& MeshUid : MeshUids)
+				{
+					TMap<int32, TArray<FString>> MeshUidsPerLodIndex;
+
+					const FInterchangeMeshGeometry& MeshGeometry = PipelineMeshesUtilities->GetMeshGeometryByUid(MeshUid);
+					const int32 LodIndex = 0;
+					TArray<FString>& TranslatedNodes = MeshUidsPerLodIndex.FindOrAdd(LodIndex);
+					TranslatedNodes.Add(MeshGeometry.MeshUid);
+
+					if (MeshUidsPerLodIndex.Num() > 0)
+					{
+						UInterchangeStaticMeshFactoryNode* StaticMeshFactoryNode = CreateStaticMeshFactoryNode(MeshUidsPerLodIndex);
+						StaticMeshFactoryNodes.Add(StaticMeshFactoryNode);
+					}
+				}
+			}
+		}
+	}
+}
+
+bool UInterchangeGenericMeshPipeline::MakeMeshFactoryNodeUidAndDisplayLabel(const TMap<int32, TArray<FString>>& MeshUidsPerLodIndex, int32 LodIndex, FString& NewNodeUid, FString& DisplayLabel)
+{
+	int32 SceneNodeCount = 0;
+
+	if (!ensure(LodIndex >= 0 && MeshUidsPerLodIndex.Num() > LodIndex))
+	{
+		return false;
+	}
+
+	for (const TPair<int32, TArray<FString>>& LodIndexAndMeshUids : MeshUidsPerLodIndex)
+	{
+		if (LodIndex == LodIndexAndMeshUids.Key)
+		{
+			const TArray<FString>& Uids = LodIndexAndMeshUids.Value;
+			if (Uids.Num() > 0)
+			{
+				const FString& Uid = Uids[0];
+				UInterchangeBaseNode* Node = BaseNodeContainer->GetNode(Uid);
+
+				if (UInterchangeMeshNode* MeshNode = Cast<UInterchangeMeshNode>(Node))
+				{
+					DisplayLabel = Node->GetDisplayLabel();
+					NewNodeUid = Uid;
+					return true;
+				}
+
+				if (UInterchangeSceneNode* SceneNode = Cast<UInterchangeSceneNode>(Node))
+				{
+					FString RefMeshUid;
+					if (SceneNode->GetCustomAssetInstanceUid(RefMeshUid))
+					{
+						UInterchangeBaseNode* MeshNode = BaseNodeContainer->GetNode(RefMeshUid);
+						if (MeshNode)
+						{
+							DisplayLabel = MeshNode->GetDisplayLabel();
+							if (Uids.Num() == 1)
+							{
+								// If we are instancing one scene node, we want to use it to name the mesh
+								DisplayLabel = SceneNode->GetDisplayLabel();
+							}
+
+							// Use the first scene node uid this LOD references, adding backslash since this uid is not asset typed (\\Mesh\\) like a mesh node
+							// @TODO: change this so that scene nodes get a type prefix of their own?
+							NewNodeUid = TEXT("\\") + Uid;
+							return true;
+						}
+					}
+				}
+			}
+
+			// We found the lod but there is no valid Mesh node to return the Uid
+			break;
+		}
+	}
+
+	return false;
+}
+
+UInterchangeStaticMeshFactoryNode* UInterchangeGenericMeshPipeline::CreateStaticMeshFactoryNode(const TMap<int32, TArray<FString>>& MeshUidsPerLodIndex)
+{
+	if (MeshUidsPerLodIndex.Num() == 0)
+	{
+		return nullptr;
+	}
+
+	// Create the static mesh factory node, name it according to the first mesh node compositing the meshes
+	FString StaticMeshUid_MeshNamePart;
+	FString DisplayLabel;
+	const int32 BaseLodIndex = 0;
+	if (!MakeMeshFactoryNodeUidAndDisplayLabel(MeshUidsPerLodIndex, BaseLodIndex, StaticMeshUid_MeshNamePart, DisplayLabel))
+	{
+		// Log an error
+		return nullptr;
+	}
+
+	const FString StaticMeshUid = TEXT("\\StaticMesh") + StaticMeshUid_MeshNamePart;
+	UInterchangeStaticMeshFactoryNode* StaticMeshFactoryNode = NewObject<UInterchangeStaticMeshFactoryNode>(BaseNodeContainer, NAME_None);
+	if (!ensure(StaticMeshFactoryNode))
+	{
+		return nullptr;
+	}
+
+	StaticMeshFactoryNode->InitializeStaticMeshNode(StaticMeshUid, DisplayLabel, UStaticMesh::StaticClass()->GetName());
+	BaseNodeContainer->AddNode(StaticMeshFactoryNode);
+
+	AddLodDataToStaticMesh(StaticMeshFactoryNode, MeshUidsPerLodIndex);
+
+	switch (VertexColorImportOption)
+	{
+	case EInterchangeVertexColorImportOption::IVCIO_Replace:
+	{
+		StaticMeshFactoryNode->SetCustomVertexColorReplace(true);
+	}
+	break;
+	case EInterchangeVertexColorImportOption::IVCIO_Ignore:
+	{
+		StaticMeshFactoryNode->SetCustomVertexColorIgnore(true);
+	}
+	break;
+	case EInterchangeVertexColorImportOption::IVCIO_Override:
+	{
+		StaticMeshFactoryNode->SetCustomVertexColorOverride(VertexOverrideColor);
+	}
+	break;
+	}
+	return StaticMeshFactoryNode;
+}
+
+
+UInterchangeStaticMeshLodDataNode* UInterchangeGenericMeshPipeline::CreateStaticMeshLodDataNode(const FString& NodeName, const FString& NodeUniqueID)
+{
+	FString DisplayLabel(NodeName);
+	FString NodeUID(NodeUniqueID);
+	UInterchangeStaticMeshLodDataNode* StaticMeshLodDataNode = NewObject<UInterchangeStaticMeshLodDataNode>(BaseNodeContainer, NAME_None);
+	if (!ensure(StaticMeshLodDataNode))
+	{
+		// @TODO: log error
+		return nullptr;
+	}
+
+	StaticMeshLodDataNode->InitializeNode(NodeUID, DisplayLabel, EInterchangeNodeContainerType::FactoryData);
+	BaseNodeContainer->AddNode(StaticMeshLodDataNode);
+	return StaticMeshLodDataNode;
+}
+
+
+void UInterchangeGenericMeshPipeline::AddLodDataToStaticMesh(UInterchangeStaticMeshFactoryNode* StaticMeshFactoryNode, const TMap<int32, TArray<FString>>& NodeUidsPerLodIndex)
+{
+	const FString StaticMeshFactoryUid = StaticMeshFactoryNode->GetUniqueID();
+
+	for (const TPair<int32, TArray<FString>>& LodIndexAndNodeUids : NodeUidsPerLodIndex)
+	{
+		const int32 LodIndex = LodIndexAndNodeUids.Key;
+		if (!bImportLods && LodIndex > 0)
+		{
+			// If the pipeline should not import lods, skip any lod over base lod
+			continue;
+		}
+
+		const TArray<FString>& NodeUids = LodIndexAndNodeUids.Value;
+
+		// Create a lod data node with all the meshes for this LOD
+		const FString StaticMeshLodDataName = TEXT("LodData") + FString::FromInt(LodIndex);
+		const FString LODDataPrefix = TEXT("\\LodData") + (LodIndex > 0 ? FString::FromInt(LodIndex) : TEXT(""));
+		const FString StaticMeshLodDataUniqueID = LODDataPrefix + StaticMeshFactoryUid;
+
+		// Create LodData node if it doesn't already exist
+		UInterchangeStaticMeshLodDataNode* LodDataNode = Cast<UInterchangeStaticMeshLodDataNode>(BaseNodeContainer->GetNode(StaticMeshLodDataUniqueID));
+		if (!LodDataNode)
+		{
+			// Add the data for the LOD (all the mesh node fbx path, so we can find them when we will create the payload data)
+			LodDataNode = CreateStaticMeshLodDataNode(StaticMeshLodDataName, StaticMeshLodDataUniqueID);
+			LodDataNode->SetParentUid(StaticMeshFactoryUid);
+			StaticMeshFactoryNode->AddLodDataUniqueId(StaticMeshLodDataUniqueID);
+		}
+
+		for (const FString& NodeUid : NodeUids)
+		{
+			TArray<FString> MaterialDependencies;
+			if (const UInterchangeSceneNode* SceneNode = Cast<UInterchangeSceneNode>(BaseNodeContainer->GetNode(NodeUid)))
+			{
+				FString MeshDependency;
+				SceneNode->GetCustomAssetInstanceUid(MeshDependency);
+				if (BaseNodeContainer->IsNodeUidValid(MeshDependency))
+				{
+					const UInterchangeMeshNode* MeshDependencyNode = Cast<UInterchangeMeshNode>(BaseNodeContainer->GetNode(MeshDependency));
+					StaticMeshFactoryNode->AddTargetNodeUid(MeshDependency);
+					MeshDependencyNode->AddTargetNodeUid(StaticMeshFactoryNode->GetUniqueID());
+					MeshDependencyNode->GetMaterialDependencies(MaterialDependencies);
+				}
+				else
+				{
+					SceneNode->GetMaterialDependencyUids(MaterialDependencies);
+				}
+			}
+			else if (const UInterchangeMeshNode* MeshNode = Cast<UInterchangeMeshNode>(BaseNodeContainer->GetNode(NodeUid)))
+			{
+				StaticMeshFactoryNode->AddTargetNodeUid(NodeUid);
+				MeshNode->AddTargetNodeUid(StaticMeshFactoryNode->GetUniqueID());
+				MeshNode->GetMaterialDependencies(MaterialDependencies);
+			}
+
+			const int32 MaterialCount = MaterialDependencies.Num();
+			for (int32 MaterialIndex = 0; MaterialIndex < MaterialCount; ++MaterialIndex)
+			{
+				const FString MaterialFactoryNodeUid = UInterchangeMaterialFactoryNode::GetMaterialFactoryNodeUidFromMaterialNodeUid(MaterialDependencies[MaterialIndex]);
+				if (BaseNodeContainer->IsNodeUidValid(MaterialFactoryNodeUid))
+				{
+					BaseNodeContainer->GetNode(MaterialFactoryNodeUid)->SetEnabled(true);
+					// Create a factory dependency so Material asset are imported before the skeletal mesh asset
+					TArray<FString> FactoryDependencies;
+					StaticMeshFactoryNode->GetFactoryDependencies(FactoryDependencies);
+					if (!FactoryDependencies.Contains(MaterialFactoryNodeUid))
+					{
+						StaticMeshFactoryNode->AddFactoryDependencyUid(MaterialFactoryNodeUid);
+					}
+				}
+			}
+
+			LodDataNode->AddMeshUid(NodeUid);
+		}
+	}
+}
diff --git a/Engine/Plugins/Experimental/Interchange/Runtime/Source/Pipelines/Public/InterchangeGenericAssetsPipeline.h b/Engine/Plugins/Experimental/Interchange/Runtime/Source/Pipelines/Public/InterchangeGenericAssetsPipeline.h
--- a/Engine/Plugins/Experimental/Interchange/Runtime/Source/Pipelines/Public/InterchangeGenericAssetsPipeline.h
+++ b/Engine/Plugins/Experimental/Interchange/Runtime/Source/Pipelines/Public/InterchangeGenericAssetsPipeline.h
@@ -13,44 +13,8 @@
 #include "InterchangeGenericAssetsPipeline.generated.h"
 
 class UInterchangeGenericMaterialPipeline;
+class UInterchangeGenericMeshPipeline;
 class UInterchangeGenericTexturePipeline;
-class UInterchangeMeshNode;
-class UInterchangePipelineMeshesUtilities;
-class UInterchangeResult;
-class UInterchangeSceneNode;
-class UInterchangeSkeletalMeshFactoryNode;
-class UInterchangeSkeletalMeshLodDataNode;
-class UInterchangeSkeletonFactoryNode;
-class UInterchangeStaticMeshFactoryNode;
-class UInterchangeStaticMeshLodDataNode;
-template<class T> class TSubclassOf;
-
-/** Force mesh type, if user want to import all meshes as one type*/
-UENUM(BlueprintType)
-enum EInterchangeForceMeshType
-{
-	/** Will import from the source type, no conversion */
-	IFMT_None UMETA(DisplayName = "None"),
-	/** Will import any mesh to static mesh. */
-	IFMT_StaticMesh UMETA(DisplayName = "Static Mesh"),
-	/** Will import any mesh to skeletal mesh. */
-	IFMT_SkeletalMesh UMETA(DisplayName = "Skeletal Mesh"),
-
-	IFMT_MAX
-};
-
-UENUM(BlueprintType)
-enum EInterchangeVertexColorImportOption
-{
-	/** Import the mesh using the vertex colors from the translated source. */
-	IVCIO_Replace UMETA(DisplayName = "Replace"),
-	/** Ignore vertex colors from the translated source. In case of a re-import keep the existing mesh vertex colors. */
-	IVCIO_Ignore UMETA(DisplayName = "Ignore"),
-	/** Override all vertex colors with the specified color. */
-	IVCIO_Override UMETA(DisplayName = "Override"),
-
-	IVCIO_MAX
-};
 
 /**
  * This pipeline is the generic pipeline option for all meshes type and should be call before specialized Mesh pipeline (like generic static mesh or skeletal mesh pipelines)
@@ -68,96 +32,20 @@ public:
 	//////////////////////////////////////////////////////////////////////////
 	// BEGIN Pre import pipeline properties, please keep by per category order for properties declaration
 
-
 	//////	COMMON_CATEGORY Properties //////
 
-
-	/** If enable and there is only one asset and one source data, we will name the asset like the source data name. */
-	UPROPERTY(EditAnywhere, BlueprintReadWrite, Category = "Common")
-	bool bUseSourceNameForAsset = true;
-
 	/* Allow user to choose the re-import strategy. */
 	UPROPERTY(EditAnywhere, BlueprintReadWrite, Category = "Common")
 	EReimportStrategyFlags ReimportStrategy = EReimportStrategyFlags::ApplyNoProperties;
 
+	/** If enable and there is only one asset and one source data, we will name the asset like the source data name. */
+	UPROPERTY(EditAnywhere, BlueprintReadWrite, Category = "Common")
+	bool bUseSourceNameForAsset = true;
 
-	//////	COMMON_MESHES_CATEGORY Properties //////
-
-
-	/** Allow to convert mesh to a particular type */
- 	UPROPERTY(EditAnywhere, BlueprintReadWrite, Category = "Common Meshes")
- 	TEnumAsByte<enum EInterchangeForceMeshType> ForceAllMeshAsType = EInterchangeForceMeshType::IFMT_None;
-
-	/** If enable, meshes LODs will be imported. Note that it required the advanced bBakeMesh property to be enabled. */
-	UPROPERTY(EditAnywhere, BlueprintReadWrite, Category = "Common Meshes", meta = (editcondition = "bBakeMeshes"))
-	bool bImportLods = true;
-
-	/** If enable, meshes will be baked with the scene instance hierarchy transform. */
-	UPROPERTY(EditAnywhere, BlueprintReadWrite, Category = "Common Meshes")
-	bool bBakeMeshes = true;
-
-	/** Specify how vertex colors should be imported */
-	UPROPERTY(EditAnywhere, BlueprintReadWrite, Category = "Common Meshes")
-	TEnumAsByte<enum EInterchangeVertexColorImportOption> VertexColorImportOption = EInterchangeVertexColorImportOption::IVCIO_Replace;
-
-	/** Specify override color in the case that VertexColorImportOption is set to Override */
-	UPROPERTY(EditAnywhere, BlueprintReadWrite, Category = "Common Meshes")
-	FColor VertexOverrideColor;
-	
-	//////	STATIC_MESHES_CATEGORY Properties //////
-
-
-	/** If enable, import the animation asset find in the sources. */
-	UPROPERTY(EditAnywhere, BlueprintReadWrite, Category = "Static Meshes")
-	bool bImportStaticMeshes = true;
-
-	/** If enable all translated static mesh node will be imported has a one static mesh. */
-	UPROPERTY(EditAnywhere, BlueprintReadWrite, Category = "Static Meshes")
-	bool bCombineStaticMeshes = false;
-
-	
-	//////  COMMON_SKELETAL_ANIMATIONS_CATEGORY //////
-
-
-	/** Skeleton to use for imported asset. When importing a skeletal mesh, leaving this as "None" will create a new skeleton. When importing an animation this MUST be specified to import the asset. */
- 	UPROPERTY(EditAnywhere, BlueprintReadWrite, Category = "Common Skeletal Mesh and Animations")
- 	TObjectPtr<class USkeleton> Skeleton;
-
-	
-	//////	SKELETAL_MESHES_CATEGORY Properties //////
-
-	
-	/** If enable, import the animation asset find in the sources. */
-	UPROPERTY(EditAnywhere, BlueprintReadWrite, Category = "Skeletal Meshes")
-	bool bImportSkeletalMeshes = true;
-
-	/** If enable all translated skinned mesh node will be imported has a one skeletal mesh, note that it can still create several skeletal mesh for each different skeleton root joint. */
-	UPROPERTY(EditAnywhere, BlueprintReadWrite, Category = "Skeletal Meshes")
-	bool bCombineSkeletalMeshes = true;
-
-	/** If enable any morph target shape will be imported. */
-	UPROPERTY(EditAnywhere, BlueprintReadWrite, Category = "Skeletal Meshes")
-	bool bImportMorphTargets = true;
-
-	/** Enable this option to update Skeleton (of the mesh)'s reference pose. Mesh's reference pose is always updated.  */
-	UPROPERTY(EditAnywhere, BlueprintReadWrite, Category = "Skeletal Meshes")
-	bool bUpdateSkeletonReferencePose = false;
-
-	/** If checked, meshes nested in bone hierarchies will be imported instead of being converted to bones. */
-	UPROPERTY(EditAnywhere, BlueprintReadWrite, Category = "Skeletal Meshes")
-	bool bImportMeshesInBoneHierarchy = true;
-
-	/** Enable this option to use frame 0 as reference pose */
-	UPROPERTY(EditAnywhere, BlueprintReadWrite, Category = "Skeletal Meshes")
-	bool bUseT0AsRefPose = false;
-
-	/** If checked, create new PhysicsAsset if it doesn't have it */
-	UPROPERTY(EditAnywhere, BlueprintReadWrite, Category = "Skeletal Meshes")
-	bool bCreatePhysicsAsset = true;
+	//////	MESHES_CATEGORY Properties //////
 
-	/** If this is set, use this specified PhysicsAsset. If its not set and bCreatePhysicsAsset is false, the importer will not generate or set any physic asset. */
-	UPROPERTY(EditAnywhere, BlueprintReadWrite, Category = "Skeletal Meshes", meta = (editcondition = "!bCreatePhysicsAsset"))
-	TObjectPtr<class UPhysicsAsset> PhysicsAsset;
+	UPROPERTY(VisibleAnywhere, Instanced, Category = "Meshes")
+	TObjectPtr<UInterchangeGenericMeshPipeline> MeshPipeline;
 	
 	//////	ANIMATIONS_CATEGORY Properties //////
 
@@ -199,91 +87,15 @@ protected:
 
 	//virtual bool ExecuteExportPipeline(UInterchangeBaseNodeContainer* BaseNodeContainer) override;
 private:
-
-	UInterchangeBaseNodeContainer* BaseNodeContainer;
-	TArray<const UInterchangeSourceData*> SourceDatas;
-	
-	/************************************************************************/
-	/* Skeletal mesh API BEGIN                                              */
-
-	/** Skeleton factory assets nodes */
-	TArray<UInterchangeSkeletonFactoryNode*> SkeletonFactoryNodes;
-
-	/** Create a UInterchangeSkeletonFactorynode */
-	UInterchangeSkeletonFactoryNode* CreateSkeletonFactoryNode(const FString& RootJointUid);
-	
-	/** Skeletal mesh factory assets nodes */
-	TArray<UInterchangeSkeletalMeshFactoryNode*> SkeletalMeshFactoryNodes;
 	
-	/**
-	 * This function can create a UInterchangeSkeletalMeshFactoryNode
-	 * @param MeshUidsPerLodIndex - The MeshUids can represent a SceneNode pointing on a MeshNode or directly a MeshNode
-	 */
-	UInterchangeSkeletalMeshFactoryNode* CreateSkeletalMeshFactoryNode(const FString& RootJointUid, const TMap<int32, TArray<FString>>& MeshUidsPerLodIndex);
-	
-	/** This function can create a UInterchangeSkeletalMeshLodDataNode which represent the LOD data need by the factory to create a lod mesh */
-	UInterchangeSkeletalMeshLodDataNode* CreateSkeletalMeshLodDataNode(const FString& NodeName, const FString& NodeUniqueID);
-
-	/**
-	 * This function add all lod data node to the skeletal mesh.
-	 * @param NodeUidsPerLodIndex - The NodeUids can be a UInterchangeSceneNode or a UInterchangeMeshNode. The scene node can bake each instance of the mesh versus the mesh node will import only the modelled mesh.
-	 */
-	void AddLodDataToSkeletalMesh(const UInterchangeSkeletonFactoryNode* SkeletonFactoryNode, UInterchangeSkeletalMeshFactoryNode* SkeletalMeshFactoryNode, const TMap<int32, TArray<FString>>& NodeUidsPerLodIndex);
-
-	/**
-	 * This function will create any skeletalmesh we need to create according to the pipeline options
-	 */
-	void ExecutePreImportPipelineSkeletalMesh();
-
-	/**
-	 * This function will finish creating the skeletalmesh asset
-	 */
-	void PostImportSkeletalMesh(UObject* CreatedAsset, UInterchangeBaseNode* Node);
-
-	/**
-	 * This function will finish creating the physics asset with the skeletalmesh render data
-	 */
-	void PostImportPhysicsAssetImport(UObject* CreatedAsset, UInterchangeBaseNode* Node);
-
-	/* Skeletal mesh API END                                                */
-	/************************************************************************/
-
-
-	/** Static mesh factory assets nodes */
-	TArray<UInterchangeStaticMeshFactoryNode*> StaticMeshFactoryNodes;
-
-	/**
-	 * This function can create a UInterchangeStaticMeshFactoryNode
-	 * @param MeshUidsPerLodIndex - The MeshUids can represent a SceneNode pointing on a MeshNode or directly a MeshNode
-	 */
-	UInterchangeStaticMeshFactoryNode* CreateStaticMeshFactoryNode(const TMap<int32, TArray<FString>>& MeshUidsPerLodIndex);
-	
-	/** This function can create a UInterchangeStaticMeshLodDataNode which represents the LOD data needed by the factory to create a lod mesh */
-	UInterchangeStaticMeshLodDataNode* CreateStaticMeshLodDataNode(const FString& NodeName, const FString& NodeUniqueID);
-
-	/**
-	 * This function add all lod data nodes to the static mesh.
-	 * @param NodeUidsPerLodIndex - The NodeUids can be a UInterchangeSceneNode or a UInterchangeMeshNode. The scene node can bake each instance of the mesh versus the mesh node will import only the modelled mesh.
-	 */
-	void AddLodDataToStaticMesh(UInterchangeStaticMeshFactoryNode* StaticMeshFactoryNode, const TMap<int32, TArray<FString>>& NodeUidsPerLodIndex);
-
-	/**
-	 * Return a reasonable UID and display label for a new mesh factory node.
-	 */
-	bool MakeMeshFactoryNodeUidAndDisplayLabel(const TMap<int32, TArray<FString>>& MeshUidsPerLodIndex, int32 LodIndex, FString& NewMeshUid, FString& DisplayLabel);
-
-
-	/* Meshes utilities, to parse the translated graph and extract the meshes informations. */
-	UInterchangePipelineMeshesUtilities* PipelineMeshesUtilities;
-
 	/**
 	 * Implement pipeline option bUseSourceNameForAsset
 	 */
 	void ImplementUseSourceNameForAssetOption();
 
-	/** Specialize for skeletalmesh */
-	void ImplementUseSourceNameForAssetOptionSkeletalMesh(const int32 MeshesAndAnimsImportedNodeCount);
-
+	UInterchangeBaseNodeContainer* BaseNodeContainer;
+	TArray<const UInterchangeSourceData*> SourceDatas;
+	
 };
 
 
diff --git a/Engine/Plugins/Experimental/Interchange/Runtime/Source/Pipelines/Public/InterchangeGenericMaterialPipeline.h b/Engine/Plugins/Experimental/Interchange/Runtime/Source/Pipelines/Public/InterchangeGenericMaterialPipeline.h
--- a/Engine/Plugins/Experimental/Interchange/Runtime/Source/Pipelines/Public/InterchangeGenericMaterialPipeline.h
+++ b/Engine/Plugins/Experimental/Interchange/Runtime/Source/Pipelines/Public/InterchangeGenericMaterialPipeline.h
@@ -40,7 +40,7 @@ public:
 	UPROPERTY(EditAnywhere, BlueprintReadWrite, Category = "Materials", Meta= (EditCondition="MaterialImport==EInterchangeMaterialImportOption::ImportAsMaterialInstances", AllowedClasses="MaterialInterface"))
 	FSoftObjectPath ParentMaterial;
 
-public:
+protected:
 	virtual void ExecutePreImportPipeline(UInterchangeBaseNodeContainer* InBaseNodeContainer, const TArray<UInterchangeSourceData*>& InSourceDatas) override;
 
 private:
diff --git a/Engine/Plugins/Experimental/Interchange/Runtime/Source/Pipelines/Public/InterchangeGenericMeshPipeline.h b/Engine/Plugins/Experimental/Interchange/Runtime/Source/Pipelines/Public/InterchangeGenericMeshPipeline.h
new file mode 100644
--- /dev/null
+++ b/Engine/Plugins/Experimental/Interchange/Runtime/Source/Pipelines/Public/InterchangeGenericMeshPipeline.h
@@ -0,0 +1,241 @@
+// Copyright Epic Games, Inc. All Rights Reserved.
+
+#pragma once
+
+#include "CoreMinimal.h"
+#include "InterchangePipelineBase.h"
+#include "InterchangeSourceData.h"
+#include "Nodes/InterchangeBaseNodeContainer.h"
+#include "UObject/Object.h"
+#include "UObject/ObjectMacros.h"
+
+#include "InterchangeGenericMeshPipeline.generated.h"
+
+class UInterchangeMeshNode;
+class UInterchangePipelineMeshesUtilities;
+class UInterchangeSkeletalMeshFactoryNode;
+class UInterchangeSkeletalMeshLodDataNode;
+class UInterchangeSkeletonFactoryNode;
+class UInterchangeStaticMeshFactoryNode;
+class UInterchangeStaticMeshLodDataNode;
+
+/** Force mesh type, if user want to import all meshes as one type*/
+UENUM(BlueprintType)
+enum EInterchangeForceMeshType
+{
+	/** Will import from the source type, no conversion */
+	IFMT_None UMETA(DisplayName = "None"),
+	/** Will import any mesh to static mesh. */
+	IFMT_StaticMesh UMETA(DisplayName = "Static Mesh"),
+	/** Will import any mesh to skeletal mesh. */
+	IFMT_SkeletalMesh UMETA(DisplayName = "Skeletal Mesh"),
+
+	IFMT_MAX
+};
+
+UENUM(BlueprintType)
+enum EInterchangeVertexColorImportOption
+{
+	/** Import the mesh using the vertex colors from the translated source. */
+	IVCIO_Replace UMETA(DisplayName = "Replace"),
+	/** Ignore vertex colors from the translated source. In case of a re-import keep the existing mesh vertex colors. */
+	IVCIO_Ignore UMETA(DisplayName = "Ignore"),
+	/** Override all vertex colors with the specified color. */
+	IVCIO_Override UMETA(DisplayName = "Override"),
+
+	IVCIO_MAX
+};
+
+UCLASS(BlueprintType, Experimental)
+class INTERCHANGEPIPELINES_API UInterchangeGenericMeshPipeline : public UInterchangePipelineBase
+{
+	GENERATED_BODY()
+
+public:
+
+	//////	COMMON_MESHES_CATEGORY Properties //////
+
+	/** Allow to convert mesh to a particular type */
+ 	UPROPERTY(EditAnywhere, BlueprintReadWrite, Category = "Common Meshes")
+ 	TEnumAsByte<enum EInterchangeForceMeshType> ForceAllMeshAsType = EInterchangeForceMeshType::IFMT_None;
+
+	/** If enable, meshes LODs will be imported. Note that it required the advanced bBakeMesh property to be enabled. */
+	UPROPERTY(EditAnywhere, BlueprintReadWrite, Category = "Common Meshes", meta = (editcondition = "bBakeMeshes"))
+	bool bImportLods = true;
+
+	/** If enable, meshes will be baked with the scene instance hierarchy transform. */
+	UPROPERTY(EditAnywhere, BlueprintReadWrite, Category = "Common Meshes")
+	bool bBakeMeshes = true;
+
+	/** Specify how vertex colors should be imported */
+	UPROPERTY(EditAnywhere, BlueprintReadWrite, Category = "Common Meshes")
+	TEnumAsByte<enum EInterchangeVertexColorImportOption> VertexColorImportOption = EInterchangeVertexColorImportOption::IVCIO_Replace;
+
+	/** Specify override color in the case that VertexColorImportOption is set to Override */
+	UPROPERTY(EditAnywhere, BlueprintReadWrite, Category = "Common Meshes")
+	FColor VertexOverrideColor;
+	
+	//////	STATIC_MESHES_CATEGORY Properties //////
+
+	/** If enable, import the animation asset find in the sources. */
+	UPROPERTY(EditAnywhere, BlueprintReadWrite, Category = "Static Meshes")
+	bool bImportStaticMeshes = true;
+
+	/** If enable all translated static mesh node will be imported has a one static mesh. */
+	UPROPERTY(EditAnywhere, BlueprintReadWrite, Category = "Static Meshes")
+	bool bCombineStaticMeshes = false;
+
+	
+	//////  COMMON_SKELETAL_ANIMATIONS_CATEGORY //////
+
+
+	/** Skeleton to use for imported asset. When importing a skeletal mesh, leaving this as "None" will create a new skeleton. When importing an animation this MUST be specified to import the asset. */
+ 	UPROPERTY(EditAnywhere, BlueprintReadWrite, Category = "Common Skeletal Mesh and Animations")
+ 	TObjectPtr<class USkeleton> Skeleton;
+
+	
+	//////	SKELETAL_MESHES_CATEGORY Properties //////
+
+	/** If enable, import the animation asset find in the sources. */
+	UPROPERTY(EditAnywhere, BlueprintReadWrite, Category = "Skeletal Meshes")
+	bool bImportSkeletalMeshes = true;
+
+	/** If enable all translated skinned mesh node will be imported has a one skeletal mesh, note that it can still create several skeletal mesh for each different skeleton root joint. */
+	UPROPERTY(EditAnywhere, BlueprintReadWrite, Category = "Skeletal Meshes")
+	bool bCombineSkeletalMeshes = true;
+
+	/** If enable any morph target shape will be imported. */
+	UPROPERTY(EditAnywhere, BlueprintReadWrite, Category = "Skeletal Meshes")
+	bool bImportMorphTargets = true;
+
+	/** Enable this option to update Skeleton (of the mesh)'s reference pose. Mesh's reference pose is always updated.  */
+	UPROPERTY(EditAnywhere, BlueprintReadWrite, Category = "Skeletal Meshes")
+	bool bUpdateSkeletonReferencePose = false;
+
+	/** If checked, meshes nested in bone hierarchies will be imported instead of being converted to bones. */
+	UPROPERTY(EditAnywhere, BlueprintReadWrite, Category = "Skeletal Meshes")
+	bool bImportMeshesInBoneHierarchy = true;
+
+	/** Enable this option to use frame 0 as reference pose */
+	UPROPERTY(EditAnywhere, BlueprintReadWrite, Category = "Skeletal Meshes")
+	bool bUseT0AsRefPose = false;
+
+	/** If checked, create new PhysicsAsset if it doesn't have it */
+	UPROPERTY(EditAnywhere, BlueprintReadWrite, Category = "Skeletal Meshes")
+	bool bCreatePhysicsAsset = true;
+
+	/** If this is set, use this specified PhysicsAsset. If its not set and bCreatePhysicsAsset is false, the importer will not generate or set any physic asset. */
+	UPROPERTY(EditAnywhere, BlueprintReadWrite, Category = "Skeletal Meshes", meta = (editcondition = "!bCreatePhysicsAsset"))
+	TObjectPtr<class UPhysicsAsset> PhysicsAsset;
+
+	virtual void PreDialogCleanup(const FName PipelineStackName) override;
+
+protected:
+	virtual void ExecutePreImportPipeline(UInterchangeBaseNodeContainer* InBaseNodeContainer, const TArray<UInterchangeSourceData*>& InSourceDatas) override;
+
+	virtual void ExecutePostImportPipeline(const UInterchangeBaseNodeContainer* InBaseNodeContainer, const FString& NodeKey, UObject* CreatedAsset, bool bIsAReimport) override;
+
+	virtual bool CanExecuteOnAnyThread(EInterchangePipelineTask PipelineTask) override
+	{
+		return true;
+	}
+
+private:
+
+	/* Meshes utilities, to parse the translated graph and extract the meshes informations. */
+	UInterchangePipelineMeshesUtilities* PipelineMeshesUtilities;
+
+
+	/************************************************************************/
+	/* Skeletal mesh API BEGIN                                              */
+
+	/**
+	 * This function will create any skeletalmesh we need to create according to the pipeline options
+	 */
+	void ExecutePreImportPipelineSkeletalMesh();
+
+	/** Skeleton factory assets nodes */
+	TArray<UInterchangeSkeletonFactoryNode*> SkeletonFactoryNodes;
+
+	/** Create a UInterchangeSkeletonFactorynode */
+	UInterchangeSkeletonFactoryNode* CreateSkeletonFactoryNode(const FString& RootJointUid);
+
+	/** Skeletal mesh factory assets nodes */
+	TArray<UInterchangeSkeletalMeshFactoryNode*> SkeletalMeshFactoryNodes;
+
+	/**
+	 * This function can create a UInterchangeSkeletalMeshFactoryNode
+	 * @param MeshUidsPerLodIndex - The MeshUids can represent a SceneNode pointing on a MeshNode or directly a MeshNode
+	 */
+	UInterchangeSkeletalMeshFactoryNode* CreateSkeletalMeshFactoryNode(const FString& RootJointUid, const TMap<int32, TArray<FString>>& MeshUidsPerLodIndex);
+
+	/** This function can create a UInterchangeSkeletalMeshLodDataNode which represent the LOD data need by the factory to create a lod mesh */
+	UInterchangeSkeletalMeshLodDataNode* CreateSkeletalMeshLodDataNode(const FString& NodeName, const FString& NodeUniqueID);
+
+	/**
+	 * This function add all lod data node to the skeletal mesh.
+	 * @param NodeUidsPerLodIndex - The NodeUids can be a UInterchangeSceneNode or a UInterchangeMeshNode. The scene node can bake each instance of the mesh versus the mesh node will import only the modelled mesh.
+	 */
+	void AddLodDataToSkeletalMesh(const UInterchangeSkeletonFactoryNode* SkeletonFactoryNode, UInterchangeSkeletalMeshFactoryNode* SkeletalMeshFactoryNode, const TMap<int32, TArray<FString>>& NodeUidsPerLodIndex);
+
+	/**
+	 * This function will finish creating the skeletalmesh asset
+	 */
+	void PostImportSkeletalMesh(UObject* CreatedAsset, UInterchangeBaseNode* Node);
+
+	/**
+	 * This function will finish creating the physics asset with the skeletalmesh render data
+	 */
+	void PostImportPhysicsAssetImport(UObject* CreatedAsset, UInterchangeBaseNode* Node);
+public:
+	
+	/** Specialize for skeletalmesh */
+	void ImplementUseSourceNameForAssetOptionSkeletalMesh(const int32 MeshesAndAnimsImportedNodeCount, const bool bUseSourceNameForAsset);
+
+private:
+	/* Skeletal mesh API END                                                */
+	/************************************************************************/
+
+
+	/************************************************************************/
+	/* Static mesh API BEGIN                                              */
+
+	/**
+	 * This function will create any skeletalmesh we need to create according to the pipeline options
+	 */
+	void ExecutePreImportPipelineStaticMesh();
+
+	/** Static mesh factory assets nodes */
+	TArray<UInterchangeStaticMeshFactoryNode*> StaticMeshFactoryNodes;
+
+	/**
+	 * This function can create a UInterchangeStaticMeshFactoryNode
+	 * @param MeshUidsPerLodIndex - The MeshUids can represent a SceneNode pointing on a MeshNode or directly a MeshNode
+	 */
+	UInterchangeStaticMeshFactoryNode* CreateStaticMeshFactoryNode(const TMap<int32, TArray<FString>>& MeshUidsPerLodIndex);
+
+	/** This function can create a UInterchangeStaticMeshLodDataNode which represents the LOD data needed by the factory to create a lod mesh */
+	UInterchangeStaticMeshLodDataNode* CreateStaticMeshLodDataNode(const FString& NodeName, const FString& NodeUniqueID);
+
+	/**
+	 * This function add all lod data nodes to the static mesh.
+	 * @param NodeUidsPerLodIndex - The NodeUids can be a UInterchangeSceneNode or a UInterchangeMeshNode. The scene node can bake each instance of the mesh versus the mesh node will import only the modelled mesh.
+	 */
+	void AddLodDataToStaticMesh(UInterchangeStaticMeshFactoryNode* StaticMeshFactoryNode, const TMap<int32, TArray<FString>>& NodeUidsPerLodIndex);
+
+	/**
+	 * Return a reasonable UID and display label for a new mesh factory node.
+	 */
+	bool MakeMeshFactoryNodeUidAndDisplayLabel(const TMap<int32, TArray<FString>>& MeshUidsPerLodIndex, int32 LodIndex, FString& NewMeshUid, FString& DisplayLabel);
+
+	/* Static mesh API END                                                */
+	/************************************************************************/
+
+private:
+
+	UInterchangeBaseNodeContainer* BaseNodeContainer = nullptr;
+	TArray<const UInterchangeSourceData*> SourceDatas;
+
+};
+
+
diff --git a/Engine/Plugins/MeshPainting/Source/MeshPaintEditorMode/Private/MeshPaintModeHelpers.cpp b/Engine/Plugins/MeshPainting/Source/MeshPaintEditorMode/Private/MeshPaintModeHelpers.cpp
--- a/Engine/Plugins/MeshPainting/Source/MeshPaintEditorMode/Private/MeshPaintModeHelpers.cpp
+++ b/Engine/Plugins/MeshPainting/Source/MeshPaintEditorMode/Private/MeshPaintModeHelpers.cpp
@@ -48,6 +48,7 @@
 #include "MeshTexturePaintingTool.h"
 #include "InterchangeAssetImportData.h"
 #include "InterchangeGenericAssetsPipeline.h"
+#include "InterchangeGenericMeshPipeline.h"
 
 
 void UMeshPaintModeSubsystem::SetViewportColorMode(EMeshPaintDataColorViewMode ColorViewMode, FEditorViewportClient* ViewportClient)
@@ -321,11 +322,11 @@ void UMeshPaintModeSubsystem::ImportVertexColorsToSkeletalMesh(USkeletalMesh* Sk
 			{
 				if (UInterchangeGenericAssetsPipeline* GenericAssetPipeline = Cast<UInterchangeGenericAssetsPipeline>(PipelineBase.Get()))
 				{
-					if (GenericAssetPipeline->VertexColorImportOption != EInterchangeVertexColorImportOption::IVCIO_Ignore)
+					if (GenericAssetPipeline->MeshPipeline && GenericAssetPipeline->MeshPipeline->VertexColorImportOption != EInterchangeVertexColorImportOption::IVCIO_Ignore)
 					{
 						InterchangeAssetImportData->SetFlags(RF_Transactional);
 						InterchangeAssetImportData->Modify();
-						GenericAssetPipeline->VertexColorImportOption = EInterchangeVertexColorImportOption::IVCIO_Ignore;
+						GenericAssetPipeline->MeshPipeline->VertexColorImportOption = EInterchangeVertexColorImportOption::IVCIO_Ignore;
 					}
 				}
 			}
diff --git a/Engine/Plugins/MeshPainting/Source/MeshPaintingToolset/Private/MeshPaintSkeletalMeshAdapter.cpp b/Engine/Plugins/MeshPainting/Source/MeshPaintingToolset/Private/MeshPaintSkeletalMeshAdapter.cpp
--- a/Engine/Plugins/MeshPainting/Source/MeshPaintingToolset/Private/MeshPaintSkeletalMeshAdapter.cpp
+++ b/Engine/Plugins/MeshPainting/Source/MeshPaintingToolset/Private/MeshPaintSkeletalMeshAdapter.cpp
@@ -11,6 +11,7 @@
 #include "IndexTypes.h"
 #include "InterchangeAssetImportData.h"
 #include "InterchangeGenericAssetsPipeline.h"
+#include "InterchangeGenericMeshPipeline.h"
 
 DEFINE_LOG_CATEGORY_STATIC(LogMeshPaintSkeletalMeshAdapter, Log, All);
 //////////////////////////////////////////////////////////////////////////
@@ -370,11 +371,11 @@ void FMeshPaintSkeletalMeshComponentAdapter::PreEdit()
 			{
 				if (UInterchangeGenericAssetsPipeline* GenericAssetPipeline = Cast<UInterchangeGenericAssetsPipeline>(PipelineBase.Get()))
 				{
-					if (GenericAssetPipeline->VertexColorImportOption != EInterchangeVertexColorImportOption::IVCIO_Ignore)
+					if (GenericAssetPipeline->MeshPipeline && GenericAssetPipeline->MeshPipeline->VertexColorImportOption != EInterchangeVertexColorImportOption::IVCIO_Ignore)
 					{
 						GenericAssetPipeline->SetFlags(RF_Transactional);
 						GenericAssetPipeline->Modify();
-						GenericAssetPipeline->VertexColorImportOption = EInterchangeVertexColorImportOption::IVCIO_Ignore;
+						GenericAssetPipeline->MeshPipeline->VertexColorImportOption = EInterchangeVertexColorImportOption::IVCIO_Ignore;
 					}
 				}
 			}
