Interchange Materials
- Added support for the clear coat and thin translucent shading models.
- Fixed the lerp node not handling connections to its factor input.
- Flag textures as linear when connected to a non-color material input.
- Added support for OneMinus when visiting nodes to extract their values for the graph to instance conversion.

Interchange GLTF
- Added support for KHR_materials_clearcoat and KHR_materials_transmission extensions.
- We'll now append the node index when generating the asset names since names aren't guaranteed to be unique in GLTF.

diff --git a/Engine/Plugins/Experimental/Interchange/Runtime/Source/FactoryNodes/Private/InterchangeMaterialFactoryNode.cpp b/Engine/Plugins/Experimental/Interchange/Runtime/Source/FactoryNodes/Private/InterchangeMaterialFactoryNode.cpp
--- a/Engine/Plugins/Experimental/Interchange/Runtime/Source/FactoryNodes/Private/InterchangeMaterialFactoryNode.cpp
+++ b/Engine/Plugins/Experimental/Interchange/Runtime/Source/FactoryNodes/Private/InterchangeMaterialFactoryNode.cpp
@@ -164,7 +164,67 @@ bool UInterchangeMaterialFactoryNode::ConnectToRefraction(const FString& Attribu
 
 bool UInterchangeMaterialFactoryNode::ConnectOutputToRefraction(const FString& ExpressionNodeUid, const FString& OutputName)
 {
-	return UInterchangeShaderPortsAPI::ConnectOuputToInput(this, UE::Interchange::Materials::PBR::Parameters::IndexOfRefraction.ToString(), ExpressionNodeUid, OutputName);
+	return UInterchangeShaderPortsAPI::ConnectOuputToInput(this, UE::Interchange::Materials::Common::Parameters::IndexOfRefraction.ToString(), ExpressionNodeUid, OutputName);
+}
+
+bool UInterchangeMaterialFactoryNode::GetClearCoatConnection(FString& ExpressionNodeUid, FString& OutputName) const
+{
+	return UInterchangeShaderPortsAPI::GetInputConnection(this, UE::Interchange::Materials::ClearCoat::Parameters::ClearCoat.ToString(), ExpressionNodeUid, OutputName);
+}
+
+bool UInterchangeMaterialFactoryNode::ConnectToClearCoat(const FString& AttributeValue)
+{
+	return UInterchangeShaderPortsAPI::ConnectDefaultOuputToInput(this, UE::Interchange::Materials::ClearCoat::Parameters::ClearCoat.ToString(), AttributeValue);
+}
+
+bool UInterchangeMaterialFactoryNode::ConnectOutputToClearCoat(const FString& ExpressionNodeUid, const FString& OutputName)
+{
+	return UInterchangeShaderPortsAPI::ConnectOuputToInput(this, UE::Interchange::Materials::ClearCoat::Parameters::ClearCoat.ToString(), ExpressionNodeUid, OutputName);
+}
+
+bool UInterchangeMaterialFactoryNode::GetClearCoatRoughnessConnection(FString& ExpressionNodeUid, FString& OutputName) const
+{
+	return UInterchangeShaderPortsAPI::GetInputConnection(this, UE::Interchange::Materials::ClearCoat::Parameters::ClearCoatRoughness.ToString(), ExpressionNodeUid, OutputName);
+}
+
+bool UInterchangeMaterialFactoryNode::ConnectToClearCoatRoughness(const FString& AttributeValue)
+{
+	return UInterchangeShaderPortsAPI::ConnectDefaultOuputToInput(this, UE::Interchange::Materials::ClearCoat::Parameters::ClearCoatRoughness.ToString(), AttributeValue);
+}
+
+bool UInterchangeMaterialFactoryNode::ConnectOutputToClearCoatRoughness(const FString& ExpressionNodeUid, const FString& OutputName)
+{
+	return UInterchangeShaderPortsAPI::ConnectOuputToInput(this, UE::Interchange::Materials::ClearCoat::Parameters::ClearCoatRoughness.ToString(), ExpressionNodeUid, OutputName);
+}
+
+bool UInterchangeMaterialFactoryNode::GetClearCoatNormalConnection(FString& ExpressionNodeUid, FString& OutputName) const
+{
+	return UInterchangeShaderPortsAPI::GetInputConnection(this, UE::Interchange::Materials::ClearCoat::Parameters::ClearCoatNormal.ToString(), ExpressionNodeUid, OutputName);
+}
+
+bool UInterchangeMaterialFactoryNode::ConnectToClearCoatNormal(const FString& AttributeValue)
+{
+	return UInterchangeShaderPortsAPI::ConnectDefaultOuputToInput(this, UE::Interchange::Materials::ClearCoat::Parameters::ClearCoatNormal.ToString(), AttributeValue);
+}
+
+bool UInterchangeMaterialFactoryNode::ConnectOutputToClearCoatNormal(const FString& ExpressionNodeUid, const FString& OutputName)
+{
+	return UInterchangeShaderPortsAPI::ConnectOuputToInput(this, UE::Interchange::Materials::ClearCoat::Parameters::ClearCoatNormal.ToString(), ExpressionNodeUid, OutputName);
+}
+
+bool UInterchangeMaterialFactoryNode::GetTransmissionColorConnection(FString& ExpressionNodeUid, FString& OutputName) const
+{
+	return UInterchangeShaderPortsAPI::GetInputConnection(this, UE::Interchange::Materials::ThinTranslucent::Parameters::TransmissionColor.ToString(), ExpressionNodeUid, OutputName);
+}
+
+bool UInterchangeMaterialFactoryNode::ConnectToTransmissionColor(const FString& AttributeValue)
+{
+	return UInterchangeShaderPortsAPI::ConnectDefaultOuputToInput(this, UE::Interchange::Materials::ThinTranslucent::Parameters::TransmissionColor.ToString(), AttributeValue);
+}
+
+bool UInterchangeMaterialFactoryNode::ConnectOutputToTransmissionColor(const FString& ExpressionNodeUid, const FString& OutputName)
+{
+	return UInterchangeShaderPortsAPI::ConnectOuputToInput(this, UE::Interchange::Materials::ThinTranslucent::Parameters::TransmissionColor.ToString(), ExpressionNodeUid, OutputName);
 }
 
 bool UInterchangeMaterialFactoryNode::GetCustomShadingModel(TEnumAsByte<EMaterialShadingModel>& AttributeValue) const
diff --git a/Engine/Plugins/Experimental/Interchange/Runtime/Source/FactoryNodes/Public/InterchangeMaterialFactoryNode.h b/Engine/Plugins/Experimental/Interchange/Runtime/Source/FactoryNodes/Public/InterchangeMaterialFactoryNode.h
--- a/Engine/Plugins/Experimental/Interchange/Runtime/Source/FactoryNodes/Public/InterchangeMaterialFactoryNode.h
+++ b/Engine/Plugins/Experimental/Interchange/Runtime/Source/FactoryNodes/Public/InterchangeMaterialFactoryNode.h
@@ -117,6 +117,42 @@ public:
 	UFUNCTION(BlueprintCallable, Category = "Interchange | Node | MaterialFactory")
 	bool ConnectOutputToRefraction(const FString& ExpressionNodeUid, const FString& OutputName);
 
+	UFUNCTION(BlueprintCallable, Category = "Interchange | Node | MaterialFactory")
+	bool GetClearCoatConnection(FString& ExpressionNodeUid, FString& OutputName) const;
+
+	UFUNCTION(BlueprintCallable, Category = "Interchange | Node | MaterialFactory")
+	bool ConnectToClearCoat(const FString& AttributeValue);
+
+	UFUNCTION(BlueprintCallable, Category = "Interchange | Node | MaterialFactory")
+	bool ConnectOutputToClearCoat(const FString& ExpressionNodeUid, const FString& OutputName);
+
+	UFUNCTION(BlueprintCallable, Category = "Interchange | Node | MaterialFactory")
+	bool GetClearCoatRoughnessConnection(FString& ExpressionNodeUid, FString& OutputName) const;
+
+	UFUNCTION(BlueprintCallable, Category = "Interchange | Node | MaterialFactory")
+	bool ConnectToClearCoatRoughness(const FString& AttributeValue);
+
+	UFUNCTION(BlueprintCallable, Category = "Interchange | Node | MaterialFactory")
+	bool ConnectOutputToClearCoatRoughness(const FString& ExpressionNodeUid, const FString& OutputName);
+
+	UFUNCTION(BlueprintCallable, Category = "Interchange | Node | MaterialFactory")
+	bool GetClearCoatNormalConnection(FString& ExpressionNodeUid, FString& OutputName) const;
+
+	UFUNCTION(BlueprintCallable, Category = "Interchange | Node | MaterialFactory")
+	bool ConnectToClearCoatNormal(const FString& AttributeValue);
+
+	UFUNCTION(BlueprintCallable, Category = "Interchange | Node | MaterialFactory")
+	bool ConnectOutputToClearCoatNormal(const FString& ExpressionNodeUid, const FString& OutputName);
+
+	UFUNCTION(BlueprintCallable, Category = "Interchange | Node | MaterialFactory")
+	bool GetTransmissionColorConnection(FString& ExpressionNodeUid, FString& OutputName) const;
+
+	UFUNCTION(BlueprintCallable, Category = "Interchange | Node | MaterialFactory")
+	bool ConnectToTransmissionColor(const FString& AttributeValue);
+
+	UFUNCTION(BlueprintCallable, Category = "Interchange | Node | MaterialFactory")
+	bool ConnectOutputToTransmissionColor(const FString& ExpressionNodeUid, const FString& OutputName);
+
 // Material parameters
 public:
 	UFUNCTION(BlueprintCallable, Category = "Interchange | Node | Material")
diff --git a/Engine/Plugins/Experimental/Interchange/Runtime/Source/Import/Private/Gltf/InterchangeGltfTranslator.cpp b/Engine/Plugins/Experimental/Interchange/Runtime/Source/Import/Private/Gltf/InterchangeGltfTranslator.cpp
--- a/Engine/Plugins/Experimental/Interchange/Runtime/Source/Import/Private/Gltf/InterchangeGltfTranslator.cpp
+++ b/Engine/Plugins/Experimental/Interchange/Runtime/Source/Import/Private/Gltf/InterchangeGltfTranslator.cpp
@@ -19,9 +19,25 @@
 #include "StaticMeshAttributes.h"
 #include "UObject/GCObjectScopeGuard.h"
 
-void UInterchangeGltfTranslator::HandleGltfNode( UInterchangeBaseNodeContainer& NodeContainer, const GLTF::FNode& GltfNode, const FString& ParentNodeUid ) const
+namespace UE::Interchange::Gltf::Private
 {
-	const FString NodeUid = ParentNodeUid + TEXT("\\") + GltfNode.Name;
+	FString GenerateUniqueIdForGltfNode( const FString& NodeName, int32 NodeIndex )
+	{
+		if ( NodeIndex == 0 )
+		{
+			return NodeName;
+		}
+		else
+		{
+			return NodeName + TEXT("_") + LexToString( NodeIndex );
+		}
+	}
+}
+void UInterchangeGltfTranslator::HandleGltfNode( UInterchangeBaseNodeContainer& NodeContainer, const GLTF::FNode& GltfNode, const FString& ParentNodeUid, const int32 NodeIndex ) const
+{
+	using namespace UE::Interchange::Gltf::Private;
+
+	const FString NodeUid = ParentNodeUid + TEXT("\\") + GenerateUniqueIdForGltfNode( GltfNode.Name, NodeIndex );
 
 	UInterchangeSceneNode* ParentSceneNode = Cast< UInterchangeSceneNode >( NodeContainer.GetNode( ParentNodeUid ) );
 
@@ -40,7 +56,7 @@ void UInterchangeGltfTranslator::HandleGltfNode( UInterchangeBaseNodeContainer&
 		{
 			if ( GltfAsset.Meshes.IsValidIndex( GltfNode.MeshIndex ) )
 			{
-				const FString MeshNodeUid = TEXT("\\Mesh\\") + GltfAsset.Meshes[ GltfNode.MeshIndex ].Name;
+				const FString MeshNodeUid = TEXT("\\Mesh\\") + GenerateUniqueIdForGltfNode( GltfAsset.Meshes[ GltfNode.MeshIndex ].Name, GltfNode.MeshIndex );
 				InterchangeSceneNode->SetCustomAssetInstanceUid( MeshNodeUid );
 			}
 			break;
@@ -52,7 +68,7 @@ void UInterchangeGltfTranslator::HandleGltfNode( UInterchangeBaseNodeContainer&
 
 			if ( GltfAsset.Cameras.IsValidIndex( GltfNode.CameraIndex ) )
 			{
-				const FString CameraNodeUid = TEXT("\\Camera\\") + GltfAsset.Cameras[ GltfNode.CameraIndex ].Name;
+				const FString CameraNodeUid = TEXT("\\Camera\\") + GenerateUniqueIdForGltfNode( GltfAsset.Cameras[ GltfNode.CameraIndex ].Name, GltfNode.CameraIndex );
 				InterchangeSceneNode->SetCustomAssetInstanceUid( CameraNodeUid );
 			}
 			break;
@@ -64,7 +80,7 @@ void UInterchangeGltfTranslator::HandleGltfNode( UInterchangeBaseNodeContainer&
 
 			if ( GltfAsset.Lights.IsValidIndex( GltfNode.LightIndex ) )
 			{
-				const FString LightNodeUid = TEXT("\\Light\\") + GltfAsset.Lights[ GltfNode.LightIndex ].Name;
+				const FString LightNodeUid = TEXT("\\Light\\") + GenerateUniqueIdForGltfNode( GltfAsset.Lights[ GltfNode.LightIndex ].Name, GltfNode.LightIndex );
 				InterchangeSceneNode->SetCustomAssetInstanceUid( LightNodeUid );
 			}
 		}
@@ -87,27 +103,27 @@ void UInterchangeGltfTranslator::HandleGltfNode( UInterchangeBaseNodeContainer&
 	{
 		if ( GltfAsset.Nodes.IsValidIndex( ChildIndex ) )
 		{
-			HandleGltfNode( NodeContainer, GltfAsset.Nodes[ ChildIndex ], NodeUid );
+			HandleGltfNode( NodeContainer, GltfAsset.Nodes[ ChildIndex ], NodeUid, ChildIndex );
 		}
 	}
 }
 
-void UInterchangeGltfTranslator::HandleGltfMaterialParameter( UInterchangeBaseNodeContainer& NodeContainer, const GLTF::FTextureMap& TextureMap, UInterchangeShaderGraphNode& ShaderGraphNode,
+void UInterchangeGltfTranslator::HandleGltfMaterialParameter( UInterchangeBaseNodeContainer& NodeContainer, const GLTF::FTextureMap& TextureMap, UInterchangeShaderNode& ShaderNode,
 		const FString& MapName, const TVariant< FLinearColor, float >& MapFactor, const FString& OutputChannel, const bool bInverse ) const
 {
 	using namespace UE::Interchange::Materials;
 
-	UInterchangeShaderNode* NodeToConnectTo = &ShaderGraphNode;
+	UInterchangeShaderNode* NodeToConnectTo = &ShaderNode;
 	FString InputToConnectTo = MapName;
 
 	if (bInverse)
 	{
 		const FString OneMinusNodeName = MapName + TEXT("OneMinus");
-		const FString OneMinusNodeUid = ShaderGraphNode.GetUniqueID() + TEXT("_") + OneMinusNodeName;
+		const FString OneMinusNodeUid = ShaderNode.GetUniqueID() + TEXT("_") + OneMinusNodeName;
 		UInterchangeShaderNode* OneMinusNode = NewObject< UInterchangeShaderNode >( &NodeContainer );
 		OneMinusNode->InitializeNode( OneMinusNodeUid, OneMinusNodeName, EInterchangeNodeContainerType::TranslatedAsset );
 		NodeContainer.AddNode( OneMinusNode );
-		NodeContainer.SetNodeParentUid( OneMinusNodeUid, ShaderGraphNode.GetUniqueID() );
+		NodeContainer.SetNodeParentUid( OneMinusNodeUid, ShaderNode.GetUniqueID() );
 
 		OneMinusNode->SetCustomShaderType(Standard::Nodes::OneMinus::Name.ToString());
 
@@ -117,15 +133,26 @@ void UInterchangeGltfTranslator::HandleGltfMaterialParameter( UInterchangeBaseNo
 		InputToConnectTo = Standard::Nodes::OneMinus::Inputs::Input.ToString();
 	}
 
-	if ( GltfAsset.Textures.IsValidIndex( TextureMap.TextureIndex ) )
+	bool bTextureHasImportance = true;
+
+	if ( MapFactor.IsType< float >() )
+	{
+		bTextureHasImportance = !FMath::IsNearlyZero( MapFactor.Get< float >() );
+	}
+	else if ( MapFactor.IsType< FLinearColor >() )
+	{
+		bTextureHasImportance = !MapFactor.Get< FLinearColor >().IsAlmostBlack();
+	}
+
+	if ( bTextureHasImportance && GltfAsset.Textures.IsValidIndex( TextureMap.TextureIndex ) )
 	{
 		const FString NodeName = MapName;
-		const FString NodeUid = ShaderGraphNode.GetUniqueID() + TEXT("_") + NodeName;
+		const FString NodeUid = ShaderNode.GetUniqueID() + TEXT("_") + NodeName;
 
 		UInterchangeShaderNode* ColorNode = NewObject< UInterchangeShaderNode >( &NodeContainer );
 		ColorNode->InitializeNode( NodeUid, NodeName, EInterchangeNodeContainerType::TranslatedAsset );
 		NodeContainer.AddNode( ColorNode );
-		NodeContainer.SetNodeParentUid( NodeUid, ShaderGraphNode.GetUniqueID() );
+		NodeContainer.SetNodeParentUid( NodeUid, ShaderNode.GetUniqueID() );
 
 		ColorNode->SetCustomShaderType( Standard::Nodes::TextureSample::Name.ToString() );
 		const FString TextureUid = TEXT("\\Texture\\") + GltfAsset.Textures[ TextureMap.TextureIndex ].Source.URI;
@@ -148,7 +175,7 @@ void UInterchangeGltfTranslator::HandleGltfMaterialParameter( UInterchangeBaseNo
 			UInterchangeShaderNode* FactorNode = NewObject< UInterchangeShaderNode >( &NodeContainer );
 			FactorNode->InitializeNode( FactorNodeUid, NodeName + TEXT("_Factor"), EInterchangeNodeContainerType::TranslatedAsset );
 			NodeContainer.AddNode( FactorNode );
-			NodeContainer.SetNodeParentUid( FactorNodeUid, ShaderGraphNode.GetUniqueID() );
+			NodeContainer.SetNodeParentUid( FactorNodeUid, ShaderNode.GetUniqueID() );
 
 			FactorNode->SetCustomShaderType( Standard::Nodes::Multiply::Name.ToString() );
 
@@ -298,10 +325,226 @@ void UInterchangeGltfTranslator::HandleGltfMaterial( UInterchangeBaseNodeContain
 
 			HandleGltfMaterialParameter( NodeContainer, GltfMaterial.BaseColor, ShaderGraphNode, PBR::Parameters::Opacity.ToString(),
 				OpacityFactor, Standard::Nodes::TextureSample::Outputs::A.ToString() );
+		}
 
+		// IOR
+		if ( GltfMaterial.bHasIOR )
+		{
 			ShaderGraphNode.AddFloatAttribute( UInterchangeShaderPortsAPI::MakeInputValueKey( PBR::Parameters::IndexOfRefraction.ToString() ), GltfMaterial.IOR );
 		}
 	}
+
+	if ( GltfMaterial.bHasClearCoat )
+	{
+		HandleGltfClearCoat( NodeContainer, GltfMaterial, ShaderGraphNode );
+	}
+
+	if ( GltfMaterial.bHasTransmission )
+	{
+		HandleGltfTransmission( NodeContainer, GltfMaterial, ShaderGraphNode );
+	}
+}
+
+void UInterchangeGltfTranslator::HandleGltfClearCoat( UInterchangeBaseNodeContainer& NodeContainer, const GLTF::FMaterial& GltfMaterial, UInterchangeShaderGraphNode& ShaderGraphNode ) const
+{
+	using namespace UE::Interchange::Materials;
+
+	if ( !GltfMaterial.bHasClearCoat || FMath::IsNearlyZero( GltfMaterial.ClearCoat.ClearCoatFactor ) )
+	{
+		return;
+	}
+
+	// ClearCoat::Parameters::ClearCoat
+	{
+		TVariant< FLinearColor, float > ClearCoatFactor;
+		ClearCoatFactor.Set< float >( GltfMaterial.ClearCoat.ClearCoatFactor );
+
+		HandleGltfMaterialParameter( NodeContainer, GltfMaterial.ClearCoat.ClearCoatMap, ShaderGraphNode, ClearCoat::Parameters::ClearCoat.ToString(),
+			ClearCoatFactor, Standard::Nodes::TextureSample::Outputs::RGB.ToString() );
+	}
+
+	//  ClearCoat::Parameters::ClearCoatRoughness
+	{
+		TVariant< FLinearColor, float > ClearCoatRoughnessFactor;
+		ClearCoatRoughnessFactor.Set< float >( GltfMaterial.ClearCoat.Roughness );
+
+		HandleGltfMaterialParameter( NodeContainer, GltfMaterial.ClearCoat.RoughnessMap, ShaderGraphNode, ClearCoat::Parameters::ClearCoatRoughness.ToString(),
+			ClearCoatRoughnessFactor, Standard::Nodes::TextureSample::Outputs::RGB.ToString() );
+	}
+
+	// ClearCoat::Parameters::ClearCoatNormal
+	{
+		TVariant< FLinearColor, float > ClearCoatNormalFactor;
+		ClearCoatNormalFactor.Set< FLinearColor >( FLinearColor::White );
+
+		HandleGltfMaterialParameter( NodeContainer, GltfMaterial.ClearCoat.NormalMap, ShaderGraphNode, ClearCoat::Parameters::ClearCoatNormal.ToString(),
+			ClearCoatNormalFactor, Standard::Nodes::TextureSample::Outputs::RGB.ToString() );
+	}
+}
+
+/**
+ * GLTF transmission is handled a little differently than UE's.
+ * GLTF doesn't allow having different reflected and transmitted colors, UE does (base color vs transmittance color).
+ * GLTF controls the amount of reflected light vs transmitted light using the transmission factor, UE does that through opacity.
+ * GLTF opacity means that the medium is present of not, so it's normal for transmission materials to be considered opaque,
+ * meaning that the medium is full present, and the transmission factor determines how much lights is transmitted.
+ * When a transmission material isn't fully opaque, we reduce the transmission color by the opacity to mimic GLTF's BTDF.
+ * Ideally, this would be better represented by blending a default lit alpha blended material with a thin translucent material based on GLTF's opacity.
+ */
+void UInterchangeGltfTranslator::HandleGltfTransmission( UInterchangeBaseNodeContainer& NodeContainer, const GLTF::FMaterial& GltfMaterial, UInterchangeShaderGraphNode& ShaderGraphNode ) const
+{
+	using namespace UE::Interchange::Materials;
+
+	if ( !GltfMaterial.bHasTransmission || FMath::IsNearlyZero( GltfMaterial.Transmission.TransmissionFactor ) )
+	{
+		return;
+	}
+ 
+	FString OpacityNodeUid;
+	FString OpacityNodeOutput;
+
+	// Common::Parameters::Opacity
+	{
+		/**
+		 * Per the spec, the red channel of the transmission texture drives how much light is transmitted vs diffused.
+		 * So we're setting the inverse of the red channel as the opacity.
+		 */
+		const FString OneMinusNodeName =  TEXT("OpacityOneMinus");
+		const FString OneMinusNodeUid = ShaderGraphNode.GetUniqueID() + TEXT("_") + OneMinusNodeName;
+		UInterchangeShaderNode* OneMinusNode = NewObject< UInterchangeShaderNode >( &NodeContainer );
+		OneMinusNode->InitializeNode( OneMinusNodeUid, OneMinusNodeName, EInterchangeNodeContainerType::TranslatedAsset );
+		NodeContainer.AddNode( OneMinusNode );
+		NodeContainer.SetNodeParentUid( OneMinusNodeUid, ShaderGraphNode.GetUniqueID() );
+
+		OneMinusNode->SetCustomShaderType(Standard::Nodes::OneMinus::Name.ToString());
+
+		UInterchangeShaderNode* CurrentNode = OneMinusNode;
+		FString CurrentInput = Standard::Nodes::OneMinus::Inputs::Input.ToString();
+
+		TVariant< FLinearColor, float > TransmissionFactor;
+		TransmissionFactor.Set< float >( GltfMaterial.Transmission.TransmissionFactor );
+
+		HandleGltfMaterialParameter( NodeContainer, GltfMaterial.Transmission.TransmissionMap, *CurrentNode, CurrentInput,
+			TransmissionFactor, Standard::Nodes::TextureSample::Outputs::R.ToString() );
+
+		// The GLTF transmission model specifies that metallic surfaces don't transmit light, so adjust Common::Parameters::Opacity so that metallic surfaces are opaque.
+		{
+			FString MetallicNodeUid;
+			FString MetallicNodeOutput;
+
+			if ( UInterchangeShaderPortsAPI::GetInputConnection( &ShaderGraphNode, PBR::Parameters::Metallic.ToString(), MetallicNodeUid, MetallicNodeOutput ) )
+			{
+				const FString MetallicLerpNodeName =  TEXT("OpacityMetallicLerp");
+				const FString MetallicLerpNodeUid = ShaderGraphNode.GetUniqueID() + TEXT("_") + MetallicLerpNodeName;
+
+				UInterchangeShaderNode* LerpMetallicNode = NewObject< UInterchangeShaderNode >( &NodeContainer );
+				LerpMetallicNode->InitializeNode( MetallicLerpNodeUid, MetallicLerpNodeName, EInterchangeNodeContainerType::TranslatedAsset );
+				LerpMetallicNode->SetCustomShaderType( Standard::Nodes::Lerp::Name.ToString() );
+
+				NodeContainer.AddNode( LerpMetallicNode );
+				NodeContainer.SetNodeParentUid( MetallicLerpNodeUid, ShaderGraphNode.GetUniqueID() );
+
+				LerpMetallicNode->AddFloatAttribute( UInterchangeShaderPortsAPI::MakeInputValueKey( Standard::Nodes::Lerp::Inputs::B.ToString() ), 1.f );
+				UInterchangeShaderPortsAPI::ConnectDefaultOuputToInput( LerpMetallicNode, Standard::Nodes::Lerp::Inputs::A.ToString(), CurrentNode->GetUniqueID() );
+				UInterchangeShaderPortsAPI::ConnectOuputToInput( LerpMetallicNode, Standard::Nodes::Lerp::Inputs::Factor.ToString(), MetallicNodeUid, MetallicNodeOutput );
+
+				CurrentNode = LerpMetallicNode;
+				CurrentInput = TEXT("");
+			}
+		}
+
+		if ( GltfMaterial.AlphaMode != GLTF::FMaterial::EAlphaMode::Opaque )
+		{
+			if ( UInterchangeShaderPortsAPI::GetInputConnection( &ShaderGraphNode, PBR::Parameters::Opacity.ToString(), OpacityNodeUid, OpacityNodeOutput ) )
+			{
+				const FString OpacityLerpNodeName =  TEXT("OpacityLerp");
+				const FString OpaciotyLerpNodeUid = ShaderGraphNode.GetUniqueID() + TEXT("_") + OpacityLerpNodeName;
+
+				UInterchangeShaderNode* OpacityLerpNode = NewObject< UInterchangeShaderNode >( &NodeContainer );
+				OpacityLerpNode->InitializeNode( OpaciotyLerpNodeUid, OpacityLerpNodeName, EInterchangeNodeContainerType::TranslatedAsset );
+				OpacityLerpNode->SetCustomShaderType( Standard::Nodes::Lerp::Name.ToString() );
+
+				NodeContainer.AddNode( OpacityLerpNode );
+				NodeContainer.SetNodeParentUid( OpaciotyLerpNodeUid, ShaderGraphNode.GetUniqueID() );
+
+				OpacityLerpNode->AddFloatAttribute( UInterchangeShaderPortsAPI::MakeInputValueKey( Standard::Nodes::Lerp::Inputs::A.ToString() ), 0.f );
+				UInterchangeShaderPortsAPI::ConnectDefaultOuputToInput( OpacityLerpNode, Standard::Nodes::Lerp::Inputs::B.ToString(), CurrentNode->GetUniqueID() );
+				UInterchangeShaderPortsAPI::ConnectOuputToInput( OpacityLerpNode, Standard::Nodes::Lerp::Inputs::Factor.ToString(), OpacityNodeUid, OpacityNodeOutput );
+
+				CurrentNode = OpacityLerpNode;
+				CurrentInput = TEXT("");
+			}
+		}
+
+		UInterchangeShaderPortsAPI::ConnectDefaultOuputToInput( &ShaderGraphNode, Common::Parameters::Opacity.ToString(), CurrentNode->GetUniqueID() );
+	}
+
+	// ThinTranslucent::Parameters::Transmissioncolor
+	{
+		// There's no separation of reflected and transmitted color in this model. So the same color is used for the base color and the transmitted color.
+		// Since this extension is only supported with the metallic-roughness model, we can reuse its base color
+		UInterchangeBaseNode* CurrentNode = &ShaderGraphNode;
+		FString CurrentOuput = TEXT("");
+		FLinearColor CurrentColor = FLinearColor::White;
+
+		FString BaseColorNodeUid;
+		FString BaseColorNodeOutput;
+
+		if ( UInterchangeShaderPortsAPI::GetInputConnection( CurrentNode, PBR::Parameters::BaseColor.ToString(), BaseColorNodeUid, BaseColorNodeOutput ) )
+		{
+			CurrentNode = NodeContainer.GetNode( BaseColorNodeUid );
+			CurrentOuput = BaseColorNodeOutput;
+		}
+		else
+		{
+			FLinearColor BaseColor;
+			if ( ShaderGraphNode.GetLinearColorAttribute( UInterchangeShaderPortsAPI::MakeInputValueKey( PBR::Parameters::BaseColor.ToString() ), BaseColor ) )
+			{
+				CurrentNode = nullptr;
+				CurrentColor = BaseColor;
+			}
+		}
+
+		if ( GltfMaterial.AlphaMode != GLTF::FMaterial::EAlphaMode::Opaque )
+		{
+			if ( !OpacityNodeUid.IsEmpty() )
+			{
+				const FString OpacityLerpNodeName =  TEXT("OpacityTransmissionLerp");
+				const FString OpaciotyLerpNodeUid = ShaderGraphNode.GetUniqueID() + TEXT("_") + OpacityLerpNodeName;
+
+				UInterchangeShaderNode* OpacityLerpNode = NewObject< UInterchangeShaderNode >( &NodeContainer );
+				OpacityLerpNode->InitializeNode( OpaciotyLerpNodeUid, OpacityLerpNodeName, EInterchangeNodeContainerType::TranslatedAsset );
+				OpacityLerpNode->SetCustomShaderType( Standard::Nodes::Lerp::Name.ToString() );
+
+				NodeContainer.AddNode( OpacityLerpNode );
+				NodeContainer.SetNodeParentUid( OpaciotyLerpNodeUid, ShaderGraphNode.GetUniqueID() );
+
+				OpacityLerpNode->AddLinearColorAttribute( UInterchangeShaderPortsAPI::MakeInputValueKey( Standard::Nodes::Lerp::Inputs::A.ToString() ), FLinearColor::White );
+				UInterchangeShaderPortsAPI::ConnectOuputToInput( OpacityLerpNode, Standard::Nodes::Lerp::Inputs::Factor.ToString(), OpacityNodeUid, OpacityNodeOutput );
+
+				if ( CurrentNode )
+				{
+					UInterchangeShaderPortsAPI::ConnectOuputToInput( OpacityLerpNode, Standard::Nodes::Lerp::Inputs::B.ToString(), CurrentNode->GetUniqueID(), CurrentOuput );
+				}
+				else
+				{
+					OpacityLerpNode->AddLinearColorAttribute( UInterchangeShaderPortsAPI::MakeInputValueKey( Standard::Nodes::Lerp::Inputs::B.ToString() ), CurrentColor );
+				}
+
+				CurrentNode = OpacityLerpNode;
+				CurrentOuput = TEXT("");
+			}
+		}
+
+		if ( CurrentNode )
+		{
+			UInterchangeShaderPortsAPI::ConnectOuputToInput( &ShaderGraphNode, ThinTranslucent::Parameters::TransmissionColor.ToString(), CurrentNode->GetUniqueID(), CurrentOuput );
+		}
+		else
+		{
+			ShaderGraphNode.AddLinearColorAttribute( UInterchangeShaderPortsAPI::MakeInputValueKey( ThinTranslucent::Parameters::TransmissionColor.ToString() ), CurrentColor );
+		}
+	}
 }
 
 EInterchangeTranslatorType UInterchangeGltfTranslator::GetTranslatorType() const
@@ -321,6 +564,8 @@ TArray<FString> UInterchangeGltfTranslator::GetSupportedFormats() const
 
 bool UInterchangeGltfTranslator::Translate( UInterchangeBaseNodeContainer& NodeContainer ) const
 {
+	using namespace UE::Interchange::Gltf::Private;
+
 	FString Filename = GetSourceData()->GetFilename();
 	if ( !FPaths::FileExists( Filename ) )
 	{
@@ -348,14 +593,16 @@ bool UInterchangeGltfTranslator::Translate( UInterchangeBaseNodeContainer& NodeC
 
 	// Materials
 	{
+		int32 MaterialIndex = 0;
 		for ( const GLTF::FMaterial& GltfMaterial : GltfAsset.Materials )
 		{
 			UInterchangeShaderGraphNode* ShaderGraphNode = NewObject< UInterchangeShaderGraphNode >( &NodeContainer );
-			FString ShaderGraphNodeUid = TEXT("\\Material\\") + GltfMaterial.Name;
-			ShaderGraphNode->InitializeNode( ShaderGraphNodeUid, GltfMaterial.Name, EInterchangeNodeContainerType::TranslatedAsset );
+			FString ShaderGraphNodeUid = TEXT("\\Material\\") + GenerateUniqueIdForGltfNode( GltfMaterial.Name, MaterialIndex );
+			ShaderGraphNode->InitializeNode( ShaderGraphNodeUid, GenerateUniqueIdForGltfNode( GltfMaterial.Name, MaterialIndex ), EInterchangeNodeContainerType::TranslatedAsset );
 			NodeContainer.AddNode( ShaderGraphNode );
 
 			HandleGltfMaterial( NodeContainer, GltfMaterial, *ShaderGraphNode );
+			++MaterialIndex;
 		}
 	}
 
@@ -365,7 +612,8 @@ bool UInterchangeGltfTranslator::Translate( UInterchangeBaseNodeContainer& NodeC
 		for ( const GLTF::FMesh& GltfMesh : GltfAsset.Meshes )
 		{
 			UInterchangeMeshNode* MeshNode = NewObject< UInterchangeMeshNode >( &NodeContainer );
-			FString MeshNodeUid = TEXT("\\Mesh\\") + GltfMesh.Name;
+			FString MeshNodeUid = TEXT("\\Mesh\\") + GenerateUniqueIdForGltfNode( GltfMesh.Name, MeshIndex );
+
 			MeshNode->InitializeNode( MeshNodeUid, GltfMesh.Name, EInterchangeNodeContainerType::TranslatedAsset );
 			MeshNode->SetPayLoadKey( LexToString( MeshIndex++ ) );
 			NodeContainer.AddNode( MeshNode );
@@ -375,7 +623,7 @@ bool UInterchangeGltfTranslator::Translate( UInterchangeBaseNodeContainer& NodeC
 			{
 				if ( GltfAsset.Materials.IsValidIndex( Primitive.MaterialIndex ) )
 				{
-					const FString ShaderGraphNodeUid = TEXT("\\Material\\") + GltfAsset.Materials[ Primitive.MaterialIndex ].Name;
+					const FString ShaderGraphNodeUid = TEXT("\\Material\\") + GenerateUniqueIdForGltfNode( GltfAsset.Materials[ Primitive.MaterialIndex ].Name, Primitive.MaterialIndex );
 					MeshNode->SetMaterialDependencyUid( ShaderGraphNodeUid );
 				}
 			}
@@ -384,23 +632,27 @@ bool UInterchangeGltfTranslator::Translate( UInterchangeBaseNodeContainer& NodeC
 
 	// Cameras
 	{
+		int32 CameraIndex = 0;
 		for ( const GLTF::FCamera& GltfCamera : GltfAsset.Cameras )
 		{
 			UInterchangeCameraNode* CameraNode = NewObject< UInterchangeCameraNode >( &NodeContainer );
-			FString CameraNodeUid = TEXT("\\Camera\\") + GltfCamera.Name;
+			FString CameraNodeUid = TEXT("\\Camera\\") + GenerateUniqueIdForGltfNode( GltfCamera.Name, CameraIndex );
 			CameraNode->InitializeNode( CameraNodeUid, GltfCamera.Name, EInterchangeNodeContainerType::TranslatedAsset );
 			NodeContainer.AddNode( CameraNode );
+			++CameraIndex;
 		}
 	}
 
 	// Lights
 	{
+		int32 LightIndex = 0;
 		for ( const GLTF::FLight& GltfLight : GltfAsset.Lights )
 		{
 			UInterchangeLightNode* LightNode = NewObject< UInterchangeLightNode >( &NodeContainer );
-			FString LightNodeUid = TEXT("\\Light\\") + GltfLight.Name;
+			FString LightNodeUid = TEXT("\\Light\\") + GenerateUniqueIdForGltfNode( GltfLight.Name, LightIndex );
 			LightNode->InitializeNode( LightNodeUid, GltfLight.Name, EInterchangeNodeContainerType::TranslatedAsset );
 			NodeContainer.AddNode( LightNode );
+			++LightIndex;
 		}
 	}
 
@@ -415,12 +667,10 @@ bool UInterchangeGltfTranslator::Translate( UInterchangeBaseNodeContainer& NodeC
 			if (SceneName.IsEmpty())
 			{
 				SceneName = TEXT("Scene");
-				if ( GltfAsset.Scenes.Num() > 1 )
-				{
-					SceneName += TEXT("_") + LexToString( SceneIndex );
-				}
 			}
 
+			SceneName = GenerateUniqueIdForGltfNode( SceneName, SceneIndex );
+
 			FString SceneNodeUid = TEXT("\\Scene\\") + SceneName;
 			SceneNode->InitializeNode( SceneNodeUid, SceneName, EInterchangeNodeContainerType::TranslatedScene );
 			NodeContainer.AddNode( SceneNode );
@@ -429,7 +679,7 @@ bool UInterchangeGltfTranslator::Translate( UInterchangeBaseNodeContainer& NodeC
 			{
 				if ( GltfAsset.Nodes.IsValidIndex( NodeIndex ) )
 				{
-					HandleGltfNode( NodeContainer, GltfAsset.Nodes[ NodeIndex ], SceneNodeUid );
+					HandleGltfNode( NodeContainer, GltfAsset.Nodes[ NodeIndex ], SceneNodeUid, NodeIndex );
 				}
 			}
 
@@ -442,6 +692,8 @@ bool UInterchangeGltfTranslator::Translate( UInterchangeBaseNodeContainer& NodeC
 
 TFuture< TOptional< UE::Interchange::FStaticMeshPayloadData > > UInterchangeGltfTranslator::GetStaticMeshPayloadData( const FString& PayLoadKey ) const
 {
+	using namespace UE::Interchange::Gltf::Private;
+
 	TPromise< TOptional< UE::Interchange::FStaticMeshPayloadData > > MeshPayloadDataPromise;
 
 	int32 MeshIndex = 0;
@@ -471,7 +723,7 @@ TFuture< TOptional< UE::Interchange::FStaticMeshPayloadData > > UInterchangeGltf
 
 			if ( GltfAsset.Materials.IsValidIndex( MaterialIndex ) )
 			{
-				StaticMeshAttributes.GetPolygonGroupMaterialSlotNames()[ MaterialSlotIndex ] = *GltfAsset.Materials[ MaterialIndex ].Name;
+				StaticMeshAttributes.GetPolygonGroupMaterialSlotNames()[ MaterialSlotIndex ] = *GenerateUniqueIdForGltfNode( GltfAsset.Materials[ MaterialIndex ].Name, MaterialIndex );
 			}
 		}
 	}
diff --git a/Engine/Plugins/Experimental/Interchange/Runtime/Source/Import/Private/Material/InterchangeMaterialFactory.cpp b/Engine/Plugins/Experimental/Interchange/Runtime/Source/Import/Private/Material/InterchangeMaterialFactory.cpp
--- a/Engine/Plugins/Experimental/Interchange/Runtime/Source/Import/Private/Material/InterchangeMaterialFactory.cpp
+++ b/Engine/Plugins/Experimental/Interchange/Runtime/Source/Import/Private/Material/InterchangeMaterialFactory.cpp
@@ -10,9 +10,11 @@
 #include "InterchangeSourceData.h"
 #include "Materials/Material.h"
 #include "Materials/MaterialExpression.h"
+#include "Materials/MaterialExpressionClearCoatNormalCustomOutput.h"
 #include "Materials/MaterialExpressionScalarParameter.h"
 #include "Materials/MaterialExpressionTextureCoordinate.h"
 #include "Materials/MaterialExpressionTextureSample.h"
+#include "Materials/MaterialExpressionThinTranslucentMaterialOutput.h"
 #include "Materials/MaterialExpressionVectorParameter.h"
 #include "Materials/MaterialInterface.h"
 #include "Materials/MaterialInstanceConstant.h"
@@ -544,6 +546,92 @@ void UInterchangeMaterialFactory::SetupMaterial(UMaterial* Material, const FCrea
 		}
 	}
 
+	// Clear Coat
+	{
+		FString ClearCoatUid;
+		FString OutputName;
+
+		if (MaterialFactoryNode->GetClearCoatConnection(ClearCoatUid, OutputName))
+		{
+			const UInterchangeMaterialExpressionFactoryNode* ClearCoatNode = Cast<UInterchangeMaterialExpressionFactoryNode>(Arguments.NodeContainer->GetNode(ClearCoatUid));
+
+			if (ClearCoatNode)
+			{
+				if (UMaterialExpression* ClearCoatExpression = CreateExpressionsForNode(Material, Arguments, ClearCoatNode, Expressions))
+				{
+					if (FExpressionInput* ClearCoatInput = Material->GetExpressionInputForProperty(MP_CustomData0))
+					{
+						ClearCoatExpression->ConnectExpression(ClearCoatInput, GetOutputIndex(*ClearCoatExpression, OutputName));
+					}
+				}
+			}
+		}
+	}
+
+	// Clear Coat Roughness
+	{
+		FString ClearCoatRoughnessUid;
+		FString OutputName;
+
+		if (MaterialFactoryNode->GetClearCoatRoughnessConnection(ClearCoatRoughnessUid, OutputName))
+		{
+			const UInterchangeMaterialExpressionFactoryNode* ClearCoatRoughnessNode = Cast<UInterchangeMaterialExpressionFactoryNode>(Arguments.NodeContainer->GetNode(ClearCoatRoughnessUid));
+
+			if (ClearCoatRoughnessNode)
+			{
+				if (UMaterialExpression* ClearCoatRoughnessExpression = CreateExpressionsForNode(Material, Arguments, ClearCoatRoughnessNode, Expressions))
+				{
+					if (FExpressionInput* ClearCoatRoughnessInput = Material->GetExpressionInputForProperty(MP_CustomData1))
+					{
+						ClearCoatRoughnessExpression->ConnectExpression(ClearCoatRoughnessInput, GetOutputIndex(*ClearCoatRoughnessExpression, OutputName));
+					}
+				}
+			}
+		}
+	}
+
+	// Clear Coat Normal
+	{
+		FString ClearCoatNormalUid;
+		FString OutputName;
+
+		if (MaterialFactoryNode->GetClearCoatNormalConnection(ClearCoatNormalUid, OutputName))
+		{
+			const UInterchangeMaterialExpressionFactoryNode* ClearCoatNormalNode = Cast<UInterchangeMaterialExpressionFactoryNode>(Arguments.NodeContainer->GetNode(ClearCoatNormalUid));
+
+			if (ClearCoatNormalNode)
+			{
+				if (UMaterialExpression* ClearCoatNormalExpression = CreateExpressionsForNode(Material, Arguments, ClearCoatNormalNode, Expressions))
+				{
+					UMaterialExpression* ClearCoatNormalCustomOutput = CreateMaterialExpression(Material, UMaterialExpressionClearCoatNormalCustomOutput::StaticClass());
+					
+					ClearCoatNormalExpression->ConnectExpression(ClearCoatNormalCustomOutput->GetInput(0), GetOutputIndex(*ClearCoatNormalExpression, OutputName));
+				}
+			}
+		}
+	}
+
+	// Thin Translucent
+	{
+		FString TransmissionColorUid;
+		FString OutputName;
+
+		if (MaterialFactoryNode->GetTransmissionColorConnection(TransmissionColorUid, OutputName))
+		{
+			const UInterchangeMaterialExpressionFactoryNode* TransmissionColorNode = Cast<UInterchangeMaterialExpressionFactoryNode>(Arguments.NodeContainer->GetNode(TransmissionColorUid));
+
+			if (TransmissionColorNode)
+			{
+				if (UMaterialExpression* TransmissionColorExpression = CreateExpressionsForNode(Material, Arguments, TransmissionColorNode, Expressions))
+				{
+					UMaterialExpression* ThinTranslucentMaterialOutput = CreateMaterialExpression(Material, UMaterialExpressionThinTranslucentMaterialOutput::StaticClass());
+					
+					TransmissionColorExpression->ConnectExpression(ThinTranslucentMaterialOutput->GetInput(0), GetOutputIndex(*TransmissionColorExpression, OutputName));
+				}
+			}
+		}
+	}
+
 	UMaterialEditingLibrary::LayoutMaterialExpressions(Material);
 }
 
diff --git a/Engine/Plugins/Experimental/Interchange/Runtime/Source/Import/Public/Gltf/InterchangeGltfTranslator.h b/Engine/Plugins/Experimental/Interchange/Runtime/Source/Import/Public/Gltf/InterchangeGltfTranslator.h
--- a/Engine/Plugins/Experimental/Interchange/Runtime/Source/Import/Public/Gltf/InterchangeGltfTranslator.h
+++ b/Engine/Plugins/Experimental/Interchange/Runtime/Source/Import/Public/Gltf/InterchangeGltfTranslator.h
@@ -13,6 +13,7 @@
 #include "InterchangeGltfTranslator.generated.h"
 
 class UInterchangeShaderGraphNode;
+class UInterchangeShaderNode;
 
 /* Gltf translator class support import of texture, material, static mesh, skeletal mesh, */
 
@@ -42,11 +43,16 @@ public:
 	/* IInterchangeTexturePayloadInterface End */
 
 protected:
-	void HandleGltfNode( UInterchangeBaseNodeContainer& NodeContainer, const GLTF::FNode& GltfNode, const FString& ParentNodeUid ) const;
+	void HandleGltfNode( UInterchangeBaseNodeContainer& NodeContainer, const GLTF::FNode& GltfNode, const FString& ParentNodeUid, const int32 NodeIndex ) const;
 	void HandleGltfMaterial( UInterchangeBaseNodeContainer& NodeContainer, const GLTF::FMaterial& GltfMaterial, UInterchangeShaderGraphNode& ShaderGraphNode ) const;
-	void HandleGltfMaterialParameter( UInterchangeBaseNodeContainer& NodeContainer, const GLTF::FTextureMap& TextureMap, UInterchangeShaderGraphNode& ShaderGraphNode,
+	void HandleGltfMaterialParameter( UInterchangeBaseNodeContainer& NodeContainer, const GLTF::FTextureMap& TextureMap, UInterchangeShaderNode& ShaderNode,
 		const FString& MapName, const TVariant< FLinearColor, float >& MapFactor, const FString& OutputChannel, const bool bInverse = false ) const;
 
+	/** Support for KHR_materials_clearcoat */
+	void HandleGltfClearCoat( UInterchangeBaseNodeContainer& NodeContainer, const GLTF::FMaterial& GltfMaterial, UInterchangeShaderGraphNode& ShaderGraphNode ) const;
+	/** Support for KHR_materials_transmission */
+	void HandleGltfTransmission( UInterchangeBaseNodeContainer& NodeContainer, const GLTF::FMaterial& GltfMaterial, UInterchangeShaderGraphNode& ShaderGraphNode ) const;
+
 private:
 	GLTF::FAsset GltfAsset;
 };
diff --git a/Engine/Plugins/Experimental/Interchange/Runtime/Source/Nodes/Public/InterchangeMaterialDefinitions.h b/Engine/Plugins/Experimental/Interchange/Runtime/Source/Nodes/Public/InterchangeMaterialDefinitions.h
--- a/Engine/Plugins/Experimental/Interchange/Runtime/Source/Nodes/Public/InterchangeMaterialDefinitions.h
+++ b/Engine/Plugins/Experimental/Interchange/Runtime/Source/Nodes/Public/InterchangeMaterialDefinitions.h
@@ -114,6 +114,28 @@ namespace Interchange
 				const FName Roughness = TEXT("Roughness"); // Type: float
 			}
 		}
+
+		namespace ClearCoat
+		{
+			namespace Parameters
+			{
+				using namespace PBR::Parameters;
+
+				const FName ClearCoat = TEXT("ClearCoat"); // Type: float
+				const FName ClearCoatRoughness = TEXT("ClearCoatRoughness"); // Type: float
+				const FName ClearCoatNormal = TEXT("ClearCoatNormal"); // Type: vector3
+			}
+		}
+
+		namespace ThinTranslucent
+		{
+			namespace Parameters
+			{
+				using namespace PBR::Parameters;
+
+				const FName TransmissionColor = TEXT("TransmissionColor"); // Type: vector3
+			}
+		}
 	}
 }
 }
diff --git a/Engine/Plugins/Experimental/Interchange/Runtime/Source/Pipelines/Private/InterchangeGenericMaterialPipeline.cpp b/Engine/Plugins/Experimental/Interchange/Runtime/Source/Pipelines/Private/InterchangeGenericMaterialPipeline.cpp
--- a/Engine/Plugins/Experimental/Interchange/Runtime/Source/Pipelines/Private/InterchangeGenericMaterialPipeline.cpp
+++ b/Engine/Plugins/Experimental/Interchange/Runtime/Source/Pipelines/Private/InterchangeGenericMaterialPipeline.cpp
@@ -16,6 +16,7 @@
 #include "InterchangeSourceData.h"
 
 #include "Materials/MaterialExpressionDivide.h"
+#include "Materials/MaterialExpressionFresnel.h"
 #include "Materials/MaterialExpressionLinearInterpolate.h"
 #include "Materials/MaterialExpressionMaterialFunctionCall.h"
 #include "Materials/MaterialExpressionOneMinus.h"
@@ -122,6 +123,24 @@ UInterchangeBaseMaterialFactoryNode* UInterchangeGenericMaterialPipeline::Create
 	return MaterialFactoryNode;
 }
 
+bool UInterchangeGenericMaterialPipeline::IsClearCoatModel(const UInterchangeShaderGraphNode* ShaderGraphNode) const
+{
+	using namespace UE::Interchange::Materials::ClearCoat;
+
+	const bool bHasClearCoatInput = UInterchangeShaderPortsAPI::HasInput(ShaderGraphNode, Parameters::ClearCoat);
+
+	return bHasClearCoatInput;
+}
+
+bool UInterchangeGenericMaterialPipeline::IsThinTranslucentModel(const UInterchangeShaderGraphNode* ShaderGraphNode) const
+{
+	using namespace UE::Interchange::Materials::ThinTranslucent;
+
+	const bool bHasTransmissionColorInput = UInterchangeShaderPortsAPI::HasInput(ShaderGraphNode, Parameters::TransmissionColor);
+
+	return bHasTransmissionColorInput;
+}
+
 bool UInterchangeGenericMaterialPipeline::IsPBRModel(const UInterchangeShaderGraphNode* ShaderGraphNode) const
 {
 	using namespace UE::Interchange::Materials::PBR;
@@ -190,6 +209,8 @@ bool UInterchangeGenericMaterialPipeline::HandlePhongModel(const UInterchangeSha
 
 		// Specular
 		{
+			TGuardValue<bool> ParsingForLinearInputGuard(bParsingForLinearInput, true);
+
 			TTuple<UInterchangeMaterialExpressionFactoryNode*, FString> SpecularExpressionFactoryNode =
 				CreateMaterialExpressionForInput(MaterialFactoryNode, ShaderGraphNode, Parameters::SpecularColor.ToString(), FunctionCallExpression->GetUniqueID());
 
@@ -205,6 +226,8 @@ bool UInterchangeGenericMaterialPipeline::HandlePhongModel(const UInterchangeSha
 			const bool bHasShininessInput = UInterchangeShaderPortsAPI::HasInput(ShaderGraphNode, Parameters::Shininess);
 			if (bHasShininessInput)
 			{
+				TGuardValue<bool> ParsingForLinearInputGuard(bParsingForLinearInput, true);
+
 				TTuple<UInterchangeMaterialExpressionFactoryNode*, FString> ShininessExpressionFactoryNode =
 					CreateMaterialExpressionForInput(MaterialFactoryNode, ShaderGraphNode, Parameters::Shininess.ToString(), MaterialFactoryNode->GetUniqueID());
 
@@ -290,6 +313,8 @@ bool UInterchangeGenericMaterialPipeline::HandlePBRModel(const UInterchangeShade
 
 		if (bHasInput)
 		{
+			TGuardValue<bool> ParsingForLinearInputGuard(bParsingForLinearInput, true);
+
 			TTuple<UInterchangeMaterialExpressionFactoryNode*, FString> ExpressionFactoryNode =
 				CreateMaterialExpressionForInput(MaterialFactoryNode, ShaderGraphNode, Parameters::Metallic.ToString(), MaterialFactoryNode->GetUniqueID());
 
@@ -308,6 +333,8 @@ bool UInterchangeGenericMaterialPipeline::HandlePBRModel(const UInterchangeShade
 
 		if (bHasInput)
 		{
+			TGuardValue<bool> ParsingForLinearInputGuard(bParsingForLinearInput, true);
+
 			TTuple<UInterchangeMaterialExpressionFactoryNode*, FString> ExpressionFactoryNode =
 				CreateMaterialExpressionForInput(MaterialFactoryNode, ShaderGraphNode, Parameters::Specular.ToString(), MaterialFactoryNode->GetUniqueID());
 
@@ -326,6 +353,8 @@ bool UInterchangeGenericMaterialPipeline::HandlePBRModel(const UInterchangeShade
 
 		if (bHasInput)
 		{
+			TGuardValue<bool> ParsingForLinearInputGuard(bParsingForLinearInput, true);
+
 			TTuple<UInterchangeMaterialExpressionFactoryNode*, FString> ExpressionFactoryNode =
 				CreateMaterialExpressionForInput(MaterialFactoryNode, ShaderGraphNode, Parameters::Roughness.ToString(), MaterialFactoryNode->GetUniqueID());
 
@@ -341,6 +370,114 @@ bool UInterchangeGenericMaterialPipeline::HandlePBRModel(const UInterchangeShade
 	return bShadingModelHandled;
 }
 
+bool UInterchangeGenericMaterialPipeline::HandleClearCoat(const UInterchangeShaderGraphNode* ShaderGraphNode, UInterchangeMaterialFactoryNode* MaterialFactoryNode)
+{
+	using namespace UE::Interchange::Materials::ClearCoat;
+
+	bool bShadingModelHandled = false;
+
+	// Clear Coat
+	{
+		const bool bHasInput = UInterchangeShaderPortsAPI::HasInput(ShaderGraphNode, Parameters::ClearCoat);
+
+		if (bHasInput)
+		{
+			TGuardValue<bool> ParsingForLinearInputGuard(bParsingForLinearInput, true);
+
+			TTuple<UInterchangeMaterialExpressionFactoryNode*, FString> ExpressionFactoryNode =
+				CreateMaterialExpressionForInput(MaterialFactoryNode, ShaderGraphNode, Parameters::ClearCoat.ToString(), MaterialFactoryNode->GetUniqueID());
+
+			if (ExpressionFactoryNode.Get<0>())
+			{
+				MaterialFactoryNode->ConnectOutputToClearCoat(ExpressionFactoryNode.Get<0>()->GetUniqueID(), ExpressionFactoryNode.Get<1>());
+			}
+
+			bShadingModelHandled = true;
+		}
+	}
+
+	// Clear Coat Roughness
+	{
+		const bool bHasInput = UInterchangeShaderPortsAPI::HasInput(ShaderGraphNode, Parameters::ClearCoatRoughness);
+
+		if (bHasInput)
+		{
+			TGuardValue<bool> ParsingForLinearInputGuard(bParsingForLinearInput, true);
+
+			TTuple<UInterchangeMaterialExpressionFactoryNode*, FString> ExpressionFactoryNode =
+				CreateMaterialExpressionForInput(MaterialFactoryNode, ShaderGraphNode, Parameters::ClearCoatRoughness.ToString(), MaterialFactoryNode->GetUniqueID());
+
+			if (ExpressionFactoryNode.Get<0>())
+			{
+				MaterialFactoryNode->ConnectOutputToClearCoatRoughness(ExpressionFactoryNode.Get<0>()->GetUniqueID(), ExpressionFactoryNode.Get<1>());
+			}
+
+			bShadingModelHandled = true;
+		}
+	}
+
+	// Clear Coat Normal
+	{
+		const bool bHasInput = UInterchangeShaderPortsAPI::HasInput(ShaderGraphNode, Parameters::ClearCoatNormal);
+
+		if (bHasInput)
+		{
+			TGuardValue<bool> ParsingForNormalInputGuard(bParsingForNormalInput, true);
+
+			TTuple<UInterchangeMaterialExpressionFactoryNode*, FString> ExpressionFactoryNode =
+				CreateMaterialExpressionForInput(MaterialFactoryNode, ShaderGraphNode, Parameters::ClearCoatNormal.ToString(), MaterialFactoryNode->GetUniqueID());
+
+			if (ExpressionFactoryNode.Get<0>())
+			{
+				MaterialFactoryNode->ConnectOutputToClearCoatNormal(ExpressionFactoryNode.Get<0>()->GetUniqueID(), ExpressionFactoryNode.Get<1>());
+			}
+
+			bShadingModelHandled = true;
+		}
+	}
+
+	if (bShadingModelHandled)
+	{
+		MaterialFactoryNode->SetCustomShadingModel(EMaterialShadingModel::MSM_ClearCoat);
+	}
+
+	return bShadingModelHandled;
+}
+
+bool UInterchangeGenericMaterialPipeline::HandleThinTranslucent(const UInterchangeShaderGraphNode* ShaderGraphNode, UInterchangeMaterialFactoryNode* MaterialFactoryNode)
+{
+	using namespace UE::Interchange::Materials::ThinTranslucent;
+
+	bool bShadingModelHandled = false;
+
+	// Transmission Color
+	{
+		const bool bHasInput = UInterchangeShaderPortsAPI::HasInput(ShaderGraphNode, Parameters::TransmissionColor);
+
+		if (bHasInput)
+		{
+			TTuple<UInterchangeMaterialExpressionFactoryNode*, FString> ExpressionFactoryNode =
+				CreateMaterialExpressionForInput(MaterialFactoryNode, ShaderGraphNode, Parameters::TransmissionColor.ToString(), MaterialFactoryNode->GetUniqueID());
+
+			if (ExpressionFactoryNode.Get<0>())
+			{
+				MaterialFactoryNode->ConnectOutputToTransmissionColor(ExpressionFactoryNode.Get<0>()->GetUniqueID(), ExpressionFactoryNode.Get<1>());
+			}
+
+			bShadingModelHandled = true;
+		}
+	}
+
+	if (bShadingModelHandled)
+	{
+		MaterialFactoryNode->SetCustomBlendMode(EBlendMode::BLEND_Translucent);
+		MaterialFactoryNode->SetCustomShadingModel(EMaterialShadingModel::MSM_ThinTranslucent);
+		MaterialFactoryNode->SetCustomTranslucencyLightingMode(ETranslucencyLightingMode::TLM_SurfacePerPixelLighting);
+	}
+
+	return bShadingModelHandled;
+}
+
 void UInterchangeGenericMaterialPipeline::HandleCommonParameters(const UInterchangeShaderGraphNode* ShaderGraphNode, UInterchangeMaterialFactoryNode* MaterialFactoryNode)
 {
 	using namespace UE::Interchange::Materials::Common;
@@ -388,6 +525,8 @@ void UInterchangeGenericMaterialPipeline::HandleCommonParameters(const UIntercha
 
 	// Opacity
 	{
+		TGuardValue<bool> ParsingForLinearInputGuard(bParsingForLinearInput, true);
+
 		const bool bHasOpacityInput = UInterchangeShaderPortsAPI::HasInput(ShaderGraphNode, Parameters::Opacity);
 
 		if (bHasOpacityInput)
@@ -408,8 +547,19 @@ void UInterchangeGenericMaterialPipeline::HandleCommonParameters(const UIntercha
 				if (OpacityExpressionFactoryNode.Get<0>())
 				{
 					MaterialFactoryNode->ConnectOutputToOpacity(OpacityExpressionFactoryNode.Get<0>()->GetUniqueID(), OpacityExpressionFactoryNode.Get<1>());
-					MaterialFactoryNode->SetCustomBlendMode(EBlendMode::BLEND_Translucent);
-					MaterialFactoryNode->SetCustomTranslucencyLightingMode(ETranslucencyLightingMode::TLM_Surface);
+
+					// Don't change the blend mode or the lighting mode if they were already set
+					TEnumAsByte<EBlendMode> BlendMode = EBlendMode::BLEND_Translucent;
+					if (!MaterialFactoryNode->GetCustomBlendMode(BlendMode))
+					{
+						MaterialFactoryNode->SetCustomBlendMode(BlendMode);
+
+						TEnumAsByte<ETranslucencyLightingMode> LightingMode = ETranslucencyLightingMode::TLM_Surface;
+						if (!MaterialFactoryNode->GetCustomTranslucencyLightingMode(LightingMode))
+						{
+							MaterialFactoryNode->SetCustomTranslucencyLightingMode(LightingMode);
+						}
+					}
 				}
 			}
 		}
@@ -417,6 +567,8 @@ void UInterchangeGenericMaterialPipeline::HandleCommonParameters(const UIntercha
 
 	// Ambient Occlusion
 	{
+		TGuardValue<bool> ParsingForLinearInputGuard(bParsingForLinearInput, true);
+
 		const bool bHasOcclusionInput = UInterchangeShaderPortsAPI::HasInput(ShaderGraphNode, Parameters::Occlusion);
 
 		if (bHasOcclusionInput)
@@ -431,8 +583,11 @@ void UInterchangeGenericMaterialPipeline::HandleCommonParameters(const UIntercha
 		}
 	}
 
-	// Index of Refraction
+	// Index of Refraction (IOR)
+	// We'll lerp between Air IOR (1) and the IOR from the shader graph based on a fresnel, as per UE doc on refraction.
 	{
+		TGuardValue<bool> ParsingForLinearInputGuard(bParsingForLinearInput, true);
+
 		const bool bHasIorInput = UInterchangeShaderPortsAPI::HasInput(ShaderGraphNode, Parameters::IndexOfRefraction);
 
 		if (bHasIorInput)
@@ -442,7 +597,21 @@ void UInterchangeGenericMaterialPipeline::HandleCommonParameters(const UIntercha
 
 			if (ExpressionFactoryNode.Get<0>())
 			{
-				MaterialFactoryNode->ConnectOutputToRefraction(ExpressionFactoryNode.Get<0>()->GetUniqueID(), ExpressionFactoryNode.Get<1>());
+				UInterchangeMaterialExpressionFactoryNode* IORLerp = CreateExpressionNode(TEXT("IORLerp"), ShaderGraphNode->GetUniqueID(), UMaterialExpressionLinearInterpolate::StaticClass());
+
+				const float AirIOR = 1.f;
+				const FString ConstAMemberName = GET_MEMBER_NAME_CHECKED(UMaterialExpressionLinearInterpolate, ConstA).ToString();
+				IORLerp->AddFloatAttribute(ConstAMemberName, AirIOR);
+				IORLerp->AddApplyAndFillDelegates<float>(ConstAMemberName, UMaterialExpressionLinearInterpolate::StaticClass(), *ConstAMemberName);
+
+				UInterchangeShaderPortsAPI::ConnectOuputToInput(IORLerp, GET_MEMBER_NAME_CHECKED(UMaterialExpressionLinearInterpolate, B).ToString(),
+					ExpressionFactoryNode.Get<0>()->GetUniqueID(), ExpressionFactoryNode.Get<1>());
+
+				UInterchangeMaterialExpressionFactoryNode* IORFresnel = CreateExpressionNode(TEXT("IORFresnel"), ShaderGraphNode->GetUniqueID(), UMaterialExpressionFresnel::StaticClass());
+
+				UInterchangeShaderPortsAPI::ConnectDefaultOuputToInput(IORLerp, GET_MEMBER_NAME_CHECKED(UMaterialExpressionLinearInterpolate, Alpha).ToString(), IORFresnel->GetUniqueID());
+
+				MaterialFactoryNode->ConnectToRefraction(IORLerp->GetUniqueID());
 			}
 		}
 	}
@@ -496,6 +665,13 @@ void UInterchangeGenericMaterialPipeline::HandleTextureSampleNode(const UInterch
 			TextureFactoryNode->SetCustomCompressionSettings(TextureCompressionSettings::TC_Normalmap);
 		}
 	}
+	else if (bParsingForLinearInput)
+	{
+		if (UInterchangeTextureFactoryNode* TextureFactoryNode = Cast<UInterchangeTextureFactoryNode>(BaseNodeContainer->GetNode(TextureFactoryUid)))
+		{
+			TextureFactoryNode->SetCustomSRGB(false);
+		}
+	}
 
 	HandleTextureCoordinates(ShaderNode, TextureSampleFactoryNode);
 }
@@ -536,28 +712,41 @@ void UInterchangeGenericMaterialPipeline::HandleLerpNode(const UInterchangeShade
 
 	LerpFactoryNode->SetCustomExpressionClassName(UMaterialExpressionLinearInterpolate::StaticClass()->GetName());
 
-	TTuple<UInterchangeMaterialExpressionFactoryNode*, FString> ColorAExpression =
-		CreateMaterialExpressionForInput(MaterialFactoryNode, ShaderNode, Nodes::Lerp::Inputs::A.ToString(), LerpFactoryNode->GetUniqueID());
-
-	if (ColorAExpression.Get<0>())
+	// A
 	{
-		UInterchangeShaderPortsAPI::ConnectOuputToInput(LerpFactoryNode, Nodes::Lerp::Inputs::A.ToString(), ColorAExpression.Get<0>()->GetUniqueID(), ColorAExpression.Get<1>());
-	}
+		TTuple<UInterchangeMaterialExpressionFactoryNode*, FString> ColorAExpression =
+			CreateMaterialExpressionForInput(MaterialFactoryNode, ShaderNode, Nodes::Lerp::Inputs::A.ToString(), LerpFactoryNode->GetUniqueID());
 
-	TTuple<UInterchangeMaterialExpressionFactoryNode*, FString> ColorBExpression =
-		CreateMaterialExpressionForInput(MaterialFactoryNode, ShaderNode, Nodes::Lerp::Inputs::B.ToString(), LerpFactoryNode->GetUniqueID());
-
-	if (ColorBExpression.Get<0>())
-	{
-		UInterchangeShaderPortsAPI::ConnectOuputToInput(LerpFactoryNode, Nodes::Lerp::Inputs::B.ToString(), ColorBExpression.Get<0>()->GetUniqueID(), ColorBExpression.Get<1>());
+		if (ColorAExpression.Get<0>())
+		{
+			UInterchangeShaderPortsAPI::ConnectOuputToInput(LerpFactoryNode, GET_MEMBER_NAME_CHECKED(UMaterialExpressionLinearInterpolate, A).ToString(),
+				ColorAExpression.Get<0>()->GetUniqueID(), ColorAExpression.Get<1>());
+		}
 	}
+	
+	// B
+	{
+		TTuple<UInterchangeMaterialExpressionFactoryNode*, FString> ColorBExpression =
+			CreateMaterialExpressionForInput(MaterialFactoryNode, ShaderNode, Nodes::Lerp::Inputs::B.ToString(), LerpFactoryNode->GetUniqueID());
 
-	float LerpFactor = 0.5f;
-	ShaderNode->GetFloatAttribute(UInterchangeShaderPortsAPI::MakeInputValueKey(Nodes::Lerp::Inputs::Factor.ToString()), LerpFactor);
+		if (ColorBExpression.Get<0>())
+		{
+			UInterchangeShaderPortsAPI::ConnectOuputToInput(LerpFactoryNode, GET_MEMBER_NAME_CHECKED(UMaterialExpressionLinearInterpolate, B).ToString(),
+				ColorBExpression.Get<0>()->GetUniqueID(), ColorBExpression.Get<1>());
+		}
+	}
+	
+	// Factor
+	{
+		TTuple<UInterchangeMaterialExpressionFactoryNode*, FString> FactorExpression =
+			CreateMaterialExpressionForInput(MaterialFactoryNode, ShaderNode, Nodes::Lerp::Inputs::Factor.ToString(), LerpFactoryNode->GetUniqueID());
 
-	const FName ConstAlphaMemberName = GET_MEMBER_NAME_CHECKED(UMaterialExpressionLinearInterpolate, ConstAlpha);
-	LerpFactoryNode->AddFloatAttribute(ConstAlphaMemberName.ToString(), LerpFactor);
-	LerpFactoryNode->AddApplyAndFillDelegates<float>(ConstAlphaMemberName.ToString(), UMaterialExpressionLinearInterpolate::StaticClass(), ConstAlphaMemberName);
+		if (FactorExpression.Get<0>())
+		{
+			UInterchangeShaderPortsAPI::ConnectOuputToInput(LerpFactoryNode, GET_MEMBER_NAME_CHECKED(UMaterialExpressionLinearInterpolate, Alpha).ToString(),
+				FactorExpression.Get<0>()->GetUniqueID(), FactorExpression.Get<1>());
+		}
+	}
 }
 
 UInterchangeMaterialExpressionFactoryNode* UInterchangeGenericMaterialPipeline::CreateMaterialExpressionForShaderNode(UInterchangeMaterialFactoryNode* MaterialFactoryNode,
@@ -746,6 +935,13 @@ UInterchangeMaterialFactoryNode* UInterchangeGenericMaterialPipeline::CreateMate
 	}
 
 	HandlePBRModel(ShaderGraphNode, MaterialFactoryNode); // Always process the PBR parameters. If they were already assigned from Phong or Lambert, they will be ignored.
+	
+	if (!HandleClearCoat(ShaderGraphNode, MaterialFactoryNode))
+	{
+		// Can't have both clear coat and thin translucent since they are different shading models
+		HandleThinTranslucent(ShaderGraphNode, MaterialFactoryNode);
+	}
+
 	HandleCommonParameters(ShaderGraphNode, MaterialFactoryNode);
 
 	return MaterialFactoryNode;
@@ -760,6 +956,14 @@ UInterchangeMaterialInstanceFactoryNode* UInterchangeGenericMaterialPipeline::Cr
 	{
 		MaterialInstanceFactoryNode->SetCustomParent(ParentMaterialObj->GetPathName());
 	}
+	else if (IsThinTranslucentModel(ShaderGraphNode))
+	{
+		MaterialInstanceFactoryNode->SetCustomParent(TEXT("Material'/Interchange/Materials/ThinTranslucentMaterial.ThinTranslucentMaterial'"));
+	}
+	else if (IsClearCoatModel(ShaderGraphNode))
+	{
+		MaterialInstanceFactoryNode->SetCustomParent(TEXT("Material'/Interchange/Materials/ClearCoatMaterial.ClearCoatMaterial'"));
+	}
 	else if (IsPBRModel(ShaderGraphNode))
 	{
 		MaterialInstanceFactoryNode->SetCustomParent(TEXT("Material'/Interchange/Materials/PBRSurfaceMaterial.PBRSurfaceMaterial'"));
@@ -863,6 +1067,10 @@ TVariant<FString, FLinearColor, float> UInterchangeGenericMaterialPipeline::Visi
 		{
 			return VisitMultiplyNode(ShaderNode);
 		}
+		else if (*ShaderType == OneMinus::Name)
+		{
+			return VisitOneMinusNode(ShaderNode);
+		}
 	}
 
 	{
@@ -1017,6 +1225,32 @@ TVariant<FString, FLinearColor, float> UInterchangeGenericMaterialPipeline::Visi
 	return ResultA;
 }
 
+TVariant<FString, FLinearColor, float> UInterchangeGenericMaterialPipeline::VisitOneMinusNode(const UInterchangeShaderNode* ShaderNode) const
+{
+	using namespace UE::Interchange::Materials::Standard::Nodes;
+
+	TVariant<FString, FLinearColor, float> ResultInput = VisitShaderInput(ShaderNode, OneMinus::Inputs::Input.ToString());
+
+	if (ResultInput.IsType<FLinearColor>())
+	{
+		const FLinearColor Value = ResultInput.Get<FLinearColor>();
+
+		TVariant<FString, FLinearColor, float> Result;
+		Result.Set<FLinearColor>(FLinearColor::White - Value);
+		return Result;
+	}
+	else if (ResultInput.IsType<float>())
+	{
+		const float Value = ResultInput.Get<float>();
+
+		TVariant<FString, FLinearColor, float> Result;
+		Result.Set<float>(1.f - Value);
+		return Result;
+	}
+
+	return ResultInput;
+}
+
 TVariant<FString, FLinearColor, float> UInterchangeGenericMaterialPipeline::VisitTextureSampleNode(const UInterchangeShaderNode* ShaderNode) const
 {
 	using namespace UE::Interchange::Materials::Standard::Nodes;
diff --git a/Engine/Plugins/Experimental/Interchange/Runtime/Source/Pipelines/Public/InterchangeGenericMaterialPipeline.h b/Engine/Plugins/Experimental/Interchange/Runtime/Source/Pipelines/Public/InterchangeGenericMaterialPipeline.h
--- a/Engine/Plugins/Experimental/Interchange/Runtime/Source/Pipelines/Public/InterchangeGenericMaterialPipeline.h
+++ b/Engine/Plugins/Experimental/Interchange/Runtime/Source/Pipelines/Public/InterchangeGenericMaterialPipeline.h
@@ -60,6 +60,12 @@ private:
 	UInterchangeMaterialFactoryNode* CreateMaterialFactoryNode(const UInterchangeShaderGraphNode* ShaderGraphNode);
 	UInterchangeMaterialInstanceFactoryNode* CreateMaterialInstanceFactoryNode(const UInterchangeShaderGraphNode* ShaderGraphNode);
 
+	/** True if the shader graph has a clear coat input. */
+	bool IsClearCoatModel(const UInterchangeShaderGraphNode* ShaderGraphNode) const;
+
+	/** True if the shader graph has a transmission color input. */
+	bool IsThinTranslucentModel(const UInterchangeShaderGraphNode* ShaderGraphNode) const;
+
 	/** True if the shader graph has a base color input. */
 	bool IsPBRModel(const UInterchangeShaderGraphNode* ShaderGraphNode) const;
 
@@ -72,6 +78,8 @@ private:
 	bool HandlePhongModel(const UInterchangeShaderGraphNode* ShaderGraphNode, UInterchangeMaterialFactoryNode* MaterialFactoryNode);
 	bool HandleLambertModel(const UInterchangeShaderGraphNode* ShaderGraphNode, UInterchangeMaterialFactoryNode* MaterialFactoryNode);
 	bool HandlePBRModel(const UInterchangeShaderGraphNode* ShaderGraphNode, UInterchangeMaterialFactoryNode* MaterialFactoryNode);
+	bool HandleClearCoat(const UInterchangeShaderGraphNode* ShaderGraphNode, UInterchangeMaterialFactoryNode* MaterialFactoryNode);
+	bool HandleThinTranslucent(const UInterchangeShaderGraphNode* ShaderGraphNode, UInterchangeMaterialFactoryNode* MaterialFactoryNode);
 	void HandleCommonParameters(const UInterchangeShaderGraphNode* ShaderGraphNode, UInterchangeMaterialFactoryNode* MaterialFactoryNode);
 	void HandleTextureSampleNode(const UInterchangeShaderNode* ShaderNode, UInterchangeMaterialExpressionFactoryNode* TextureSampleFactoryNode);
 	void HandleTextureCoordinates(const UInterchangeShaderNode* ShaderNode, UInterchangeMaterialExpressionFactoryNode* TextureSampleFactoryNode);
@@ -99,10 +107,12 @@ private:
 	 */
 	TVariant<FString, FLinearColor, float> VisitLerpNode(const UInterchangeShaderNode* ShaderNode) const;
 	TVariant<FString, FLinearColor, float> VisitMultiplyNode(const UInterchangeShaderNode* ShaderNode) const;
+	TVariant<FString, FLinearColor, float> VisitOneMinusNode(const UInterchangeShaderNode* ShaderNode) const;
 	TVariant<FString, FLinearColor, float> VisitTextureSampleNode(const UInterchangeShaderNode* ShaderNode) const;
 
 private:
 	bool bParsingForNormalInput;
+	bool bParsingForLinearInput; // True when parsing non-color inputs (metallic, roughness, specular, etc.)
 };
 
 
