fix for merge in InterchangeTextureFactoryNode
goes with "clean up the mess caused by AlphaCoverageThresholds"

diff --git a/Engine/Plugins/Experimental/Interchange/Runtime/Source/FactoryNodes/Public/InterchangeTextureFactoryNode.h b/Engine/Plugins/Experimental/Interchange/Runtime/Source/FactoryNodes/Public/InterchangeTextureFactoryNode.h
--- a/Engine/Plugins/Experimental/Interchange/Runtime/Source/FactoryNodes/Public/InterchangeTextureFactoryNode.h
+++ b/Engine/Plugins/Experimental/Interchange/Runtime/Source/FactoryNodes/Public/InterchangeTextureFactoryNode.h
@@ -824,8 +824,11 @@ private:
 	const UE::Interchange::FAttributeKey Macro_CustomCompositeTextureModeKey = UE::Interchange::FAttributeKey(TEXT("CompositeTextureMode"));
 	const UE::Interchange::FAttributeKey Macro_CustomCompositePowerKey = UE::Interchange::FAttributeKey(TEXT("CompositePower"));
 
+	// no IMPLEMENT_NODE_ATTRIBUTE_APPLY_UOBJECT here ?
+
 	//TODO support per platform data in the FAttributeStorage, so we can set different value per platform at the pipeline stage, We set only the default value for now
 	//IMPLEMENT_NODE_ATTRIBUTE_APPLY_UOBJECT(Downscale, float, UTexture, );
+
 #if WITH_ENGINE
 	bool ApplyCustomDownscaleToAsset(UObject* Asset) const
 	{
