Fixed a crash in the interchange texture factory when a payload wasn't available from the source data.

diff --git a/Engine/Plugins/Experimental/Interchange/Runtime/Source/Import/Private/Texture/InterchangeTextureFactory.cpp b/Engine/Plugins/Experimental/Interchange/Runtime/Source/Import/Private/Texture/InterchangeTextureFactory.cpp
--- a/Engine/Plugins/Experimental/Interchange/Runtime/Source/Import/Private/Texture/InterchangeTextureFactory.cpp
+++ b/Engine/Plugins/Experimental/Interchange/Runtime/Source/Import/Private/Texture/InterchangeTextureFactory.cpp
@@ -615,6 +615,17 @@ namespace UE::Interchange::Private::InterchangeTextureFactory
 		return *this;
 	}
 
+	bool FProcessedPayload::IsValid() const
+	{
+		if (SettingsFromPayload.IsType<FEmptyVariantState>())
+		{
+			return false;
+		}
+
+		return true;
+	}
+
+
 	TArray<FString> GetFilesToHash(const FTextureNodeVariant& TextureNodeVariant, const FTexturePayloadVariant& TexturePayload)
 	{
 		TArray<FString> FilesToHash;
@@ -935,14 +946,7 @@ UObject* UInterchangeTextureFactory::CreateAsset(const FCreateAssetParams& Argum
 	if (!bCanSetup)
 	{
 		LogErrorInvalidPayload(Texture->GetClass()->GetName(), Texture->GetName());
-
-		// The texture is not supported
-		if (!Arguments.ReimportObject)
-		{
-			Texture->RemoveFromRoot();
-			Texture->MarkAsGarbage();
-		}
-		return nullptr;
+		return Texture;
 	}
 
 	FGraphEventArray TasksToDo = GenerateHashSourceFilesTasks(Arguments.SourceData, GetFilesToHash(TextureNodeVariant, TexturePayload), SourceFiles);
@@ -974,7 +978,7 @@ void UInterchangeTextureFactory::PreImportPreCompletedCallback(const FImportPreC
 #if WITH_EDITOR
 
 	// Finish the import on the game thread by doing the setup on the texture here
-	if (Texture)
+	if (Texture && ProcessedPayload.IsValid())
 	{
 		Texture->PreEditChange(nullptr);
 
@@ -1022,13 +1026,23 @@ void UInterchangeTextureFactory::PreImportPreCompletedCallback(const FImportPreC
 			UE::Interchange::FFactoryCommon::ApplyReimportStrategyToAsset(Texture, PreviousNode, CurrentNode, TextureFactoryNode);
 		}
 	}
+	else
+	{
+		// The texture is not supported
+		if (!Arguments.bIsReimport)
+		{
+			// Not thread safe. So those should stay on the game thread.
+			Texture->RemoveFromRoot();
+			Texture->MarkAsGarbage();
+		}
+	}
 #endif //WITH_EDITOR
 
 	Super::PreImportPreCompletedCallback(Arguments);
 
 	//TODO make sure this work at runtime
 #if WITH_EDITORONLY_DATA
-	if (ensure(Texture && Arguments.SourceData))
+	if (ensure(Texture && Arguments.SourceData) && ProcessedPayload.IsValid())
 	{
 		//We must call the Update of the asset source file in the main thread because UAssetImportData::Update execute some delegate we do not control
 		UE::Interchange::FFactoryCommon::FSetImportAssetDataParameters SetImportAssetDataParameters(Texture
diff --git a/Engine/Plugins/Experimental/Interchange/Runtime/Source/Import/Public/Texture/InterchangeTextureFactory.h b/Engine/Plugins/Experimental/Interchange/Runtime/Source/Import/Public/Texture/InterchangeTextureFactory.h
--- a/Engine/Plugins/Experimental/Interchange/Runtime/Source/Import/Public/Texture/InterchangeTextureFactory.h
+++ b/Engine/Plugins/Experimental/Interchange/Runtime/Source/Import/Public/Texture/InterchangeTextureFactory.h
@@ -36,6 +36,8 @@ namespace UE::Interchange::Private::InterchangeTextureFactory
 
 		FProcessedPayload& operator=(UE::Interchange::Private::InterchangeTextureFactory::FTexturePayloadVariant&& InPayloadVariant);
 
+		bool IsValid() const;
+
 		UE::Interchange::Private::InterchangeTextureFactory::FTexturePayloadVariant SettingsFromPayload;
 		UE::Serialization::FEditorBulkData::FSharedBufferWithID PayloadAndId;
 	};
