Fixed .obj bounds calculation.

diff --git a/Engine/Plugins/Experimental/Interchange/Runtime/Source/Import/Private/Mesh/InterchangeOBJTranslator.cpp b/Engine/Plugins/Experimental/Interchange/Runtime/Source/Import/Private/Mesh/InterchangeOBJTranslator.cpp
--- a/Engine/Plugins/Experimental/Interchange/Runtime/Source/Import/Private/Mesh/InterchangeOBJTranslator.cpp
+++ b/Engine/Plugins/Experimental/Interchange/Runtime/Source/Import/Private/Mesh/InterchangeOBJTranslator.cpp
@@ -173,7 +173,7 @@ TArray<int32> FObjData::GetUVIndicesUsedByGroup(const FGroupData& GroupData) con
 
 FBox FObjData::GetGroupBoundingBox(const FString& GroupName) const
 {
-	FBox Box;
+	FBox Box(ForceInit);
 
 	const FGroupData* GroupDataPtr = Groups.Find(GroupName);
 	if (!GroupDataPtr)
