Interchange static meshes
- Set build settings so that destination lightmap index doesn't overwrite an existing UV channel.

diff --git a/Engine/Plugins/Experimental/Interchange/Runtime/Source/Import/Private/Mesh/InterchangeStaticMeshFactory.cpp b/Engine/Plugins/Experimental/Interchange/Runtime/Source/Import/Private/Mesh/InterchangeStaticMeshFactory.cpp
--- a/Engine/Plugins/Experimental/Interchange/Runtime/Source/Import/Private/Mesh/InterchangeStaticMeshFactory.cpp
+++ b/Engine/Plugins/Experimental/Interchange/Runtime/Source/Import/Private/Mesh/InterchangeStaticMeshFactory.cpp
@@ -387,13 +387,26 @@ UObject* UInterchangeStaticMeshFactory::CreateAsset(const FCreateAssetParams& Ar
 
 			SectionIndex++;
 		}
-		if (CurrentLodIndex >= FinalLodCount)
+
+		if (!Arguments.ReimportObject)
 		{
-			// Set static mesh build settings
 			FStaticMeshSourceModel& SrcModel = StaticMesh->GetSourceModel(CurrentLodIndex);
-			SrcModel.ReductionSettings.MaxDeviation = 0.0f;
-			SrcModel.ReductionSettings.PercentTriangles = 1.0f;
-			SrcModel.ReductionSettings.PercentVertices = 1.0f;
+
+			const int32 NumUVChannels = StaticMeshAttributes.GetVertexInstanceUVs().GetNumChannels();
+			const int32 FirstOpenUVChannel = NumUVChannels >= MAX_MESH_TEXTURE_COORDS_MD ? 1 : NumUVChannels;
+			SrcModel.BuildSettings.DstLightmapIndex = FirstOpenUVChannel;
+
+			if (CurrentLodIndex == 0)
+			{
+				StaticMesh->SetLightMapCoordinateIndex(FirstOpenUVChannel);
+			}
+
+			if (CurrentLodIndex >= FinalLodCount)
+			{
+				SrcModel.ReductionSettings.MaxDeviation = 0.0f;
+				SrcModel.ReductionSettings.PercentTriangles = 1.0f;
+				SrcModel.ReductionSettings.PercentVertices = 1.0f;
+			}
 		}
 
 		CurrentLodIndex++;
