Moved the default udim setup from a translator to the pipeline and with that also moved the responsability to generate the blocked payload into the interchange texture factory.

Removed the with_engine defines from the texture factory nodes as those are no longer needed.
Made the texture factory node abstract to reflect the engine structure for those asset.
Removed the asset class attribute for the texture node as it is no longer needed.

diff --git a/Engine/Plugins/Enterprise/DatasmithImporter/Source/DatasmithImporter/Private/DatasmithTextureImporter.cpp b/Engine/Plugins/Enterprise/DatasmithImporter/Source/DatasmithImporter/Private/DatasmithTextureImporter.cpp
--- a/Engine/Plugins/Enterprise/DatasmithImporter/Source/DatasmithImporter/Private/DatasmithTextureImporter.cpp
+++ b/Engine/Plugins/Enterprise/DatasmithImporter/Source/DatasmithImporter/Private/DatasmithTextureImporter.cpp
@@ -34,7 +34,7 @@
 #include "ObjectTools.h"
 #include "RHI.h"
 #include "InterchangeSourceData.h"
-#include "InterchangeTextureFactoryNode.h"
+#include "InterchangeTexture2DFactoryNode.h"
 #include "InterchangeTextureNode.h"
 
 #define LOCTEXT_NAMESPACE "DatasmithTextureImport"
@@ -431,13 +431,13 @@ void UDatasmithTexturePipeline::ExecutePreImportPipeline(UInterchangeBaseNodeCon
 
 	FString DisplayLabel = TextureTranslatedNode->GetDisplayLabel();
 	FString NodeUID = UInterchangeTextureFactoryNode::GetTextureFactoryNodeUidFromTextureNodeUid(TextureTranslatedNode->GetUniqueID());
-	UInterchangeTextureFactoryNode* TextureFactoryNode = NewObject<UInterchangeTextureFactoryNode>(BaseNodeContainer, NAME_None);
+	UInterchangeTexture2DFactoryNode* TextureFactoryNode = NewObject<UInterchangeTexture2DFactoryNode>(BaseNodeContainer, NAME_None);
 	if (!ensure(TextureFactoryNode))
 	{
 		return;
 	}
 	//Creating a UTexture2D
-	TextureFactoryNode->InitializeTextureNode(NodeUID, DisplayLabel, UTexture2D::StaticClass()->GetName(), TextureTranslatedNode->GetDisplayLabel());
+	TextureFactoryNode->InitializeTextureNode(NodeUID, DisplayLabel, TextureTranslatedNode->GetDisplayLabel());
 	TextureFactoryNode->SetCustomTranslatedTextureNodeUid(TextureTranslatedNode->GetUniqueID());
 	BaseNodeContainer->AddNode(TextureFactoryNode);
 
diff --git a/Engine/Plugins/Experimental/Interchange/Runtime/Source/FactoryNodes/Public/InterchangeTexture2DArrayFactoryNode.h b/Engine/Plugins/Experimental/Interchange/Runtime/Source/FactoryNodes/Public/InterchangeTexture2DArrayFactoryNode.h
--- a/Engine/Plugins/Experimental/Interchange/Runtime/Source/FactoryNodes/Public/InterchangeTexture2DArrayFactoryNode.h
+++ b/Engine/Plugins/Experimental/Interchange/Runtime/Source/FactoryNodes/Public/InterchangeTexture2DArrayFactoryNode.h
@@ -21,7 +21,13 @@ class INTERCHANGEFACTORYNODES_API UInterchangeTexture2DArrayFactoryNode : public
 
 public:
 
-	virtual bool SetCustomAddressX(const uint8 AttributeValue, bool bAddApplyDelegate = true) override
+	UFUNCTION(BlueprintCallable, Category = "Interchange | Node | Texture")
+	bool GetCustomAddressX(uint8& AttributeValue) const
+	{
+		IMPLEMENT_NODE_ATTRIBUTE_GETTER(AddressX, uint8);
+	}
+
+	bool SetCustomAddressX(const uint8 AttributeValue, bool bAddApplyDelegate = true)
 	{
 #if WITH_EDITORONLY_DATA
 		IMPLEMENT_NODE_ATTRIBUTE_SETTER(UInterchangeTexture2DArrayFactoryNode, AddressX, uint8, UTexture2DArray)
@@ -30,7 +36,13 @@ public:
 #endif
 	}
 
-	virtual bool SetCustomAddressY(const uint8 AttributeValue, bool bAddApplyDelegate = true) override
+	UFUNCTION(BlueprintCallable, Category = "Interchange | Node | Texture")
+	bool GetCustomAddressY(uint8& AttributeValue) const
+	{
+		IMPLEMENT_NODE_ATTRIBUTE_GETTER(AddressY, uint8);
+	}
+
+	bool SetCustomAddressY(const uint8 AttributeValue, bool bAddApplyDelegate = true)
 	{
 #if WITH_EDITORONLY_DATA
 		IMPLEMENT_NODE_ATTRIBUTE_SETTER(UInterchangeTexture2DArrayFactoryNode, AddressY, uint8, UTexture2DArray)
@@ -55,22 +67,14 @@ public:
 #endif
 	}
 
-
 private:
-#if WITH_ENGINE
-	virtual void FillAssetClassFromClassName(const FString& ClassName) override
+	virtual UClass* GetObjectClass() const override
 	{
-		if (ClassName.Equals(UTexture2DArray::StaticClass()->GetName()))
-		{
-			AssetClass = UTexture2DArray::StaticClass();
-			bIsTextureNodeClassInitialized = true;
-		}
+		return UTexture2DArray::StaticClass();
 	}
-#endif
 
 	// Addressing
 	IMPLEMENT_NODE_ATTRIBUTE_KEY(AddressX);
 	IMPLEMENT_NODE_ATTRIBUTE_KEY(AddressY);
 	IMPLEMENT_NODE_ATTRIBUTE_KEY(AddressZ);
-
 };
\ No newline at end of file
diff --git a/Engine/Plugins/Experimental/Interchange/Runtime/Source/FactoryNodes/Public/InterchangeTexture2DFactoryNode.h b/Engine/Plugins/Experimental/Interchange/Runtime/Source/FactoryNodes/Public/InterchangeTexture2DFactoryNode.h
new file mode 100644
--- /dev/null
+++ b/Engine/Plugins/Experimental/Interchange/Runtime/Source/FactoryNodes/Public/InterchangeTexture2DFactoryNode.h
@@ -0,0 +1,117 @@
+// Copyright Epic Games, Inc. All Rights Reserved.
+
+#pragma once
+
+#include "CoreMinimal.h"
+#include "InterchangeTextureFactoryNode.h"
+
+#include "Engine/Texture2D.h"
+
+#include "InterchangeTexture2DFactoryNode.generated.h"
+
+namespace UE::Interchange
+{ 
+	struct FTexture2DFactoryNodeStaticData : public FBaseNodeStaticData
+	{
+		static const FString& GetBaseSourceBlocksKey()
+		{
+			static FString StringKey(TEXT("SourceBlocks"));
+			return StringKey;
+		}
+	};
+}//ns UE::Interchange
+
+UCLASS(BlueprintType, Experimental)
+class INTERCHANGEFACTORYNODES_API UInterchangeTexture2DFactoryNode : public UInterchangeTextureFactoryNode
+{
+	GENERATED_BODY()
+
+public:
+
+
+	/** Return false if the Attribute was not set previously.*/
+	UFUNCTION(BlueprintCallable, Category = "Interchange | Node | Texture")
+	bool GetCustomAddressX(uint8& AttributeValue) const
+	{
+		IMPLEMENT_NODE_ATTRIBUTE_GETTER(AddressX, uint8);
+	}
+
+	UFUNCTION(BlueprintCallable, Category = "Interchange | Node | Texture")
+	bool SetCustomAddressX(const uint8 AttributeValue, bool bAddApplyDelegate = true)
+	{
+		IMPLEMENT_NODE_ATTRIBUTE_SETTER(UInterchangeTexture2DFactoryNode, AddressX, uint8, UTexture2D)
+	}
+
+	/** Return false if the Attribute was not set previously.*/
+	UFUNCTION(BlueprintCallable, Category = "Interchange | Node | Texture")
+	bool GetCustomAddressY(uint8& AttributeValue) const
+	{
+		IMPLEMENT_NODE_ATTRIBUTE_GETTER(AddressY, uint8);
+	}
+
+	UFUNCTION(BlueprintCallable, Category = "Interchange | Node | Texture")
+	bool SetCustomAddressY(const uint8 AttributeValue, bool bAddApplyDelegate = true)
+	{
+		IMPLEMENT_NODE_ATTRIBUTE_SETTER(UInterchangeTexture2DFactoryNode, AddressY, uint8, UTexture2D)
+	}
+
+	//////////////////////////////////////////////////////////////////////////
+	// UDIMs begin here
+	// UDIM base texture use a different model for the source data
+
+	/**
+	 * Get the source blocks for the texture
+	 * If the map is empty then the texture will be simply be imported as normal texture using the payload key
+	 */
+	UFUNCTION(BlueprintCallable, Category = "Interchange | Node | Texture | UDIMs")
+	TMap<int32, FString> GetSourceBlocks() const
+	{
+		return SourceBlocks.ToMap();
+	}
+
+	/**
+	 * Set the source blocks
+	 * Using this will suggest the pipeline to consider this texture as UDIM and it can chose to pass or not these block to the texture factory node.
+	 * @param InSourceBlocks The blocks and their source image that compose the whole texture.
+	 * The textures must be of the same format and use the same pixel format
+	 * The first block in the map is used to determine the accepted texture format and pixel format
+	 */
+	UFUNCTION(BlueprintCallable, Category = "Interchange | Node | Texture | UDIMs")
+	void SetSourceBlocks(const TMap<int32, FString>& InSourceBlocks)
+	{
+		SourceBlocks = InSourceBlocks;
+	}
+
+	
+	/**
+	 * Set the source blocks
+	 * Using this will suggest the pipeline to consider this texture as UDIM and it can chose to pass or not these block to the texture factory node.
+	 * @param InSourceBlocks The blocks and their source image that compose the whole texture.
+	 * The textures must be of the same format and use the same pixel format
+	 * The first block in the map is used to determine the accepted texture format and pixel format
+	 */
+	void SetSourceBlocks(TMap<int32, FString>&& InSourceBlocks)
+	{
+		SourceBlocks = InSourceBlocks;
+	}
+
+	// UDIMs ends here
+	//////////////////////////////////////////////////////////////////////////
+
+	virtual UClass* GetObjectClass() const override
+	{
+		return UTexture2D::StaticClass();
+	}
+
+	virtual void PostInitProperties()
+	{
+		Super::PostInitProperties();
+		SourceBlocks.Initialize(Attributes.ToSharedRef(), UE::Interchange::FTexture2DFactoryNodeStaticData::GetBaseSourceBlocksKey());
+	}
+
+private:
+	IMPLEMENT_NODE_ATTRIBUTE_KEY(AddressX);
+	IMPLEMENT_NODE_ATTRIBUTE_KEY(AddressY);
+
+	UE::Interchange::TMapAttributeHelper<int32, FString> SourceBlocks;
+};
\ No newline at end of file
diff --git a/Engine/Plugins/Experimental/Interchange/Runtime/Source/FactoryNodes/Public/InterchangeTextureCubeFactoryNode.h b/Engine/Plugins/Experimental/Interchange/Runtime/Source/FactoryNodes/Public/InterchangeTextureCubeFactoryNode.h
--- a/Engine/Plugins/Experimental/Interchange/Runtime/Source/FactoryNodes/Public/InterchangeTextureCubeFactoryNode.h
+++ b/Engine/Plugins/Experimental/Interchange/Runtime/Source/FactoryNodes/Public/InterchangeTextureCubeFactoryNode.h
@@ -18,16 +18,8 @@ class INTERCHANGEFACTORYNODES_API UInterchangeTextureCubeFactoryNode : public UI
 	GENERATED_BODY()
 
 private:
-
-#if WITH_ENGINE
-	virtual void FillAssetClassFromClassName(const FString& ClassName) override
+	virtual UClass* GetObjectClass() const override
 	{
-		if (ClassName.Equals(UTextureCube::StaticClass()->GetName()))
-		{
-			AssetClass = UTextureCube::StaticClass();
-			bIsTextureNodeClassInitialized = true;
-		}
+		return UTextureCube::StaticClass();
 	}
-#endif
-
 };
\ No newline at end of file
diff --git a/Engine/Plugins/Experimental/Interchange/Runtime/Source/FactoryNodes/Public/InterchangeTextureFactoryNode.h b/Engine/Plugins/Experimental/Interchange/Runtime/Source/FactoryNodes/Public/InterchangeTextureFactoryNode.h
--- a/Engine/Plugins/Experimental/Interchange/Runtime/Source/FactoryNodes/Public/InterchangeTextureFactoryNode.h
+++ b/Engine/Plugins/Experimental/Interchange/Runtime/Source/FactoryNodes/Public/InterchangeTextureFactoryNode.h
@@ -5,15 +5,12 @@
 #include "CoreMinimal.h"
 #include "Nodes/InterchangeFactoryBaseNode.h"
 
-#if WITH_ENGINE
 #include "Engine/Texture.h"
-#include "Engine/Texture2D.h"
-#endif
 
 #include "InterchangeTextureFactoryNode.generated.h"
 
 
-UCLASS(BlueprintType, Experimental)
+UCLASS(BlueprintType, Abstract, Experimental)
 class INTERCHANGEFACTORYNODES_API UInterchangeTextureFactoryNode : public UInterchangeFactoryBaseNode
 {
 	GENERATED_BODY()
@@ -26,35 +23,16 @@ public:
 	 * @param InAssetClass - The class the texture factory will create for this node.
 	 */
 	UFUNCTION(BlueprintCallable, Category = "Interchange | Node | Texture")
-	void InitializeTextureNode(const FString& UniqueID, const FString& DisplayLabel, const FString& InAssetClass, const FString& InAssetName)
+	void InitializeTextureNode(const FString& UniqueID, const FString& DisplayLabel, const FString& InAssetName)
 	{
 		//Initialize
-		bIsTextureNodeClassInitialized = false;
 		InitializeNode(UniqueID, DisplayLabel, EInterchangeNodeContainerType::FactoryData);
-		
-		//Set the class
-		FString OperationName = GetTypeName() + TEXT(".SetAssetClassName");
-		InterchangePrivateNodeBase::SetCustomAttribute<FString>(*Attributes, ClassNameAttributeKey, OperationName, InAssetClass);
-		FillAssetClassFromAttribute();
-		
+
 		//Set the asset name
-		OperationName = GetTypeName() + TEXT(".SetAssetName");
+		FString OperationName = GetTypeName() + TEXT(".SetAssetName");
 		InterchangePrivateNodeBase::SetCustomAttribute<FString>(*Attributes, AssetNameKey, OperationName, InAssetName);
 	}
 
-	virtual void Serialize(FArchive& Ar) override
-	{
-		Super::Serialize(Ar);
-#if WITH_ENGINE
-		if (Ar.IsLoading())
-		{
-			//Make sure the class is properly set when we compile with engine, this will set the
-			//bIsTextureNodeClassInitialized to true.
-			SetTextureNodeClassFromClassAttribute();
-		}
-#endif //#if WITH_ENGINE
-	}
-
 	/**
 	 * Return the node type name of the class, we use this when reporting error
 	 */
@@ -66,14 +44,11 @@ public:
 
 	/** Get the class this node want to create */
 	UFUNCTION(BlueprintCallable, Category = "Interchange | Node | Texture")
-	virtual class UClass* GetObjectClass() const override
+	virtual UClass* GetObjectClass() const override
 	{
-		ensure(bIsTextureNodeClassInitialized);
-#if WITH_ENGINE
-		return AssetClass.Get() != nullptr ? AssetClass.Get() : UTexture::StaticClass();
-#else
+		// This function should always be overridden by the sub texture factory node
+		ensure(false);
 		return nullptr;
-#endif
 	}
 
 	virtual FGuid GetHash() const override
@@ -553,32 +528,6 @@ public:
 		IMPLEMENT_NODE_ATTRIBUTE_SETTER(UInterchangeTextureFactoryNode, VirtualTextureStreaming, bool, UTexture)
 	}
 
-	/** Return false if the Attribute was not set previously.*/
-	UFUNCTION(BlueprintCallable, Category = "Interchange | Node | Texture")
-	virtual bool GetCustomAddressX(uint8& AttributeValue) const
-	{
-		IMPLEMENT_NODE_ATTRIBUTE_GETTER(AddressX, uint8);
-	}
-
-	UFUNCTION(BlueprintCallable, Category = "Interchange | Node | Texture")
-	virtual bool SetCustomAddressX(const uint8 AttributeValue, bool bAddApplyDelegate = true)
-	{
-		IMPLEMENT_NODE_ATTRIBUTE_SETTER(UInterchangeTextureFactoryNode, AddressX, uint8, UTexture2D)
-	}
-
-	/** Return false if the Attribute was not set previously.*/
-	UFUNCTION(BlueprintCallable, Category = "Interchange | Node | Texture")
-	virtual bool GetCustomAddressY(uint8& AttributeValue) const
-	{
-		IMPLEMENT_NODE_ATTRIBUTE_GETTER(AddressY, uint8);
-	}
-
-	UFUNCTION(BlueprintCallable, Category = "Interchange | Node | Texture")
-	virtual bool SetCustomAddressY(const uint8 AttributeValue, bool bAddApplyDelegate = true)
-	{
-		IMPLEMENT_NODE_ATTRIBUTE_SETTER(UInterchangeTextureFactoryNode, AddressY, uint8, UTexture2D)
-	}
-
 	//////////////////////////////////////////////////////////////////////////
 	//Level of Detail
 
@@ -708,41 +657,6 @@ public:
 
 private:
 
-	void FillAssetClassFromAttribute()
-	{
-#if WITH_ENGINE
-		FString OperationName = GetTypeName() + TEXT(".GetAssetClassName");
-		FString ClassName;
-		InterchangePrivateNodeBase::GetCustomAttribute<FString>(*Attributes, ClassNameAttributeKey, OperationName, ClassName);
-		FillAssetClassFromClassName(ClassName);
-#endif
-	}
-
-#if WITH_ENGINE
-	virtual void FillAssetClassFromClassName(const FString& ClassName)
-	{
-		if (ClassName.Equals(UTexture2D::StaticClass()->GetName()))
-		{
-			AssetClass = UTexture2D::StaticClass();
-			bIsTextureNodeClassInitialized = true;
-		}
-		else if (ClassName.Equals(UTexture::StaticClass()->GetName()))
-		{
-			AssetClass = UTexture::StaticClass();
-			bIsTextureNodeClassInitialized = true;
-		}
-	}
-#endif
-
-	bool SetTextureNodeClassFromClassAttribute()
-	{
-		if (!bIsTextureNodeClassInitialized)
-		{
-			FillAssetClassFromAttribute();
-		}
-		return bIsTextureNodeClassInitialized;
-	}
-
 	const UE::Interchange::FAttributeKey ClassNameAttributeKey = UE::Interchange::FBaseNodeStaticData::ClassTypeAttributeKey();
 	const UE::Interchange::FAttributeKey AssetNameKey = UE::Interchange::FBaseNodeStaticData::AssetNameKey();
 
@@ -782,8 +696,6 @@ private:
 	const UE::Interchange::FAttributeKey Macro_CustomSRGBKey = UE::Interchange::FAttributeKey(TEXT("SRGB"));
 	const UE::Interchange::FAttributeKey Macro_CustombUseLegacyGammaKey = UE::Interchange::FAttributeKey(TEXT("bUseLegacyGamma"));
 	const UE::Interchange::FAttributeKey Macro_CustomVirtualTextureStreamingKey = UE::Interchange::FAttributeKey(TEXT("VirtualTextureStreaming"));
-	const UE::Interchange::FAttributeKey Macro_CustomAddressXKey = UE::Interchange::FAttributeKey(TEXT("AddressX"));
-	const UE::Interchange::FAttributeKey Macro_CustomAddressYKey = UE::Interchange::FAttributeKey(TEXT("AddressY"));
 
 	//Level of Detail
 	const UE::Interchange::FAttributeKey Macro_CustombPreserveBorderKey = UE::Interchange::FAttributeKey(TEXT("bPreserveBorder"));
@@ -803,7 +715,7 @@ private:
 	//TODO support per platform data in the FAttributeStorage, so we can set different value per platform at the pipeline stage, We set only the default value for now
 	//IMPLEMENT_NODE_ATTRIBUTE_APPLY_UOBJECT(Downscale, float, UTexture, );
 
-#if WITH_ENGINE
+
 	bool ApplyCustomDownscaleToAsset(UObject* Asset) const
 	{
 		if (!Asset)
@@ -841,12 +753,4 @@ private:
 		}
 		return false;
 	}
-#endif //WITH_ENGINE
-
-protected:
-#if WITH_ENGINE
-	TSubclassOf<UTexture> AssetClass = nullptr;
-#endif
-
-	bool bIsTextureNodeClassInitialized = false;
 };
diff --git a/Engine/Plugins/Experimental/Interchange/Runtime/Source/FactoryNodes/Public/InterchangeTextureLightProfileFactoryNode.h b/Engine/Plugins/Experimental/Interchange/Runtime/Source/FactoryNodes/Public/InterchangeTextureLightProfileFactoryNode.h
--- a/Engine/Plugins/Experimental/Interchange/Runtime/Source/FactoryNodes/Public/InterchangeTextureLightProfileFactoryNode.h
+++ b/Engine/Plugins/Experimental/Interchange/Runtime/Source/FactoryNodes/Public/InterchangeTextureLightProfileFactoryNode.h
@@ -3,7 +3,7 @@
 #pragma once
 
 #include "CoreMinimal.h"
-#include "InterchangeTextureFactoryNode.h"
+#include "InterchangeTexture2DFactoryNode.h"
 #include "Nodes/InterchangeBaseNode.h"
 
 #if WITH_ENGINE
@@ -15,7 +15,7 @@
 #include "InterchangeTextureLightProfileFactoryNode.generated.h"
 
 UCLASS(BlueprintType, Experimental)
-class INTERCHANGEFACTORYNODES_API UInterchangeTextureLightProfileFactoryNode : public UInterchangeTextureFactoryNode
+class INTERCHANGEFACTORYNODES_API UInterchangeTextureLightProfileFactoryNode : public UInterchangeTexture2DFactoryNode
 {
 	GENERATED_BODY()
 
@@ -56,16 +56,10 @@ public:
 
 
 private:
-#if WITH_ENGINE
-	virtual void FillAssetClassFromClassName(const FString& ClassName) override
+	virtual UClass* GetObjectClass() const override
 	{
-		if (ClassName.Equals(UTextureLightProfile::StaticClass()->GetName()))
-		{
-			AssetClass = UTextureLightProfile::StaticClass();
-			bIsTextureNodeClassInitialized = true;
-		}
+		return UTextureLightProfile::StaticClass();
 	}
-#endif
 
 	// Addressing
 	IMPLEMENT_NODE_ATTRIBUTE_KEY(Brightness);
diff --git a/Engine/Plugins/Experimental/Interchange/Runtime/Source/Import/Private/InterchangeImportModule.cpp b/Engine/Plugins/Experimental/Interchange/Runtime/Source/Import/Private/InterchangeImportModule.cpp
--- a/Engine/Plugins/Experimental/Interchange/Runtime/Source/Import/Private/InterchangeImportModule.cpp
+++ b/Engine/Plugins/Experimental/Interchange/Runtime/Source/Import/Private/InterchangeImportModule.cpp
@@ -28,7 +28,6 @@
 #include "Texture/InterchangeTextureFactory.h"
 #include "Texture/InterchangeTGATranslator.h"
 #include "Texture/InterchangeTIFFTranslator.h"
-#include "Texture/InterchangeUDIMTranslator.h"
 
 DEFINE_LOG_CATEGORY(LogInterchangeImport);
 
@@ -56,12 +55,6 @@ void FInterchangeImportModule::StartupModule()
 		InterchangeManager.RegisterTranslator(UInterchangeOBJTranslator::StaticClass());
 
 		//Textures
-
-		// UDIM must be registered before the other texture translators
-		// @todo: this seems a bit fragile; can we do better?
-		InterchangeManager.RegisterTranslator(UInterchangeUDIMTranslator::StaticClass());
-		InterchangeManager.RegisterTextureOnlyTranslatorClass(UInterchangeUDIMTranslator::StaticClass());
-
 		InterchangeManager.RegisterTranslator(UInterchangeBMPTranslator::StaticClass());
 		InterchangeManager.RegisterTextureOnlyTranslatorClass(UInterchangeBMPTranslator::StaticClass());
 		InterchangeManager.RegisterTranslator(UInterchangeDDSTranslator::StaticClass());
diff --git a/Engine/Plugins/Experimental/Interchange/Runtime/Source/Import/Private/Texture/InterchangeTextureFactory.cpp b/Engine/Plugins/Experimental/Interchange/Runtime/Source/Import/Private/Texture/InterchangeTextureFactory.cpp
--- a/Engine/Plugins/Experimental/Interchange/Runtime/Source/Import/Private/Texture/InterchangeTextureFactory.cpp
+++ b/Engine/Plugins/Experimental/Interchange/Runtime/Source/Import/Private/Texture/InterchangeTextureFactory.cpp
@@ -20,6 +20,7 @@
 #include "InterchangeTextureCubeFactoryNode.h"
 #include "InterchangeTextureCubeNode.h"
 #include "InterchangeTextureFactoryNode.h"
+#include "InterchangeTexture2DFactoryNode.h"
 #include "InterchangeTextureLightProfileFactoryNode.h"
 #include "InterchangeTextureLightProfileNode.h"
 #include "InterchangeTextureNode.h"
@@ -44,13 +45,86 @@
 
 namespace UE::Interchange::Private::InterchangeTextureFactory
 {
+	class FScopedTranslatorsAndSourceData
+	{
+	public:
+		FScopedTranslatorsAndSourceData(const UInterchangeTranslatorBase& BaseTranslator, uint32 NumRequired)
+		{
+			UObject* TransientPackage = GetTransientPackage();
+
+			TranslatorsAndSourcesData.Reserve(NumRequired);
+
+			UClass* TranslatorClass = BaseTranslator.GetClass();
+			for (uint32 Index = 0; Index < NumRequired; Index++)
+			{
+				TranslatorsAndSourcesData.Emplace(
+					NewObject<UInterchangeTranslatorBase>(TransientPackage, TranslatorClass, NAME_None)
+					, NewObject<UInterchangeSourceData>(TransientPackage, UInterchangeSourceData::StaticClass(), NAME_None)
+				);
+
+			}
+		};
+
+		~FScopedTranslatorsAndSourceData()
+		{
+			for (TPair<UInterchangeTranslatorBase*, UInterchangeSourceData*>& TranslatorAndSourceData : TranslatorsAndSourcesData)
+			{
+				TranslatorAndSourceData.Key->ClearInternalFlags(EInternalObjectFlags::Async);
+				TranslatorAndSourceData.Value->ClearInternalFlags(EInternalObjectFlags::Async);
+			}
+		};
+
+		TPair<UInterchangeTranslatorBase*, UInterchangeSourceData*>& operator[](uint32 Index)
+		{
+			return TranslatorsAndSourcesData[Index];
+		}
+
+		const TPair<UInterchangeTranslatorBase*, UInterchangeSourceData*>& operator[](uint32 Index) const
+		{
+			return TranslatorsAndSourcesData[Index];
+		}
+
+		uint32 Num() const
+		{
+			return TranslatorsAndSourcesData.Num();
+		}
+
+	private:
+		TArray<TPair<UInterchangeTranslatorBase*, UInterchangeSourceData*>> TranslatorsAndSourcesData;
+	};
+
+	class FScopedClearAsyncFlag
+	{
+	public:
+		FScopedClearAsyncFlag(UObject* InObject)
+			: Object(nullptr)
+		{
+			if (InObject && InObject->HasAnyInternalFlags(EInternalObjectFlags::Async))
+			{
+				Object = InObject;
+			}
+		}
+
+		~FScopedClearAsyncFlag()
+		{
+			if (Object)
+			{
+				Object->ClearInternalFlags(EInternalObjectFlags::Async);
+			}
+		}
+
+	private:
+		UObject* Object;
+	};
+
+
 	/**
 	 * Return the supported class if the node is one otherwise return nullptr
 	 */
 	UClass* GetSupportedFactoryNodeClass(const UInterchangeBaseNode* AssetNode)
 	{
 		UClass* TextureCubeFactoryClass = UInterchangeTextureCubeFactoryNode::StaticClass();
-		UClass* TextureFactoryClass = UInterchangeTextureFactoryNode::StaticClass();
+		UClass* Texture2DFactoryClass = UInterchangeTexture2DFactoryNode::StaticClass();
 		UClass* Texture2DArrayFactoryClass = UInterchangeTexture2DArrayFactoryNode::StaticClass();
 		UClass* TextureLightProfileFactoryClass = UInterchangeTextureLightProfileFactoryNode::StaticClass();
 		UClass* AssetClass = AssetNode->GetClass();
@@ -67,15 +141,15 @@ namespace UE::Interchange::Private::InterchangeTextureFactory
 		{
 			return TextureLightProfileFactoryClass;
 		}
-		else if (AssetClass->IsChildOf(TextureFactoryClass))
+		else if (AssetClass->IsChildOf(Texture2DFactoryClass))
 		{
-			return TextureFactoryClass;
+			return Texture2DFactoryClass;
 		}
 #else
 		while (AssetClass)
 		{
 			if (AssetClass == TextureCubeFactoryClass
-				|| AssetClass == TextureFactoryClass
+				|| AssetClass == Texture2DFactoryClass
 				|| AssetClass == Texture2DArrayFactoryClass
 				|| AssetClass == TextureLightProfileFactoryClass)
 			{
@@ -90,7 +164,7 @@ namespace UE::Interchange::Private::InterchangeTextureFactory
 	}
 
 	using FTextureFactoryNodeVariant = TVariant<FEmptyVariantState
-		, UInterchangeTextureFactoryNode*
+		, UInterchangeTexture2DFactoryNode*
 		, UInterchangeTextureCubeFactoryNode*
 		, UInterchangeTexture2DArrayFactoryNode*
 		,UInterchangeTextureLightProfileFactoryNode* >;
@@ -104,9 +178,9 @@ namespace UE::Interchange::Private::InterchangeTextureFactory
 				SupportedFactoryNodeClass = GetSupportedFactoryNodeClass(AssetNode);
 			}
 
-			if (SupportedFactoryNodeClass == UInterchangeTextureFactoryNode::StaticClass())
+			if (SupportedFactoryNodeClass == UInterchangeTexture2DFactoryNode::StaticClass())
 			{
-				return FTextureFactoryNodeVariant(TInPlaceType<UInterchangeTextureFactoryNode*>(), static_cast<UInterchangeTextureFactoryNode*>(AssetNode));
+				return FTextureFactoryNodeVariant(TInPlaceType<UInterchangeTexture2DFactoryNode*>(), static_cast<UInterchangeTexture2DFactoryNode*>(AssetNode));
 			}
 
 			if (SupportedFactoryNodeClass == UInterchangeTextureCubeFactoryNode::StaticClass())
@@ -138,9 +212,9 @@ namespace UE::Interchange::Private::InterchangeTextureFactory
 	{
 		FString TextureNodeUniqueID;
 
-		if (UInterchangeTextureFactoryNode* const* TextureFactoryNode = FactoryVariant.TryGet<UInterchangeTextureFactoryNode*>())
+		if (UInterchangeTexture2DFactoryNode* const* Texture2DFactoryNode = FactoryVariant.TryGet<UInterchangeTexture2DFactoryNode*>())
 		{
-			(*TextureFactoryNode)->GetCustomTranslatedTextureNodeUid(TextureNodeUniqueID);
+			(*Texture2DFactoryNode)->GetCustomTranslatedTextureNodeUid(TextureNodeUniqueID);
 		}
 		else if (UInterchangeTextureCubeFactoryNode* const* TextureCubeFactoryNode = FactoryVariant.TryGet<UInterchangeTextureCubeFactoryNode*>())
 		{
@@ -231,18 +305,120 @@ namespace UE::Interchange::Private::InterchangeTextureFactory
 		return {};
 	}
 
-	FTexturePayloadVariant GetTexturePayload(const UInterchangeSourceData* SourceData, const FString& PayloadKey, const FTextureNodeVariant& TextureNodeVariant, const UInterchangeTranslatorBase* Translator)
+	TOptional<FImportBlockedImage> GetBlockedTexturePayloadDataFromSourceFiles(const UInterchangeSourceData* SourceData, const TMap<int32, FString>& InBlockSourcesData, const UInterchangeTranslatorBase* Translator)
+	{
+		check(!InBlockSourcesData.IsEmpty());
+
+		TArray<const TPair<int32, FString>*> UDIMsAndSourcesFileArray;
+		UDIMsAndSourcesFileArray.Reserve(InBlockSourcesData.Num());
+		int32 OrginalSourceDataIndex = INDEX_NONE;
+		for (const TPair<int32, FString>& Pair : InBlockSourcesData)
+		{
+			if (OrginalSourceDataIndex == INDEX_NONE && SourceData->GetFilename() == Pair.Value)
+			{
+				OrginalSourceDataIndex = UDIMsAndSourcesFileArray.Num();
+			}
+
+			UDIMsAndSourcesFileArray.Add(&Pair);
+		}
+
+		if (OrginalSourceDataIndex != INDEX_NONE && 0 != OrginalSourceDataIndex)
+		{
+			// Move the original file to the front.
+			UDIMsAndSourcesFileArray.Swap(0, OrginalSourceDataIndex);
+		}
+
+		FScopedTranslatorsAndSourceData TranslatorsAndSourceData(*Translator, UDIMsAndSourcesFileArray.Num());
+
+		/**
+			* Possible improvement notes.
+			* If we are able at some point to extract the size and format of the textures from the translate step for some formats,
+			* We could use those information to init the blocked image RawData and set the ImportsImages RawData to be a view into the blocked image.
+			*/ 
+		TArray<UE::Interchange::FImportImage> Images;
+		Images.AddDefaulted(UDIMsAndSourcesFileArray.Num());
+
+		ParallelFor(UDIMsAndSourcesFileArray.Num(), [&TranslatorsAndSourceData, &UDIMsAndSourcesFileArray, &Images](int32 Index)
+			{
+				TPair<UInterchangeTranslatorBase*,UInterchangeSourceData*>& TranslatorAndSourceData = TranslatorsAndSourceData[Index];
+				UInterchangeTranslatorBase* Translator = TranslatorAndSourceData.Key;
+				IInterchangeTexturePayloadInterface* TextureTranslator = CastChecked<IInterchangeTexturePayloadInterface>(TranslatorAndSourceData.Key);
+				UInterchangeSourceData* SourceDataForBlock = TranslatorAndSourceData.Value;
+
+				const TPair<int32, FString>& UDIMAndFilename = *UDIMsAndSourcesFileArray[Index];
+				const int32 CurrentUDIM = UDIMAndFilename.Key;
+				SourceDataForBlock->SetFilename(UDIMAndFilename.Value);
+				
+				TOptional<UE::Interchange::FImportImage> Payload;
+				if (Translator->CanImportSourceData(SourceDataForBlock))
+				{
+					Translator->SourceData = SourceDataForBlock;
+					Payload = TextureTranslator->GetTexturePayloadData(SourceDataForBlock, SourceDataForBlock->GetFilename());
+				}
+
+				if (Payload.IsSet())
+				{
+					UE::Interchange::FImportImage& Image = Payload.GetValue();
+					// Consume the image
+					Images[Index] = MoveTemp(Image);
+				}
+				else
+				{
+					// Todo Capture the error message from the translator?
+				}
+
+				// Let the translator release his data.
+				Translator->ImportFinish();
+			}
+			, EParallelForFlags::Unbalanced | EParallelForFlags::BackgroundPriority);
+
+
+		UE::Interchange::FImportBlockedImage BlockedImage;
+		// If the image that triggered the import is not valid. We shouldn't import the rest of the images into the UDIM.
+		if (BlockedImage.InitDataSharedAmongBlocks(Images[0]))
+		{
+			BlockedImage.BlocksData.Reserve(Images.Num());
+
+			for (int32 Index = 0; Index < Images.Num(); ++Index)
+			{
+				int32 UDIMIndex = UDIMsAndSourcesFileArray[Index]->Key;
+				const int32 BlockX = (UDIMIndex - 1001) % 10;
+				const int32 BlockY = (UDIMIndex - 1001) / 10;
+				BlockedImage.InitBlockFromImage(BlockX, BlockY, Images[Index]);
+			}
+
+			BlockedImage.MigrateDataFromImagesToRawData(Images);
+		}
+
+		return BlockedImage;
+	}
+
+	FTexturePayloadVariant GetTexturePayload(const UInterchangeSourceData* SourceData, const FString& PayloadKey, const FTextureNodeVariant& TextureNodeVariant, const FTextureFactoryNodeVariant& FactoryNodeVariant, const UInterchangeTranslatorBase* Translator)
 	{
 		// Standard texture 2D payload
 		if (const UInterchangeTexture2DNode* const* TextureNode =  TextureNodeVariant.TryGet<const UInterchangeTexture2DNode*>())
 		{
+			TMap<int32, FString> BlockAndSourceDataFiles;
+			if (const UInterchangeTexture2DFactoryNode* const* Texture2DFactoryNode = FactoryNodeVariant.TryGet<UInterchangeTexture2DFactoryNode*>())
+			{
+				BlockAndSourceDataFiles = (*Texture2DFactoryNode)->GetSourceBlocks();
+			}
+
+			// Is there a case were a translator can be both interface and how should the factory chose which to invoke?
 			if (const IInterchangeTexturePayloadInterface* TextureTranslator = Cast<IInterchangeTexturePayloadInterface>(Translator))
 			{
-				return FTexturePayloadVariant(TInPlaceType<TOptional<FImportImage>>(), TextureTranslator->GetTexturePayloadData(SourceData, PayloadKey));
+				if (BlockAndSourceDataFiles.IsEmpty())
+				{
+					return FTexturePayloadVariant(TInPlaceType<TOptional<FImportImage>>(), TextureTranslator->GetTexturePayloadData(SourceData, PayloadKey));
+				}
+				else
+				{
+					return FTexturePayloadVariant(TInPlaceType<TOptional<FImportBlockedImage>>(), GetBlockedTexturePayloadDataFromSourceFiles(SourceData, BlockAndSourceDataFiles, Translator));
+				}
 			}
 			else if (const IInterchangeBlockedTexturePayloadInterface* BlockedTextureTranslator = Cast<IInterchangeBlockedTexturePayloadInterface>(Translator))
 			{
-				return FTexturePayloadVariant(TInPlaceType<TOptional<FImportBlockedImage>>(), BlockedTextureTranslator->GetBlockedTexturePayloadData((*TextureNode)->GetSourceBlocks(), SourceData));
+				return FTexturePayloadVariant(TInPlaceType<TOptional<FImportBlockedImage>>(), BlockedTextureTranslator->GetBlockedTexturePayloadData(SourceData, PayloadKey));
 			}
 		}
 
@@ -626,11 +802,11 @@ namespace UE::Interchange::Private::InterchangeTextureFactory
 	}
 
 
-	TArray<FString> GetFilesToHash(const FTextureNodeVariant& TextureNodeVariant, const FTexturePayloadVariant& TexturePayload)
+	TArray<FString> GetFilesToHash(const FTextureFactoryNodeVariant& TextureFactoryNodeVariant, const FTexturePayloadVariant& TexturePayload)
 	{
 		TArray<FString> FilesToHash;
 		// Standard texture 2D payload
-		if (const UInterchangeTexture2DNode* const* TextureNode = TextureNodeVariant.TryGet<const UInterchangeTexture2DNode*>())
+		if (const UInterchangeTexture2DFactoryNode* const* TextureNode = TextureFactoryNodeVariant.TryGet<UInterchangeTexture2DFactoryNode*>())
 		{
 			using namespace UE::Interchange;
 			if (const TOptional<FImportBlockedImage>* OptionalBlockedPayload = TexturePayload.TryGet<TOptional<FImportBlockedImage>>())
@@ -874,7 +1050,7 @@ UObject* UInterchangeTextureFactory::CreateAsset(const FCreateAssetParams& Argum
 	UClass* SupportedFactoryNodeClass = GetSupportedFactoryNodeClass(Arguments.AssetNode);
 	if (SupportedFactoryNodeClass == nullptr)
 	{
-		UE_LOG(LogInterchangeImport, Error, TEXT("UInterchangeTextureFactory: Asset node parameter is not a UInterchangeTextureFactoryNode or UInterchangeTextureCubeFactoryNode."));
+		UE_LOG(LogInterchangeImport, Error, TEXT("UInterchangeTextureFactory: Asset node parameter is not a child of UInterchangeTextureFactoryNode"));
 		return nullptr;
 	}
 
@@ -895,7 +1071,7 @@ UObject* UInterchangeTextureFactory::CreateAsset(const FCreateAssetParams& Argum
 		return nullptr;
 	}
 
-	FTexturePayloadVariant TexturePayload = GetTexturePayload(Arguments.SourceData, PayLoadKey.GetValue(), TextureNodeVariant, Arguments.Translator);
+	FTexturePayloadVariant TexturePayload = GetTexturePayload(Arguments.SourceData, PayLoadKey.GetValue(), TextureNodeVariant, TextureFactoryNodeVariant, Arguments.Translator);
 
 	if(TexturePayload.IsType<FEmptyVariantState>())
 	{
@@ -949,7 +1125,7 @@ UObject* UInterchangeTextureFactory::CreateAsset(const FCreateAssetParams& Argum
 		return Texture;
 	}
 
-	FGraphEventArray TasksToDo = GenerateHashSourceFilesTasks(Arguments.SourceData, GetFilesToHash(TextureNodeVariant, TexturePayload), SourceFiles);
+	FGraphEventArray TasksToDo = GenerateHashSourceFilesTasks(Arguments.SourceData, GetFilesToHash(TextureFactoryNodeVariant, TexturePayload), SourceFiles);
 
 	// Hash the payload while we hash the source files
 
diff --git a/Engine/Plugins/Experimental/Interchange/Runtime/Source/Import/Private/Texture/InterchangeUDIMTranslator.cpp b/Engine/Plugins/Experimental/Interchange/Runtime/Source/Import/Private/Texture/InterchangeUDIMTranslator.cpp
deleted file mode 100644
diff --git a/Engine/Plugins/Experimental/Interchange/Runtime/Source/Import/Public/Texture/InterchangeBlockedTexturePayloadInterface.h b/Engine/Plugins/Experimental/Interchange/Runtime/Source/Import/Public/Texture/InterchangeBlockedTexturePayloadInterface.h
--- a/Engine/Plugins/Experimental/Interchange/Runtime/Source/Import/Public/Texture/InterchangeBlockedTexturePayloadInterface.h
+++ b/Engine/Plugins/Experimental/Interchange/Runtime/Source/Import/Public/Texture/InterchangeBlockedTexturePayloadInterface.h
@@ -26,16 +26,14 @@ class INTERCHANGEIMPORT_API IInterchangeBlockedTexturePayloadInterface
 {
 	GENERATED_BODY()
 public:
-
-
 	/**
 	 * Once the translation is done, the import process need a way to retrieve payload data.
 	 * This payload will be use by the factories to create the asset.
 	 *
-	 * @param SourceData - The sources data containing the data to translate
-	 * @param OriginalFile - The file that was used to start the import
+	 * @param SourceData - The source data containing the data to translate
+	 * @param PayloadKey - The key to retrieve the a particular payload contain into the specified source data
 	 * @return a PayloadData containing the import blocked image data. The TOptional will not be set if there is an error.
 	 */
-	virtual TOptional<UE::Interchange::FImportBlockedImage> GetBlockedTexturePayloadData(const TMap<int32, FString>& InSourcesData,  const UInterchangeSourceData* OriginalFile) const = 0;
+	virtual TOptional<UE::Interchange::FImportBlockedImage> GetBlockedTexturePayloadData(const UInterchangeSourceData* PayloadSourceData, const FString& PayLoadKey) const {return {};}
 };
 
diff --git a/Engine/Plugins/Experimental/Interchange/Runtime/Source/Import/Public/Texture/InterchangeUDIMTranslator.h b/Engine/Plugins/Experimental/Interchange/Runtime/Source/Import/Public/Texture/InterchangeUDIMTranslator.h
deleted file mode 100644
diff --git a/Engine/Plugins/Experimental/Interchange/Runtime/Source/Nodes/Public/InterchangeTexture2DNode.h b/Engine/Plugins/Experimental/Interchange/Runtime/Source/Nodes/Public/InterchangeTexture2DNode.h
--- a/Engine/Plugins/Experimental/Interchange/Runtime/Source/Nodes/Public/InterchangeTexture2DNode.h
+++ b/Engine/Plugins/Experimental/Interchange/Runtime/Source/Nodes/Public/InterchangeTexture2DNode.h
@@ -8,20 +8,17 @@
 #include "InterchangeTexture2DNode.generated.h"
 
 //Interchange namespace
-namespace UE
+namespace UE::Interchange
 {
-	namespace Interchange
+	struct FTexture2DNodeStaticData : public FBaseNodeStaticData
 	{
-		struct FTexture2DNodeStaticData : public FBaseNodeStaticData
+		static const FString& GetBaseSourceBlocksKey()
 		{
-			static const FString& GetBaseSourceBlocksKey()
-			{
-				static FString StringKey(TEXT("SourceBlocks"));
-				return StringKey;
-			}
-		};
-	}//ns Interchange
-}//ns UE
+			static FString StringKey(TEXT("SourceBlocks"));
+			return StringKey;
+		}
+	};
+}//ns UE::Interchange
 
 UCLASS(BlueprintType, Experimental)
 class INTERCHANGENODES_API UInterchangeTexture2DNode : public UInterchangeTextureNode
@@ -61,25 +58,27 @@ public:
 
 	/**
 	 * Set the source blocks
-	 * Using this will force the texture factory to consider this texture as UDIM.
+	 * Using this will suggest the pipeline to consider this texture as UDIM and it can chose to pass or not these block to the texture factory node.
 	 * @param InSourceBlocks The blocks and their source image that compose the whole texture.
 	 * The textures must be of the same format and use the same pixel format
-	 * The first block inserted in the map be used to determine the accepted texture format and pixel format
+	 * The first block in the map is used to determine the accepted texture format and pixel format
 	 */
-	UFUNCTION(BlueprintCallable, Category = "Interchange | Node | Texture | UDIMs")
-	void SetSourceBlocks(const TMap<int32, FString>& InSourceBlocks)
+	void SetSourceBlocks(TMap<int32, FString>&& InSourceBlocks)
 	{
 		SourceBlocks = InSourceBlocks;
 	}
 
+	
 	/**
-	 * Remove the source block data from texture node
-	 * This is a easy way to stop a texture from being imported as a UDIM
+	 * Set the source blocks
+	 * Using this will suggest the pipeline to consider this texture as UDIM and it can chose to pass or not these block to the texture factory node.
+	 * @param InSourceBlocks The blocks and their source image that compose the whole texture.
+	 * The textures must be of the same format and use the same pixel format
+	 * The first block in the map is used to determine the accepted texture format and pixel format
 	 */
-	UFUNCTION(BlueprintCallable, Category = "Interchange | Node | Texture | UDIMs")
-	void ClearSourceBlocksData()
+	void SetSourceBlocks(const TMap<int32, FString>& InSourceBlocks)
 	{
-		SourceBlocks.Empty();
+		SourceBlocks = InSourceBlocks;
 	}
 
 	// UDIMs ends here
diff --git a/Engine/Plugins/Experimental/Interchange/Runtime/Source/Pipelines/Private/InterchangeGenericTexturePipeline.cpp b/Engine/Plugins/Experimental/Interchange/Runtime/Source/Pipelines/Private/InterchangeGenericTexturePipeline.cpp
--- a/Engine/Plugins/Experimental/Interchange/Runtime/Source/Pipelines/Private/InterchangeGenericTexturePipeline.cpp
+++ b/Engine/Plugins/Experimental/Interchange/Runtime/Source/Pipelines/Private/InterchangeGenericTexturePipeline.cpp
@@ -10,6 +10,7 @@
 #include "InterchangeTexture2DArrayFactoryNode.h"
 #include "InterchangeTexture2DArrayNode.h"
 #include "InterchangeTexture2DArrayNode.h"
+#include "InterchangeTexture2DFactoryNode.h"
 #include "InterchangeTexture2DNode.h"
 #include "InterchangeTexture2DNode.h"
 #include "InterchangeTextureCubeFactoryNode.h"
@@ -27,6 +28,7 @@
 #if WITH_EDITOR
 #include "NormalMapIdentification.h"
 #include "TextureCompiler.h"
+#include "UDIMUtilities.h"
 #endif //WITH_EDITOR
 
 namespace UE::Interchange::Private
@@ -35,7 +37,7 @@ namespace UE::Interchange::Private
 	{
 		if (UInterchangeTexture2DNode::StaticClass() == NodeClass)
 		{
-			return UInterchangeTextureFactoryNode::StaticClass();
+			return UInterchangeTexture2DFactoryNode::StaticClass();
 		}
 
 		if (UInterchangeTextureCubeNode::StaticClass() == NodeClass)
@@ -56,31 +58,6 @@ namespace UE::Interchange::Private
 		return nullptr;
 	}
 
-	UClass* GetDefaultAssetClassFromFactoryClass(UClass* NodeClass)
-	{
-		if (UInterchangeTextureFactoryNode::StaticClass() == NodeClass)
-		{
-			return UTexture2D::StaticClass();
-		}
-
-		if (UInterchangeTextureCubeFactoryNode::StaticClass() == NodeClass)
-		{
-			return UTextureCube::StaticClass();
-		}
-
-		if (UInterchangeTexture2DArrayFactoryNode::StaticClass() == NodeClass)
-		{
-			return UTexture2DArray::StaticClass();
-		}
-
-		if (UInterchangeTextureLightProfileFactoryNode::StaticClass() == NodeClass)
-		{
-			return UTextureLightProfile::StaticClass();
-		}
-
-		return nullptr;
-	}
-
 #if WITH_EDITOR
 	void AdjustTextureForNormalMap(UTexture* Texture, bool bFlipNormalMapGreenChannel)
 	{
@@ -162,10 +139,11 @@ UInterchangeTextureFactoryNode* UInterchangeGenericTexturePipeline::HandleCreati
 {
 	UClass* FactoryClass = UE::Interchange::Private::GetDefaultFactoryClassFromTextureNodeClass(TextureNode->GetClass());
 
+	TOptional<FString> SourceFile = TextureNode->GetPayLoadKey();
 #if WITH_EDITORONLY_DATA
-	if (FactoryClass == UInterchangeTextureFactoryNode::StaticClass())
+	if (FactoryClass == UInterchangeTexture2DFactoryNode::StaticClass())
 	{ 
-		if (TOptional<FString> SourceFile = TextureNode->GetPayLoadKey())
+		if (SourceFile)
 		{
 			const FString Extension = FPaths::GetExtension(SourceFile.GetValue()).ToLower();
 			if (FileExtensionsToImportAsLongLatCubemap.Contains(Extension))
@@ -176,7 +154,37 @@ UInterchangeTextureFactoryNode* UInterchangeGenericTexturePipeline::HandleCreati
 	}
 #endif
 
-	return CreateTextureFactoryNode(TextureNode, FactoryClass);
+	UInterchangeTextureFactoryNode* InterchangeTextureFactoryNode =  CreateTextureFactoryNode(TextureNode, FactoryClass);
+
+	if (FactoryClass == UInterchangeTexture2DFactoryNode::StaticClass() && InterchangeTextureFactoryNode)
+	{
+		// Forward the UDIM from the translator to the factory node
+		TMap<int32, FString> SourceBlocks;
+		UInterchangeTexture2DFactoryNode* Texture2DFactoryNode = static_cast<UInterchangeTexture2DFactoryNode*>(InterchangeTextureFactoryNode);
+		if (const UInterchangeTexture2DNode* Texture2DNode = Cast<UInterchangeTexture2DNode>(TextureNode))
+		{
+			SourceBlocks = Texture2DNode->GetSourceBlocks();
+		}
+
+#if WITH_EDITOR
+		if (SourceBlocks.IsEmpty() && bImportUDIMs && SourceFile)
+		{
+			FString PrettyAssetName;
+			SourceBlocks = UE::TextureUtilitiesCommon::GetUDIMBlocksFromSourceFile(SourceFile.GetValue(), UE::TextureUtilitiesCommon::DefaultUdimRegexPattern, &PrettyAssetName);
+			if (!PrettyAssetName.IsEmpty())
+			{
+				InterchangeTextureFactoryNode->SetAssetName(PrettyAssetName);
+			}
+		}
+#endif
+
+		if (!SourceBlocks.IsEmpty())
+		{
+			Texture2DFactoryNode->SetSourceBlocks(MoveTemp(SourceBlocks));
+		}
+	}
+
+	return InterchangeTextureFactoryNode;
 }
 
 UInterchangeTextureFactoryNode* UInterchangeGenericTexturePipeline::CreateTextureFactoryNode(const UInterchangeTextureNode* TextureNode, const TSubclassOf<UInterchangeTextureFactoryNode>& FactorySubclass)
@@ -201,20 +209,13 @@ UInterchangeTextureFactoryNode* UInterchangeGenericTexturePipeline::CreateTextur
 			return nullptr;
 		}
 
-		UClass* TextureClass = UE::Interchange::Private::GetDefaultAssetClassFromFactoryClass(FactoryClass);
-		if (!ensure(TextureClass))
-		{
-			// Log an error
-			return nullptr;
-		}
-
 		TextureFactoryNode = NewObject<UInterchangeTextureFactoryNode>(BaseNodeContainer, FactoryClass);
 		if (!ensure(TextureFactoryNode))
 		{
 			return nullptr;
 		}
-		//Creating a UTexture2D
-		TextureFactoryNode->InitializeTextureNode(NodeUid, DisplayLabel, TextureClass->GetName(), TextureNode->GetDisplayLabel());
+		//Creating a Texture
+		TextureFactoryNode->InitializeTextureNode(NodeUid, DisplayLabel, TextureNode->GetDisplayLabel());
 		TextureFactoryNode->SetCustomTranslatedTextureNodeUid(TextureNode->GetUniqueID());
 		BaseNodeContainer->AddNode(TextureFactoryNode);
 		TextureFactoryNodes.Add(TextureFactoryNode);
diff --git a/Engine/Plugins/Experimental/Interchange/Runtime/Source/Pipelines/Public/InterchangeGenericTexturePipeline.h b/Engine/Plugins/Experimental/Interchange/Runtime/Source/Pipelines/Public/InterchangeGenericTexturePipeline.h
--- a/Engine/Plugins/Experimental/Interchange/Runtime/Source/Pipelines/Public/InterchangeGenericTexturePipeline.h
+++ b/Engine/Plugins/Experimental/Interchange/Runtime/Source/Pipelines/Public/InterchangeGenericTexturePipeline.h
@@ -25,7 +25,7 @@ public:
 
 #if WITH_EDITORONLY_DATA
 	/** 
-	 * If enable, after a new import a test will be run to see if the texture is a normal map
+	 * If enable, after a new import a test will be run to see if the texture is a normal map.
 	 * If the texture is a normal map the SRG, CompressionSettings and LODGroup settings will be adjusted.
 	 */
 	UPROPERTY(EditAnywhere, BlueprintReadWrite, Category = "Textures")
@@ -35,6 +35,10 @@ public:
 	UPROPERTY(EditAnywhere, Category = "Textures")
 	bool bFlipNormalMapGreenChannel = false;
 
+	/** If enabled detect if a texture use a UDIM pattern and if so import it as UIDMs. */
+	UPROPERTY(EditAnywhere, Category = "Textures")
+	bool bImportUDIMs = true;
+
 	/** Specify the files type that should be imported as long/lat cubemap */
 	UPROPERTY(EditAnywhere, Category = "Textures")
 	TSet<FString> FileExtensionsToImportAsLongLatCubemap = {"hdr"};
diff --git a/Engine/Source/Runtime/TextureUtilitiesCommon/Private/UDIMUtilities.cpp b/Engine/Source/Runtime/TextureUtilitiesCommon/Private/UDIMUtilities.cpp
--- a/Engine/Source/Runtime/TextureUtilitiesCommon/Private/UDIMUtilities.cpp
+++ b/Engine/Source/Runtime/TextureUtilitiesCommon/Private/UDIMUtilities.cpp
@@ -3,6 +3,8 @@
 #include "UDIMUtilities.h"
 
 #include "Internationalization/Regex.h"
+#include "Misc/Paths.h"
+#include "HAL/FileManager.h"
 
 namespace UE
 {
@@ -52,6 +54,41 @@ namespace UE
 			return BlockY * 10 + BlockX + 1001;
 		}
 
+		TEXTUREUTILITIESCOMMON_API TMap<int32, FString> GetUDIMBlocksFromSourceFile(const FString& File, const FString& UdimRegexPattern, FString* OutFilenameWithoutUdimPatternAndExtension)
+		{
+			FString PreUDIMName;
+			FString PostUDIMName;
+			TMap<int32, FString> UDIMsAndSourcesFile;
+
+			if (ParseUDIMName(FPaths::GetBaseFilename(File), UdimRegexPattern, PreUDIMName, PostUDIMName) != INDEX_NONE)
+			{ 
+				const FString Path = FPaths::GetPath(File);
+				const FString UDIMFilter = (Path / PreUDIMName) + TEXT("*") + PostUDIMName + FPaths::GetExtension(File, true);
+
+				TArray<FString> UDIMFiles;
+				IFileManager::Get().FindFiles(UDIMFiles, *UDIMFilter, true, false);
+				UDIMFiles.Reserve(UDIMFiles.Num());
+
+				for (FString& UDIMFile : UDIMFiles)
+				{
+					const int32 UDIMIndex = ParseUDIMName(FPaths::GetBaseFilename(UDIMFile), UdimRegexPattern, PreUDIMName, PostUDIMName);
+					if (UDIMIndex != INDEX_NONE)
+					{
+						const FString UDIMName = PreUDIMName + PostUDIMName;
+						FString UDIMPath = FPaths::Combine(Path, UDIMFile);
+						UDIMsAndSourcesFile.Add(UDIMIndex, MoveTemp(UDIMPath));
+					}
+				}
+
+				if (!UDIMsAndSourcesFile.IsEmpty() && OutFilenameWithoutUdimPatternAndExtension)
+				{
+					*OutFilenameWithoutUdimPatternAndExtension = PreUDIMName + PostUDIMName;
+				}
+			}
+
+			return UDIMsAndSourcesFile;
+		}
+
 		void ExtractUDIMCoordinates(int32 UDIMIndex, int32& OutBlockX, int32& OutBlockY)
 		{
 			OutBlockX = (UDIMIndex - 1001) % 10;
diff --git a/Engine/Source/Runtime/TextureUtilitiesCommon/Public/UDIMUtilities.h b/Engine/Source/Runtime/TextureUtilitiesCommon/Public/UDIMUtilities.h
--- a/Engine/Source/Runtime/TextureUtilitiesCommon/Public/UDIMUtilities.h
+++ b/Engine/Source/Runtime/TextureUtilitiesCommon/Public/UDIMUtilities.h
@@ -16,6 +16,12 @@ namespace UE
 
 		TEXTUREUTILITIESCOMMON_API int32 GetUDIMIndex(int32 BlockX, int32 BlockY);
 
+		/**
+		 * Parse the file for the UDIM pattern
+		 * If the pattern is found search in the folder of the file for the others blocks
+		 */
+		TEXTUREUTILITIESCOMMON_API TMap<int32, FString> GetUDIMBlocksFromSourceFile(const FString& File, const FString& UdimRegexPattern, FString* OutFilenameWithoutUdimPatternAndExtension = nullptr);
+
 		TEXTUREUTILITIESCOMMON_API void ExtractUDIMCoordinates(int32 UDIMIndex, int32& OutBlockX, int32& OutBlockY);
 	}
 }
