Fixed a non unity compile issue (missing an include).

diff --git a/Engine/Plugins/Experimental/Interchange/Runtime/Source/Import/Private/Texture/InterchangeTextureFactory.cpp b/Engine/Plugins/Experimental/Interchange/Runtime/Source/Import/Private/Texture/InterchangeTextureFactory.cpp
--- a/Engine/Plugins/Experimental/Interchange/Runtime/Source/Import/Private/Texture/InterchangeTextureFactory.cpp
+++ b/Engine/Plugins/Experimental/Interchange/Runtime/Source/Import/Private/Texture/InterchangeTextureFactory.cpp
@@ -2,6 +2,7 @@
 
 #include "Texture/InterchangeTextureFactory.h"
 
+#include "Async/ParallelFor.h"
 #include "Async/TaskGraphInterfaces.h"
 #include "Async/TaskGraphInterfaces.h"
 #include "Engine/Texture.h"
@@ -16,11 +17,11 @@
 #include "InterchangeImportLog.h"
 #include "InterchangeTexture2DArrayFactoryNode.h"
 #include "InterchangeTexture2DArrayNode.h"
+#include "InterchangeTexture2DFactoryNode.h"
 #include "InterchangeTexture2DNode.h"
 #include "InterchangeTextureCubeFactoryNode.h"
 #include "InterchangeTextureCubeNode.h"
 #include "InterchangeTextureFactoryNode.h"
-#include "InterchangeTexture2DFactoryNode.h"
 #include "InterchangeTextureLightProfileFactoryNode.h"
 #include "InterchangeTextureLightProfileNode.h"
 #include "InterchangeTextureNode.h"
