Add Interchange support for Import geo/skinning only workflow
#UETOOL-4729
#rnx

diff --git a/Engine/Plugins/Experimental/Interchange/Runtime/Source/FactoryNodes/Private/InterchangeSkeletalMeshFactoryNode.cpp b/Engine/Plugins/Experimental/Interchange/Runtime/Source/FactoryNodes/Private/InterchangeSkeletalMeshFactoryNode.cpp
--- a/Engine/Plugins/Experimental/Interchange/Runtime/Source/FactoryNodes/Private/InterchangeSkeletalMeshFactoryNode.cpp
+++ b/Engine/Plugins/Experimental/Interchange/Runtime/Source/FactoryNodes/Private/InterchangeSkeletalMeshFactoryNode.cpp
@@ -178,6 +178,52 @@ bool UInterchangeSkeletalMeshFactoryNode::SetCustomVertexColorOverride(const FCo
 	IMPLEMENT_NODE_ATTRIBUTE_SETTER_NODELEGATE(VertexColorOverride, FColor)
 }
 
+bool UInterchangeSkeletalMeshFactoryNode::GetCustomImportContentType(EInterchangeSkeletalMeshContentType& AttributeValue) const
+{
+	FString OperationName = GetTypeName() + TEXT(".GetImportContentType");
+	uint8 EnumRawValue = 0;
+	const bool bResult = InterchangePrivateNodeBase::GetCustomAttribute<uint8>(*Attributes, Macro_CustomImportContentTypeKey, OperationName, EnumRawValue);
+	if (bResult)
+	{
+		AttributeValue = static_cast<EInterchangeSkeletalMeshContentType>(EnumRawValue);
+	}
+	return bResult;
+}
+
+bool UInterchangeSkeletalMeshFactoryNode::SetCustomImportContentType(const EInterchangeSkeletalMeshContentType& AttributeValue)
+{
+	FString OperationName = GetTypeName() + TEXT(".SetImportContentType");
+	uint8 EnumRawValue = static_cast<uint8>(AttributeValue);
+	return InterchangePrivateNodeBase::SetCustomAttribute<uint8>(*Attributes, Macro_CustomImportContentTypeKey, OperationName, EnumRawValue);
+}
+
+void UInterchangeSkeletalMeshFactoryNode::AppendAssetRegistryTags(TArray<FAssetRegistryTag>& OutTags) const
+{
+	Super::AppendAssetRegistryTags(OutTags);
+#if WITH_EDITORONLY_DATA
+	EInterchangeSkeletalMeshContentType ContentType;
+	if (GetCustomImportContentType(ContentType))
+	{
+		auto ImportContentTypeToString = [](const EInterchangeSkeletalMeshContentType value)-> FString
+		{
+			switch (value)
+			{
+			case EInterchangeSkeletalMeshContentType::All:
+				return NSSkeletalMeshSourceFileLabels::GeoAndSkinningMetaDataValue();
+			case EInterchangeSkeletalMeshContentType::Geometry:
+				return NSSkeletalMeshSourceFileLabels::GeometryMetaDataValue();
+			case EInterchangeSkeletalMeshContentType::SkinningWeights:
+				return NSSkeletalMeshSourceFileLabels::SkinningMetaDataValue();
+			}
+			return NSSkeletalMeshSourceFileLabels::GeoAndSkinningMetaDataValue();
+		};
+
+		FString EnumString = ImportContentTypeToString(ContentType);
+		OutTags.Add(FAssetRegistryTag(NSSkeletalMeshSourceFileLabels::GetSkeletalMeshLastImportContentTypeMetadataKey(), EnumString, FAssetRegistryTag::TT_Hidden));
+	}
+#endif
+}
+
 void UInterchangeSkeletalMeshFactoryNode::FillAssetClassFromAttribute()
 {
 #if WITH_ENGINE
diff --git a/Engine/Plugins/Experimental/Interchange/Runtime/Source/FactoryNodes/Public/InterchangeSkeletalMeshFactoryNode.h b/Engine/Plugins/Experimental/Interchange/Runtime/Source/FactoryNodes/Public/InterchangeSkeletalMeshFactoryNode.h
--- a/Engine/Plugins/Experimental/Interchange/Runtime/Source/FactoryNodes/Public/InterchangeSkeletalMeshFactoryNode.h
+++ b/Engine/Plugins/Experimental/Interchange/Runtime/Source/FactoryNodes/Public/InterchangeSkeletalMeshFactoryNode.h
@@ -25,6 +25,15 @@ namespace UE
 	}//ns Interchange
 }//ns UE
 
+UENUM(BlueprintType)
+enum class EInterchangeSkeletalMeshContentType : uint8
+{
+	All UMETA(DisplayName = "Geometry and Skinning Weights.", ToolTip = "Import all skelatl mesh content: geometry, skinning and weights."),
+	Geometry UMETA(DisplayName = "Geometry Only", ToolTip = "Import the skeletal mesh geometry only (will create a default skeleton, or map the geometry to the existing one). Morph and LOD can be imported with it."),
+	SkinningWeights UMETA(DisplayName = "Skinning Weights Only", ToolTip = "Import the skeletal mesh skinning and weights only (no geometry will be imported). Morph and LOD will not be imported with this settings."),
+	MAX,
+};
+
 UCLASS(BlueprintType, Experimental)
 class INTERCHANGEFACTORYNODES_API UInterchangeSkeletalMeshFactoryNode : public UInterchangeFactoryBaseNode
 {
@@ -124,6 +133,19 @@ public:
 	/** Set weather the skeletal mesh factory should override the vertex color. Return false if the attribute cannot be set.*/
 	UFUNCTION(BlueprintCallable, Category = "Interchange | Node | SkeletalMesh")
 	bool SetCustomVertexColorOverride(const FColor& AttributeValue);
+
+	/** Query the skeletal mesh import content type. The content type is use by the factory to import partial or full translated content. Return false if the attribute was not set.*/
+	UFUNCTION(BlueprintCallable, Category = "Interchange | Node | SkeletalMesh")
+	bool GetCustomImportContentType(EInterchangeSkeletalMeshContentType& AttributeValue) const;
+
+	/** Set the skeletal mesh import content type. The content type is use by the factory to import partial or full translated content. Return false if the attribute cannot be set.*/
+	UFUNCTION(BlueprintCallable, Category = "Interchange | Node | SkeletalMesh")
+	bool SetCustomImportContentType(const EInterchangeSkeletalMeshContentType& AttributeValue);
+
+	/**
+	 * The skeletal mesh thumbnail can have an overlay if the last re-import was geometry only. This thumbnail overlay feature use the metadata to find out if the last import was geometry only.
+	 */
+	virtual void AppendAssetRegistryTags(TArray<FAssetRegistryTag>& OutTags) const override;
 private:
 
 	void FillAssetClassFromAttribute();
@@ -140,6 +162,7 @@ private:
 	const UE::Interchange::FAttributeKey Macro_CustomVertexColorReplaceKey = UE::Interchange::FAttributeKey(TEXT("VertexColorReplace"));
 	const UE::Interchange::FAttributeKey Macro_CustomVertexColorIgnoreKey = UE::Interchange::FAttributeKey(TEXT("VertexColorIgnore"));
 	const UE::Interchange::FAttributeKey Macro_CustomVertexColorOverrideKey = UE::Interchange::FAttributeKey(TEXT("VertexColorOverride"));
+	const UE::Interchange::FAttributeKey Macro_CustomImportContentTypeKey = UE::Interchange::FAttributeKey(TEXT("ImportContentType"));
 
 	UE::Interchange::TArrayAttributeHelper<FString> LodDependencies;
 protected:
diff --git a/Engine/Plugins/Experimental/Interchange/Runtime/Source/Import/Private/InterchangeImportCommon.cpp b/Engine/Plugins/Experimental/Interchange/Runtime/Source/Import/Private/InterchangeImportCommon.cpp
--- a/Engine/Plugins/Experimental/Interchange/Runtime/Source/Import/Private/InterchangeImportCommon.cpp
+++ b/Engine/Plugins/Experimental/Interchange/Runtime/Source/Import/Private/InterchangeImportCommon.cpp
@@ -80,6 +80,21 @@ namespace UE
 
 		UAssetImportData* FFactoryCommon::UpdateImportAssetData(FUpdateImportAssetDataParameters& Parameters)
 		{
+			return UpdateImportAssetData(Parameters, [&Parameters](UInterchangeAssetImportData* AssetImportData)
+				{
+#if WITH_EDITORONLY_DATA
+					//Set the asset import data file source to allow reimport. TODO: manage MD5 Hash properly
+					TOptional<FMD5Hash> FileContentHash = Parameters.SourceData->GetFileContentHash();
+
+					//Update the first filename, TODO for asset using multiple source file we have to update the correct index
+					AssetImportData->Update(Parameters.SourceData->GetFilename(), FileContentHash.IsSet() ? &FileContentHash.GetValue() : nullptr);
+#endif //WITH_EDITORONLY_DATA
+				});
+
+		}
+
+		UAssetImportData* FFactoryCommon::UpdateImportAssetData(FUpdateImportAssetDataParameters& Parameters, TFunctionRef<void(UInterchangeAssetImportData*)> CustomFileSourceUpdate)
+		{
 #if WITH_EDITORONLY_DATA
 			if (!ensure(IsInGameThread()))
 			{
@@ -102,11 +117,7 @@ namespace UE
 				}
 			}
 
-			//Set the asset import data file source to allow reimport. TODO: manage MD5 Hash properly
-			TOptional<FMD5Hash> FileContentHash = Parameters.SourceData->GetFileContentHash();
-
-			//Update the first filename, TODO for asset using multiple source file we have to update the correct index
-			AssetImportData->Update(Parameters.SourceData->GetFilename(), FileContentHash.IsSet() ? &FileContentHash.GetValue() : nullptr);
+			CustomFileSourceUpdate(AssetImportData);
 
 
 			Private::ImportCommon::EndSetupAssetData(Parameters, AssetImportData);
@@ -117,7 +128,6 @@ namespace UE
 			return nullptr;
 		}
 
-
 #if WITH_EDITORONLY_DATA
 		FFactoryCommon::FSetImportAssetDataParameters::FSetImportAssetDataParameters(UObject* InAssetImportDataOuter
 																					, UAssetImportData* InAssetImportData
@@ -187,10 +197,7 @@ namespace UE
 			{
 				case EReimportStrategyFlags::ApplyNoProperties:
 				{
-					//We want to have no effect, we want to keep the original pipeline node.
-					//So we copy the previous asset node into the pipeline mode, the pipeline node will be saved
-					//in the import asset data, and it will save the original import node.
-					UInterchangeBaseNode::CopyStorage(PreviousAssetNode, PipelineAssetNode);
+					//We want to have no effect
 					break;
 				}
 					
diff --git a/Engine/Plugins/Experimental/Interchange/Runtime/Source/Import/Private/InterchangeImportCommon.h b/Engine/Plugins/Experimental/Interchange/Runtime/Source/Import/Private/InterchangeImportCommon.h
--- a/Engine/Plugins/Experimental/Interchange/Runtime/Source/Import/Private/InterchangeImportCommon.h
+++ b/Engine/Plugins/Experimental/Interchange/Runtime/Source/Import/Private/InterchangeImportCommon.h
@@ -5,6 +5,7 @@
 #include "EditorFramework/AssetImportData.h"
 
 class UAssetImportData;
+class UInterchangeAssetImportData;
 class UInterchangeBaseNode;
 class UInterchangeBaseNodeContainer;
 class UInterchangePipelineBase;
@@ -47,6 +48,14 @@ namespace UE
 			 */
 			static UAssetImportData* UpdateImportAssetData(FUpdateImportAssetDataParameters& Parameters);
 
+			/**
+			 * Update the AssetImportData of the specified asset in the parameters. Also update the node container and the node unique id.
+			 * If the AssetImportData is null it will create one. If the AssetImportData is not an UInterchangeAssetImportData it will create a new one.
+			 * The file source update is done by calling the function parameter CustomFileSourceUpdate, so its the client responsability to properly update the file source.
+			 * @return The source data that should be stored on the asset or nullptr if a parameter is invalid
+			 */
+			static UAssetImportData* UpdateImportAssetData(FUpdateImportAssetDataParameters& Parameters, TFunctionRef<void(UInterchangeAssetImportData*)> CustomFileSourceUpdate);
+
 #if WITH_EDITORONLY_DATA
 			struct FSetImportAssetDataParameters : public FUpdateImportAssetDataParameters
 			{
diff --git a/Engine/Plugins/Experimental/Interchange/Runtime/Source/Import/Private/Mesh/InterchangeSkeletalMeshFactory.cpp b/Engine/Plugins/Experimental/Interchange/Runtime/Source/Import/Private/Mesh/InterchangeSkeletalMeshFactory.cpp
--- a/Engine/Plugins/Experimental/Interchange/Runtime/Source/Import/Private/Mesh/InterchangeSkeletalMeshFactory.cpp
+++ b/Engine/Plugins/Experimental/Interchange/Runtime/Source/Import/Private/Mesh/InterchangeSkeletalMeshFactory.cpp
@@ -1152,8 +1152,32 @@ UObject* UInterchangeSkeletalMeshFactory::CreateAsset(const FCreateAssetParams&
 	ensure(LodDataUniqueIds.Num() == LodCount);
 	int32 CurrentLodIndex = 0;
 
+	EInterchangeSkeletalMeshContentType ImportContent = EInterchangeSkeletalMeshContentType::All;
+	SkeletalMeshFactoryNode->GetCustomImportContentType(ImportContent);
+	const bool bApplyGeometry = !bIsReImport || (ImportContent == EInterchangeSkeletalMeshContentType::All || ImportContent == EInterchangeSkeletalMeshContentType::Geometry);
+	const bool bApplySkinning = !bIsReImport || (ImportContent == EInterchangeSkeletalMeshContentType::All || ImportContent == EInterchangeSkeletalMeshContentType::SkinningWeights);
+	const bool bApplyPartialContent = bIsReImport && ImportContent != EInterchangeSkeletalMeshContentType::All;
+	const bool bApplyGeometryOnly = bApplyPartialContent && bApplyGeometry;
+	const bool bApplySkinningOnly = bApplyPartialContent && bApplySkinning;
+	
+	if (bApplySkinningOnly)
+	{
+		//Ignore vertex color when we import only the skinning
+		constexpr bool bForceIgnoreVertexColor = true;
+		SkeletalMeshFactoryNode->SetCustomVertexColorIgnore(bForceIgnoreVertexColor);
+		constexpr bool bFalseSetting = false;
+		SkeletalMeshFactoryNode->SetCustomVertexColorReplace(bFalseSetting);
+	}
+
 	for (int32 LodIndex = 0; LodIndex < LodCount; ++LodIndex)
 	{
+		ESkeletalMeshGeoImportVersions GeoImportVersion = ESkeletalMeshGeoImportVersions::LatestVersion;
+		ESkeletalMeshSkinningImportVersions SkinningImportVersion = ESkeletalMeshSkinningImportVersions::LatestVersion;
+		if (bIsReImport && SkeletalMesh->GetImportedModel() && SkeletalMesh->GetImportedModel()->LODModels.IsValidIndex(CurrentLodIndex))
+		{
+			SkeletalMesh->GetLODImportedDataVersions(CurrentLodIndex, GeoImportVersion, SkinningImportVersion);
+		}
+
 		FString LodUniqueId = LodDataUniqueIds[LodIndex];
 		const UInterchangeSkeletalMeshLodDataNode* LodDataNode = Cast<UInterchangeSkeletalMeshLodDataNode>(Arguments.NodeContainer->GetNode(LodUniqueId));
 		if (!LodDataNode)
@@ -1331,9 +1355,13 @@ UObject* UInterchangeSkeletalMeshFactory::CreateAsset(const FCreateAssetParams&
 					}
 				}
 			}
-			// Store whether or not this mesh has vertex colors
-			SkeletalMesh->SetHasVertexColors(SkeletalMeshImportData.bHasVertexColors);
-			SkeletalMesh->SetVertexColorGuid(SkeletalMesh->GetHasVertexColors() ? FGuid::NewGuid() : FGuid());
+
+			if (bApplyGeometry)
+			{
+				// Store whether or not this mesh has vertex colors
+				SkeletalMesh->SetHasVertexColors(SkeletalMeshImportData.bHasVertexColors);
+				SkeletalMesh->SetVertexColorGuid(SkeletalMesh->GetHasVertexColors() ? FGuid::NewGuid() : FGuid());
+			}
 		}
 
 		if (bIsReImport)
@@ -1373,10 +1401,32 @@ UObject* UInterchangeSkeletalMeshFactory::CreateAsset(const FCreateAssetParams&
 
 		UE::Interchange::Private::ProcessImportMeshMaterials(SkeletalMesh->GetMaterials(), SkeletalMeshImportData, AvailableMaterials);
 		UE::Interchange::Private::ProcessImportMeshInfluences(SkeletalMeshImportData.Wedges.Num(), SkeletalMeshImportData.Influences);
+
+		if (bApplyGeometryOnly)
+		{
+			FSkeletalMeshImportData::ReplaceSkeletalMeshRigImportData(SkeletalMesh, &SkeletalMeshImportData, CurrentLodIndex);
+		}
+		else if(bApplySkinningOnly)
+		{
+			FSkeletalMeshImportData::ReplaceSkeletalMeshGeometryImportData(SkeletalMesh, &SkeletalMeshImportData, CurrentLodIndex);
+		}
+
 		//Store the original fbx import data the SkelMeshImportDataPtr should not be modified after this
 		SkeletalMesh->SaveLODImportedData(CurrentLodIndex, SkeletalMeshImportData);
-		//We reimport both
-		SkeletalMesh->SetLODImportedDataVersions(CurrentLodIndex, ESkeletalMeshGeoImportVersions::LatestVersion, ESkeletalMeshSkinningImportVersions::LatestVersion);
+
+		if (bApplySkinningOnly)
+		{
+			SkeletalMesh->SetLODImportedDataVersions(CurrentLodIndex, GeoImportVersion, ESkeletalMeshSkinningImportVersions::LatestVersion);
+		}
+		else if (bApplyGeometryOnly)
+		{
+			SkeletalMesh->SetLODImportedDataVersions(CurrentLodIndex, ESkeletalMeshGeoImportVersions::LatestVersion, SkinningImportVersion);
+		}
+		else
+		{
+			//We reimport both
+			SkeletalMesh->SetLODImportedDataVersions(CurrentLodIndex, ESkeletalMeshGeoImportVersions::LatestVersion, ESkeletalMeshSkinningImportVersions::LatestVersion);
+		}
 
 		auto AddLodInfo = [&SkeletalMesh]()
 		{
@@ -1446,10 +1496,17 @@ UObject* UInterchangeSkeletalMeshFactory::CreateAsset(const FCreateAssetParams&
 		CurrentLodIndex++;
 	}
 
-	SkeletonReference->MergeAllBonesToBoneTree(SkeletalMesh);
-	if (SkeletalMesh->GetSkeleton() != SkeletonReference)
+	if(SkeletonReference)
 	{
-		SkeletalMesh->SetSkeleton(SkeletonReference);
+		SkeletonReference->MergeAllBonesToBoneTree(SkeletalMesh);
+		if (SkeletalMesh->GetSkeleton() != SkeletonReference)
+		{
+			SkeletalMesh->SetSkeleton(SkeletonReference);
+		}
+	}
+	else
+	{
+		UE_LOG(LogInterchangeImport, Error, TEXT("Interchange Import UInterchangeSkeletalMeshFactory::CreateAsset, USkeleton* SkeletonReference is nullptr."));
 	}
 
 	SkeletalMesh->CalculateInvRefMatrices();
@@ -1518,7 +1575,58 @@ void UInterchangeSkeletalMeshFactory::PreImportPreCompletedCallback(const FImpor
 																										  , Arguments.NodeContainer
 																										  , Arguments.Pipelines);
 
-		ImportDataPtr = UE::Interchange::FFactoryCommon::UpdateImportAssetData(UpdateImportAssetDataParameters);
+		ImportDataPtr = UE::Interchange::FFactoryCommon::UpdateImportAssetData(UpdateImportAssetDataParameters, [&Arguments, SkeletalMesh](UInterchangeAssetImportData* AssetImportData)
+			{
+				auto GetSourceIndexFromContentType = [](const EInterchangeSkeletalMeshContentType& ImportContentType)->int32
+				{
+					return (ImportContentType == EInterchangeSkeletalMeshContentType::Geometry)
+						? 1
+						: (ImportContentType == EInterchangeSkeletalMeshContentType::SkinningWeights)
+						? 2
+						: 0;
+				};
+
+				auto GetSourceLabelFromSourceIndex = [](const int32 SourceIndex)->FString
+				{
+					return (SourceIndex == 1)
+						? NSSkeletalMeshSourceFileLabels::GeometryText().ToString()
+						: (SourceIndex == 2)
+						? NSSkeletalMeshSourceFileLabels::SkinningText().ToString()
+						: NSSkeletalMeshSourceFileLabels::GeoAndSkinningText().ToString();
+				};
+
+				if (UInterchangeSkeletalMeshFactoryNode* SkeletalMeshFactoryNode = Cast<UInterchangeSkeletalMeshFactoryNode>(Arguments.NodeContainer->GetNode(Arguments.NodeUniqueID)))
+				{
+					EInterchangeSkeletalMeshContentType ImportContentType = EInterchangeSkeletalMeshContentType::All;
+					SkeletalMeshFactoryNode->GetCustomImportContentType(ImportContentType);
+					const FString& NewSourceFilename = Arguments.SourceData->GetFilename();
+					const int32 NewSourceIndex = GetSourceIndexFromContentType(ImportContentType);
+					//NewSourceIndex should be 0, 1 or 2 (All, Geo, Skinning)
+					check(NewSourceIndex >= 0 && NewSourceIndex < 3);
+					const FString DefaultFilename = AssetImportData->ScriptGetFirstFilename();
+					const TArray<FString> OldFilenames = AssetImportData->ScriptExtractFilenames();
+					for (int32 SourceIndex = 0; SourceIndex < 3; ++SourceIndex)
+					{
+						FString SourceLabel = GetSourceLabelFromSourceIndex(SourceIndex);
+						if (SourceIndex == NewSourceIndex)
+						{
+							AssetImportData->ScriptedAddFilename(NewSourceFilename, SourceIndex, SourceLabel);
+						}
+						else
+						{
+							//Extract filename create a default path if the FSourceFile::RelativeFilename is empty
+							//We want to fill the entry with the base source file (SourceIndex 0, All) in this case.
+							const bool bValidOldFilename = AssetImportData->SourceData.SourceFiles.IsValidIndex(SourceIndex)
+								&& !AssetImportData->SourceData.SourceFiles[SourceIndex].RelativeFilename.IsEmpty()
+								&& OldFilenames.IsValidIndex(SourceIndex);
+							const FString& OldFilename = bValidOldFilename ? OldFilenames[SourceIndex] : DefaultFilename;
+							AssetImportData->ScriptedAddFilename(OldFilename, SourceIndex, SourceLabel);
+						}
+					}
+				}
+			});
+
+		
 		SkeletalMesh->SetAssetImportData(ImportDataPtr);
 	}
 #endif
diff --git a/Engine/Plugins/Experimental/Interchange/Runtime/Source/Pipelines/Private/InterchangeGenericAssetsPipeline.cpp b/Engine/Plugins/Experimental/Interchange/Runtime/Source/Pipelines/Private/InterchangeGenericAssetsPipeline.cpp
--- a/Engine/Plugins/Experimental/Interchange/Runtime/Source/Pipelines/Private/InterchangeGenericAssetsPipeline.cpp
+++ b/Engine/Plugins/Experimental/Interchange/Runtime/Source/Pipelines/Private/InterchangeGenericAssetsPipeline.cpp
@@ -27,9 +27,9 @@
 
 UInterchangeGenericAssetsPipeline::UInterchangeGenericAssetsPipeline()
 {
+	TexturePipeline = CreateDefaultSubobject<UInterchangeGenericTexturePipeline>("TexturePipeline");
 	MaterialPipeline = CreateDefaultSubobject<UInterchangeGenericMaterialPipeline>("MaterialPipeline");
 	MeshPipeline = CreateDefaultSubobject<UInterchangeGenericMeshPipeline>("MeshPipeline");
-	TexturePipeline = CreateDefaultSubobject<UInterchangeGenericTexturePipeline>("TexturePipeline");
 }
 
 void UInterchangeGenericAssetsPipeline::PreDialogCleanup(const FName PipelineStackName)
@@ -112,6 +112,24 @@ void UInterchangeGenericAssetsPipeline::ExecutePostImportPipeline(const UInterch
 	}
 }
 
+void UInterchangeGenericAssetsPipeline::SetReimportSourceIndex(UClass* ReimportObjectClass, const int32 SourceFileIndex)
+{
+	if (TexturePipeline)
+	{
+		TexturePipeline->ScriptedSetReimportSourceIndex(ReimportObjectClass, SourceFileIndex);
+	}
+
+	if (MaterialPipeline)
+	{
+		MaterialPipeline->ScriptedSetReimportSourceIndex(ReimportObjectClass, SourceFileIndex);
+	}
+
+	if (MeshPipeline)
+	{
+		MeshPipeline->ScriptedSetReimportSourceIndex(ReimportObjectClass, SourceFileIndex);
+	}
+}
+
 void UInterchangeGenericAssetsPipeline::ImplementUseSourceNameForAssetOption()
 {
 	const UClass* SkeletalMeshFactoryNodeClass = UInterchangeSkeletalMeshFactoryNode::StaticClass();
diff --git a/Engine/Plugins/Experimental/Interchange/Runtime/Source/Pipelines/Private/InterchangeGenericMeshPipeline.cpp b/Engine/Plugins/Experimental/Interchange/Runtime/Source/Pipelines/Private/InterchangeGenericMeshPipeline.cpp
--- a/Engine/Plugins/Experimental/Interchange/Runtime/Source/Pipelines/Private/InterchangeGenericMeshPipeline.cpp
+++ b/Engine/Plugins/Experimental/Interchange/Runtime/Source/Pipelines/Private/InterchangeGenericMeshPipeline.cpp
@@ -51,12 +51,51 @@ void UInterchangeGenericMeshPipeline::ExecutePostImportPipeline(const UInterchan
 		return;
 	}
 
+	//Set the last content type import
+	LastSkeletalMeshImportContentType = SkeletalMeshImportContentType;
+
 	PostImportSkeletalMesh(CreatedAsset, Node);
 
 	//Finish the physics asset import, it need the skeletal mesh render data to create the physics collision geometry
 	PostImportPhysicsAssetImport(CreatedAsset, Node);
 }
 
+void UInterchangeGenericMeshPipeline::SetReimportSourceIndex(UClass* ReimportObjectClass, const int32 SourceFileIndex)
+{
+	if (ReimportObjectClass == USkeletalMesh::StaticClass())
+	{
+		switch (SourceFileIndex)
+		{
+			case 0:
+			{
+				//Geo and skinning
+				SkeletalMeshImportContentType = EInterchangeSkeletalMeshContentType::All;
+			}
+			break;
+
+			case 1:
+			{
+				//Geo only
+				SkeletalMeshImportContentType = EInterchangeSkeletalMeshContentType::Geometry;
+			}
+			break;
+
+			case 2:
+			{
+				//Skinning only
+				SkeletalMeshImportContentType = EInterchangeSkeletalMeshContentType::SkinningWeights;
+			}
+			break;
+
+			default:
+			{
+				//In case SourceFileIndex == INDEX_NONE //No specified options, we use the last imported content type
+				SkeletalMeshImportContentType = LastSkeletalMeshImportContentType;
+			}
+		};
+	}
+}
+
 void UInterchangeGenericMeshPipeline::PreDialogCleanup(const FName PipelineStackName)
 {
 	//We always clean the pipeline skeleton when showing the dialog
diff --git a/Engine/Plugins/Experimental/Interchange/Runtime/Source/Pipelines/Private/InterchangeGenericSkeletalMeshPipeline.cpp b/Engine/Plugins/Experimental/Interchange/Runtime/Source/Pipelines/Private/InterchangeGenericSkeletalMeshPipeline.cpp
--- a/Engine/Plugins/Experimental/Interchange/Runtime/Source/Pipelines/Private/InterchangeGenericSkeletalMeshPipeline.cpp
+++ b/Engine/Plugins/Experimental/Interchange/Runtime/Source/Pipelines/Private/InterchangeGenericSkeletalMeshPipeline.cpp
@@ -566,6 +566,8 @@ UInterchangeSkeletalMeshFactoryNode* UInterchangeGenericMeshPipeline::CreateSkel
 	AddLodDataToSkeletalMesh(SkeletonFactoryNode, SkeletalMeshFactoryNode, MeshUidsPerLodIndex);
 	SkeletalMeshFactoryNode->SetCustomImportMorphTarget(bImportMorphTargets);
 
+	SkeletalMeshFactoryNode->SetCustomImportContentType(SkeletalMeshImportContentType);
+
 	//If we have a specified skeleton
 	if (Skeleton)
 	{
@@ -740,7 +742,9 @@ void UInterchangeGenericMeshPipeline::PostImportSkeletalMesh(UObject* CreatedAss
 		return;
 	}
 
-	if (bUpdateSkeletonReferencePose && !Skeleton.IsNull() && SkeletalMesh->GetSkeleton() == Skeleton.Get())
+	//If we import only the geometry we do not want to update the skeleton reference pose.
+	const bool bImportGeometryOnlyContent = SkeletalMeshImportContentType == EInterchangeSkeletalMeshContentType::Geometry;
+	if (!bImportGeometryOnlyContent && bUpdateSkeletonReferencePose && !Skeleton.IsNull() && SkeletalMesh->GetSkeleton() == Skeleton.Get())
 	{
 		SkeletalMesh->GetSkeleton()->UpdateReferencePoseFromMesh(SkeletalMesh);
 		//TODO: notify editor the skeleton has change
diff --git a/Engine/Plugins/Experimental/Interchange/Runtime/Source/Pipelines/Public/InterchangeGenericAssetsPipeline.h b/Engine/Plugins/Experimental/Interchange/Runtime/Source/Pipelines/Public/InterchangeGenericAssetsPipeline.h
--- a/Engine/Plugins/Experimental/Interchange/Runtime/Source/Pipelines/Public/InterchangeGenericAssetsPipeline.h
+++ b/Engine/Plugins/Experimental/Interchange/Runtime/Source/Pipelines/Public/InterchangeGenericAssetsPipeline.h
@@ -85,6 +85,8 @@ protected:
 		return true;
 	}
 
+	virtual void SetReimportSourceIndex(UClass* ReimportObjectClass, const int32 SourceFileIndex) override;
+
 	//virtual bool ExecuteExportPipeline(UInterchangeBaseNodeContainer* BaseNodeContainer) override;
 private:
 	
diff --git a/Engine/Plugins/Experimental/Interchange/Runtime/Source/Pipelines/Public/InterchangeGenericMeshPipeline.h b/Engine/Plugins/Experimental/Interchange/Runtime/Source/Pipelines/Public/InterchangeGenericMeshPipeline.h
--- a/Engine/Plugins/Experimental/Interchange/Runtime/Source/Pipelines/Public/InterchangeGenericMeshPipeline.h
+++ b/Engine/Plugins/Experimental/Interchange/Runtime/Source/Pipelines/Public/InterchangeGenericMeshPipeline.h
@@ -6,6 +6,7 @@
 #include "InterchangePipelineBase.h"
 #include "InterchangeSourceData.h"
 #include "Nodes/InterchangeBaseNodeContainer.h"
+#include "InterchangeSkeletalMeshFactoryNode.h"
 #include "UObject/Object.h"
 #include "UObject/ObjectMacros.h"
 
@@ -21,7 +22,7 @@ class UInterchangeStaticMeshLodDataNode;
 
 /** Force mesh type, if user want to import all meshes as one type*/
 UENUM(BlueprintType)
-enum EInterchangeForceMeshType
+enum class EInterchangeForceMeshType : uint8
 {
 	/** Will import from the source type, no conversion */
 	IFMT_None UMETA(DisplayName = "None"),
@@ -34,7 +35,7 @@ enum EInterchangeForceMeshType
 };
 
 UENUM(BlueprintType)
-enum EInterchangeVertexColorImportOption
+enum class EInterchangeVertexColorImportOption : uint8
 {
 	/** Import the mesh using the vertex colors from the translated source. */
 	IVCIO_Replace UMETA(DisplayName = "Replace"),
@@ -57,7 +58,7 @@ public:
 
 	/** Allow to convert mesh to a particular type */
  	UPROPERTY(EditAnywhere, BlueprintReadWrite, Category = "Common Meshes")
- 	TEnumAsByte<enum EInterchangeForceMeshType> ForceAllMeshAsType = EInterchangeForceMeshType::IFMT_None;
+ 	EInterchangeForceMeshType ForceAllMeshAsType = EInterchangeForceMeshType::IFMT_None;
 
 	/** If enable, meshes LODs will be imported. Note that it required the advanced bBakeMesh property to be enabled. */
 	UPROPERTY(EditAnywhere, BlueprintReadWrite, Category = "Common Meshes", meta = (editcondition = "bBakeMeshes"))
@@ -69,7 +70,7 @@ public:
 
 	/** Specify how vertex colors should be imported */
 	UPROPERTY(EditAnywhere, BlueprintReadWrite, Category = "Common Meshes")
-	TEnumAsByte<enum EInterchangeVertexColorImportOption> VertexColorImportOption = EInterchangeVertexColorImportOption::IVCIO_Replace;
+	EInterchangeVertexColorImportOption VertexColorImportOption = EInterchangeVertexColorImportOption::IVCIO_Replace;
 
 	/** Specify override color in the case that VertexColorImportOption is set to Override */
 	UPROPERTY(EditAnywhere, BlueprintReadWrite, Category = "Common Meshes")
@@ -100,6 +101,14 @@ public:
 	UPROPERTY(EditAnywhere, BlueprintReadWrite, Category = "Skeletal Meshes")
 	bool bImportSkeletalMeshes = true;
 
+	/** Re-import partially or totally a skeletal mesh. You can choose betwwen geometry, skinning or everything.*/
+	UPROPERTY(EditAnywhere, BlueprintReadWrite, Category = "Skeletal Meshes", meta = (DisplayName = "Import Content Type"))
+	EInterchangeSkeletalMeshContentType SkeletalMeshImportContentType;
+	
+	/** The value of the content type during the last import. This cannot be edited and is set only on successful import or re-import*/
+	UPROPERTY()
+	EInterchangeSkeletalMeshContentType LastSkeletalMeshImportContentType;
+
 	/** If enable all translated skinned mesh node will be imported has a one skeletal mesh, note that it can still create several skeletal mesh for each different skeleton root joint. */
 	UPROPERTY(EditAnywhere, BlueprintReadWrite, Category = "Skeletal Meshes")
 	bool bCombineSkeletalMeshes = true;
@@ -140,6 +149,8 @@ protected:
 		return true;
 	}
 
+	virtual void SetReimportSourceIndex(UClass* ReimportObjectClass, const int32 SourceFileIndex) override;
+
 private:
 
 	/* Meshes utilities, to parse the translated graph and extract the meshes informations. */
diff --git a/Engine/Source/Developer/AssetTools/Private/AssetTypeActions/AssetTypeActions_SkeletalMesh.cpp b/Engine/Source/Developer/AssetTools/Private/AssetTypeActions/AssetTypeActions_SkeletalMesh.cpp
--- a/Engine/Source/Developer/AssetTools/Private/AssetTypeActions/AssetTypeActions_SkeletalMesh.cpp
+++ b/Engine/Source/Developer/AssetTools/Private/AssetTypeActions/AssetTypeActions_SkeletalMesh.cpp
@@ -593,37 +593,24 @@ EVisibility FAssetTypeActions_SkeletalMesh::GetThumbnailSkinningOverlayVisibilit
 		return EVisibility::Collapsed;
 	}
 
-	//Prevent loading all assets when we display the thumbnail
-	//The tags cannot change until the asset is loaded in memory
-	if (!AssetData.IsAssetLoaded())
+	//Prevent loading assets when we display the thumbnail use the asset registry tags
+
+	//Legacy fbx code
+	FAssetDataTagMapSharedView::FFindTagResult Result = AssetData.TagsAndValues.FindTag(GET_MEMBER_NAME_CHECKED(UFbxSkeletalMeshImportData, LastImportContentType));
+	if (Result.IsSet() && Result.GetValue() == TEXT("FBXICT_Geometry"))
 	{
-		FAssetDataTagMapSharedView::FFindTagResult Result = AssetData.TagsAndValues.FindTag(GET_MEMBER_NAME_CHECKED(UFbxSkeletalMeshImportData, LastImportContentType));
-		if (Result.IsSet() && Result.GetValue() == TEXT("FBXICT_Geometry"))
-		{
-			//Show the icon
-			return EVisibility::HitTestInvisible;
-		}
-		return EVisibility::Collapsed;
+		//Show the icon
+		return EVisibility::HitTestInvisible;
 	}
-
-	//The object is loaded we can use the memory value of the object to set the overlay
-	UObject* Obj = AssetData.GetAsset();
-	if(USkeletalMesh* SkeletalMesh = Cast<USkeletalMesh>(Obj))
+#if WITH_EDITORONLY_DATA
+	//Generic Interchange code
+	Result = AssetData.TagsAndValues.FindTag(NSSkeletalMeshSourceFileLabels::GetSkeletalMeshLastImportContentTypeMetadataKey());
+	if (Result.IsSet() && Result.GetValue().Equals(NSSkeletalMeshSourceFileLabels::GeometryMetaDataValue()))
 	{
-		// Don't block on asset compilation just for overlay, the thumbnail is refreshed anyway once the compilation is complete.
-		if (!SkeletalMesh->IsCompiling())
-		{
-			UAssetImportData* GenericImportData = SkeletalMesh->GetAssetImportData();
-			if (GenericImportData != nullptr)
-			{
-				UFbxSkeletalMeshImportData* ImportData = Cast<UFbxSkeletalMeshImportData>(GenericImportData);
-				if (ImportData != nullptr && ImportData->LastImportContentType == EFBXImportContentType::FBXICT_Geometry)
-				{
-					return EVisibility::HitTestInvisible;
-				}
-			}
-		}
+		//Show the icon
+		return EVisibility::HitTestInvisible;
 	}
+#endif
 	return EVisibility::Collapsed;
 }
 
diff --git a/Engine/Source/Editor/SkeletalMeshEditor/Private/SkeletalMeshEditor.cpp b/Engine/Source/Editor/SkeletalMeshEditor/Private/SkeletalMeshEditor.cpp
--- a/Engine/Source/Editor/SkeletalMeshEditor/Private/SkeletalMeshEditor.cpp
+++ b/Engine/Source/Editor/SkeletalMeshEditor/Private/SkeletalMeshEditor.cpp
@@ -56,6 +56,7 @@
 #include "Misc/CoreMisc.h"
 #include "Toolkits/AssetEditorToolkitMenuContext.h"
 #include "MeshMergeModule.h"
+#include "InterchangeAssetImportData.h"
 
 const FName SkeletalMeshEditorAppIdentifier = FName(TEXT("SkeletalMeshEditorApp"));
 
@@ -401,12 +402,15 @@ void FSkeletalMeshEditor::RegisterReimportContextMenu(const FName InBaseMenuName
 		{
 			if (USkeletalMesh* FoundSkeletalMesh = SkeletalMeshEditor->SkeletalMesh)
 			{
-				UFbxSkeletalMeshImportData* SkeletalMeshImportData = Cast<UFbxSkeletalMeshImportData>(FoundSkeletalMesh->GetAssetImportData());
-				if (SkeletalMeshImportData)
+				if (UFbxSkeletalMeshImportData* SkeletalMeshImportData = Cast<UFbxSkeletalMeshImportData>(FoundSkeletalMesh->GetAssetImportData()))
 				{
 					SkeletalMeshImportData->ImportContentType = SourceFileIndex == 0 ? EFBXImportContentType::FBXICT_All : SourceFileIndex == 1 ? EFBXImportContentType::FBXICT_Geometry : EFBXImportContentType::FBXICT_SkinningWeights;
 					SkeletalMeshEditor->HandleReimportMeshWithNewFile(SourceFileIndex);
 				}
+				else if (UInterchangeAssetImportData* InterchangeImportData = Cast<UInterchangeAssetImportData>(FoundSkeletalMesh->GetAssetImportData()))
+				{
+					SkeletalMeshEditor->HandleReimportMeshWithNewFile(SourceFileIndex);
+				}
 			}
 		}
 	};
@@ -447,27 +451,28 @@ void FSkeletalMeshEditor::RegisterReimportContextMenu(const FName InBaseMenuName
 				if (SkeletalMeshImportData)
 				{
 					SkeletalMeshImportData->ImportContentType = SourceFileIndex == 0 ? EFBXImportContentType::FBXICT_All : SourceFileIndex == 1 ? EFBXImportContentType::FBXICT_Geometry : EFBXImportContentType::FBXICT_SkinningWeights;
-					if (bReimportAll)
+				}
+
+				if (bReimportAll)
+				{
+					if (bWithNewFile)
 					{
-						if (bWithNewFile)
-						{
-							SkeletalMeshEditor->HandleReimportAllMeshWithNewFile(SourceFileIndex);
-						}
-						else
-						{
-							SkeletalMeshEditor->HandleReimportAllMesh(SourceFileIndex);
-						}
+						SkeletalMeshEditor->HandleReimportAllMeshWithNewFile(SourceFileIndex);
 					}
 					else
 					{
-						if (bWithNewFile)
-						{
-							SkeletalMeshEditor->HandleReimportMeshWithNewFile(SourceFileIndex);
-						}
-						else
-						{
-							SkeletalMeshEditor->HandleReimportMesh(SourceFileIndex);
-						}
+						SkeletalMeshEditor->HandleReimportAllMesh(SourceFileIndex);
+					}
+				}
+				else
+				{
+					if (bWithNewFile)
+					{
+						SkeletalMeshEditor->HandleReimportMeshWithNewFile(SourceFileIndex);
+					}
+					else
+					{
+						SkeletalMeshEditor->HandleReimportMesh(SourceFileIndex);
 					}
 				}
 			}
diff --git a/Engine/Source/Editor/SkeletalMeshEditor/SkeletalMeshEditor.Build.cs b/Engine/Source/Editor/SkeletalMeshEditor/SkeletalMeshEditor.Build.cs
--- a/Engine/Source/Editor/SkeletalMeshEditor/SkeletalMeshEditor.Build.cs
+++ b/Engine/Source/Editor/SkeletalMeshEditor/SkeletalMeshEditor.Build.cs
@@ -40,6 +40,8 @@ public class SkeletalMeshEditor : ModuleRules
 				"EditorSubsystem",
 				"StatusBar",
 				"PhysicsUtilities",
+				"InterchangeCore",
+				"InterchangeEngine",
 			}
 		);
 
diff --git a/Engine/Source/Editor/UnrealEd/Classes/Factories/InterchangeReimportHandler.h b/Engine/Source/Editor/UnrealEd/Classes/Factories/InterchangeReimportHandler.h
new file mode 100644
--- /dev/null
+++ b/Engine/Source/Editor/UnrealEd/Classes/Factories/InterchangeReimportHandler.h
@@ -0,0 +1,49 @@
+// Copyright Epic Games, Inc. All Rights Reserved.
+
+//~=============================================================================
+// ReimportFbxSkeletalMeshFactory
+//~=============================================================================
+
+#pragma once
+
+#include "CoreMinimal.h"
+#include "UObject/ObjectMacros.h"
+#include "EditorReimportHandler.h"
+#include "InterchangeReimportHandler.generated.h"
+
+/**
+ * This FReimportHandler class is temporary until we remove the UFactory import/reimport code.
+ * We use this ReimportHandler factory to make sure the set reimport is not touching the UInterchangeAssetImportData when reimporting an asset
+ */
+UCLASS(MinimalAPI, collapsecategories)
+class UInterchangeReimportHandler : public UObject, public FReimportHandler
+{
+	GENERATED_UCLASS_BODY()
+
+	//~ Begin FReimportHandler Interface
+	virtual bool CanReimport( UObject* Obj, TArray<FString>& OutFilenames ) override;
+	virtual void SetReimportPaths(UObject* Obj, const TArray<FString>& NewReimportPaths) override
+	{
+		for (int32 SourceIndex = 0; SourceIndex < NewReimportPaths.Num(); ++SourceIndex)
+		{
+			SetReimportPaths(Obj, NewReimportPaths[SourceIndex], SourceIndex);
+		}
+	}
+	virtual void SetReimportPaths(UObject* Obj, const FString& NewReimportPath, const int32 SourceIndex) override;
+
+	virtual void SetReimportSourceIndex(UObject* Obj, const int32 SourceIndex) override;
+
+	virtual EReimportResult::Type Reimport(UObject* Obj) override
+	{
+		return Reimport(Obj, INDEX_NONE);
+	}
+	virtual EReimportResult::Type Reimport(UObject* Obj, int32 SourceFileIndex);
+	virtual int32 GetPriority() const override;
+
+	virtual bool IsInterchangeFactory() const override
+	{
+		return true;
+	}
+	//~ End FReimportHandler Interface
+
+};
diff --git a/Engine/Source/Editor/UnrealEd/Classes/Factories/ReimportInterchangeFactory.h b/Engine/Source/Editor/UnrealEd/Classes/Factories/ReimportInterchangeFactory.h
deleted file mode 100644
diff --git a/Engine/Source/Editor/UnrealEd/Private/Editor.cpp b/Engine/Source/Editor/UnrealEd/Private/Editor.cpp
--- a/Engine/Source/Editor/UnrealEd/Private/Editor.cpp
+++ b/Engine/Source/Editor/UnrealEd/Private/Editor.cpp
@@ -392,8 +392,11 @@ bool FReimportManager::Reimport( UObject* Obj, bool bAskForNewFileIfMissing, boo
 			{
 				if (bUseInterchangeFramework && CanReimportHandler->IsInterchangeFactory())
 				{
-					UE::Interchange::FScopedSourceData ScopedSourceData(SourceFilenames[0]);
+					int32 RealSourceFileIndex = SourceFileIndex == INDEX_NONE ? 0 : SourceFileIndex;
+					int32 RealValidSourceFileIndex = SourceFilenames.IsValidIndex(RealSourceFileIndex) ? RealSourceFileIndex : 0;
+					UE::Interchange::FScopedSourceData ScopedSourceData(SourceFilenames[RealValidSourceFileIndex]);
 
+					CanReimportHandler->SetReimportSourceIndex(Obj, SourceFileIndex);
 					bool bUseATextureTranslator = false;
 					if (bUseInterchangeFrameworkForTextureOnly)
 					{
@@ -429,6 +432,7 @@ bool FReimportManager::Reimport( UObject* Obj, bool bAskForNewFileIfMissing, boo
 						FImportAssetParameters ImportAssetParameters;
 						ImportAssetParameters.bIsAutomated = GIsAutomationTesting || FApp::IsUnattended() || IsRunningCommandlet() || GIsRunningUnattendedScript;
 						ImportAssetParameters.ReimportAsset = Obj;
+						ImportAssetParameters.ReimportSourceIndex = SourceFileIndex;
 						InterchangeManager.ImportAsset(FString(), ScopedSourceData.GetSourceData(), ImportAssetParameters);
 						InterchangeManager.OnAssetPostImport.Remove(PostImportHandle);
 						return true;
diff --git a/Engine/Source/Editor/UnrealEd/Private/Factories/InterchangeReimportHandler.cpp b/Engine/Source/Editor/UnrealEd/Private/Factories/InterchangeReimportHandler.cpp
new file mode 100644
--- /dev/null
+++ b/Engine/Source/Editor/UnrealEd/Private/Factories/InterchangeReimportHandler.cpp
@@ -0,0 +1,148 @@
+// Copyright Epic Games, Inc. All Rights Reserved.
+
+#include "Factories/InterchangeReimportHandler.h"
+
+#include "CoreMinimal.h"
+#include "EngineDefines.h"
+#include "Engine/SkeletalMesh.h"
+#include "Engine/StaticMesh.h"
+#include "Engine/Texture.h"
+#include "Factories/Factory.h"
+#include "InterchangeAssetImportData.h"
+#include "InterchangeManager.h"
+#include "Settings/EditorExperimentalSettings.h"
+#include "UObject/UObjectGlobals.h"
+#include "UObject/ObjectMacros.h"
+
+#define LOCTEXT_NAMESPACE "InterchangeReimportHandler"
+
+UInterchangeReimportHandler::UInterchangeReimportHandler(const FObjectInitializer& ObjectInitializer)
+	: Super(ObjectInitializer)
+{
+}
+
+//~ Begin FReimportHandler Interface
+bool UInterchangeReimportHandler::CanReimport(UObject* Obj, TArray<FString>& OutFilenames)
+{
+	auto CanReimportAsset = [&OutFilenames](UAssetImportData* AssetImportData)
+	{
+		if (!AssetImportData)
+		{
+			return false;
+		}
+		UInterchangeAssetImportData* InterchangeAssetImportData = Cast<UInterchangeAssetImportData>(AssetImportData);
+		if (InterchangeAssetImportData == nullptr)
+		{
+			return false;
+		}
+		AssetImportData->ExtractFilenames(OutFilenames);
+		return true;
+	};
+
+	
+	const UEditorExperimentalSettings* EditorExperimentalSettings = GetDefault<UEditorExperimentalSettings>();
+	const bool bUseInterchangeFramework = EditorExperimentalSettings->bEnableInterchangeFramework;
+	const bool bUseInterchangeFrameworkForTextureOnly = (!bUseInterchangeFramework) && EditorExperimentalSettings->bEnableInterchangeFrameworkForTextureOnly;
+
+	USkeletalMesh* SkeletalMesh = Cast<USkeletalMesh>(Obj);
+	UStaticMesh* StaticMesh = Cast<UStaticMesh>(Obj);
+	//UAnimSequence* AnimSequence = Cast<UAnimSequence>(Obj);
+	UTexture* Texture = Cast<UTexture>(Obj);
+	if (bUseInterchangeFramework && SkeletalMesh)
+	{
+		if (UAssetImportData* AssetImportData = SkeletalMesh->GetAssetImportData())
+		{
+			return CanReimportAsset(AssetImportData);
+		}
+	}
+	else if (bUseInterchangeFramework && StaticMesh)
+	{
+		if (UAssetImportData* AssetImportData = StaticMesh->GetAssetImportData())
+		{
+			return CanReimportAsset(AssetImportData);
+		}
+	}
+// 	else if (bUseInterchangeFramework && AnimSequence)
+// 	{
+// 		if (UAssetImportData* AssetImportData = AnimSequence->AssetImportData)
+// 		{
+// 			return CanReimportAsset(AssetImportData);
+// 		}
+// 	}
+	else if ((bUseInterchangeFramework || bUseInterchangeFrameworkForTextureOnly) && Texture)
+	{
+		if (UAssetImportData* AssetImportData = Texture->AssetImportData)
+		{
+			return CanReimportAsset(AssetImportData);
+		}
+	}
+	return false;
+}
+
+void UInterchangeReimportHandler::SetReimportPaths(UObject* Obj, const FString& NewReimportPath, const int32 SourceFileIndex)
+{
+	USkeletalMesh* SkeletalMesh = Cast<USkeletalMesh>(Obj);
+	UStaticMesh* StaticMesh = Cast<UStaticMesh>(Obj);
+//	UAnimSequence* AnimSequence = Cast<UAnimSequence>(Obj);
+	UTexture* Texture = Cast<UTexture>(Obj);
+
+	UAssetImportData* AssetImportData = nullptr;
+	if (SkeletalMesh)
+	{
+		AssetImportData = SkeletalMesh->GetAssetImportData();
+	}
+	else if (StaticMesh)
+	{
+		AssetImportData = StaticMesh->GetAssetImportData();
+	}
+// 	else if (AnimSequence)
+// 	{
+// 		AssetImportData = AnimSequence->AssetImportData;
+// 	}
+	else if (Texture)
+	{
+		AssetImportData = Texture->AssetImportData;
+	}
+
+	if (AssetImportData)
+	{
+		int32 RealSourceFileIndex = SourceFileIndex == INDEX_NONE ? 0 : SourceFileIndex;
+		if (RealSourceFileIndex < AssetImportData->GetSourceFileCount())
+		{
+			AssetImportData->UpdateFilenameOnly(NewReimportPath, SourceFileIndex);
+		}
+		else
+		{
+			//Create a source file entry, this case happen when user import a specific content for the first time
+			FString SourceIndexLabel = USkeletalMesh::GetSourceFileLabelFromIndex(RealSourceFileIndex).ToString();
+			AssetImportData->AddFileName(NewReimportPath, RealSourceFileIndex, SourceIndexLabel);
+		}
+	}
+}
+
+void UInterchangeReimportHandler::SetReimportSourceIndex(UObject* Obj, const int32 SourceIndex)
+{
+	if (USkeletalMesh* SkeletalMesh = Cast<USkeletalMesh>(Obj))
+	{
+		if (UInterchangeAssetImportData* InterchangeAssetImportData = Cast<UInterchangeAssetImportData>(SkeletalMesh->GetAssetImportData()))
+		{
+			for (UInterchangePipelineBase* PipelineBase : InterchangeAssetImportData->Pipelines)
+			{
+				PipelineBase->ScriptedSetReimportSourceIndex(Obj->GetClass(), SourceIndex);
+			}
+		}
+	}
+}
+
+EReimportResult::Type UInterchangeReimportHandler::Reimport(UObject* Obj, int32 SourceFileIndex)
+{
+	return EReimportResult::Failed;
+}
+
+int32 UInterchangeReimportHandler::GetPriority() const
+{
+	//We want a high priority to surpass other legacy re-import handlers
+	return UFactory::GetDefaultImportPriority() + 10;
+}
+
+#undef LOCTEXT_NAMESPACE
\ No newline at end of file
diff --git a/Engine/Source/Editor/UnrealEd/Private/Factories/ReimportInterchangeFactory.cpp b/Engine/Source/Editor/UnrealEd/Private/Factories/ReimportInterchangeFactory.cpp
deleted file mode 100644
diff --git a/Engine/Source/Editor/UnrealEd/Public/EditorReimportHandler.h b/Engine/Source/Editor/UnrealEd/Public/EditorReimportHandler.h
--- a/Engine/Source/Editor/UnrealEd/Public/EditorReimportHandler.h
+++ b/Engine/Source/Editor/UnrealEd/Public/EditorReimportHandler.h
@@ -217,6 +217,16 @@ public:
 		SetReimportPaths(Obj, NewReimportPaths);
 	}
 
+	/**
+	* If we want to import a specific source index and there is some settings to update in the asset import data before starting the re-import.
+	*
+	* @param	Obj - Object for which to change re-import settings.
+	* @param	SourceIndex - the index of the SourceFile that will be re-imported
+	*/
+	virtual void SetReimportSourceIndex(UObject* Obj, const int32 SourceIndex)
+	{
+	}
+
 	/**
 	 * Attempt to reimport the specified object from its source
 	 *
diff --git a/Engine/Source/Runtime/Engine/Classes/Engine/SkeletalMesh.h b/Engine/Source/Runtime/Engine/Classes/Engine/SkeletalMesh.h
--- a/Engine/Source/Runtime/Engine/Classes/Engine/SkeletalMesh.h
+++ b/Engine/Source/Runtime/Engine/Classes/Engine/SkeletalMesh.h
@@ -742,6 +742,30 @@ namespace NSSkeletalMeshSourceFileLabels
 		static FText SkinningText = (NSLOCTEXT("FBXReimport", "ImportContentTypeSkinning", "Skinning Weights"));
 		return SkinningText;
 	}
+
+	static const FString& GeoAndSkinningMetaDataValue()
+	{
+		static FString GeoAndSkinningName(TEXT("All"));
+		return GeoAndSkinningName;
+	}
+
+	static const FString& GeometryMetaDataValue()
+	{
+		static FString GeometryName(TEXT("Geometry"));
+		return GeometryName;
+	}
+
+	static const FString& SkinningMetaDataValue()
+	{
+		static FString SkinningName(TEXT("SkinningWeights"));
+		return SkinningName;
+	}
+
+	static FName GetSkeletalMeshLastImportContentTypeMetadataKey()
+	{
+		static FName SkeletalMeshLastImportContentTypeMetadataKey("SkeletalMeshLastImportContentTypeMetadataKey");
+		return SkeletalMeshLastImportContentTypeMetadataKey;
+	}
 }
 #endif
 
diff --git a/Engine/Source/Runtime/Experimental/Interchange/Core/Private/Nodes/InterchangeBaseNode.cpp b/Engine/Source/Runtime/Experimental/Interchange/Core/Private/Nodes/InterchangeBaseNode.cpp
--- a/Engine/Source/Runtime/Experimental/Interchange/Core/Private/Nodes/InterchangeBaseNode.cpp
+++ b/Engine/Source/Runtime/Experimental/Interchange/Core/Private/Nodes/InterchangeBaseNode.cpp
@@ -312,7 +312,7 @@ bool UInterchangeBaseNode::SetReimportStrategyFlags(const EReimportStrategyFlags
 	UE::Interchange::EAttributeStorageResult Result = Attributes->RegisterAttribute(UE::Interchange::FBaseNodeStaticData::ReimportStrategyFlagsKey(), static_cast<uint8>(ReimportStrategyFlags));
 	if (IsAttributeStorageResultSuccess(Result))
 	{
-		UE::Interchange::FAttributeStorage::TAttributeHandle<FString> Handle = Attributes->GetAttributeHandle<FString>(UE::Interchange::FBaseNodeStaticData::ReimportStrategyFlagsKey());
+		UE::Interchange::FAttributeStorage::TAttributeHandle<uint8> Handle = Attributes->GetAttributeHandle<uint8>(UE::Interchange::FBaseNodeStaticData::ReimportStrategyFlagsKey());
 		return Handle.IsValid();
 	}
 	return false;
diff --git a/Engine/Source/Runtime/Experimental/Interchange/Core/Public/InterchangePipelineBase.h b/Engine/Source/Runtime/Experimental/Interchange/Core/Public/InterchangePipelineBase.h
--- a/Engine/Source/Runtime/Experimental/Interchange/Core/Public/InterchangePipelineBase.h
+++ b/Engine/Source/Runtime/Experimental/Interchange/Core/Public/InterchangePipelineBase.h
@@ -44,7 +44,7 @@ public:
 	/**
 	 * Non virtual helper to allow blueprint to implement event base function to implement a post import pipeline,
 	 * It is call after we completely import an asset. PostEditChange is already called. Some assets uses asynchronous build.
-	 * This can be useful if you need build data of an asset to finish the setup of another asset.
+	 * This can be useful if you need builded data of an asset to finish the setup of another asset.
 	 * @example - PhysicsAsset need skeletal mesh render data to be build properly.
 	 * @note - the Interchange manager is calling this function not the virtual one that is call by the default implementation.
 	 */
@@ -83,6 +83,19 @@ public:
 		return CanExecuteOnAnyThread(PipelineTask);
 	}
 
+	/**
+	 * Non virtual helper to allow blueprint to implement event base function.
+	 * the Interchange framework is calling this function not the virtual one that is called by the default implementation.
+	 */
+	UFUNCTION(BlueprintNativeEvent, BlueprintCallable, Category = "Interchange | Pipeline")
+	void ScriptedSetReimportSourceIndex(UClass* ReimportObjectClass, const int32 SourceFileIndex);
+
+	/** The default implementation (call if the blueprint do not have any implementation) will call the virtual SetReimportContentFromSourceIndex */
+	void ScriptedSetReimportSourceIndex_Implementation(UClass* ReimportObjectClass, const int32 SourceFileIndex)
+	{
+		SetReimportSourceIndex(ReimportObjectClass, SourceFileIndex);
+	}
+
 	/**
 	 * Non scripted class should return false here, we have the default to true because scripted class cannot override
 	 * this function since it can be call in a asynchronous thread, which python cannot be executed.
@@ -100,7 +113,7 @@ public:
 	void SaveSettings(const FName PipelineStackName);
 
 	virtual void PreDialogCleanup(const FName PipelineStackName) {};
-	
+
 	/**
 	 * This function is used to add the given message object directly into the results for this operation.
 	 */
@@ -155,6 +168,10 @@ protected:
 		return true;
 	}
 
+	virtual void SetReimportSourceIndex(UClass* ReimportObjectClass, const int32 SourceFileIndex)
+	{
+	}
+
 	/** This function can modify the BaseNodeContainer to create a pipeline that will set/validate the graph nodes hierarchy and options.*/
 	virtual void ExecuteExportPipeline(UInterchangeBaseNodeContainer* BaseNodeContainer)
 	{
diff --git a/Engine/Source/Runtime/Experimental/Interchange/Core/Public/Nodes/InterchangeBaseNode.h b/Engine/Source/Runtime/Experimental/Interchange/Core/Public/Nodes/InterchangeBaseNode.h
--- a/Engine/Source/Runtime/Experimental/Interchange/Core/Public/Nodes/InterchangeBaseNode.h
+++ b/Engine/Source/Runtime/Experimental/Interchange/Core/Public/Nodes/InterchangeBaseNode.h
@@ -166,12 +166,15 @@ bool FillCustom##AttributeName##FromAsset(UObject* Asset)															\
 UE::Interchange::FAttributeStorage::TAttributeHandle<ValueType> Handle = RegisterAttribute<ValueType>(UE::Interchange::FAttributeKey(NodeAttributeKey), Value);	\
 return Handle.IsValid();
 
-#define INTERCHANGE_BASE_NODE_GET_ATTRIBUTE(ValueType)																							\
-UE::Interchange::FAttributeStorage::TAttributeHandle<ValueType> Handle = GetAttributeHandle<ValueType>(UE::Interchange::FAttributeKey(NodeAttributeKey));	\
-if (Handle.IsValid())																															\
-{																																				\
-	return (Handle.Get(OutValue) == UE::Interchange::EAttributeStorageResult::Operation_Success);													\
-}																																				\
+#define INTERCHANGE_BASE_NODE_GET_ATTRIBUTE(ValueType)																						\
+if(HasAttribute(UE::Interchange::FAttributeKey(NodeAttributeKey)))																			\
+{																																			\
+	UE::Interchange::FAttributeStorage::TAttributeHandle<ValueType> Handle = GetAttributeHandle<ValueType>(UE::Interchange::FAttributeKey(NodeAttributeKey));	\
+	if (Handle.IsValid())																													\
+	{																																		\
+		return (Handle.Get(OutValue) == UE::Interchange::EAttributeStorageResult::Operation_Success);										\
+	}																																		\
+}																																			\
 return false;
 
 //Interchange namespace
@@ -593,6 +596,11 @@ public:
 
 	static void CopyStorage(const UInterchangeBaseNode* SourceNode, UInterchangeBaseNode* DestinationNode);
 	
+	virtual void AppendAssetRegistryTags(TArray<FAssetRegistryTag>& OutTags) const
+	{
+		return;
+	}
+
 	UPROPERTY()
 	mutable FSoftObjectPath ReferenceObject;
 protected:
diff --git a/Engine/Source/Runtime/Experimental/Interchange/Engine/Public/InterchangeAssetImportData.h b/Engine/Source/Runtime/Experimental/Interchange/Engine/Public/InterchangeAssetImportData.h
--- a/Engine/Source/Runtime/Experimental/Interchange/Engine/Public/InterchangeAssetImportData.h
+++ b/Engine/Source/Runtime/Experimental/Interchange/Engine/Public/InterchangeAssetImportData.h
@@ -61,7 +61,23 @@ public:
 #endif
 		return TempDisplayLabels;
 	}
-	
+
+#if WITH_EDITORONLY_DATA
+#if WITH_EDITOR
+	/**
+	 * This function add tags to the asset registry.
+	 */
+	virtual void AppendAssetRegistryTags(TArray<FAssetRegistryTag>& OutTags) override
+	{
+		if (NodeContainer->IsNodeUidValid(NodeUniqueID))
+		{
+			NodeContainer->GetNode(NodeUniqueID)->AppendAssetRegistryTags(OutTags);
+		}
+		Super::AppendAssetRegistryTags(OutTags);
+	}
+#endif
+#endif
+
 	/** The graph that was use to create this asset */
 	UPROPERTY(VisibleAnywhere, Category = "Interchange | AssetImportData")
 	TObjectPtr<UInterchangeBaseNodeContainer> NodeContainer;
diff --git a/Engine/Source/Runtime/Experimental/Interchange/Engine/Public/InterchangeManager.h b/Engine/Source/Runtime/Experimental/Interchange/Engine/Public/InterchangeManager.h
--- a/Engine/Source/Runtime/Experimental/Interchange/Engine/Public/InterchangeManager.h
+++ b/Engine/Source/Runtime/Experimental/Interchange/Engine/Public/InterchangeManager.h
@@ -247,6 +247,10 @@ struct FImportAssetParameters
 	UPROPERTY(EditAnywhere, BlueprintReadWrite, Category = "Interchange|ImportAsset")
 	TObjectPtr<UObject> ReimportAsset = nullptr;
 
+	// If we are doing a reimport, set the source index here. Some asset have more then one source file when they import partial part of there content.
+	UPROPERTY(EditAnywhere, BlueprintReadWrite, Category = "Interchange|ImportAsset")
+	int32 ReimportSourceIndex = INDEX_NONE;
+
 	// Tell interchange that import is automated and it shouldn't present a model window
 	UPROPERTY(EditAnywhere, BlueprintReadWrite, Category = "Interchange|ImportAsset")
 	bool bIsAutomated = false;
