Revert exception being disabled in the InterchangeFbxParser

This breaks Clang/LLVM platforms

diff --git a/Engine/Plugins/Experimental/Interchange/Runtime/Source/Parsers/Fbx/InterchangeFbxParser.Build.cs b/Engine/Plugins/Experimental/Interchange/Runtime/Source/Parsers/Fbx/InterchangeFbxParser.Build.cs
--- a/Engine/Plugins/Experimental/Interchange/Runtime/Source/Parsers/Fbx/InterchangeFbxParser.Build.cs
+++ b/Engine/Plugins/Experimental/Interchange/Runtime/Source/Parsers/Fbx/InterchangeFbxParser.Build.cs
@@ -20,6 +20,8 @@ namespace UnrealBuildTool.Rules
 				throw new BuildException("InterchangeFbxParser module do not support target platform {0}.", Target.Platform.ToString());
 			}
 
+			bEnableExceptions = true;
+
 			PublicDependencyModuleNames.AddRange(
 				new string[]
 				{
