Code compilation fixes for implicit FVector conversions

- gets EngineTest Development Editor compiling
- changes may not be LWC correct as I'm just trying to get code compilation unblocked
  - will need revisited

diff --git a/Engine/Plugins/Animation/ControlRig/Source/ControlRig/Private/Sequencer/MovieSceneControlRigParameterSection.cpp b/Engine/Plugins/Animation/ControlRig/Source/ControlRig/Private/Sequencer/MovieSceneControlRigParameterSection.cpp
--- a/Engine/Plugins/Animation/ControlRig/Source/ControlRig/Private/Sequencer/MovieSceneControlRigParameterSection.cpp
+++ b/Engine/Plugins/Animation/ControlRig/Source/ControlRig/Private/Sequencer/MovieSceneControlRigParameterSection.cpp
@@ -2439,7 +2439,7 @@ void UMovieSceneControlRigParameterSection::RecordControlRigKey(FFrameNumber Fra
 						FloatChannels[ChannelIndex + 2]->SetDefault(Val.Z);
 					}
 
-					AddVectorKeyToFloatChannels(ChannelIndex, FrameNumber, Val);
+					AddVectorKeyToFloatChannels(ChannelIndex, FrameNumber, FVector(Val));
 
 					break;
 				}
diff --git a/Engine/Plugins/Experimental/Interchange/Runtime/Source/Import/Private/Mesh/InterchangeOBJTranslator.cpp b/Engine/Plugins/Experimental/Interchange/Runtime/Source/Import/Private/Mesh/InterchangeOBJTranslator.cpp
--- a/Engine/Plugins/Experimental/Interchange/Runtime/Source/Import/Private/Mesh/InterchangeOBJTranslator.cpp
+++ b/Engine/Plugins/Experimental/Interchange/Runtime/Source/Import/Private/Mesh/InterchangeOBJTranslator.cpp
@@ -186,7 +186,7 @@ FBox FObjData::GetGroupBoundingBox(const FString& GroupName) const
 
 	for (int32 PositionIndex : GetVertexIndicesUsedByGroup(GroupData))
 	{
-		Box += PositionToUEBasis(Positions[PositionIndex]);
+		Box += FVector(PositionToUEBasis(Positions[PositionIndex]));
 	}
 
 	return Box;
diff --git a/Engine/Plugins/Experimental/Interchange/Runtime/Source/Import/Private/Mesh/InterchangeSkeletalMeshFactory.cpp b/Engine/Plugins/Experimental/Interchange/Runtime/Source/Import/Private/Mesh/InterchangeSkeletalMeshFactory.cpp
--- a/Engine/Plugins/Experimental/Interchange/Runtime/Source/Import/Private/Mesh/InterchangeSkeletalMeshFactory.cpp
+++ b/Engine/Plugins/Experimental/Interchange/Runtime/Source/Import/Private/Mesh/InterchangeSkeletalMeshFactory.cpp
@@ -406,7 +406,7 @@ namespace UE
 				for (FVertexID VertexID : MeshDescription.Vertices().GetElementIDs())
 				{
 					//We can use GetValue because the Meshdescription was compacted before the copy
-					DestinationVertexPositions[VertexID.GetValue()] = VertexPositions[VertexID];
+					DestinationVertexPositions[VertexID.GetValue()] = FVector(VertexPositions[VertexID]);
 				}
 
 				// Deform the vertex array with the links contained in the mesh.
@@ -490,7 +490,7 @@ namespace UE
 				for (FVertexID VertexID : MeshDescription.Vertices().GetElementIDs())
 				{
 					const int32 VertexIndex = VertexID.GetValue();
-					const FVector3f lSrcVertex = DestinationVertexPositions[VertexIndex];
+					const FVector lSrcVertex = DestinationVertexPositions[VertexIndex];
 					FVector& lDstVertex = DestinationVertexPositions[VertexIndex];
 					double Weight = SkinWeights[VertexIndex];
 
@@ -502,7 +502,7 @@ namespace UE
 						//Normalized, in case the weight is different then 1
 						lDstVertex /= Weight;
 						//Set the new vertex position in the mesh description
-						VertexPositions[VertexID] = lDstVertex;
+						VertexPositions[VertexID] = FVector3f(lDstVertex);
 					}
 				}
 			}
@@ -871,7 +871,7 @@ namespace UE
 				 */
 				FORCEINLINE static FBoxCenterAndExtent GetBoundingBox(const FSoftSkinVertex& Element)
 				{
-					return FBoxCenterAndExtent(Element.Position, FVector::ZeroVector);
+					return FBoxCenterAndExtent(FVector(Element.Position), FVector::ZeroVector);
 				}
 
 				/**
@@ -910,8 +910,8 @@ namespace UE
 				for (int32 WedgeIndex = 0; WedgeIndex < WedgeNumber; ++WedgeIndex)
 				{
 					SkeletalMeshImportData::FVertex& Wedge = SkelMeshImportData->Wedges[WedgeIndex];
-					const FVector& Position = SkelMeshImportData->Points[Wedge.VertexIndex];
-					Bounds += Position;
+					const FVector3f& Position = SkelMeshImportData->Points[Wedge.VertexIndex];
+					Bounds += FVector(Position);
 				}
 
 				TArray<FSoftSkinVertex> Vertices;
@@ -919,7 +919,7 @@ namespace UE
 				for (int32 SkinVertexIndex = 0; SkinVertexIndex < Vertices.Num(); ++SkinVertexIndex)
 				{
 					const FSoftSkinVertex& SkinVertex = Vertices[SkinVertexIndex];
-					Bounds += SkinVertex.Position;
+					Bounds += FVector(SkinVertex.Position);
 				}
 
 				TSKCVertPosOctree VertPosOctree(Bounds.GetCenter(), Bounds.GetExtent().GetMax());
@@ -947,12 +947,12 @@ namespace UE
 				for (int32 WedgeIndex = 0; WedgeIndex < WedgeNumber; ++WedgeIndex)
 				{
 					SkeletalMeshImportData::FVertex& Wedge = SkelMeshImportData->Wedges[WedgeIndex];
-					const FVector& Position = SkelMeshImportData->Points[Wedge.VertexIndex];
+					const FVector3f& Position = SkelMeshImportData->Points[Wedge.VertexIndex];
 					const FVector2f UV = Wedge.UVs[0];
 					const FVector3f& Normal = WedgeIndexToNormal.FindChecked(WedgeIndex);
 
 					TArray<FSoftSkinVertex> PointsToConsider;
-					VertPosOctree.FindNearbyElements(Position, [&PointsToConsider](const FSoftSkinVertex& Vertex)
+					VertPosOctree.FindNearbyElements(FVector(Position), [&PointsToConsider](const FSoftSkinVertex& Vertex)
 						{
 							PointsToConsider.Add(Vertex);
 						});
diff --git a/Engine/Plugins/Experimental/Interchange/Runtime/Source/Import/Private/Mesh/InterchangeStaticMeshFactory.cpp b/Engine/Plugins/Experimental/Interchange/Runtime/Source/Import/Private/Mesh/InterchangeStaticMeshFactory.cpp
--- a/Engine/Plugins/Experimental/Interchange/Runtime/Source/Import/Private/Mesh/InterchangeStaticMeshFactory.cpp
+++ b/Engine/Plugins/Experimental/Interchange/Runtime/Source/Import/Private/Mesh/InterchangeStaticMeshFactory.cpp
@@ -297,7 +297,7 @@ UObject* UInterchangeStaticMeshFactory::CreateAsset(const FCreateAssetParams& Ar
 						//If we have old vertex color (reimport), we want to keep it if the option is ignore
 						if (ExisitingVertexColorData.Num() > 0)
 						{
-							FVector3f VertexPosition = LodMeshDescription->GetVertexPosition(LodMeshDescription->GetVertexInstanceVertex(VertexInstanceID));
+							FVector VertexPosition = FVector(LodMeshDescription->GetVertexPosition(LodMeshDescription->GetVertexInstanceVertex(VertexInstanceID)));
 							const FColor* PaintedColor = ExisitingVertexColorData.Find(VertexPosition);
 							if (PaintedColor)
 							{
diff --git a/Engine/Plugins/Tests/InterchangeTests/Source/InterchangeTests/Private/ImportTestFunctions/SkeletalMeshImportTestFunctions.cpp b/Engine/Plugins/Tests/InterchangeTests/Source/InterchangeTests/Private/ImportTestFunctions/SkeletalMeshImportTestFunctions.cpp
--- a/Engine/Plugins/Tests/InterchangeTests/Source/InterchangeTests/Private/ImportTestFunctions/SkeletalMeshImportTestFunctions.cpp
+++ b/Engine/Plugins/Tests/InterchangeTests/Source/InterchangeTests/Private/ImportTestFunctions/SkeletalMeshImportTestFunctions.cpp
@@ -210,7 +210,7 @@ FInterchangeTestFunctionResult USkeletalMeshImportTestFunctions::CheckVertexInde
 		}
 		else
 		{
-			FVector VertexPosition = Mesh->GetResourceForRendering()->LODRenderData[LodIndex].StaticVertexBuffers.PositionVertexBuffer.VertexPosition(VertexIndex);
+			FVector VertexPosition = FVector(Mesh->GetResourceForRendering()->LODRenderData[LodIndex].StaticVertexBuffers.PositionVertexBuffer.VertexPosition(VertexIndex));
 			if (!VertexPosition.Equals(ExpectedVertexPosition))
 			{
 				Result.AddError(FString::Printf(TEXT("For LOD %d vertex index %d, expected position %s, imported %s."), LodIndex, VertexIndex, *ExpectedVertexPosition.ToString(), *VertexPosition.ToString()));
diff --git a/Engine/Plugins/Tests/InterchangeTests/Source/InterchangeTests/Private/ImportTestFunctions/StaticMeshImportTestFunctions.cpp b/Engine/Plugins/Tests/InterchangeTests/Source/InterchangeTests/Private/ImportTestFunctions/StaticMeshImportTestFunctions.cpp
--- a/Engine/Plugins/Tests/InterchangeTests/Source/InterchangeTests/Private/ImportTestFunctions/StaticMeshImportTestFunctions.cpp
+++ b/Engine/Plugins/Tests/InterchangeTests/Source/InterchangeTests/Private/ImportTestFunctions/StaticMeshImportTestFunctions.cpp
@@ -390,7 +390,7 @@ FInterchangeTestFunctionResult UStaticMeshImportTestFunctions::CheckVertexIndexP
 		}
 		else
 		{
-			const FVector3f VertexPosition = Attributes.GetVertexPositions()[VertexIndex];
+			const FVector VertexPosition = FVector(Attributes.GetVertexPositions()[VertexIndex]);
 			if (!VertexPosition.Equals(ExpectedVertexPosition))
 			{
 				Result.AddError(FString::Printf(TEXT("For LOD %d vertex index %d, expected position %s, imported %s."), LodIndex, VertexIndex, *ExpectedVertexPosition.ToString(), *VertexPosition.ToString()));
@@ -596,8 +596,8 @@ FInterchangeTestFunctionResult UStaticMeshImportTestFunctions::CheckAgainstGroun
 			{
 				for (int32 VertexIndex = 0; VertexIndex < VertexCount; VertexIndex++)
 				{
-					const FVector& V0 = VertexPositions[VertexIndex];
-					const FVector& V1 = GroundTruthVertexPositions[VertexIndex];
+					const FVector3f& V0 = VertexPositions[VertexIndex];
+					const FVector3f& V1 = GroundTruthVertexPositions[VertexIndex];
 
 					if (!V0.Equals(V1))
 					{
@@ -632,8 +632,8 @@ FInterchangeTestFunctionResult UStaticMeshImportTestFunctions::CheckAgainstGroun
 			{
 				for (int32 NormalIndex = 0; NormalIndex < NormalsCount; NormalIndex++)
 				{
-					const FVector& N0 = VertexInstanceNormals[NormalIndex];
-					const FVector& N1 = GroundTruthVertexInstanceNormals[NormalIndex];
+					const FVector3f& N0 = VertexInstanceNormals[NormalIndex];
+					const FVector3f& N1 = GroundTruthVertexInstanceNormals[NormalIndex];
 
 					if (!N0.Equals(N1))
 					{
diff --git a/Engine/Source/Runtime/Engine/Private/PrimitiveDrawingUtils.cpp b/Engine/Source/Runtime/Engine/Private/PrimitiveDrawingUtils.cpp
--- a/Engine/Source/Runtime/Engine/Private/PrimitiveDrawingUtils.cpp
+++ b/Engine/Source/Runtime/Engine/Private/PrimitiveDrawingUtils.cpp
@@ -788,7 +788,7 @@ void DrawTorus(FPrimitiveDrawInterface* PDI, const FMatrix& Transform, const FVe
 			}
 
 			FDynamicMeshVertex MeshVertex;
-			MeshVertex.Position = Vertex;
+			MeshVertex.Position = FVector3f(Vertex);
 			MeshVertex.TextureCoordinate[0] = TC;
 			MeshVertex.TextureCoordinate[0].X += TCStepX * InnerIndex;
 
diff --git a/Engine/Source/Runtime/Renderer/Private/RayTracing/RayTracingLighting.cpp b/Engine/Source/Runtime/Renderer/Private/RayTracing/RayTracingLighting.cpp
--- a/Engine/Source/Runtime/Renderer/Private/RayTracing/RayTracingLighting.cpp
+++ b/Engine/Source/Runtime/Renderer/Private/RayTracing/RayTracingLighting.cpp
@@ -259,7 +259,7 @@ static void SetupRaytracingLightDataPacked(
 		LightDataElement.LightProfileIndex = IESLightProfileIndex;
 
 		LightDataElement.Direction = LightParameters.Direction;
-		LightDataElement.TranslatedLightPosition = LightParameters.WorldPosition + View.ViewMatrices.GetPreViewTranslation();
+		LightDataElement.TranslatedLightPosition = FVector3f(LightParameters.WorldPosition + View.ViewMatrices.GetPreViewTranslation());
 		LightDataElement.LightColor = FVector3f(LightParameters.Color);
 		LightDataElement.Tangent = LightParameters.Tangent;
 
diff --git a/Engine/Source/Runtime/Renderer/Private/RayTracing/RayTracingScene.cpp b/Engine/Source/Runtime/Renderer/Private/RayTracing/RayTracingScene.cpp
--- a/Engine/Source/Runtime/Renderer/Private/RayTracing/RayTracingScene.cpp
+++ b/Engine/Source/Runtime/Renderer/Private/RayTracing/RayTracingScene.cpp
@@ -212,6 +212,8 @@ void FRayTracingScene::Create(FRDGBuilder& GraphBuilder, const FGPUScene& GPUSce
 				BuildRayTracingInstanceBuffer(
 					RHICmdList,
 					GPUScene,
+					FVector3f(ViewTilePosition),
+					FVector3f(RelativePreViewTranslation),
 					PassParams->InstanceBuffer->GetRHI(),
 					InstanceUploadSRV,
 					AccelerationStructureAddressesBuffer.SRV,
diff --git a/Engine/Source/Runtime/Renderer/Private/ShaderPrint.cpp b/Engine/Source/Runtime/Renderer/Private/ShaderPrint.cpp
--- a/Engine/Source/Runtime/Renderer/Private/ShaderPrint.cpp
+++ b/Engine/Source/Runtime/Renderer/Private/ShaderPrint.cpp
@@ -156,7 +156,7 @@ namespace ShaderPrint
 		Out.MaxSymbolCount = Data.MaxSymbolCount;
 		Out.MaxStateCount = Data.MaxStateCount;
 		Out.MaxLineCount = Data.MaxLineCount;
-		Out.TranslatedWorldOffset = Data.TranslatedWorldOffset;
+		Out.TranslatedWorldOffset = FVector3f(Data.TranslatedWorldOffset);
 		return TUniformBufferRef<ShaderPrint::FShaderPrintCommonParameters>::CreateUniformBufferImmediate(Out, UniformBuffer_SingleFrame);
 	}
 
