Reconsidering fixes from RES to Main for explicit Vector constructor fixes

Fix FVector/FVector3f implicit conversions
#rnx

diff --git a/Engine/Plugins/AI/MLAdapter/Source/MLAdapter/Private/Actuators/MLAdapterActuator_Camera.cpp b/Engine/Plugins/AI/MLAdapter/Source/MLAdapter/Private/Actuators/MLAdapterActuator_Camera.cpp
--- a/Engine/Plugins/AI/MLAdapter/Source/MLAdapter/Private/Actuators/MLAdapterActuator_Camera.cpp
+++ b/Engine/Plugins/AI/MLAdapter/Source/MLAdapter/Private/Actuators/MLAdapterActuator_Camera.cpp
@@ -89,7 +89,7 @@ void UMLAdapterActuator_Camera::DigestInputData(FMLAdapterMemoryReader& ValueStr
 	{
 		FVector3f HeadingVector3f;
 		ValueStream << HeadingVector3f;
-		HeadingVector = HeadingVector3f;
+		HeadingVector = FVector(HeadingVector3f);
 	}
 	else
 	{
diff --git a/Engine/Plugins/Experimental/GPULightmass/Source/GPULightmass/Private/LightmapRenderer.cpp b/Engine/Plugins/Experimental/GPULightmass/Source/GPULightmass/Private/LightmapRenderer.cpp
--- a/Engine/Plugins/Experimental/GPULightmass/Source/GPULightmass/Private/LightmapRenderer.cpp
+++ b/Engine/Plugins/Experimental/GPULightmass/Source/GPULightmass/Private/LightmapRenderer.cpp
@@ -1109,7 +1109,7 @@ bool FSceneRenderState::SetupRayTracingScene(int32 LODIndex)
 				RHICmdList,
 				nullptr,
 				AbsoluteViewOrigin.GetTile(),
-				View.ViewMatrices.GetPreViewTranslation() + ViewTileOffset,
+				FVector3f(View.ViewMatrices.GetPreViewTranslation() + ViewTileOffset),
 				InstanceBuffer.UAV,
 				InstanceUploadSRV,
 				AccelerationStructureAddressesBuffer.SRV,
diff --git a/Engine/Plugins/Experimental/Interchange/Runtime/Source/Import/Private/Mesh/InterchangeSkeletalMeshFactory.cpp b/Engine/Plugins/Experimental/Interchange/Runtime/Source/Import/Private/Mesh/InterchangeSkeletalMeshFactory.cpp
--- a/Engine/Plugins/Experimental/Interchange/Runtime/Source/Import/Private/Mesh/InterchangeSkeletalMeshFactory.cpp
+++ b/Engine/Plugins/Experimental/Interchange/Runtime/Source/Import/Private/Mesh/InterchangeSkeletalMeshFactory.cpp
@@ -396,7 +396,7 @@ namespace UE
 				const int32 VertexCount = MeshDescription.Vertices().Num();
 				const TArray<FString>& JointNames = LodMeshPayload.JointNames;
 				// Create a copy of the vertex array to receive vertex deformations.
-				TArray<FVector> DestinationVertexPositions;
+				TArray<FVector3f> DestinationVertexPositions;
 				DestinationVertexPositions.AddZeroed(VertexCount);
 
 				FSkeletalMeshAttributes Attributes(MeshDescription);
@@ -406,7 +406,7 @@ namespace UE
 				for (FVertexID VertexID : MeshDescription.Vertices().GetElementIDs())
 				{
 					//We can use GetValue because the Meshdescription was compacted before the copy
-					DestinationVertexPositions[VertexID.GetValue()] = FVector(VertexPositions[VertexID]);
+					DestinationVertexPositions[VertexID.GetValue()] = VertexPositions[VertexID];
 				}
 
 				// Deform the vertex array with the links contained in the mesh.
@@ -490,19 +490,19 @@ namespace UE
 				for (FVertexID VertexID : MeshDescription.Vertices().GetElementIDs())
 				{
 					const int32 VertexIndex = VertexID.GetValue();
-					const FVector lSrcVertex = DestinationVertexPositions[VertexIndex];
-					FVector& lDstVertex = DestinationVertexPositions[VertexIndex];
+					const FVector lSrcVertex = FVector(DestinationVertexPositions[VertexIndex]);
+					FVector3f& lDstVertex = DestinationVertexPositions[VertexIndex];
 					double Weight = SkinWeights[VertexIndex];
 
 					// Deform the vertex if there was at least a link with an influence on the vertex,
 					if (!FMath::IsNearlyZero(Weight))
 					{
 						//Apply skinning of all joints
-						lDstVertex = SkinDeformations[VertexIndex].TransformPosition(lSrcVertex);
+						lDstVertex = FVector4f(SkinDeformations[VertexIndex].TransformPosition(lSrcVertex));
 						//Normalized, in case the weight is different then 1
 						lDstVertex /= Weight;
 						//Set the new vertex position in the mesh description
-						VertexPositions[VertexID] = FVector3f(lDstVertex);
+						VertexPositions[VertexID] = lDstVertex;
 					}
 				}
 			}
@@ -910,8 +910,8 @@ namespace UE
 				for (int32 WedgeIndex = 0; WedgeIndex < WedgeNumber; ++WedgeIndex)
 				{
 					SkeletalMeshImportData::FVertex& Wedge = SkelMeshImportData->Wedges[WedgeIndex];
-					const FVector3f& Position = SkelMeshImportData->Points[Wedge.VertexIndex];
-					Bounds += FVector(Position);
+					const FVector Position = FVector(SkelMeshImportData->Points[Wedge.VertexIndex]);
+					Bounds += Position;
 				}
 
 				TArray<FSoftSkinVertex> Vertices;
@@ -947,12 +947,12 @@ namespace UE
 				for (int32 WedgeIndex = 0; WedgeIndex < WedgeNumber; ++WedgeIndex)
 				{
 					SkeletalMeshImportData::FVertex& Wedge = SkelMeshImportData->Wedges[WedgeIndex];
-					const FVector3f& Position = SkelMeshImportData->Points[Wedge.VertexIndex];
+					const FVector Position = FVector(SkelMeshImportData->Points[Wedge.VertexIndex]);
 					const FVector2f UV = Wedge.UVs[0];
 					const FVector3f& Normal = WedgeIndexToNormal.FindChecked(WedgeIndex);
 
 					TArray<FSoftSkinVertex> PointsToConsider;
-					VertPosOctree.FindNearbyElements(FVector(Position), [&PointsToConsider](const FSoftSkinVertex& Vertex)
+					VertPosOctree.FindNearbyElements(Position, [&PointsToConsider](const FSoftSkinVertex& Vertex)
 						{
 							PointsToConsider.Add(Vertex);
 						});
