Interchange Materials
- Added support for a Sheen BSDF.
- Added a flatten normal shader node.
- Added a texture transform shader node.

Interchange GLTF
- Parse the sheen texture, sheen roughness factor and sheen roughness texture.
- Prefixed the materials with "Material_" to prevent trying to create a mesh and a material with the same name.

diff --git a/Engine/Plugins/Editor/GLTFImporter/Source/GLTFCore/Private/GLTF/ExtensionsHandler.cpp b/Engine/Plugins/Editor/GLTFImporter/Source/GLTFCore/Private/GLTF/ExtensionsHandler.cpp
--- a/Engine/Plugins/Editor/GLTFImporter/Source/GLTFCore/Private/GLTF/ExtensionsHandler.cpp
+++ b/Engine/Plugins/Editor/GLTFImporter/Source/GLTFCore/Private/GLTF/ExtensionsHandler.cpp
@@ -188,7 +188,12 @@ namespace GLTF
 					const FJsonObject& Sheen = ExtObj;
 
 					Material.bHasSheen = true;
+
 					Material.Sheen.SheenColorFactor = GetVec3(Sheen, TEXT("sheenColorFactor"));
+					GLTF::SetTextureMap(Sheen, TEXT("sheenColorTexture"), nullptr, Asset->Textures, Material.Sheen.SheenColorMap, Messages);
+
+					Material.Sheen.SheenRoughnessFactor = GetScalar(Sheen, TEXT("sheenRoughnessFactor"));
+					GLTF::SetTextureMap(Sheen, TEXT("sheenRoughnessTexture"), nullptr, Asset->Textures, Material.Sheen.SheenRoughnessMap, Messages);
 
 					Asset->ExtensionsUsed.Add(EExtension::KHR_MaterialsSheen);
 				}
diff --git a/Engine/Plugins/Editor/GLTFImporter/Source/GLTFCore/Public/GLTFMaterial.h b/Engine/Plugins/Editor/GLTFImporter/Source/GLTFCore/Public/GLTFMaterial.h
--- a/Engine/Plugins/Editor/GLTFImporter/Source/GLTFCore/Public/GLTFMaterial.h
+++ b/Engine/Plugins/Editor/GLTFImporter/Source/GLTFCore/Public/GLTFMaterial.h
@@ -109,6 +109,9 @@ namespace GLTF
 		struct FSheen
 		{
 			FVector SheenColorFactor = FVector::Zero();
+			FTextureMap SheenColorMap;
+			float SheenRoughnessFactor = 0.f;
+			FTextureMap SheenRoughnessMap;
 		};
 		struct FSpecular // "specular" extension
 		{
diff --git a/Engine/Plugins/Experimental/Interchange/Runtime/Source/Dispatcher/Public/InterchangeCommands.h b/Engine/Plugins/Experimental/Interchange/Runtime/Source/Dispatcher/Public/InterchangeCommands.h
--- a/Engine/Plugins/Experimental/Interchange/Runtime/Source/Dispatcher/Public/InterchangeCommands.h
+++ b/Engine/Plugins/Experimental/Interchange/Runtime/Source/Dispatcher/Public/InterchangeCommands.h
@@ -17,7 +17,7 @@ namespace UE
 		{
 		public:
 			//Major version should be updated when there is an existing API that has been change.
-			static int32 GetMajor() { return 4; }
+			static int32 GetMajor() { return 5; }
 			//Minor version must be updated when there is an addition to the API.
 			static int32 GetMinor() { return 0; }
 			//Patch version should be update if there is some bug fixes in the private code.
diff --git a/Engine/Plugins/Experimental/Interchange/Runtime/Source/FactoryNodes/Private/InterchangeMaterialFactoryNode.cpp b/Engine/Plugins/Experimental/Interchange/Runtime/Source/FactoryNodes/Private/InterchangeMaterialFactoryNode.cpp
--- a/Engine/Plugins/Experimental/Interchange/Runtime/Source/FactoryNodes/Private/InterchangeMaterialFactoryNode.cpp
+++ b/Engine/Plugins/Experimental/Interchange/Runtime/Source/FactoryNodes/Private/InterchangeMaterialFactoryNode.cpp
@@ -227,6 +227,36 @@ bool UInterchangeMaterialFactoryNode::ConnectOutputToTransmissionColor(const FSt
 	return UInterchangeShaderPortsAPI::ConnectOuputToInput(this, UE::Interchange::Materials::ThinTranslucent::Parameters::TransmissionColor.ToString(), ExpressionNodeUid, OutputName);
 }
 
+bool UInterchangeMaterialFactoryNode::GetFuzzColorConnection(FString& ExpressionNodeUid, FString& OutputName) const
+{
+	return UInterchangeShaderPortsAPI::GetInputConnection(this, UE::Interchange::Materials::Sheen::Parameters::SheenColor.ToString(), ExpressionNodeUid, OutputName);
+}
+
+bool UInterchangeMaterialFactoryNode::ConnectToFuzzColor(const FString& AttributeValue)
+{
+	return UInterchangeShaderPortsAPI::ConnectDefaultOuputToInput(this, UE::Interchange::Materials::Sheen::Parameters::SheenColor.ToString(), AttributeValue);
+}
+
+bool UInterchangeMaterialFactoryNode::ConnectOutputToFuzzColor(const FString& ExpressionNodeUid, const FString& OutputName)
+{
+	return UInterchangeShaderPortsAPI::ConnectOuputToInput(this, UE::Interchange::Materials::Sheen::Parameters::SheenColor.ToString(), ExpressionNodeUid, OutputName);
+}
+
+bool UInterchangeMaterialFactoryNode::GetClothConnection(FString& ExpressionNodeUid, FString& OutputName) const
+{
+	return UInterchangeShaderPortsAPI::GetInputConnection(this, UE::Interchange::Materials::Sheen::Parameters::SheenRoughness.ToString(), ExpressionNodeUid, OutputName);
+}
+
+bool UInterchangeMaterialFactoryNode::ConnectToCloth(const FString& AttributeValue)
+{
+	return UInterchangeShaderPortsAPI::ConnectDefaultOuputToInput(this, UE::Interchange::Materials::Sheen::Parameters::SheenRoughness.ToString(), AttributeValue);
+}
+
+bool UInterchangeMaterialFactoryNode::ConnectOutputToCloth(const FString& ExpressionNodeUid, const FString& OutputName)
+{
+	return UInterchangeShaderPortsAPI::ConnectOuputToInput(this, UE::Interchange::Materials::Sheen::Parameters::SheenRoughness.ToString(), ExpressionNodeUid, OutputName);
+}
+
 bool UInterchangeMaterialFactoryNode::GetCustomShadingModel(TEnumAsByte<EMaterialShadingModel>& AttributeValue) const
 {
 	IMPLEMENT_NODE_ATTRIBUTE_GETTER(ShadingModel, TEnumAsByte<EMaterialShadingModel>);
diff --git a/Engine/Plugins/Experimental/Interchange/Runtime/Source/FactoryNodes/Public/InterchangeMaterialFactoryNode.h b/Engine/Plugins/Experimental/Interchange/Runtime/Source/FactoryNodes/Public/InterchangeMaterialFactoryNode.h
--- a/Engine/Plugins/Experimental/Interchange/Runtime/Source/FactoryNodes/Public/InterchangeMaterialFactoryNode.h
+++ b/Engine/Plugins/Experimental/Interchange/Runtime/Source/FactoryNodes/Public/InterchangeMaterialFactoryNode.h
@@ -153,6 +153,24 @@ public:
 	UFUNCTION(BlueprintCallable, Category = "Interchange | Node | MaterialFactory")
 	bool ConnectOutputToTransmissionColor(const FString& ExpressionNodeUid, const FString& OutputName);
 
+	UFUNCTION(BlueprintCallable, Category = "Interchange | Node | MaterialFactory")
+	bool GetFuzzColorConnection(FString& ExpressionNodeUid, FString& OutputName) const;
+
+	UFUNCTION(BlueprintCallable, Category = "Interchange | Node | MaterialFactory")
+	bool ConnectToFuzzColor(const FString& AttributeValue);
+
+	UFUNCTION(BlueprintCallable, Category = "Interchange | Node | MaterialFactory")
+	bool ConnectOutputToFuzzColor(const FString& ExpressionNodeUid, const FString& OutputName);
+
+	UFUNCTION(BlueprintCallable, Category = "Interchange | Node | MaterialFactory")
+	bool GetClothConnection(FString& ExpressionNodeUid, FString& OutputName) const;
+
+	UFUNCTION(BlueprintCallable, Category = "Interchange | Node | MaterialFactory")
+	bool ConnectToCloth(const FString& AttributeValue);
+
+	UFUNCTION(BlueprintCallable, Category = "Interchange | Node | MaterialFactory")
+	bool ConnectOutputToCloth(const FString& ExpressionNodeUid, const FString& OutputName);
+
 // Material parameters
 public:
 	UFUNCTION(BlueprintCallable, Category = "Interchange | Node | Material")
diff --git a/Engine/Plugins/Experimental/Interchange/Runtime/Source/Import/Private/Gltf/InterchangeGltfTranslator.cpp b/Engine/Plugins/Experimental/Interchange/Runtime/Source/Import/Private/Gltf/InterchangeGltfTranslator.cpp
--- a/Engine/Plugins/Experimental/Interchange/Runtime/Source/Import/Private/Gltf/InterchangeGltfTranslator.cpp
+++ b/Engine/Plugins/Experimental/Interchange/Runtime/Source/Import/Private/Gltf/InterchangeGltfTranslator.cpp
@@ -109,7 +109,7 @@ void UInterchangeGltfTranslator::HandleGltfNode( UInterchangeBaseNodeContainer&
 }
 
 void UInterchangeGltfTranslator::HandleGltfMaterialParameter( UInterchangeBaseNodeContainer& NodeContainer, const GLTF::FTextureMap& TextureMap, UInterchangeShaderNode& ShaderNode,
-		const FString& MapName, const TVariant< FLinearColor, float >& MapFactor, const FString& OutputChannel, const bool bInverse ) const
+		const FString& MapName, const TVariant< FLinearColor, float >& MapFactor, const FString& OutputChannel, const bool bInverse, const bool bIsNormal ) const
 {
 	using namespace UE::Interchange::Materials;
 
@@ -158,6 +158,11 @@ void UInterchangeGltfTranslator::HandleGltfMaterialParameter( UInterchangeBaseNo
 		const FString TextureUid = TEXT("\\Texture\\") + GltfAsset.Textures[ TextureMap.TextureIndex ].Source.URI;
 		ColorNode->AddStringAttribute( UInterchangeShaderPortsAPI::MakeInputValueKey( Standard::Nodes::TextureSample::Inputs::Texture.ToString() ), TextureUid );
 
+		if ( TextureMap.bHasTextureTransform )
+		{
+			HandleGltfTextureTransform( NodeContainer, TextureMap.TextureTransform, TextureMap.TexCoord, *ColorNode );
+		}
+
 		bool bNeedsFactorNode = false;
 
 		if ( MapFactor.IsType< float >() )
@@ -177,18 +182,41 @@ void UInterchangeGltfTranslator::HandleGltfMaterialParameter( UInterchangeBaseNo
 			NodeContainer.AddNode( FactorNode );
 			NodeContainer.SetNodeParentUid( FactorNodeUid, ShaderNode.GetUniqueID() );
 
-			FactorNode->SetCustomShaderType( Standard::Nodes::Multiply::Name.ToString() );
-
-			if ( MapFactor.IsType< float >() )
+			if ( bIsNormal )
 			{
-				FactorNode->AddFloatAttribute( UInterchangeShaderPortsAPI::MakeInputValueKey( Standard::Nodes::Multiply::Inputs::B.ToString() ), MapFactor.Get< float >() );
+				FactorNode->SetCustomShaderType( Standard::Nodes::FlattenNormal::Name.ToString() );
+
+				const FString FactorOneMinusNodeUid = FactorNodeUid + TEXT("_OneMinus");
+				UInterchangeShaderNode* FactorOneMinusNode = NewObject< UInterchangeShaderNode >( &NodeContainer );
+				FactorOneMinusNode->InitializeNode( FactorOneMinusNodeUid, NodeName + TEXT("_Factor_OneMinus"), EInterchangeNodeContainerType::TranslatedAsset );
+				FactorOneMinusNode->SetCustomShaderType(Standard::Nodes::OneMinus::Name.ToString());
+				NodeContainer.AddNode( FactorOneMinusNode );
+				NodeContainer.SetNodeParentUid( FactorOneMinusNodeUid, ShaderNode.GetUniqueID() );
+
+				if ( MapFactor.IsType< float >() )
+				{
+					FactorOneMinusNode->AddFloatAttribute( UInterchangeShaderPortsAPI::MakeInputValueKey( Standard::Nodes::OneMinus::Inputs::Input.ToString() ), MapFactor.Get< float >() );
+				}
+
+				UInterchangeShaderPortsAPI::ConnectOuputToInput( FactorNode, Standard::Nodes::FlattenNormal::Inputs::Normal.ToString(), NodeUid, OutputChannel );
+				UInterchangeShaderPortsAPI::ConnectDefaultOuputToInput( FactorNode, Standard::Nodes::FlattenNormal::Inputs::Flatness.ToString(), FactorOneMinusNodeUid );
 			}
-			else if ( MapFactor.IsType< FLinearColor >() )
+			else
 			{
-				FactorNode->AddLinearColorAttribute( UInterchangeShaderPortsAPI::MakeInputValueKey( Standard::Nodes::Multiply::Inputs::B.ToString() ), MapFactor.Get< FLinearColor >() );
+				FactorNode->SetCustomShaderType( Standard::Nodes::Multiply::Name.ToString() );
+
+				if ( MapFactor.IsType< float >() )
+				{
+					FactorNode->AddFloatAttribute( UInterchangeShaderPortsAPI::MakeInputValueKey( Standard::Nodes::Multiply::Inputs::B.ToString() ), MapFactor.Get< float >() );
+				}
+				else if ( MapFactor.IsType< FLinearColor >() )
+				{
+					FactorNode->AddLinearColorAttribute( UInterchangeShaderPortsAPI::MakeInputValueKey( Standard::Nodes::Multiply::Inputs::B.ToString() ), MapFactor.Get< FLinearColor >() );
+				}
+
+				UInterchangeShaderPortsAPI::ConnectOuputToInput( FactorNode, Standard::Nodes::Multiply::Inputs::A.ToString(), NodeUid, OutputChannel );
 			}
 
-			UInterchangeShaderPortsAPI::ConnectOuputToInput( FactorNode, Standard::Nodes::Multiply::Inputs::A.ToString(), NodeUid, OutputChannel );
 			UInterchangeShaderPortsAPI::ConnectDefaultOuputToInput( NodeToConnectTo, InputToConnectTo, FactorNodeUid );
 		}
 		else
@@ -293,8 +321,11 @@ void UInterchangeGltfTranslator::HandleGltfMaterial( UInterchangeBaseNodeContain
 			TVariant< FLinearColor, float > NormalFactor;
 			NormalFactor.Set< float >( GltfMaterial.NormalScale );
 
+			const bool bInverse = false;
+			const bool bIsNormal = true;
+
 			HandleGltfMaterialParameter( NodeContainer, GltfMaterial.Normal, ShaderGraphNode, Common::Parameters::Normal.ToString(),
-				NormalFactor, Standard::Nodes::TextureSample::Outputs::RGB.ToString() );
+				NormalFactor, Standard::Nodes::TextureSample::Outputs::RGB.ToString(), bInverse, bIsNormal );
 		}
 
 		// Emissive
@@ -314,7 +345,7 @@ void UInterchangeGltfTranslator::HandleGltfMaterial( UInterchangeBaseNodeContain
 			OcclusionFactor.Set< float >( GltfMaterial.OcclusionStrength );
 
 			HandleGltfMaterialParameter( NodeContainer, GltfMaterial.Occlusion, ShaderGraphNode, PBR::Parameters::Occlusion.ToString(),
-				OcclusionFactor, Standard::Nodes::TextureSample::Outputs::RGB.ToString() );
+				OcclusionFactor, Standard::Nodes::TextureSample::Outputs::R.ToString() );
 		}
 
 		// Opacity (use the base color alpha channel)
@@ -339,6 +370,11 @@ void UInterchangeGltfTranslator::HandleGltfMaterial( UInterchangeBaseNodeContain
 		HandleGltfClearCoat( NodeContainer, GltfMaterial, ShaderGraphNode );
 	}
 
+	if ( GltfMaterial.bHasSheen )
+	{
+		HandleGltfSheen( NodeContainer, GltfMaterial, ShaderGraphNode );
+	}
+
 	if ( GltfMaterial.bHasTransmission )
 	{
 		HandleGltfTransmission( NodeContainer, GltfMaterial, ShaderGraphNode );
@@ -360,7 +396,7 @@ void UInterchangeGltfTranslator::HandleGltfClearCoat( UInterchangeBaseNodeContai
 		ClearCoatFactor.Set< float >( GltfMaterial.ClearCoat.ClearCoatFactor );
 
 		HandleGltfMaterialParameter( NodeContainer, GltfMaterial.ClearCoat.ClearCoatMap, ShaderGraphNode, ClearCoat::Parameters::ClearCoat.ToString(),
-			ClearCoatFactor, Standard::Nodes::TextureSample::Outputs::RGB.ToString() );
+			ClearCoatFactor, Standard::Nodes::TextureSample::Outputs::R.ToString() );
 	}
 
 	//  ClearCoat::Parameters::ClearCoatRoughness
@@ -369,16 +405,47 @@ void UInterchangeGltfTranslator::HandleGltfClearCoat( UInterchangeBaseNodeContai
 		ClearCoatRoughnessFactor.Set< float >( GltfMaterial.ClearCoat.Roughness );
 
 		HandleGltfMaterialParameter( NodeContainer, GltfMaterial.ClearCoat.RoughnessMap, ShaderGraphNode, ClearCoat::Parameters::ClearCoatRoughness.ToString(),
-			ClearCoatRoughnessFactor, Standard::Nodes::TextureSample::Outputs::RGB.ToString() );
+			ClearCoatRoughnessFactor, Standard::Nodes::TextureSample::Outputs::G.ToString() );
 	}
 
 	// ClearCoat::Parameters::ClearCoatNormal
 	{
 		TVariant< FLinearColor, float > ClearCoatNormalFactor;
-		ClearCoatNormalFactor.Set< FLinearColor >( FLinearColor::White );
+		ClearCoatNormalFactor.Set< FLinearColor >( FLinearColor( FVector::UpVector ) );
+
+		const bool bInverse = false;
+		const bool bIsNormal = true;
 
 		HandleGltfMaterialParameter( NodeContainer, GltfMaterial.ClearCoat.NormalMap, ShaderGraphNode, ClearCoat::Parameters::ClearCoatNormal.ToString(),
-			ClearCoatNormalFactor, Standard::Nodes::TextureSample::Outputs::RGB.ToString() );
+			ClearCoatNormalFactor, Standard::Nodes::TextureSample::Outputs::RGB.ToString(), bInverse, bIsNormal );
+	}
+}
+
+void UInterchangeGltfTranslator::HandleGltfSheen( UInterchangeBaseNodeContainer& NodeContainer, const GLTF::FMaterial& GltfMaterial, UInterchangeShaderGraphNode& ShaderGraphNode ) const
+{
+	using namespace UE::Interchange::Materials;
+
+	if ( !GltfMaterial.bHasSheen )
+	{
+		return;
+	}
+
+	// Sheen::Parameters::SheenColor
+	{
+		TVariant< FLinearColor, float > SheenColorFactor;
+		SheenColorFactor.Set< FLinearColor >( FLinearColor( GltfMaterial.Sheen.SheenColorFactor ) );
+
+		HandleGltfMaterialParameter( NodeContainer, GltfMaterial.Sheen.SheenColorMap, ShaderGraphNode, Sheen::Parameters::SheenColor.ToString(),
+			SheenColorFactor, Standard::Nodes::TextureSample::Outputs::RGB.ToString() );
+	}
+
+	// Sheen::Parameters::SheenRoughness
+	{
+		TVariant< FLinearColor, float > SheenRoughnessFactor;
+		SheenRoughnessFactor.Set< float >( GltfMaterial.Sheen.SheenRoughnessFactor );
+
+		HandleGltfMaterialParameter( NodeContainer, GltfMaterial.Sheen.SheenRoughnessMap, ShaderGraphNode, Sheen::Parameters::SheenRoughness.ToString(),
+			SheenRoughnessFactor, Standard::Nodes::TextureSample::Outputs::A.ToString() );
 	}
 }
 
@@ -547,6 +614,73 @@ void UInterchangeGltfTranslator::HandleGltfTransmission( UInterchangeBaseNodeCon
 	}
 }
 
+void UInterchangeGltfTranslator::HandleGltfTextureTransform( UInterchangeBaseNodeContainer& NodeContainer, const GLTF::FTextureTransform& TextureTransform, const int32 TexCoordIndex, UInterchangeShaderNode& ShaderNode ) const
+{
+	using namespace UE::Interchange::Materials;
+
+	UInterchangeShaderNode* CurrentNode = &ShaderNode;
+
+	// Texture Coordinates
+	{
+		const FString TexCoordNodeUid = ShaderNode.GetUniqueID() + TEXT("_TexCoord");
+		UInterchangeShaderNode* TexCoordNode = NewObject< UInterchangeShaderNode >( &NodeContainer );
+		TexCoordNode->InitializeNode( TexCoordNodeUid, ShaderNode.GetDisplayLabel() + TEXT("_TexCoord"), EInterchangeNodeContainerType::TranslatedAsset );
+		NodeContainer.AddNode( TexCoordNode );
+		NodeContainer.SetNodeParentUid( TexCoordNodeUid, ShaderNode.GetUniqueID() );
+
+		TexCoordNode->SetCustomShaderType( Standard::Nodes::TextureCoordinate::Name.ToString() );
+
+		TexCoordNode->AddInt32Attribute( UInterchangeShaderPortsAPI::MakeInputValueKey( Standard::Nodes::TextureCoordinate::Inputs::Index.ToString() ), TexCoordIndex );
+
+		CurrentNode = TexCoordNode;
+	}
+
+	// Scale
+	if ( !FMath::IsNearlyEqual( TextureTransform.Scale[0], 1.f ) || 
+		 !FMath::IsNearlyEqual( TextureTransform.Scale[1], 1.f ) )
+	{
+		FVector2f TextureScale;
+		TextureScale.X = TextureTransform.Scale[0];
+		TextureScale.Y = TextureTransform.Scale[1];
+
+		CurrentNode->SetAttribute< FVector2f >( UInterchangeShaderPortsAPI::MakeInputValueKey( Standard::Nodes::TextureCoordinate::Inputs::Scale.ToString() ), TextureScale );
+	}
+
+	// Offset
+	if ( !FMath::IsNearlyZero( TextureTransform.Offset[0] ) || 
+		 !FMath::IsNearlyZero( TextureTransform.Offset[1] ) )
+	{
+		FVector2f TextureOffset;
+		TextureOffset.X = TextureTransform.Offset[0];
+		TextureOffset.Y = TextureTransform.Offset[1];
+
+		CurrentNode->SetAttribute< FVector2f >( UInterchangeShaderPortsAPI::MakeInputValueKey( Standard::Nodes::TextureCoordinate::Inputs::Offset.ToString() ), TextureOffset );
+	}
+
+	// Rotate
+	if ( !FMath::IsNearlyZero( TextureTransform.Rotation ) )
+	{
+		float AngleRadians = TextureTransform.Rotation;
+
+		if ( AngleRadians < 0.0f )
+		{
+			AngleRadians = TWO_PI - AngleRadians;
+		}
+
+		AngleRadians = 1.0f - ( AngleRadians / TWO_PI );
+
+		CurrentNode->AddFloatAttribute( UInterchangeShaderPortsAPI::MakeInputValueKey( Standard::Nodes::TextureCoordinate::Inputs::Rotate.ToString() ), AngleRadians );
+
+		FVector2f RotationCenter = FVector2f::Zero();
+		CurrentNode->SetAttribute< FVector2f >( UInterchangeShaderPortsAPI::MakeInputValueKey( Standard::Nodes::TextureCoordinate::Inputs::RotationCenter.ToString() ), RotationCenter );
+	}
+
+	if ( CurrentNode != &ShaderNode )
+	{
+		UInterchangeShaderPortsAPI::ConnectDefaultOuputToInput( &ShaderNode, Standard::Nodes::TextureSample::Inputs::Coordinates.ToString(), CurrentNode->GetUniqueID() );
+	}
+}
+
 EInterchangeTranslatorType UInterchangeGltfTranslator::GetTranslatorType() const
 {
 	return EInterchangeTranslatorType::Scenes;
@@ -598,7 +732,7 @@ bool UInterchangeGltfTranslator::Translate( UInterchangeBaseNodeContainer& NodeC
 		{
 			UInterchangeShaderGraphNode* ShaderGraphNode = NewObject< UInterchangeShaderGraphNode >( &NodeContainer );
 			FString ShaderGraphNodeUid = TEXT("\\Material\\") + GenerateUniqueIdForGltfNode( GltfMaterial.Name, MaterialIndex );
-			ShaderGraphNode->InitializeNode( ShaderGraphNodeUid, GenerateUniqueIdForGltfNode( GltfMaterial.Name, MaterialIndex ), EInterchangeNodeContainerType::TranslatedAsset );
+			ShaderGraphNode->InitializeNode( ShaderGraphNodeUid, TEXT("Material_") + GenerateUniqueIdForGltfNode( GltfMaterial.Name, MaterialIndex ), EInterchangeNodeContainerType::TranslatedAsset );
 			NodeContainer.AddNode( ShaderGraphNode );
 
 			HandleGltfMaterial( NodeContainer, GltfMaterial, *ShaderGraphNode );
@@ -723,7 +857,7 @@ TFuture< TOptional< UE::Interchange::FStaticMeshPayloadData > > UInterchangeGltf
 
 			if ( GltfAsset.Materials.IsValidIndex( MaterialIndex ) )
 			{
-				StaticMeshAttributes.GetPolygonGroupMaterialSlotNames()[ MaterialSlotIndex ] = *GenerateUniqueIdForGltfNode( GltfAsset.Materials[ MaterialIndex ].Name, MaterialIndex );
+				StaticMeshAttributes.GetPolygonGroupMaterialSlotNames()[ MaterialSlotIndex ] = *(TEXT("Material_") + GenerateUniqueIdForGltfNode( GltfAsset.Materials[ MaterialIndex ].Name, MaterialIndex ));
 			}
 		}
 	}
@@ -743,9 +877,10 @@ TOptional< UE::Interchange::FImportImage > UInterchangeGltfTranslator::GetTextur
 		return TOptional< UE::Interchange::FImportImage >();
 	}
 
-	GLTF::FTexture GltfTexture = GltfAsset.Textures[ TextureIndex ];
+	const GLTF::FTexture& GltfTexture = GltfAsset.Textures[ TextureIndex ];
+	const FString TextureFilePath = FPaths::ConvertRelativePathToFull( GltfTexture.Source.FilePath );
 
-	UInterchangeSourceData* PayloadSourceData = UInterchangeManager::GetInterchangeManager().CreateSourceData( GltfTexture.Source.FilePath );
+	UInterchangeSourceData* PayloadSourceData = UInterchangeManager::GetInterchangeManager().CreateSourceData( TextureFilePath );
 	FGCObjectScopeGuard ScopedSourceData( PayloadSourceData );
 	
 	if ( !PayloadSourceData )
@@ -761,5 +896,5 @@ TOptional< UE::Interchange::FImportImage > UInterchangeGltfTranslator::GetTextur
 		return TOptional<UE::Interchange::FImportImage>();
 	}
 
-	return TextureTranslator->GetTexturePayloadData( PayloadSourceData, GltfTexture.Source.FilePath );
+	return TextureTranslator->GetTexturePayloadData( PayloadSourceData, TextureFilePath );
 }
\ No newline at end of file
diff --git a/Engine/Plugins/Experimental/Interchange/Runtime/Source/Import/Private/Material/InterchangeMaterialFactory.cpp b/Engine/Plugins/Experimental/Interchange/Runtime/Source/Import/Private/Material/InterchangeMaterialFactory.cpp
--- a/Engine/Plugins/Experimental/Interchange/Runtime/Source/Import/Private/Material/InterchangeMaterialFactory.cpp
+++ b/Engine/Plugins/Experimental/Interchange/Runtime/Source/Import/Private/Material/InterchangeMaterialFactory.cpp
@@ -632,6 +632,50 @@ void UInterchangeMaterialFactory::SetupMaterial(UMaterial* Material, const FCrea
 		}
 	}
 
+	// Fuzz Color
+	{
+		FString FuzzColorUid;
+		FString OutputName;
+
+		if (MaterialFactoryNode->GetFuzzColorConnection(FuzzColorUid, OutputName))
+		{
+			const UInterchangeMaterialExpressionFactoryNode* FuzzColorNode = Cast<UInterchangeMaterialExpressionFactoryNode>(Arguments.NodeContainer->GetNode(FuzzColorUid));
+
+			if (FuzzColorNode)
+			{
+				if (UMaterialExpression* FuzzColorExpression = CreateExpressionsForNode(Material, Arguments, FuzzColorNode, Expressions))
+				{
+					if (FExpressionInput* FuzzColorInput = Material->GetExpressionInputForProperty(MP_SubsurfaceColor))
+					{
+						FuzzColorExpression->ConnectExpression(FuzzColorInput, GetOutputIndex(*FuzzColorExpression, OutputName));
+					}
+				}
+			}
+		}
+	}
+
+	// Cloth
+	{
+		FString ClothUid;
+		FString OutputName;
+
+		if (MaterialFactoryNode->GetClothConnection(ClothUid, OutputName))
+		{
+			const UInterchangeMaterialExpressionFactoryNode* ClothNode = Cast<UInterchangeMaterialExpressionFactoryNode>(Arguments.NodeContainer->GetNode(ClothUid));
+
+			if (ClothNode)
+			{
+				if (UMaterialExpression* ClothExpression = CreateExpressionsForNode(Material, Arguments, ClothNode, Expressions))
+				{
+					if (FExpressionInput* ClothInput = Material->GetExpressionInputForProperty(MP_CustomData0))
+					{
+						ClothExpression->ConnectExpression(ClothInput, GetOutputIndex(*ClothExpression, OutputName));
+					}
+				}
+			}
+		}
+	}
+
 	UMaterialEditingLibrary::LayoutMaterialExpressions(Material);
 }
 
diff --git a/Engine/Plugins/Experimental/Interchange/Runtime/Source/Import/Private/Mesh/InterchangeOBJTranslator.cpp b/Engine/Plugins/Experimental/Interchange/Runtime/Source/Import/Private/Mesh/InterchangeOBJTranslator.cpp
--- a/Engine/Plugins/Experimental/Interchange/Runtime/Source/Import/Private/Mesh/InterchangeOBJTranslator.cpp
+++ b/Engine/Plugins/Experimental/Interchange/Runtime/Source/Import/Private/Mesh/InterchangeOBJTranslator.cpp
@@ -1005,8 +1005,6 @@ void UInterchangeOBJTranslator::AddMaterialNodes(UInterchangeBaseNodeContainer&
 		}
 
 		TextureSampleShader->AddStringAttribute(UInterchangeShaderPortsAPI::MakeInputValueKey(TextureSample::Inputs::Texture.ToString()), TextureNodeUid);
-		TextureSampleShader->AddFloatAttribute(UInterchangeShaderPortsAPI::MakeInputValueKey(TextureSample::Inputs::UTiling.ToString()), 1.0f);
-		TextureSampleShader->AddFloatAttribute(UInterchangeShaderPortsAPI::MakeInputValueKey(TextureSample::Inputs::VTiling.ToString()), 1.0f);
 
 		UInterchangeShaderPortsAPI::ConnectDefaultOuputToInput(ShaderGraphNode, InputType, TextureSampleShaderUid);
 	}
diff --git a/Engine/Plugins/Experimental/Interchange/Runtime/Source/Import/Public/Gltf/InterchangeGltfTranslator.h b/Engine/Plugins/Experimental/Interchange/Runtime/Source/Import/Public/Gltf/InterchangeGltfTranslator.h
--- a/Engine/Plugins/Experimental/Interchange/Runtime/Source/Import/Public/Gltf/InterchangeGltfTranslator.h
+++ b/Engine/Plugins/Experimental/Interchange/Runtime/Source/Import/Public/Gltf/InterchangeGltfTranslator.h
@@ -46,13 +46,18 @@ protected:
 	void HandleGltfNode( UInterchangeBaseNodeContainer& NodeContainer, const GLTF::FNode& GltfNode, const FString& ParentNodeUid, const int32 NodeIndex ) const;
 	void HandleGltfMaterial( UInterchangeBaseNodeContainer& NodeContainer, const GLTF::FMaterial& GltfMaterial, UInterchangeShaderGraphNode& ShaderGraphNode ) const;
 	void HandleGltfMaterialParameter( UInterchangeBaseNodeContainer& NodeContainer, const GLTF::FTextureMap& TextureMap, UInterchangeShaderNode& ShaderNode,
-		const FString& MapName, const TVariant< FLinearColor, float >& MapFactor, const FString& OutputChannel, const bool bInverse = false ) const;
+		const FString& MapName, const TVariant< FLinearColor, float >& MapFactor, const FString& OutputChannel, const bool bInverse = false, const bool bIsNormal = false ) const;
 
 	/** Support for KHR_materials_clearcoat */
 	void HandleGltfClearCoat( UInterchangeBaseNodeContainer& NodeContainer, const GLTF::FMaterial& GltfMaterial, UInterchangeShaderGraphNode& ShaderGraphNode ) const;
+	/** Support for KHR_materials_sheen */
+	void HandleGltfSheen( UInterchangeBaseNodeContainer& NodeContainer, const GLTF::FMaterial& GltfMaterial, UInterchangeShaderGraphNode& ShaderGraphNode ) const;
 	/** Support for KHR_materials_transmission */
 	void HandleGltfTransmission( UInterchangeBaseNodeContainer& NodeContainer, const GLTF::FMaterial& GltfMaterial, UInterchangeShaderGraphNode& ShaderGraphNode ) const;
 
+	/** Support for KHR_texture_transform */
+	void HandleGltfTextureTransform( UInterchangeBaseNodeContainer& NodeContainer, const GLTF::FTextureTransform& TextureTransform, const int32 TexCoordIndex, UInterchangeShaderNode& ShaderNode ) const;
+
 private:
 	GLTF::FAsset GltfAsset;
 };
diff --git a/Engine/Plugins/Experimental/Interchange/Runtime/Source/Nodes/Public/InterchangeMaterialDefinitions.h b/Engine/Plugins/Experimental/Interchange/Runtime/Source/Nodes/Public/InterchangeMaterialDefinitions.h
--- a/Engine/Plugins/Experimental/Interchange/Runtime/Source/Nodes/Public/InterchangeMaterialDefinitions.h
+++ b/Engine/Plugins/Experimental/Interchange/Runtime/Source/Nodes/Public/InterchangeMaterialDefinitions.h
@@ -12,6 +12,28 @@ namespace Interchange
 		{
 			namespace Nodes
 			{
+				namespace Add
+				{
+					const FName Name = TEXT("Add");
+
+					namespace Inputs
+					{
+						const FName A = TEXT("A");
+						const FName B = TEXT("B");
+					}
+				}
+
+				namespace FlattenNormal
+				{
+					const FName Name = TEXT("FlattenNormal");
+
+					namespace Inputs
+					{
+						const FName Normal = TEXT("Normal");
+						const FName Flatness = TEXT("Flatness");
+					}
+				}
+
 				namespace Lerp
 				{
 					const FName Name = TEXT("Lerp");
@@ -45,15 +67,30 @@ namespace Interchange
 					}
 				}
 
+				namespace TextureCoordinate
+				{
+					const FName Name = TEXT("TextureCoordinate");
+
+					namespace Inputs
+					{
+						const FName Index = TEXT("Index"); // Type: int
+						const FName UTiling = TEXT("UTiling"); // Type: float
+						const FName VTiling = TEXT("VTiling"); // Type: float
+						const FName Offset = TEXT("Offset"); // Type: vec2
+						const FName Scale = TEXT("Scale"); // Type: vec2
+						const FName Rotate = TEXT("Rotate"); // Type: float, Range: 0-1
+						const FName RotationCenter = TEXT("RotationCenter"); // Type: vec2
+					}
+				}
+
 				namespace TextureSample
 				{
 					const FName Name = TEXT("TextureSample");
 
 					namespace Inputs
 					{
+						const FName Coordinates = TEXT("Coordinates");
 						const FName Texture = TEXT("TextureUid"); // Type: FString (unique id of a texture node)
-						const FName UTiling = TEXT("UTiling"); // Type: float
-						const FName VTiling = TEXT("VTiling"); // Type: float
 					}
 
 					namespace Outputs
@@ -136,6 +173,17 @@ namespace Interchange
 				const FName TransmissionColor = TEXT("TransmissionColor"); // Type: vector3
 			}
 		}
+
+		namespace Sheen
+		{
+			namespace Parameters
+			{
+				using namespace PBR::Parameters;
+
+				const FName SheenColor = TEXT("SheenColor"); // Type: vector3
+				const FName SheenRoughness = TEXT("SheenRoughness"); // Type: float
+			}
+		}
 	}
 }
 }
diff --git a/Engine/Plugins/Experimental/Interchange/Runtime/Source/Parsers/Fbx/Private/FbxMaterial.cpp b/Engine/Plugins/Experimental/Interchange/Runtime/Source/Parsers/Fbx/Private/FbxMaterial.cpp
--- a/Engine/Plugins/Experimental/Interchange/Runtime/Source/Parsers/Fbx/Private/FbxMaterial.cpp
+++ b/Engine/Plugins/Experimental/Interchange/Runtime/Source/Parsers/Fbx/Private/FbxMaterial.cpp
@@ -109,9 +109,9 @@ namespace UE
 				// Handles max one texture per property.
 				if (TextureCount > 0)
 				{
-					FbxFileTexture* FbxTextureFilePath = Property.GetSrcObject<FbxFileTexture>(0);
+					FbxFileTexture* FbxTexture = Property.GetSrcObject<FbxFileTexture>(0);
 
-					FString TextureFilename = UTF8_TO_TCHAR(FbxTextureFilePath->GetFileName());
+					FString TextureFilename = UTF8_TO_TCHAR(FbxTexture->GetFileName());
 					FString TextureName = FPaths::GetBaseFilename(TextureFilename);
 
 					UInterchangeShaderNode* TextureSampleShader = NewObject<UInterchangeShaderNode>(&NodeContainer);
@@ -131,8 +131,23 @@ namespace UE
 					}
 
 					TextureSampleShader->AddStringAttribute(UInterchangeShaderPortsAPI::MakeInputValueKey(TextureSample::Inputs::Texture.ToString()), TextureNode->GetUniqueID());
-					TextureSampleShader->AddFloatAttribute(UInterchangeShaderPortsAPI::MakeInputValueKey(TextureSample::Inputs::UTiling.ToString()), (float)FbxTextureFilePath->GetScaleU());
-					TextureSampleShader->AddFloatAttribute(UInterchangeShaderPortsAPI::MakeInputValueKey(TextureSample::Inputs::VTiling.ToString()), (float)FbxTextureFilePath->GetScaleV());
+
+					if (!FMath::IsNearlyEqual(FbxTexture->GetScaleU(), 1.0) || !FMath::IsNearlyEqual(FbxTexture->GetScaleV(), 1.0))
+					{
+						UInterchangeShaderNode* TextureCoordinateShader = NewObject<UInterchangeShaderNode>(&NodeContainer);
+						TextureCoordinateShader->SetCustomShaderType(TextureCoordinate::Name.ToString());
+
+						FString TextureCoordinateShaderUid = TextureSampleShader->GetUniqueID() + TEXT("_Coordinate");
+						TextureCoordinateShader->InitializeNode(TextureCoordinateShaderUid, TextureName + TEXT("_Coordinate"), EInterchangeNodeContainerType::TranslatedAsset);
+
+						NodeContainer.AddNode(TextureCoordinateShader);
+						NodeContainer.SetNodeParentUid(TextureCoordinateShaderUid, ShaderGraphNode->GetUniqueID());
+
+						TextureCoordinateShader->AddFloatAttribute(UInterchangeShaderPortsAPI::MakeInputValueKey(TextureCoordinate::Inputs::UTiling.ToString()), (float)FbxTexture->GetScaleU());
+						TextureCoordinateShader->AddFloatAttribute(UInterchangeShaderPortsAPI::MakeInputValueKey(TextureCoordinate::Inputs::VTiling.ToString()), (float)FbxTexture->GetScaleV());
+
+						UInterchangeShaderPortsAPI::ConnectDefaultOuputToInput(TextureSampleShader, TextureSample::Inputs::Coordinates.ToString(), TextureCoordinateShaderUid);
+					}
 
 					UInterchangeShaderPortsAPI::ConnectDefaultOuputToInput(NodeToConnectTo, InputToConnectTo, TextureSampleShaderUid);
 
@@ -221,7 +236,7 @@ namespace UE
 								{
 									const float Factor = 1.f;
 									TVariant<FLinearColor, float> DefaultValue;
-									DefaultValue.Set<FLinearColor>(FLinearColor::Black);
+									DefaultValue.Set<FLinearColor>(FLinearColor(FVector::UpVector));
 									ConvertPropertyToShaderNode(NodeContainer, ShaderGraphNode, CurrentProperty, Factor, Phong::Parameters::Normal, DefaultValue);
 								}
 								else if (PropertyName == FbxSurfaceMaterial::sTransparentColor)
diff --git a/Engine/Plugins/Experimental/Interchange/Runtime/Source/Pipelines/Private/InterchangeGenericMaterialPipeline.cpp b/Engine/Plugins/Experimental/Interchange/Runtime/Source/Pipelines/Private/InterchangeGenericMaterialPipeline.cpp
--- a/Engine/Plugins/Experimental/Interchange/Runtime/Source/Pipelines/Private/InterchangeGenericMaterialPipeline.cpp
+++ b/Engine/Plugins/Experimental/Interchange/Runtime/Source/Pipelines/Private/InterchangeGenericMaterialPipeline.cpp
@@ -15,10 +15,13 @@
 #include "InterchangePipelineLog.h"
 #include "InterchangeSourceData.h"
 
+#include "Materials/MaterialExpressionAdd.h"
+#include "Materials/MaterialExpressionComponentMask.h"
 #include "Materials/MaterialExpressionDivide.h"
 #include "Materials/MaterialExpressionFresnel.h"
 #include "Materials/MaterialExpressionLinearInterpolate.h"
 #include "Materials/MaterialExpressionMaterialFunctionCall.h"
+#include "Materials/MaterialExpressionMultiply.h"
 #include "Materials/MaterialExpressionOneMinus.h"
 #include "Materials/MaterialExpressionScalarParameter.h"
 #include "Materials/MaterialExpressionTextureCoordinate.h"
@@ -132,6 +135,15 @@ bool UInterchangeGenericMaterialPipeline::IsClearCoatModel(const UInterchangeSha
 	return bHasClearCoatInput;
 }
 
+bool UInterchangeGenericMaterialPipeline::IsSheenModel(const UInterchangeShaderGraphNode* ShaderGraphNode) const
+{
+	using namespace UE::Interchange::Materials::Sheen;
+
+	const bool bHasSheenColorInput = UInterchangeShaderPortsAPI::HasInput(ShaderGraphNode, Parameters::SheenColor);
+
+	return bHasSheenColorInput;
+}
+
 bool UInterchangeGenericMaterialPipeline::IsThinTranslucentModel(const UInterchangeShaderGraphNode* ShaderGraphNode) const
 {
 	using namespace UE::Interchange::Materials::ThinTranslucent;
@@ -444,6 +456,62 @@ bool UInterchangeGenericMaterialPipeline::HandleClearCoat(const UInterchangeShad
 	return bShadingModelHandled;
 }
 
+bool UInterchangeGenericMaterialPipeline::HandleSheen(const UInterchangeShaderGraphNode* ShaderGraphNode, UInterchangeMaterialFactoryNode* MaterialFactoryNode)
+{
+	using namespace UE::Interchange::Materials::Sheen;
+
+	bool bShadingModelHandled = false;
+
+	// Sheen Color
+	{
+		const bool bHasInput = UInterchangeShaderPortsAPI::HasInput(ShaderGraphNode, Parameters::SheenColor);
+
+		if (bHasInput)
+		{
+			TTuple<UInterchangeMaterialExpressionFactoryNode*, FString> ExpressionFactoryNode =
+				CreateMaterialExpressionForInput(MaterialFactoryNode, ShaderGraphNode, Parameters::SheenColor.ToString(), MaterialFactoryNode->GetUniqueID());
+
+			if (ExpressionFactoryNode.Get<0>())
+			{
+				MaterialFactoryNode->ConnectOutputToFuzzColor(ExpressionFactoryNode.Get<0>()->GetUniqueID(), ExpressionFactoryNode.Get<1>());
+			}
+
+			bShadingModelHandled = true;
+		}
+	}
+
+	// Sheen Roughness
+	{
+		const bool bHasInput = UInterchangeShaderPortsAPI::HasInput(ShaderGraphNode, Parameters::SheenRoughness);
+
+		if (bHasInput)
+		{
+			TTuple<UInterchangeMaterialExpressionFactoryNode*, FString> ExpressionFactoryNode =
+				CreateMaterialExpressionForInput(MaterialFactoryNode, ShaderGraphNode, Parameters::SheenRoughness.ToString(), MaterialFactoryNode->GetUniqueID());
+
+			if (ExpressionFactoryNode.Get<0>())
+			{
+				UInterchangeMaterialExpressionFactoryNode* InverseSheenRoughnessNode =
+						CreateExpressionNode(TEXT("InverseSheenRoughness"), ExpressionFactoryNode.Get<0>()->GetUniqueID(), UMaterialExpressionOneMinus::StaticClass());
+
+				UInterchangeShaderPortsAPI::ConnectOuputToInput(InverseSheenRoughnessNode, GET_MEMBER_NAME_CHECKED(UMaterialExpressionOneMinus, Input).ToString(),
+					ExpressionFactoryNode.Get<0>()->GetUniqueID(), ExpressionFactoryNode.Get<1>());
+
+				MaterialFactoryNode->ConnectToCloth(InverseSheenRoughnessNode->GetUniqueID());
+			}
+
+			bShadingModelHandled = true;
+		}
+	}
+
+	if (bShadingModelHandled)
+	{
+		MaterialFactoryNode->SetCustomShadingModel(EMaterialShadingModel::MSM_Cloth);
+	}
+
+	return bShadingModelHandled;
+}
+
 bool UInterchangeGenericMaterialPipeline::HandleThinTranslucent(const UInterchangeShaderGraphNode* ShaderGraphNode, UInterchangeMaterialFactoryNode* MaterialFactoryNode)
 {
 	using namespace UE::Interchange::Materials::ThinTranslucent;
@@ -617,7 +685,43 @@ void UInterchangeGenericMaterialPipeline::HandleCommonParameters(const UIntercha
 	}
 }
 
-void UInterchangeGenericMaterialPipeline::HandleTextureSampleNode(const UInterchangeShaderNode* ShaderNode, UInterchangeMaterialExpressionFactoryNode* TextureSampleFactoryNode)
+void UInterchangeGenericMaterialPipeline::HandleFlattenNormalNode(const UInterchangeShaderNode* ShaderNode, UInterchangeMaterialFactoryNode* MaterialFactoryNode,
+	UInterchangeMaterialExpressionFactoryNode* FlattenNormalFactoryNode)
+{
+	using namespace UE::Interchange::Materials::Standard::Nodes::FlattenNormal;
+
+	FlattenNormalFactoryNode->SetCustomExpressionClassName(UMaterialExpressionMaterialFunctionCall::StaticClass()->GetName());
+
+	const FString MaterialFunctionMemberName = GET_MEMBER_NAME_CHECKED(UMaterialExpressionMaterialFunctionCall, MaterialFunction).ToString();
+	FlattenNormalFactoryNode->AddStringAttribute(MaterialFunctionMemberName, TEXT("/Engine/Functions/Engine_MaterialFunctions01/Texturing/FlattenNormal.FlattenNormal"));
+	FlattenNormalFactoryNode->AddApplyAndFillDelegates<FString>(MaterialFunctionMemberName, UMaterialExpressionMaterialFunctionCall::StaticClass(), *MaterialFunctionMemberName);
+
+	// Normal
+	{
+		TTuple<UInterchangeMaterialExpressionFactoryNode*, FString> NormalExpression =
+			CreateMaterialExpressionForInput(MaterialFactoryNode, ShaderNode, Inputs::Normal.ToString(), FlattenNormalFactoryNode->GetUniqueID());
+
+		if (NormalExpression.Get<0>())
+		{
+			UInterchangeShaderPortsAPI::ConnectOuputToInput(FlattenNormalFactoryNode, TEXT("Normal"),
+				NormalExpression.Get<0>()->GetUniqueID(), NormalExpression.Get<1>());
+		}
+	}
+
+	// Flatness
+	{
+		TTuple<UInterchangeMaterialExpressionFactoryNode*, FString> FlatnessExpression =
+			CreateMaterialExpressionForInput(MaterialFactoryNode, ShaderNode, Inputs::Flatness.ToString(), FlattenNormalFactoryNode->GetUniqueID());
+
+		if (FlatnessExpression.Get<0>())
+		{
+			UInterchangeShaderPortsAPI::ConnectOuputToInput(FlattenNormalFactoryNode, TEXT("Flatness"),
+				FlatnessExpression.Get<0>()->GetUniqueID(), FlatnessExpression.Get<1>());
+		}
+	}
+}
+
+void UInterchangeGenericMaterialPipeline::HandleTextureSampleNode(const UInterchangeShaderNode* ShaderNode, UInterchangeMaterialFactoryNode* MaterialFactoryNode, UInterchangeMaterialExpressionFactoryNode* TextureSampleFactoryNode)
 {
 	using namespace UE::Interchange::Materials::Standard::Nodes::TextureSample;
 
@@ -673,36 +777,127 @@ void UInterchangeGenericMaterialPipeline::HandleTextureSampleNode(const UInterch
 		}
 	}
 
-	HandleTextureCoordinates(ShaderNode, TextureSampleFactoryNode);
+	// Coordinates
+	{
+		TTuple<UInterchangeMaterialExpressionFactoryNode*, FString> CoordinatesExpression =
+			CreateMaterialExpressionForInput(MaterialFactoryNode, ShaderNode, Inputs::Coordinates.ToString(), TextureSampleFactoryNode->GetUniqueID());
+
+		if (CoordinatesExpression.Get<0>())
+		{
+			UInterchangeShaderPortsAPI::ConnectOuputToInput(TextureSampleFactoryNode, GET_MEMBER_NAME_CHECKED(UMaterialExpressionTextureSample, Coordinates).ToString(),
+				CoordinatesExpression.Get<0>()->GetUniqueID(), CoordinatesExpression.Get<1>());
+		}
+	}
 }
 
-void UInterchangeGenericMaterialPipeline::HandleTextureCoordinates(const UInterchangeShaderNode* ShaderNode, UInterchangeMaterialExpressionFactoryNode* TextureSampleFactoryNode)
+void UInterchangeGenericMaterialPipeline::HandleTextureCoordinateNode(const UInterchangeShaderNode* ShaderNode, UInterchangeMaterialFactoryNode* MaterialFactoryNode,
+	UInterchangeMaterialExpressionFactoryNode*& TexCoordFactoryNode)
 {
 	using namespace UE::Interchange::Materials::Standard;
 
-	float UTilingValue = 1.f;
-	bool bHasUTiling = ShaderNode->GetFloatAttribute(UInterchangeShaderPortsAPI::MakeInputValueKey(Nodes::TextureSample::Inputs::UTiling.ToString()), UTilingValue);
-	bHasUTiling = bHasUTiling && (UTilingValue != GetDefault<UMaterialExpressionTextureCoordinate>()->UTiling);
+	TexCoordFactoryNode->SetCustomExpressionClassName(UMaterialExpressionTextureCoordinate::StaticClass()->GetName());
 
-	float VTilingValue = 1.f;
-	bool bHasVTiling = ShaderNode->GetFloatAttribute(UInterchangeShaderPortsAPI::MakeInputValueKey(Nodes::TextureSample::Inputs::VTiling.ToString()), VTilingValue);
-	bHasVTiling = bHasVTiling && (VTilingValue != GetDefault<UMaterialExpressionTextureCoordinate>()->VTiling);
+	// Index
+	{
+		TTuple<UInterchangeMaterialExpressionFactoryNode*, FString> IndexExpression =
+			CreateMaterialExpressionForInput(MaterialFactoryNode, ShaderNode, Nodes::TextureCoordinate::Inputs::Index.ToString(), TexCoordFactoryNode->GetUniqueID());
 
-	if (bHasUTiling || bHasVTiling)
+		if (IndexExpression.Get<0>())
+		{
+			UInterchangeShaderPortsAPI::ConnectOuputToInput(TexCoordFactoryNode, GET_MEMBER_NAME_CHECKED(UMaterialExpressionTextureCoordinate, CoordinateIndex).ToString(),
+				IndexExpression.Get<0>()->GetUniqueID(), IndexExpression.Get<1>());
+		}
+	}
+
+	// U tiling
 	{
-		UInterchangeMaterialExpressionFactoryNode* TextureCoordinate = NewObject<UInterchangeMaterialExpressionFactoryNode>(BaseNodeContainer);
-		const FString TextureCoordinateUid = TextureSampleFactoryNode->GetUniqueID() + TEXT("Coordinate");
+		TTuple<UInterchangeMaterialExpressionFactoryNode*, FString> UTilingExpression =
+			CreateMaterialExpressionForInput(MaterialFactoryNode, ShaderNode, Nodes::TextureCoordinate::Inputs::UTiling.ToString(), TexCoordFactoryNode->GetUniqueID());
 
-		TextureCoordinate->InitializeNode(TextureCoordinateUid, ShaderNode->GetDisplayLabel(), EInterchangeNodeContainerType::FactoryData);
-		BaseNodeContainer->AddNode(TextureCoordinate);
-		BaseNodeContainer->SetNodeParentUid(TextureCoordinateUid, TextureSampleFactoryNode->GetUniqueID());
+		if (UTilingExpression.Get<0>())
+		{
+			UInterchangeShaderPortsAPI::ConnectOuputToInput(TexCoordFactoryNode, GET_MEMBER_NAME_CHECKED(UMaterialExpressionTextureCoordinate, UTiling).ToString(),
+				UTilingExpression.Get<0>()->GetUniqueID(), UTilingExpression.Get<1>());
+		}
+	}
 
-		TextureCoordinate->SetCustomExpressionClassName(UMaterialExpressionTextureCoordinate::StaticClass()->GetName());
+	// V tiling
+	{
+		TTuple<UInterchangeMaterialExpressionFactoryNode*, FString> VTilingExpression =
+			CreateMaterialExpressionForInput(MaterialFactoryNode, ShaderNode, Nodes::TextureCoordinate::Inputs::VTiling.ToString(), TexCoordFactoryNode->GetUniqueID());
 
-		TextureCoordinate->AddFloatAttribute(GET_MEMBER_NAME_CHECKED(UMaterialExpressionTextureCoordinate, UTiling).ToString(), UTilingValue);
-		TextureCoordinate->AddFloatAttribute(GET_MEMBER_NAME_CHECKED(UMaterialExpressionTextureCoordinate, VTiling).ToString(), VTilingValue);
+		if (VTilingExpression.Get<0>())
+		{
+			UInterchangeShaderPortsAPI::ConnectOuputToInput(TexCoordFactoryNode, GET_MEMBER_NAME_CHECKED(UMaterialExpressionTextureCoordinate, VTiling).ToString(),
+				VTilingExpression.Get<0>()->GetUniqueID(), VTilingExpression.Get<1>());
+		}
+	}
 
-		UInterchangeShaderPortsAPI::ConnectDefaultOuputToInput(TextureSampleFactoryNode, GET_MEMBER_NAME_CHECKED(UMaterialExpressionTextureSample, Coordinates).ToString(), TextureCoordinate->GetUniqueID());
+	// Scale
+	{
+		TTuple<UInterchangeMaterialExpressionFactoryNode*, FString> ScaleExpression =
+			CreateMaterialExpressionForInput(MaterialFactoryNode, ShaderNode, Nodes::TextureCoordinate::Inputs::Scale.ToString(), TexCoordFactoryNode->GetUniqueID());
+
+		if (ScaleExpression.Get<0>())
+		{
+			UInterchangeMaterialExpressionFactoryNode* MultiplyExpression =
+				CreateExpressionNode(ScaleExpression.Get<0>()->GetDisplayLabel() + TEXT("_Multiply"), TexCoordFactoryNode->GetUniqueID(), UMaterialExpressionMultiply::StaticClass());
+
+			UInterchangeShaderPortsAPI::ConnectDefaultOuputToInput(MultiplyExpression, GET_MEMBER_NAME_CHECKED(UMaterialExpressionMultiply, A).ToString(),
+				TexCoordFactoryNode->GetUniqueID());
+			UInterchangeShaderPortsAPI::ConnectOuputToInput(MultiplyExpression, GET_MEMBER_NAME_CHECKED(UMaterialExpressionMultiply, B).ToString(),
+				ScaleExpression.Get<0>()->GetUniqueID(), ScaleExpression.Get<1>());
+
+			TexCoordFactoryNode = MultiplyExpression;
+		}
+	}
+	
+	// Rotate
+	{
+		TTuple<UInterchangeMaterialExpressionFactoryNode*, FString> RotateExpression =
+			CreateMaterialExpressionForInput(MaterialFactoryNode, ShaderNode, Nodes::TextureCoordinate::Inputs::Rotate.ToString(), TexCoordFactoryNode->GetUniqueID());
+
+		if (RotateExpression.Get<0>())
+		{
+			UInterchangeMaterialExpressionFactoryNode* CallRotatorExpression =
+				CreateExpressionNode(RotateExpression.Get<0>()->GetDisplayLabel() + TEXT("_Rotator"), TexCoordFactoryNode->GetUniqueID(), UMaterialExpressionMaterialFunctionCall::StaticClass());
+
+			const FString MaterialFunctionMemberName = GET_MEMBER_NAME_CHECKED(UMaterialExpressionMaterialFunctionCall, MaterialFunction).ToString();
+			CallRotatorExpression->AddStringAttribute(MaterialFunctionMemberName, TEXT("/Engine/Functions/Engine_MaterialFunctions02/Texturing/CustomRotator.CustomRotator"));
+			CallRotatorExpression->AddApplyAndFillDelegates<FString>(MaterialFunctionMemberName, UMaterialExpressionMaterialFunctionCall::StaticClass(), *MaterialFunctionMemberName);
+
+			UInterchangeShaderPortsAPI::ConnectDefaultOuputToInput(CallRotatorExpression, TEXT("UVs"), TexCoordFactoryNode->GetUniqueID());
+			UInterchangeShaderPortsAPI::ConnectOuputToInput(CallRotatorExpression, TEXT("Rotation Angle (0-1)"), RotateExpression.Get<0>()->GetUniqueID(), RotateExpression.Get<1>());
+
+			TTuple<UInterchangeMaterialExpressionFactoryNode*, FString> RotationCenterExpression =
+				CreateMaterialExpressionForInput(MaterialFactoryNode, ShaderNode, Nodes::TextureCoordinate::Inputs::RotationCenter.ToString(), TexCoordFactoryNode->GetUniqueID());
+
+			if (RotationCenterExpression.Get<0>())
+			{
+				UInterchangeShaderPortsAPI::ConnectOuputToInput(CallRotatorExpression, TEXT("Rotation Center"), RotationCenterExpression.Get<0>()->GetUniqueID(), RotationCenterExpression.Get<1>());
+			}
+
+			TexCoordFactoryNode = CallRotatorExpression;
+		}
+	}
+
+	// Offset
+	{
+		TTuple<UInterchangeMaterialExpressionFactoryNode*, FString> OffsetExpression =
+			CreateMaterialExpressionForInput(MaterialFactoryNode, ShaderNode, Nodes::TextureCoordinate::Inputs::Offset.ToString(), TexCoordFactoryNode->GetUniqueID());
+
+		if (OffsetExpression.Get<0>())
+		{
+			UInterchangeMaterialExpressionFactoryNode* AddExpression =
+				CreateExpressionNode(OffsetExpression.Get<0>()->GetDisplayLabel() + TEXT("_Add"), TexCoordFactoryNode->GetUniqueID(), UMaterialExpressionAdd::StaticClass());
+
+			UInterchangeShaderPortsAPI::ConnectDefaultOuputToInput(AddExpression, GET_MEMBER_NAME_CHECKED(UMaterialExpressionAdd, A).ToString(),
+				TexCoordFactoryNode->GetUniqueID());
+			UInterchangeShaderPortsAPI::ConnectOuputToInput(AddExpression, GET_MEMBER_NAME_CHECKED(UMaterialExpressionAdd, B).ToString(),
+				OffsetExpression.Get<0>()->GetUniqueID(), OffsetExpression.Get<1>());
+
+			TexCoordFactoryNode = AddExpression;
+		}
 	}
 }
 
@@ -775,14 +970,22 @@ UInterchangeMaterialExpressionFactoryNode* UInterchangeGenericMaterialPipeline::
 	FString ShaderType;
 	ShaderNode->GetCustomShaderType(ShaderType);
 
-	if (*ShaderType == Nodes::TextureSample::Name)
+	if (*ShaderType == Nodes::FlattenNormal::Name)
 	{
-		HandleTextureSampleNode(ShaderNode, MaterialExpression);
+		HandleFlattenNormalNode(ShaderNode, MaterialFactoryNode, MaterialExpression);
 	}
 	else if (*ShaderType == Nodes::Lerp::Name)
 	{
 		HandleLerpNode(ShaderNode, MaterialFactoryNode, MaterialExpression);
 	}
+	else if (*ShaderType == Nodes::TextureCoordinate::Name)
+	{
+		HandleTextureCoordinateNode(ShaderNode, MaterialFactoryNode, MaterialExpression);
+	}
+	else if (*ShaderType == Nodes::TextureSample::Name)
+	{
+		HandleTextureSampleNode(ShaderNode, MaterialFactoryNode, MaterialExpression);
+	}
 	else
 	{
 		const FString ExpressionClassName = TEXT("MaterialExpression") + ShaderType;
@@ -884,6 +1087,29 @@ UInterchangeMaterialExpressionFactoryNode* UInterchangeGenericMaterialPipeline::
 	return MaterialExpressionFactoryNode;
 }
 
+UInterchangeMaterialExpressionFactoryNode* UInterchangeGenericMaterialPipeline::CreateVector2ParameterExpression(const UInterchangeShaderNode* ShaderNode, const FString& InputName, const FString& ParentUid)
+{
+	FVector2f InputValue;
+	if (ShaderNode->GetAttribute<FVector2f>(UInterchangeShaderPortsAPI::MakeInputValueKey(InputName), InputValue))
+	{
+		UInterchangeMaterialExpressionFactoryNode* VectorParameterFactoryNode = CreateExpressionNode(InputName, ParentUid, UMaterialExpressionVectorParameter::StaticClass());
+
+		const FName DefaultValueMemberName = GET_MEMBER_NAME_CHECKED(UMaterialExpressionVectorParameter, DefaultValue);
+		VectorParameterFactoryNode->AddLinearColorAttribute(DefaultValueMemberName.ToString(), FLinearColor(InputValue.X, InputValue.Y, 0.f));
+		VectorParameterFactoryNode->AddApplyAndFillDelegates<FLinearColor>(DefaultValueMemberName.ToString(), UMaterialExpressionVectorParameter::StaticClass(), DefaultValueMemberName);
+
+		// Defaults to R&G
+		UInterchangeMaterialExpressionFactoryNode* ComponentMaskFactoryNode = CreateExpressionNode(InputName + TEXT("_Mask"), ParentUid, UMaterialExpressionComponentMask::StaticClass());
+
+		UInterchangeShaderPortsAPI::ConnectDefaultOuputToInput(ComponentMaskFactoryNode, GET_MEMBER_NAME_CHECKED(UMaterialExpressionComponentMask, Input).ToString(),
+			VectorParameterFactoryNode->GetUniqueID() );
+
+		return ComponentMaskFactoryNode;
+	}
+
+	return nullptr;
+}
+
 TTuple<UInterchangeMaterialExpressionFactoryNode*, FString> UInterchangeGenericMaterialPipeline::CreateMaterialExpressionForInput(UInterchangeMaterialFactoryNode* MaterialFactoryNode, const UInterchangeShaderNode* ShaderNode, const FString& InputName, const FString& ParentUid)
 {
 	// Make sure we don't create an expression for an input if it already has one
@@ -919,6 +1145,9 @@ TTuple<UInterchangeMaterialExpressionFactoryNode*, FString> UInterchangeGenericM
 		case UE::Interchange::EAttributeTypes::LinearColor:
 			MaterialExpressionFactoryNode = CreateVectorParameterExpression(ShaderNode, InputName, ParentUid);
 			break;
+		case UE::Interchange::EAttributeTypes::Vector2f:
+			MaterialExpressionFactoryNode = CreateVector2ParameterExpression(ShaderNode, InputName, ParentUid);
+			break;
 		}
 	}
 
@@ -938,8 +1167,11 @@ UInterchangeMaterialFactoryNode* UInterchangeGenericMaterialPipeline::CreateMate
 	
 	if (!HandleClearCoat(ShaderGraphNode, MaterialFactoryNode))
 	{
-		// Can't have both clear coat and thin translucent since they are different shading models
-		HandleThinTranslucent(ShaderGraphNode, MaterialFactoryNode);
+		// Can't have different shading models
+		if (!HandleThinTranslucent(ShaderGraphNode, MaterialFactoryNode))
+		{
+			HandleSheen(ShaderGraphNode, MaterialFactoryNode);
+		}
 	}
 
 	HandleCommonParameters(ShaderGraphNode, MaterialFactoryNode);
@@ -964,6 +1196,10 @@ UInterchangeMaterialInstanceFactoryNode* UInterchangeGenericMaterialPipeline::Cr
 	{
 		MaterialInstanceFactoryNode->SetCustomParent(TEXT("Material'/Interchange/Materials/ClearCoatMaterial.ClearCoatMaterial'"));
 	}
+	else if (IsSheenModel(ShaderGraphNode))
+	{
+		MaterialInstanceFactoryNode->SetCustomParent(TEXT("Material'/Interchange/Materials/SheenMaterial.SheenMaterial'"));
+	}
 	else if (IsPBRModel(ShaderGraphNode))
 	{
 		MaterialInstanceFactoryNode->SetCustomParent(TEXT("Material'/Interchange/Materials/PBRSurfaceMaterial.PBRSurfaceMaterial'"));
diff --git a/Engine/Plugins/Experimental/Interchange/Runtime/Source/Pipelines/Public/InterchangeGenericMaterialPipeline.h b/Engine/Plugins/Experimental/Interchange/Runtime/Source/Pipelines/Public/InterchangeGenericMaterialPipeline.h
--- a/Engine/Plugins/Experimental/Interchange/Runtime/Source/Pipelines/Public/InterchangeGenericMaterialPipeline.h
+++ b/Engine/Plugins/Experimental/Interchange/Runtime/Source/Pipelines/Public/InterchangeGenericMaterialPipeline.h
@@ -63,6 +63,9 @@ private:
 	/** True if the shader graph has a clear coat input. */
 	bool IsClearCoatModel(const UInterchangeShaderGraphNode* ShaderGraphNode) const;
 
+	/** True if the shader graph has a sheen color input. */
+	bool IsSheenModel(const UInterchangeShaderGraphNode* ShaderGraphNode) const;
+
 	/** True if the shader graph has a transmission color input. */
 	bool IsThinTranslucentModel(const UInterchangeShaderGraphNode* ShaderGraphNode) const;
 
@@ -79,10 +82,13 @@ private:
 	bool HandleLambertModel(const UInterchangeShaderGraphNode* ShaderGraphNode, UInterchangeMaterialFactoryNode* MaterialFactoryNode);
 	bool HandlePBRModel(const UInterchangeShaderGraphNode* ShaderGraphNode, UInterchangeMaterialFactoryNode* MaterialFactoryNode);
 	bool HandleClearCoat(const UInterchangeShaderGraphNode* ShaderGraphNode, UInterchangeMaterialFactoryNode* MaterialFactoryNode);
+	bool HandleSheen(const UInterchangeShaderGraphNode* ShaderGraphNode, UInterchangeMaterialFactoryNode* MaterialFactoryNode);
 	bool HandleThinTranslucent(const UInterchangeShaderGraphNode* ShaderGraphNode, UInterchangeMaterialFactoryNode* MaterialFactoryNode);
 	void HandleCommonParameters(const UInterchangeShaderGraphNode* ShaderGraphNode, UInterchangeMaterialFactoryNode* MaterialFactoryNode);
-	void HandleTextureSampleNode(const UInterchangeShaderNode* ShaderNode, UInterchangeMaterialExpressionFactoryNode* TextureSampleFactoryNode);
-	void HandleTextureCoordinates(const UInterchangeShaderNode* ShaderNode, UInterchangeMaterialExpressionFactoryNode* TextureSampleFactoryNode);
+
+	void HandleFlattenNormalNode(const UInterchangeShaderNode* ShaderNode, UInterchangeMaterialFactoryNode* MaterialFactoryNode, UInterchangeMaterialExpressionFactoryNode* FlattenNormalFactoryNode);
+	void HandleTextureSampleNode(const UInterchangeShaderNode* ShaderNode, UInterchangeMaterialFactoryNode* MaterialFactoryNode, UInterchangeMaterialExpressionFactoryNode* TextureSampleFactoryNode);
+	void HandleTextureCoordinateNode(const UInterchangeShaderNode* ShaderNode, UInterchangeMaterialFactoryNode* MaterialFactoryNode, UInterchangeMaterialExpressionFactoryNode*& TextureSampleFactoryNode);
 	void HandleLerpNode(const UInterchangeShaderNode* ShaderNode, UInterchangeMaterialFactoryNode* MaterialFactoryNode, UInterchangeMaterialExpressionFactoryNode* LerpFactoryNode);
 
 	UInterchangeMaterialExpressionFactoryNode* CreateMaterialExpressionForShaderNode(UInterchangeMaterialFactoryNode* MaterialFactoryNode, const UInterchangeShaderNode* ShaderNode, const FString& ParentUid);
@@ -91,6 +97,7 @@ private:
 	UInterchangeMaterialExpressionFactoryNode* CreateExpressionNode(const FString& ExpressionName, const FString& ParentUid, UClass* MaterialExpressionClass);
 	UInterchangeMaterialExpressionFactoryNode* CreateScalarParameterExpression(const UInterchangeShaderNode* ShaderNode, const FString& InputName, const FString& ParentUid);
 	UInterchangeMaterialExpressionFactoryNode* CreateVectorParameterExpression(const UInterchangeShaderNode* ShaderNode, const FString& InputName, const FString& ParentUid);
+	UInterchangeMaterialExpressionFactoryNode* CreateVector2ParameterExpression(const UInterchangeShaderNode* ShaderNode, const FString& InputName, const FString& ParentUid);
 
 	/**
 	 * Visits a given shader node and its connections to find its strongest value.
