Interchange Materials - Fixes for texture coordinates and clearcoat normal. Texture coordinates have inputs that can't handle connections in UE so we can only accept plain values.

diff --git a/Engine/Plugins/Experimental/Interchange/Runtime/Source/Import/Private/Gltf/InterchangeGltfTranslator.cpp b/Engine/Plugins/Experimental/Interchange/Runtime/Source/Import/Private/Gltf/InterchangeGltfTranslator.cpp
--- a/Engine/Plugins/Experimental/Interchange/Runtime/Source/Import/Private/Gltf/InterchangeGltfTranslator.cpp
+++ b/Engine/Plugins/Experimental/Interchange/Runtime/Source/Import/Private/Gltf/InterchangeGltfTranslator.cpp
@@ -411,7 +411,7 @@ void UInterchangeGltfTranslator::HandleGltfClearCoat( UInterchangeBaseNodeContai
 	// ClearCoat::Parameters::ClearCoatNormal
 	{
 		TVariant< FLinearColor, float > ClearCoatNormalFactor;
-		ClearCoatNormalFactor.Set< FLinearColor >( FLinearColor( FVector::UpVector ) );
+		ClearCoatNormalFactor.Set< float >( GltfMaterial.ClearCoat.NormalMapUVScale );
 
 		const bool bInverse = false;
 		const bool bIsNormal = true;
diff --git a/Engine/Plugins/Experimental/Interchange/Runtime/Source/Nodes/Public/InterchangeMaterialDefinitions.h b/Engine/Plugins/Experimental/Interchange/Runtime/Source/Nodes/Public/InterchangeMaterialDefinitions.h
--- a/Engine/Plugins/Experimental/Interchange/Runtime/Source/Nodes/Public/InterchangeMaterialDefinitions.h
+++ b/Engine/Plugins/Experimental/Interchange/Runtime/Source/Nodes/Public/InterchangeMaterialDefinitions.h
@@ -89,7 +89,7 @@ namespace Interchange
 
 					namespace Inputs
 					{
-						const FName Coordinates = TEXT("Coordinates");
+						const FName Coordinates = TEXT("Coordinates"); // Type: vec2
 						const FName Texture = TEXT("TextureUid"); // Type: FString (unique id of a texture node)
 					}
 
diff --git a/Engine/Plugins/Experimental/Interchange/Runtime/Source/Pipelines/Private/InterchangeGenericMaterialPipeline.cpp b/Engine/Plugins/Experimental/Interchange/Runtime/Source/Pipelines/Private/InterchangeGenericMaterialPipeline.cpp
--- a/Engine/Plugins/Experimental/Interchange/Runtime/Source/Pipelines/Private/InterchangeGenericMaterialPipeline.cpp
+++ b/Engine/Plugins/Experimental/Interchange/Runtime/Source/Pipelines/Private/InterchangeGenericMaterialPipeline.cpp
@@ -799,37 +799,36 @@ void UInterchangeGenericMaterialPipeline::HandleTextureCoordinateNode(const UInt
 
 	// Index
 	{
-		TTuple<UInterchangeMaterialExpressionFactoryNode*, FString> IndexExpression =
-			CreateMaterialExpressionForInput(MaterialFactoryNode, ShaderNode, Nodes::TextureCoordinate::Inputs::Index.ToString(), TexCoordFactoryNode->GetUniqueID());
-
-		if (IndexExpression.Get<0>())
+		int32 CoordIndex;
+		if (ShaderNode->GetInt32Attribute(UInterchangeShaderPortsAPI::MakeInputValueKey(Nodes::TextureCoordinate::Inputs::Index.ToString()), CoordIndex))
 		{
-			UInterchangeShaderPortsAPI::ConnectOuputToInput(TexCoordFactoryNode, GET_MEMBER_NAME_CHECKED(UMaterialExpressionTextureCoordinate, CoordinateIndex).ToString(),
-				IndexExpression.Get<0>()->GetUniqueID(), IndexExpression.Get<1>());
+			const FString CoordinateIndexMemberName = GET_MEMBER_NAME_CHECKED(UMaterialExpressionTextureCoordinate, CoordinateIndex).ToString();
+			TexCoordFactoryNode->AddInt32Attribute(CoordinateIndexMemberName, CoordIndex);
+			TexCoordFactoryNode->AddApplyAndFillDelegates<int32>(CoordinateIndexMemberName, UMaterialExpressionTextureCoordinate::StaticClass(), *CoordinateIndexMemberName);
 		}
 	}
 
 	// U tiling
 	{
-		TTuple<UInterchangeMaterialExpressionFactoryNode*, FString> UTilingExpression =
-			CreateMaterialExpressionForInput(MaterialFactoryNode, ShaderNode, Nodes::TextureCoordinate::Inputs::UTiling.ToString(), TexCoordFactoryNode->GetUniqueID());
+		TVariant<FString, FLinearColor, float> UTilingValue = VisitShaderInput(ShaderNode, Nodes::TextureCoordinate::Inputs::UTiling.ToString());
 
-		if (UTilingExpression.Get<0>())
+		if (UTilingValue.IsType<float>())
 		{
-			UInterchangeShaderPortsAPI::ConnectOuputToInput(TexCoordFactoryNode, GET_MEMBER_NAME_CHECKED(UMaterialExpressionTextureCoordinate, UTiling).ToString(),
-				UTilingExpression.Get<0>()->GetUniqueID(), UTilingExpression.Get<1>());
+			const FString UTilingMemberName = GET_MEMBER_NAME_CHECKED(UMaterialExpressionTextureCoordinate, UTiling).ToString();
+			TexCoordFactoryNode->AddFloatAttribute(UTilingMemberName, UTilingValue.Get<float>());
+			TexCoordFactoryNode->AddApplyAndFillDelegates<float>(UTilingMemberName, UMaterialExpressionTextureCoordinate::StaticClass(), *UTilingMemberName);
 		}
 	}
 
 	// V tiling
 	{
-		TTuple<UInterchangeMaterialExpressionFactoryNode*, FString> VTilingExpression =
-			CreateMaterialExpressionForInput(MaterialFactoryNode, ShaderNode, Nodes::TextureCoordinate::Inputs::VTiling.ToString(), TexCoordFactoryNode->GetUniqueID());
+		TVariant<FString, FLinearColor, float> VTilingValue = VisitShaderInput(ShaderNode, Nodes::TextureCoordinate::Inputs::VTiling.ToString());
 
-		if (VTilingExpression.Get<0>())
+		if (VTilingValue.IsType<float>())
 		{
-			UInterchangeShaderPortsAPI::ConnectOuputToInput(TexCoordFactoryNode, GET_MEMBER_NAME_CHECKED(UMaterialExpressionTextureCoordinate, VTiling).ToString(),
-				VTilingExpression.Get<0>()->GetUniqueID(), VTilingExpression.Get<1>());
+			const FString VTilingMemberName = GET_MEMBER_NAME_CHECKED(UMaterialExpressionTextureCoordinate, VTiling).ToString();
+			TexCoordFactoryNode->AddFloatAttribute(VTilingMemberName, VTilingValue.Get<float>());
+			TexCoordFactoryNode->AddApplyAndFillDelegates<float>(VTilingMemberName, UMaterialExpressionTextureCoordinate::StaticClass(), *VTilingMemberName);
 		}
 	}
 
