Added an ensure to warn the users of interchange if they create a subclass of an interchange texture factory node to expose their custom pipeline attributes.

diff --git a/Engine/Plugins/Experimental/Interchange/Runtime/Source/Import/Private/Texture/InterchangeTextureFactory.cpp b/Engine/Plugins/Experimental/Interchange/Runtime/Source/Import/Private/Texture/InterchangeTextureFactory.cpp
--- a/Engine/Plugins/Experimental/Interchange/Runtime/Source/Import/Private/Texture/InterchangeTextureFactory.cpp
+++ b/Engine/Plugins/Experimental/Interchange/Runtime/Source/Import/Private/Texture/InterchangeTextureFactory.cpp
@@ -200,6 +200,25 @@ namespace UE::Interchange::Private::InterchangeTextureFactory
 			}
 		}
 
+		ensureMsgf(false
+			, TEXT("Unknow factory node class (%s). To add support for a new texture type either update this factory or register a new factory to interchange.\n\
+				If the factory node class is used to add functions for some custom attributes of a pipeline please don't do that.\n\
+				Use some static functions libraries or use some helper structs/objects instead.\n\
+				Here is exemple of what this could look like in c++:\n\n\
+				\tif (UCustomAttributeInterface::HasInterface(TextureFactoryNode))\n\
+				\t{\n\
+					\t\tFString MyPath = UCustomAttributeInterface::GetCustomPath(TextureFactoryNode);\n\
+					\t\t...\n\
+				\t}\n\
+				or\n\
+				\tif (UCustomAttributeInterface* CustomInterface = UCustomAttributeInterface::GetInterface(TextureFactoryNode))\n\
+				\t{\n\
+					\t\tFString MyPath = CustomInterface->GetPath();\n\
+					\t\t...\n\
+				\t}\n\n")
+			, *AssetNode->GetClass()->GetName()
+			);
+
 		return {};
 	}
 
