Add support for UMaterialExpressionComponentMask in GenericMaterialPipeline.
By default the channels are disabled, the user can decide to enable them by adding a Boolean attribute to the ShaderNode.

diff --git a/Engine/Plugins/Experimental/Interchange/Runtime/Source/Nodes/Public/InterchangeMaterialDefinitions.h b/Engine/Plugins/Experimental/Interchange/Runtime/Source/Nodes/Public/InterchangeMaterialDefinitions.h
--- a/Engine/Plugins/Experimental/Interchange/Runtime/Source/Nodes/Public/InterchangeMaterialDefinitions.h
+++ b/Engine/Plugins/Experimental/Interchange/Runtime/Source/Nodes/Public/InterchangeMaterialDefinitions.h
@@ -46,6 +46,24 @@ namespace Interchange
 					}
 				}
 
+				namespace Mask
+				{
+					const FName Name = TEXT("ComponentMask");
+
+					namespace Attributes
+					{
+						const FName R = TEXT("R");
+						const FName G = TEXT("G");
+						const FName B = TEXT("B");
+						const FName A = TEXT("A");
+					}
+
+					namespace Inputs
+					{
+						const FName Input = TEXT("Input");
+					}
+				}
+
 				namespace Multiply
 				{
 					const FName Name = TEXT("Multiply");
diff --git a/Engine/Plugins/Experimental/Interchange/Runtime/Source/Pipelines/Private/InterchangeGenericMaterialPipeline.cpp b/Engine/Plugins/Experimental/Interchange/Runtime/Source/Pipelines/Private/InterchangeGenericMaterialPipeline.cpp
--- a/Engine/Plugins/Experimental/Interchange/Runtime/Source/Pipelines/Private/InterchangeGenericMaterialPipeline.cpp
+++ b/Engine/Plugins/Experimental/Interchange/Runtime/Source/Pipelines/Private/InterchangeGenericMaterialPipeline.cpp
@@ -943,6 +943,66 @@ void UInterchangeGenericMaterialPipeline::HandleLerpNode(const UInterchangeShade
 	}
 }
 
+void UInterchangeGenericMaterialPipeline::HandleMaskNode(const UInterchangeShaderNode* ShaderNode, UInterchangeMaterialFactoryNode* MaterialFactoryNode, UInterchangeMaterialExpressionFactoryNode* MaskFactoryNode)
+{
+	using namespace UE::Interchange::Materials::Standard::Nodes;
+
+	MaskFactoryNode->SetCustomExpressionClassName(UMaterialExpressionComponentMask::StaticClass()->GetName());
+
+	bool bRChannel = false;
+	bool bGChannel = false;
+	bool bBChannel = false;
+	bool bAChannel = false;
+	bool bIsAnyMaskChannelSet = 
+		ShaderNode->GetBooleanAttribute(Mask::Attributes::R.ToString(), bRChannel) ||
+		ShaderNode->GetBooleanAttribute(Mask::Attributes::G.ToString(), bGChannel) ||
+		ShaderNode->GetBooleanAttribute(Mask::Attributes::B.ToString(), bBChannel) ||
+		ShaderNode->GetBooleanAttribute(Mask::Attributes::A.ToString(), bAChannel);
+
+	if(bIsAnyMaskChannelSet)
+	{
+		// R
+		{
+			const FString RMemberName = GET_MEMBER_NAME_CHECKED(UMaterialExpressionComponentMask, R).ToString();
+			MaskFactoryNode->AddBooleanAttribute(RMemberName, bRChannel);
+			MaskFactoryNode->AddApplyAndFillDelegates<bool>(RMemberName, UMaterialExpressionComponentMask::StaticClass(), *RMemberName);
+		}
+
+		// G
+		{
+			const FString GMemberName = GET_MEMBER_NAME_CHECKED(UMaterialExpressionComponentMask, G).ToString();
+			MaskFactoryNode->AddBooleanAttribute(GMemberName, bGChannel);
+			MaskFactoryNode->AddApplyAndFillDelegates<bool>(GMemberName, UMaterialExpressionComponentMask::StaticClass(), *GMemberName);
+		}
+
+		// B
+		{
+			const FString BMemberName = GET_MEMBER_NAME_CHECKED(UMaterialExpressionComponentMask, B).ToString();
+			MaskFactoryNode->AddBooleanAttribute(BMemberName, bBChannel);
+			MaskFactoryNode->AddApplyAndFillDelegates<bool>(BMemberName, UMaterialExpressionComponentMask::StaticClass(), *BMemberName);
+		}
+
+		// A
+		{
+			const FString AMemberName = GET_MEMBER_NAME_CHECKED(UMaterialExpressionComponentMask, A).ToString();
+			MaskFactoryNode->AddBooleanAttribute(AMemberName, bAChannel);
+			MaskFactoryNode->AddApplyAndFillDelegates<bool>(AMemberName, UMaterialExpressionComponentMask::StaticClass(), *AMemberName);
+		}
+	}
+
+	// Input
+	{
+		TTuple<UInterchangeMaterialExpressionFactoryNode*, FString> InputExpression =
+			CreateMaterialExpressionForInput(MaterialFactoryNode, ShaderNode, Mask::Inputs::Input.ToString(), MaskFactoryNode->GetUniqueID());
+
+		if(InputExpression.Get<0>())
+		{
+			UInterchangeShaderPortsAPI::ConnectOuputToInput(MaskFactoryNode, GET_MEMBER_NAME_CHECKED(UMaterialExpressionComponentMask, Input).ToString(),
+				InputExpression.Get<0>()->GetUniqueID(), InputExpression.Get<1>());
+		}
+	}
+}
+
 UInterchangeMaterialExpressionFactoryNode* UInterchangeGenericMaterialPipeline::CreateMaterialExpressionForShaderNode(UInterchangeMaterialFactoryNode* MaterialFactoryNode,
 	const UInterchangeShaderNode* ShaderNode, const FString& ParentUid)
 {
@@ -977,6 +1037,10 @@ UInterchangeMaterialExpressionFactoryNode* UInterchangeGenericMaterialPipeline::
 	{
 		HandleLerpNode(ShaderNode, MaterialFactoryNode, MaterialExpression);
 	}
+	else if(*ShaderType == Nodes::Mask::Name)
+	{
+		HandleMaskNode(ShaderNode, MaterialFactoryNode, MaterialExpression);
+	}
 	else if (*ShaderType == Nodes::TextureCoordinate::Name)
 	{
 		HandleTextureCoordinateNode(ShaderNode, MaterialFactoryNode, MaterialExpression);
diff --git a/Engine/Plugins/Experimental/Interchange/Runtime/Source/Pipelines/Public/InterchangeGenericMaterialPipeline.h b/Engine/Plugins/Experimental/Interchange/Runtime/Source/Pipelines/Public/InterchangeGenericMaterialPipeline.h
--- a/Engine/Plugins/Experimental/Interchange/Runtime/Source/Pipelines/Public/InterchangeGenericMaterialPipeline.h
+++ b/Engine/Plugins/Experimental/Interchange/Runtime/Source/Pipelines/Public/InterchangeGenericMaterialPipeline.h
@@ -90,6 +90,7 @@ private:
 	void HandleTextureSampleNode(const UInterchangeShaderNode* ShaderNode, UInterchangeMaterialFactoryNode* MaterialFactoryNode, UInterchangeMaterialExpressionFactoryNode* TextureSampleFactoryNode);
 	void HandleTextureCoordinateNode(const UInterchangeShaderNode* ShaderNode, UInterchangeMaterialFactoryNode* MaterialFactoryNode, UInterchangeMaterialExpressionFactoryNode*& TextureSampleFactoryNode);
 	void HandleLerpNode(const UInterchangeShaderNode* ShaderNode, UInterchangeMaterialFactoryNode* MaterialFactoryNode, UInterchangeMaterialExpressionFactoryNode* LerpFactoryNode);
+	void HandleMaskNode(const UInterchangeShaderNode* ShaderNode, UInterchangeMaterialFactoryNode* MaterialFactoryNode, UInterchangeMaterialExpressionFactoryNode* MaskFactoryNode);
 
 	UInterchangeMaterialExpressionFactoryNode* CreateMaterialExpressionForShaderNode(UInterchangeMaterialFactoryNode* MaterialFactoryNode, const UInterchangeShaderNode* ShaderNode, const FString& ParentUid);
 	TTuple<UInterchangeMaterialExpressionFactoryNode*, FString> CreateMaterialExpressionForInput(UInterchangeMaterialFactoryNode* MaterialFactoryNode, const UInterchangeShaderNode* ShaderNode, const FString& InputName, const FString& ParentUid);
