Support for collision in static mesh import.
- Added pipeline support for parsing meshes named with certain prefixes as collision primitives applied to a corresponding unprefixed mesh.
- Added build support for the four collision primitive types, compatible with async / interchange.
- Meshes imported without collision will be assigned a standard kdop-18 convex primitive.

diff --git a/Engine/Plugins/Experimental/Interchange/Runtime/Source/FactoryNodes/Private/InterchangeStaticMeshLodDataNode.cpp b/Engine/Plugins/Experimental/Interchange/Runtime/Source/FactoryNodes/Private/InterchangeStaticMeshLodDataNode.cpp
--- a/Engine/Plugins/Experimental/Interchange/Runtime/Source/FactoryNodes/Private/InterchangeStaticMeshLodDataNode.cpp
+++ b/Engine/Plugins/Experimental/Interchange/Runtime/Source/FactoryNodes/Private/InterchangeStaticMeshLodDataNode.cpp
@@ -11,6 +11,31 @@ namespace UE
 			static FString MeshUids_BaseKey = TEXT("__MeshUids__Key");
 			return MeshUids_BaseKey;
 		}
+
+		const FString& FStaticMeshNodeLodDataStaticData::GetBoxCollisionMeshUidsBaseKey()
+		{
+			static FString CollisionMeshUids_BaseKey = TEXT("__BoxCollisionMeshUids__Key");
+			return CollisionMeshUids_BaseKey;
+		}
+
+		const FString& FStaticMeshNodeLodDataStaticData::GetCapsuleCollisionMeshUidsBaseKey()
+		{
+			static FString CollisionMeshUids_BaseKey = TEXT("__CapsuleCollisionMeshUids__Key");
+			return CollisionMeshUids_BaseKey;
+		}
+
+		const FString& FStaticMeshNodeLodDataStaticData::GetSphereCollisionMeshUidsBaseKey()
+		{
+			static FString CollisionMeshUids_BaseKey = TEXT("__SphereCollisionMeshUids__Key");
+			return CollisionMeshUids_BaseKey;
+		}
+
+		const FString& FStaticMeshNodeLodDataStaticData::GetConvexCollisionMeshUidsBaseKey()
+		{
+			static FString CollisionMeshUids_BaseKey = TEXT("__ConvexCollisionMeshUids__Key");
+			return CollisionMeshUids_BaseKey;
+		}
+
 	} // namespace Interchange
 } // namespace UE
 
@@ -18,11 +43,15 @@ namespace UE
 UInterchangeStaticMeshLodDataNode::UInterchangeStaticMeshLodDataNode()
 {
 	MeshUids.Initialize(Attributes, UE::Interchange::FStaticMeshNodeLodDataStaticData::GetMeshUidsBaseKey());
+	BoxCollisionMeshUids.Initialize(Attributes, UE::Interchange::FStaticMeshNodeLodDataStaticData::GetBoxCollisionMeshUidsBaseKey());
+	CapsuleCollisionMeshUids.Initialize(Attributes, UE::Interchange::FStaticMeshNodeLodDataStaticData::GetCapsuleCollisionMeshUidsBaseKey());
+	SphereCollisionMeshUids.Initialize(Attributes, UE::Interchange::FStaticMeshNodeLodDataStaticData::GetSphereCollisionMeshUidsBaseKey());
+	ConvexCollisionMeshUids.Initialize(Attributes, UE::Interchange::FStaticMeshNodeLodDataStaticData::GetConvexCollisionMeshUidsBaseKey());
 }
 
 /**
-	* Return the node type name of the class, we use this when reporting error
-	*/
+ * Return the node type name of the class, we use this when reporting error
+ */
 FString UInterchangeStaticMeshLodDataNode::GetTypeName() const
 {
 	const FString TypeName = TEXT("StaticMeshLodDataNode");
@@ -86,6 +115,116 @@ bool UInterchangeStaticMeshLodDataNode::RemoveAllMeshes()
 	return MeshUids.RemoveAllItems();
 }
 
+int32 UInterchangeStaticMeshLodDataNode::GetBoxCollisionMeshUidsCount() const
+{
+	return BoxCollisionMeshUids.GetCount();
+}
+
+void UInterchangeStaticMeshLodDataNode::GetBoxCollisionMeshUids(TArray<FString>& OutMeshNames) const
+{
+	BoxCollisionMeshUids.GetItems(OutMeshNames);
+}
+
+bool UInterchangeStaticMeshLodDataNode::AddBoxCollisionMeshUid(const FString& MeshName)
+{
+	return BoxCollisionMeshUids.AddItem(MeshName);
+}
+
+bool UInterchangeStaticMeshLodDataNode::RemoveBoxCollisionMeshUid(const FString& MeshName)
+{
+	return BoxCollisionMeshUids.RemoveItem(MeshName);
+}
+
+bool UInterchangeStaticMeshLodDataNode::RemoveAllBoxCollisionMeshes()
+{
+	return BoxCollisionMeshUids.RemoveAllItems();
+}
+
+int32 UInterchangeStaticMeshLodDataNode::GetCapsuleCollisionMeshUidsCount() const
+{
+	return CapsuleCollisionMeshUids.GetCount();
+}
+
+void UInterchangeStaticMeshLodDataNode::GetCapsuleCollisionMeshUids(TArray<FString>& OutMeshNames) const
+{
+	CapsuleCollisionMeshUids.GetItems(OutMeshNames);
+}
+
+bool UInterchangeStaticMeshLodDataNode::AddCapsuleCollisionMeshUid(const FString& MeshName)
+{
+	return CapsuleCollisionMeshUids.AddItem(MeshName);
+}
+
+bool UInterchangeStaticMeshLodDataNode::RemoveCapsuleCollisionMeshUid(const FString& MeshName)
+{
+	return CapsuleCollisionMeshUids.RemoveItem(MeshName);
+}
+
+bool UInterchangeStaticMeshLodDataNode::RemoveAllCapsuleCollisionMeshes()
+{
+	return CapsuleCollisionMeshUids.RemoveAllItems();
+}
+
+int32 UInterchangeStaticMeshLodDataNode::GetSphereCollisionMeshUidsCount() const
+{
+	return SphereCollisionMeshUids.GetCount();
+}
+
+void UInterchangeStaticMeshLodDataNode::GetSphereCollisionMeshUids(TArray<FString>& OutMeshNames) const
+{
+	SphereCollisionMeshUids.GetItems(OutMeshNames);
+}
+
+bool UInterchangeStaticMeshLodDataNode::AddSphereCollisionMeshUid(const FString& MeshName)
+{
+	return SphereCollisionMeshUids.AddItem(MeshName);
+}
+
+bool UInterchangeStaticMeshLodDataNode::RemoveSphereCollisionMeshUid(const FString& MeshName)
+{
+	return SphereCollisionMeshUids.RemoveItem(MeshName);
+}
+
+bool UInterchangeStaticMeshLodDataNode::RemoveAllSphereCollisionMeshes()
+{
+	return SphereCollisionMeshUids.RemoveAllItems();
+}
+
+int32 UInterchangeStaticMeshLodDataNode::GetConvexCollisionMeshUidsCount() const
+{
+	return ConvexCollisionMeshUids.GetCount();
+}
+
+void UInterchangeStaticMeshLodDataNode::GetConvexCollisionMeshUids(TArray<FString>& OutMeshNames) const
+{
+	ConvexCollisionMeshUids.GetItems(OutMeshNames);
+}
+
+bool UInterchangeStaticMeshLodDataNode::AddConvexCollisionMeshUid(const FString& MeshName)
+{
+	return ConvexCollisionMeshUids.AddItem(MeshName);
+}
+
+bool UInterchangeStaticMeshLodDataNode::RemoveConvexCollisionMeshUid(const FString& MeshName)
+{
+	return ConvexCollisionMeshUids.RemoveItem(MeshName);
+}
+
+bool UInterchangeStaticMeshLodDataNode::RemoveAllConvexCollisionMeshes()
+{
+	return ConvexCollisionMeshUids.RemoveAllItems();
+}
+
+bool UInterchangeStaticMeshLodDataNode::GetOneConvexHullPerUCX(bool& AttributeValue) const
+{
+	IMPLEMENT_NODE_ATTRIBUTE_GETTER(OneConvexHullPerUCX, bool);
+}
+
+bool UInterchangeStaticMeshLodDataNode::SetOneConvexHullPerUCX(bool AttributeValue)
+{
+	IMPLEMENT_NODE_ATTRIBUTE_SETTER_NODELEGATE(OneConvexHullPerUCX, bool)
+}
+
 bool UInterchangeStaticMeshLodDataNode::IsEditorOnlyDataDefined()
 {
 #if WITH_EDITORONLY_DATA
diff --git a/Engine/Plugins/Experimental/Interchange/Runtime/Source/FactoryNodes/Public/InterchangeStaticMeshLodDataNode.h b/Engine/Plugins/Experimental/Interchange/Runtime/Source/FactoryNodes/Public/InterchangeStaticMeshLodDataNode.h
--- a/Engine/Plugins/Experimental/Interchange/Runtime/Source/FactoryNodes/Public/InterchangeStaticMeshLodDataNode.h
+++ b/Engine/Plugins/Experimental/Interchange/Runtime/Source/FactoryNodes/Public/InterchangeStaticMeshLodDataNode.h
@@ -3,6 +3,7 @@
 #pragma once
 
 #include "CoreMinimal.h"
+#include "UObject/ObjectMacros.h"
 #include "Nodes/InterchangeBaseNode.h"
 
 #include "InterchangeStaticMeshLodDataNode.generated.h"
@@ -15,6 +16,10 @@ namespace UE
 		struct FStaticMeshNodeLodDataStaticData : public FBaseNodeStaticData
 		{
 			static const FString& GetMeshUidsBaseKey();
+			static const FString& GetBoxCollisionMeshUidsBaseKey();
+			static const FString& GetCapsuleCollisionMeshUidsBaseKey();
+			static const FString& GetSphereCollisionMeshUidsBaseKey();
+			static const FString& GetConvexCollisionMeshUidsBaseKey();
 		};
 	} // namespace Interchange
 } // namespace UE
@@ -55,10 +60,82 @@ public:
 	UFUNCTION(BlueprintCallable, Category = "Interchange | Node | StaticMeshLodData")
 	bool RemoveAllMeshes();
 
+	UFUNCTION(BlueprintCallable, Category = "Interchange | Node | StaticMeshLodData")
+	int32 GetBoxCollisionMeshUidsCount() const;
+
+	UFUNCTION(BlueprintCallable, Category = "Interchange | Node | StaticMeshLodData")
+	void GetBoxCollisionMeshUids(TArray<FString>& OutMeshNames) const;
+
+	UFUNCTION(BlueprintCallable, Category = "Interchange | Node | StaticMeshLodData")
+	bool AddBoxCollisionMeshUid(const FString& MeshName);
+
+	UFUNCTION(BlueprintCallable, Category = "Interchange | Node | StaticMeshLodData")
+	bool RemoveBoxCollisionMeshUid(const FString& MeshName);
+
+	UFUNCTION(BlueprintCallable, Category = "Interchange | Node | StaticMeshLodData")
+	bool RemoveAllBoxCollisionMeshes();
+
+	UFUNCTION(BlueprintCallable, Category = "Interchange | Node | StaticMeshLodData")
+	int32 GetCapsuleCollisionMeshUidsCount() const;
+
+	UFUNCTION(BlueprintCallable, Category = "Interchange | Node | StaticMeshLodData")
+	void GetCapsuleCollisionMeshUids(TArray<FString>& OutMeshNames) const;
+
+	UFUNCTION(BlueprintCallable, Category = "Interchange | Node | StaticMeshLodData")
+	bool AddCapsuleCollisionMeshUid(const FString& MeshName);
+
+	UFUNCTION(BlueprintCallable, Category = "Interchange | Node | StaticMeshLodData")
+	bool RemoveCapsuleCollisionMeshUid(const FString& MeshName);
+
+	UFUNCTION(BlueprintCallable, Category = "Interchange | Node | StaticMeshLodData")
+	bool RemoveAllCapsuleCollisionMeshes();
+
+	UFUNCTION(BlueprintCallable, Category = "Interchange | Node | StaticMeshLodData")
+	int32 GetSphereCollisionMeshUidsCount() const;
+
+	UFUNCTION(BlueprintCallable, Category = "Interchange | Node | StaticMeshLodData")
+	void GetSphereCollisionMeshUids(TArray<FString>& OutMeshNames) const;
+
+	UFUNCTION(BlueprintCallable, Category = "Interchange | Node | StaticMeshLodData")
+	bool AddSphereCollisionMeshUid(const FString& MeshName);
+
+	UFUNCTION(BlueprintCallable, Category = "Interchange | Node | StaticMeshLodData")
+	bool RemoveSphereCollisionMeshUid(const FString& MeshName);
+
+	UFUNCTION(BlueprintCallable, Category = "Interchange | Node | StaticMeshLodData")
+	bool RemoveAllSphereCollisionMeshes();
+
+	UFUNCTION(BlueprintCallable, Category = "Interchange | Node | StaticMeshLodData")
+	int32 GetConvexCollisionMeshUidsCount() const;
+
+	UFUNCTION(BlueprintCallable, Category = "Interchange | Node | StaticMeshLodData")
+	void GetConvexCollisionMeshUids(TArray<FString>& OutMeshNames) const;
+
+	UFUNCTION(BlueprintCallable, Category = "Interchange | Node | StaticMeshLodData")
+	bool AddConvexCollisionMeshUid(const FString& MeshName);
+
+	UFUNCTION(BlueprintCallable, Category = "Interchange | Node | StaticMeshLodData")
+	bool RemoveConvexCollisionMeshUid(const FString& MeshName);
+
+	UFUNCTION(BlueprintCallable, Category = "Interchange | Node | StaticMeshLodData")
+	bool RemoveAllConvexCollisionMeshes();
+
+	UFUNCTION(BlueprintCallable, Category = "Interchange | Node | StaticMeshLodData")
+	bool GetOneConvexHullPerUCX(bool& AttributeValue) const;
+
+	UFUNCTION(BlueprintCallable, Category = "Interchange | Node | StaticMeshLodData")
+	bool SetOneConvexHullPerUCX(bool AttributeValue);
+
+
 private:
 
 	bool IsEditorOnlyDataDefined();
 
 	UE::Interchange::TArrayAttributeHelper<FString> MeshUids;
-protected:
+	UE::Interchange::TArrayAttributeHelper<FString> BoxCollisionMeshUids;
+	UE::Interchange::TArrayAttributeHelper<FString> CapsuleCollisionMeshUids;
+	UE::Interchange::TArrayAttributeHelper<FString> SphereCollisionMeshUids;
+	UE::Interchange::TArrayAttributeHelper<FString> ConvexCollisionMeshUids;
+
+	const UE::Interchange::FAttributeKey Macro_CustomOneConvexHullPerUCXKey = UE::Interchange::FAttributeKey(TEXT("__OneConvexHullPerUCX__Key"));
 };
diff --git a/Engine/Plugins/Experimental/Interchange/Runtime/Source/Import/InterchangeImport.Build.cs b/Engine/Plugins/Experimental/Interchange/Runtime/Source/Import/InterchangeImport.Build.cs
--- a/Engine/Plugins/Experimental/Interchange/Runtime/Source/Import/InterchangeImport.Build.cs
+++ b/Engine/Plugins/Experimental/Interchange/Runtime/Source/Import/InterchangeImport.Build.cs
@@ -43,6 +43,7 @@ namespace UnrealBuildTool.Rules
 				PrivateDependencyModuleNames.AddRange(
 					new string[]
 					{
+						"BSPUtils",
 						"MaterialEditor",
 						"UnrealEd",
 					}
diff --git a/Engine/Plugins/Experimental/Interchange/Runtime/Source/Import/Private/Mesh/InterchangeStaticMeshFactory.cpp b/Engine/Plugins/Experimental/Interchange/Runtime/Source/Import/Private/Mesh/InterchangeStaticMeshFactory.cpp
--- a/Engine/Plugins/Experimental/Interchange/Runtime/Source/Import/Private/Mesh/InterchangeStaticMeshFactory.cpp
+++ b/Engine/Plugins/Experimental/Interchange/Runtime/Source/Import/Private/Mesh/InterchangeStaticMeshFactory.cpp
@@ -1,7 +1,12 @@
 // Copyright Epic Games, Inc. All Rights Reserved. 
+
 #include "Mesh/InterchangeStaticMeshFactory.h"
 
+#if WITH_EDITOR
+#include "BSPOps.h"
+#endif
 #include "Components.h"
+#include "Engine/Polys.h"
 #include "Engine/StaticMesh.h"
 #include "InterchangeAssetImportData.h"
 #include "InterchangeImportCommon.h"
@@ -18,15 +23,20 @@
 #include "Materials/MaterialInterface.h"
 #include "Mesh/InterchangeStaticMeshPayload.h"
 #include "Mesh/InterchangeStaticMeshPayloadInterface.h"
+#include "Model.h"
 #include "Nodes/InterchangeBaseNode.h"
 #include "Nodes/InterchangeBaseNodeContainer.h"
+#include "PhysicsEngine/AggregateGeom.h"
+#include "PhysicsEngine/BodySetup.h"
+#include "PhysicsEngine/BoxElem.h"
+#include "PhysicsEngine/ConvexElem.h"
+#include "PhysicsEngine/SphereElem.h"
+#include "PhysicsEngine/SphylElem.h"
 #include "StaticMeshAttributes.h"
 #include "StaticMeshOperations.h"
 
 #if WITH_EDITORONLY_DATA
-
 #include "EditorFramework/AssetImportData.h"
-
 #endif //WITH_EDITORONLY_DATA
 
 
@@ -75,6 +85,9 @@ UObject* UInterchangeStaticMeshFactory::CreateEmptyAsset(const FCreateAssetParam
 		UE_LOG(LogInterchangeImport, Warning, TEXT("Could not create StaticMesh asset %s"), *Arguments.AssetName);
 		return nullptr;
 	}
+
+	// create the BodySetup on the game thread
+	StaticMesh->CreateBodySetup();
 	
 	StaticMesh->PreEditChange(nullptr);	
 	return StaticMesh;
@@ -102,13 +115,6 @@ UObject* UInterchangeStaticMeshFactory::CreateAsset(const FCreateAssetParams& Ar
 		return nullptr;
 	}
 
-	const IInterchangeStaticMeshPayloadInterface* StaticMeshTranslatorPayloadInterface = Cast<IInterchangeStaticMeshPayloadInterface>(Arguments.Translator);
-	if (!StaticMeshTranslatorPayloadInterface)
-	{
-		UE_LOG(LogInterchangeImport, Error, TEXT("Cannot import static mesh, the translator does not implement the IInterchangeStaticMeshPayloadInterface."));
-		return nullptr;
-	}
-
 	const UClass* StaticMeshClass = StaticMeshFactoryNode->GetObjectClass();
 	check(StaticMeshClass && StaticMeshClass->IsChildOf(GetFactoryClass()));
 
@@ -146,6 +152,8 @@ UObject* UInterchangeStaticMeshFactory::CreateAsset(const FCreateAssetParams& Ar
 			
 	int32 LodCount = StaticMeshFactoryNode->GetLodDataCount();
 	int32 FinalLodCount = LodCount;
+
+	// If we are reimporting, cache the existing vertex colors so they can be optionally reapplied after reimport
 	TMap<FVector, FColor> ExisitingVertexColorData;
 	if (Arguments.ReimportObject)
 	{
@@ -153,15 +161,49 @@ UObject* UInterchangeStaticMeshFactory::CreateAsset(const FCreateAssetParams& Ar
 		{
 			StaticMesh->GetVertexColorData(ExisitingVertexColorData);
 		}
-		//When we reimport we dont want to reduce the number of existing LODs
+
+		// When we reimport we don't want to reduce the number of existing LODs
 		FinalLodCount = FMath::Max(StaticMesh->GetNumLODs(), LodCount);
 	}
+
 	StaticMesh->SetNumSourceModels(FinalLodCount);
 
+	// Set material slots from imported materials
+	TArray<FString> FactoryDependencies;
+	StaticMeshFactoryNode->GetFactoryDependencies(FactoryDependencies);
+
+	for (int32 DependencyIndex = 0; DependencyIndex < FactoryDependencies.Num(); ++DependencyIndex)
+	{
+		const UInterchangeBaseMaterialFactoryNode* MaterialFactoryNode = Cast<UInterchangeBaseMaterialFactoryNode>(Arguments.NodeContainer->GetNode(FactoryDependencies[DependencyIndex]));
+		if (!MaterialFactoryNode || !MaterialFactoryNode->ReferenceObject.IsValid())
+		{
+			continue;
+		}
+		if (!MaterialFactoryNode->IsEnabled())
+		{
+			continue;
+		}
+
+		FName MaterialSlotName = *MaterialFactoryNode->GetDisplayLabel();
+		UMaterialInterface* MaterialInterface = Cast<UMaterialInterface>(MaterialFactoryNode->ReferenceObject.ResolveObject());
+		if (!MaterialInterface)
+		{
+			MaterialInterface = UMaterial::GetDefaultMaterial(MD_Surface);
+		}
+
+		int32 MaterialSlotIndex = StaticMesh->GetMaterialIndexFromImportedMaterialSlotName(MaterialSlotName);
+		if (MaterialSlotIndex == INDEX_NONE)
+		{
+			StaticMesh->GetStaticMaterials().Emplace(MaterialInterface, MaterialSlotName);
+		}
+	}
+
+	// Now import geometry for each LOD
 	TArray<FString> LodDataUniqueIds;
 	StaticMeshFactoryNode->GetLodDataUniqueIds(LodDataUniqueIds);
 	ensure(LodDataUniqueIds.Num() == LodCount);
 
+	bool bImportedCustomCollision = false;
 	int32 CurrentLodIndex = 0;
 	for (int32 LodIndex = 0; LodIndex < LodCount; ++LodIndex)
 	{
@@ -173,69 +215,6 @@ UObject* UInterchangeStaticMeshFactory::CreateAsset(const FCreateAssetParams& Ar
 			continue;
 		}
 
-		// Get the mesh node context for each MeshUids
-		struct FMeshNodeContext
-		{
-			const UInterchangeMeshNode* MeshNode = nullptr;
-			const UInterchangeSceneNode* SceneNode = nullptr;
-			TOptional<FTransform> SceneGlobalTransform;
-			FString TranslatorPayloadKey;
-		};
-		TArray<FMeshNodeContext> MeshReferences;
-
-		// Scope to query the mesh node
-		{
-			TArray<FString> MeshUids;
-			LodDataNode->GetMeshUids(MeshUids);
-			MeshReferences.Reserve(MeshUids.Num());
-
-			for (const FString& MeshUid : MeshUids)
-			{
-				FMeshNodeContext MeshReference;
-				MeshReference.MeshNode = Cast<UInterchangeMeshNode>(Arguments.NodeContainer->GetNode(MeshUid));
-				if (!MeshReference.MeshNode)
-				{
-					// The reference is a scene node and we need to bake the geometry
-					MeshReference.SceneNode = Cast<UInterchangeSceneNode>(Arguments.NodeContainer->GetNode(MeshUid));
-					if (!ensure(MeshReference.SceneNode != nullptr))
-					{
-						UE_LOG(LogInterchangeImport, Warning, TEXT("Invalid LOD mesh reference when importing StaticMesh asset %s"), *Arguments.AssetName);
-						continue;
-					}
-
-					FString MeshDependencyUid;
-					MeshReference.SceneNode->GetCustomAssetInstanceUid(MeshDependencyUid);
-					MeshReference.MeshNode = Cast<UInterchangeMeshNode>(Arguments.NodeContainer->GetNode(MeshDependencyUid));
-
-					// Cache the scene node global matrix, we will use this matrix to bake the vertices
-					FTransform SceneNodeGlobalTransform;
-					if (MeshReference.SceneNode->GetCustomGlobalTransform(Arguments.NodeContainer, SceneNodeGlobalTransform))
-					{
-						MeshReference.SceneGlobalTransform = SceneNodeGlobalTransform;
-					}
-				}
-
-				if (!ensure(MeshReference.MeshNode != nullptr))
-				{
-					UE_LOG(LogInterchangeImport, Warning, TEXT("Invalid LOD mesh reference when importing StaticMesh asset %s"), *Arguments.AssetName);
-					continue;
-				}
-
-				TOptional<FString> MeshPayloadKey = MeshReference.MeshNode->GetPayLoadKey();
-				if (MeshPayloadKey.IsSet())
-				{
-					MeshReference.TranslatorPayloadKey = MeshPayloadKey.GetValue();
-				}
-				else
-				{
-					UE_LOG(LogInterchangeImport, Warning, TEXT("Empty LOD mesh reference payload when importing StaticMesh asset %s"), *Arguments.AssetName);
-					continue;
-				}
-
-				MeshReferences.Add(MeshReference);
-			}
-		}
-
 		// Add the lod mesh data to the static mesh
 		FMeshDescription* LodMeshDescription = StaticMesh->CreateMeshDescription(CurrentLodIndex);
 
@@ -244,43 +223,32 @@ UObject* UInterchangeStaticMeshFactory::CreateAsset(const FCreateAssetParams& Ar
 		{
 			AppendSettings.bMergeUVChannels[ChannelIdx] = true;
 		}
-
-		// Fetch the payloads for all mesh references in parallel
-		TMap<FString, TFuture<TOptional<UE::Interchange::FStaticMeshPayloadData>>> MeshPayloads;
-		for (const FMeshNodeContext& MeshNodeContext : MeshReferences)
-		{
-			const FString& PayloadKey = MeshNodeContext.TranslatorPayloadKey;
-			MeshPayloads.Add(PayloadKey, StaticMeshTranslatorPayloadInterface->GetStaticMeshPayloadData(PayloadKey));
-		}
 				
+
+		TArray<FString> MeshUids;
+		LodDataNode->GetMeshUids(MeshUids);
+
 		// Fill the lod mesh description using all combined mesh parts
-		for (const FMeshNodeContext& MeshNodeContext : MeshReferences)
+		TArray<FMeshPayload> MeshPayloads = GetMeshPayloads(Arguments, MeshUids);
+		for (const FMeshPayload& MeshPayload : MeshPayloads)
 		{
-			TOptional<UE::Interchange::FStaticMeshPayloadData> LodMeshPayload = MeshPayloads.FindChecked(MeshNodeContext.TranslatorPayloadKey).Get();
+			const TOptional<UE::Interchange::FStaticMeshPayloadData>& LodMeshPayload = MeshPayload.PayloadData.Get();
 			if (!LodMeshPayload.IsSet())
 			{
-				UE_LOG(LogInterchangeImport, Warning, TEXT("Invalid static mesh payload key [%s] StaticMesh asset %s"), *MeshNodeContext.TranslatorPayloadKey, *Arguments.AssetName);
+				UE_LOG(LogInterchangeImport, Warning, TEXT("Invalid static mesh payload key, StaticMesh asset %s"), *Arguments.AssetName);
 				continue;
 			}
 
 			// Bake the payload, with the provided transform
-			if (MeshNodeContext.SceneGlobalTransform.IsSet())
-			{
-				AppendSettings.MeshTransform = MeshNodeContext.SceneGlobalTransform;
-			}
-			else
-			{
-				AppendSettings.MeshTransform.Reset();
-			}
-
+			AppendSettings.MeshTransform = MeshPayload.Transform;
 			FStaticMeshOperations::AppendMeshDescription(LodMeshPayload->MeshDescription, *LodMeshDescription, AppendSettings);
 		}
 
-		//////////////////////////////////////////////////////////////////////////
-		//Manage vertex color
-		//Replace -> do nothing, we want to use the translated source data
-		//Ignore -> remove vertex color from import data (when we re-import, ignore have to put back the current mesh vertex color)
-		//Override -> replace the vertex color by the override color
+		// Manage vertex color
+		// Replace -> do nothing, we want to use the translated source data
+		// Ignore -> remove vertex color from import data (when we re-import, ignore have to put back the current mesh vertex color)
+		// Override -> replace the vertex color by the override color
+		// @todo: new mesh description attribute for painted vertex colors?
 		{
 			FStaticMeshAttributes Attributes(*LodMeshDescription);
 			TVertexInstanceAttributesRef<FVector4f> VertexInstanceColors = Attributes.GetVertexInstanceColors();
@@ -335,37 +303,7 @@ UObject* UInterchangeStaticMeshFactory::CreateAsset(const FCreateAssetParams& Ar
 		CommitMeshDescriptionParams.bMarkPackageDirty = false; // Marking packages dirty isn't threadsafe
 		StaticMesh->CommitMeshDescription(CurrentLodIndex, CommitMeshDescriptionParams);
 
-		// Handle materials
-		// @TODO: move this outside of the LOD loop?
-		TArray<FString> FactoryDependencies;
-		StaticMeshFactoryNode->GetFactoryDependencies(FactoryDependencies);
-
-		for (int32 DependencyIndex = 0; DependencyIndex < FactoryDependencies.Num(); ++DependencyIndex)
-		{
-			const UInterchangeBaseMaterialFactoryNode* MaterialFactoryNode = Cast<UInterchangeBaseMaterialFactoryNode>(Arguments.NodeContainer->GetNode(FactoryDependencies[DependencyIndex]));
-			if (!MaterialFactoryNode || !MaterialFactoryNode->ReferenceObject.IsValid())
-			{
-				continue;
-			}
-			if (!MaterialFactoryNode->IsEnabled())
-			{
-				continue;
-			}
-
-			FName MaterialSlotName = *MaterialFactoryNode->GetDisplayLabel();
-			UMaterialInterface* MaterialInterface = Cast<UMaterialInterface>(MaterialFactoryNode->ReferenceObject.ResolveObject());
-			if (!MaterialInterface)
-			{
-				MaterialInterface = UMaterial::GetDefaultMaterial(MD_Surface);
-			}
-
-			int32 MaterialSlotIndex = StaticMesh->GetMaterialIndexFromImportedMaterialSlotName(MaterialSlotName);
-			if (MaterialSlotIndex == INDEX_NONE)
-			{
-				StaticMesh->GetStaticMaterials().Emplace(MaterialInterface, MaterialSlotName);
-			}
-		}
-
+		// Build section info map from materials
 		FStaticMeshConstAttributes StaticMeshAttributes(*LodMeshDescription);
 		TPolygonGroupAttributesRef<const FName> SlotNames = StaticMeshAttributes.GetPolygonGroupMaterialSlotNames();
 		int32 SectionIndex = 0;
@@ -388,6 +326,15 @@ UObject* UInterchangeStaticMeshFactory::CreateAsset(const FCreateAssetParams& Ar
 			SectionIndex++;
 		}
 
+		// Import collision geometry
+		if (CurrentLodIndex == 0)
+		{
+			bImportedCustomCollision |= ImportBoxCollision(Arguments, StaticMesh, LodDataNode);
+			bImportedCustomCollision |= ImportCapsuleCollision(Arguments, StaticMesh, LodDataNode);
+			bImportedCustomCollision |= ImportSphereCollision(Arguments, StaticMesh, LodDataNode);
+			bImportedCustomCollision |= ImportConvexCollision(Arguments, StaticMesh, LodDataNode);
+		}
+
 		if (!Arguments.ReimportObject)
 		{
 			FStaticMeshSourceModel& SrcModel = StaticMesh->GetSourceModel(CurrentLodIndex);
@@ -431,7 +378,16 @@ UObject* UInterchangeStaticMeshFactory::CreateAsset(const FCreateAssetParams& Ar
 		CurrentNode->FillAllCustomAttributeFromObject(StaticMesh);
 		UE::Interchange::FFactoryCommon::ApplyReimportStrategyToAsset(StaticMesh, PreviousNode, CurrentNode, StaticMeshFactoryNode);
 	}
-		
+	
+	if (!bImportedCustomCollision)
+	{
+		GenerateKDopCollision(Arguments, StaticMesh);
+	}
+	else
+	{
+		StaticMesh->bCustomizedCollision = true;
+	}
+
 	// Getting the file Hash will cache it into the source data
 	Arguments.SourceData->GetFileContentHash();
 
@@ -460,4 +416,864 @@ void UInterchangeStaticMeshFactory::PreImportPreCompletedCallback(const FImportP
 		StaticMesh->SetAssetImportData(ImportDataPtr);
 	}
 #endif
-}
\ No newline at end of file
+}
+
+
+TArray<UInterchangeStaticMeshFactory::FMeshPayload> UInterchangeStaticMeshFactory::GetMeshPayloads(const FCreateAssetParams& Arguments, const TArray<FString>& MeshUids) const
+{
+	TArray<FMeshPayload> Payloads;
+	Payloads.Reserve(MeshUids.Num());
+
+	const IInterchangeStaticMeshPayloadInterface* StaticMeshTranslatorPayloadInterface = Cast<IInterchangeStaticMeshPayloadInterface>(Arguments.Translator);
+	if (!StaticMeshTranslatorPayloadInterface)
+	{
+		UE_LOG(LogInterchangeImport, Error, TEXT("Cannot import static mesh, the translator does not implement the IInterchangeStaticMeshPayloadInterface."));
+		return Payloads;
+	}
+
+	for (const FString& MeshUid : MeshUids)
+	{
+		FMeshPayload Payload;
+
+		const UInterchangeBaseNode* Node = Arguments.NodeContainer->GetNode(MeshUid);
+		const UInterchangeMeshNode* MeshNode = Cast<const UInterchangeMeshNode>(Node);
+		if (MeshNode == nullptr)
+		{
+			// MeshUid must refer to a scene node
+			const UInterchangeSceneNode* SceneNode = Cast<const UInterchangeSceneNode>(Node);
+			if (!ensure(SceneNode))
+			{
+				UE_LOG(LogInterchangeImport, Warning, TEXT("Invalid LOD mesh reference when importing StaticMesh asset %s"), *Arguments.AssetName);
+				continue;
+			}
+
+			// Get the transform from the scene node
+			FTransform SceneNodeGlobalTransform;
+			if (SceneNode->GetCustomGlobalTransform(Arguments.NodeContainer, SceneNodeGlobalTransform))
+			{
+				Payload.Transform = SceneNodeGlobalTransform;
+			}
+
+			// And get the mesh node which it references
+			FString MeshDependencyUid;
+			SceneNode->GetCustomAssetInstanceUid(MeshDependencyUid);
+			MeshNode = Cast<UInterchangeMeshNode>(Arguments.NodeContainer->GetNode(MeshDependencyUid));
+		}
+
+		if (!ensure(MeshNode))
+		{
+			UE_LOG(LogInterchangeImport, Warning, TEXT("Invalid LOD mesh reference when importing StaticMesh asset %s"), *Arguments.AssetName);
+			continue;
+		}
+
+		TOptional<FString> MeshPayloadKey = MeshNode->GetPayLoadKey();
+		if (!ensure(MeshPayloadKey.IsSet()))
+		{
+			UE_LOG(LogInterchangeImport, Warning, TEXT("Empty LOD mesh reference payload when importing StaticMesh asset %s"), *Arguments.AssetName);
+			continue;
+		}
+
+		Payload.MeshName = *MeshPayloadKey;
+		Payload.PayloadData = StaticMeshTranslatorPayloadInterface->GetStaticMeshPayloadData(*MeshPayloadKey);
+
+		Payloads.Emplace(MoveTemp(Payload));
+	}
+
+	return Payloads;
+}
+
+
+bool UInterchangeStaticMeshFactory::AddConvexGeomFromVertices(const FCreateAssetParams& Arguments, const FMeshDescription& MeshDescription, const FTransform& Transform, FKAggregateGeom& AggGeom)
+{
+	FStaticMeshConstAttributes Attributes(MeshDescription);
+	TVertexAttributesRef<const FVector3f> VertexPositions = Attributes.GetVertexPositions();
+	
+	if (VertexPositions.GetNumElements() == 0)
+	{
+		return false;
+	}
+
+	FKConvexElem& ConvexElem = AggGeom.ConvexElems.Emplace_GetRef();
+	ConvexElem.VertexData.AddUninitialized(VertexPositions.GetNumElements());
+
+	for (int32 Index = 0; Index < VertexPositions.GetNumElements(); Index++)
+	{
+		ConvexElem.VertexData[Index] = Transform.TransformPosition(FVector(VertexPositions[Index]));
+	}
+
+	ConvexElem.UpdateElemBox();
+
+	return true;
+}
+
+
+bool UInterchangeStaticMeshFactory::DecomposeConvexMesh(const FCreateAssetParams& Arguments, const FMeshDescription& MeshDescription, const FTransform& Transform, UBodySetup* BodySetup)
+{
+#if WITH_EDITOR
+
+	// Construct a bit array containing a bit for each triangle ID in the mesh description.
+	// We are assuming the mesh description is compact, i.e. it has no holes, and so the number of triangles is equal to the array size.
+	// The aim is to identify 'islands' of adjacent triangles which will form separate convex hulls
+
+	check(MeshDescription.Triangles().Num() == MeshDescription.Triangles().GetArraySize());
+	TBitArray<> BitArray(0, MeshDescription.Triangles().Num());
+
+	// Here we build the groups of triangle IDs
+
+	TArray<TArray<FTriangleID>> TriangleGroups;
+
+	int32 FirstIndex = BitArray.FindAndSetFirstZeroBit();
+	while (FirstIndex != INDEX_NONE)
+	{
+		// Find the first index we haven't used yet, and use it as the beginning of a new triangle group
+
+		TArray<FTriangleID>& TriangleGroup = TriangleGroups.Emplace_GetRef();
+		TriangleGroup.Emplace(FirstIndex);
+
+		// Now iterate through the TriangleGroup array, finding unused adjacent triangles to each index, and appending them
+		// to the end of the array.  Note we deliberately check the array size each time round the loop, as each iteration
+		// can cause it to grow.
+
+		for (int32 CheckIndex = 0; CheckIndex < TriangleGroup.Num(); ++CheckIndex)
+		{
+			for (FTriangleID AdjacentTriangleID : MeshDescription.GetTriangleAdjacentTriangles(TriangleGroup[CheckIndex]))
+			{
+				if (BitArray[AdjacentTriangleID] == 0)
+				{
+					// Append unused adjacent triangles to the TriangleGroup, to be considered for adjacency later
+					TriangleGroup.Emplace(AdjacentTriangleID);
+					BitArray[AdjacentTriangleID] = 1;
+				}
+			}
+		}
+
+		// When we exhaust the triangle group array, there are no more triangles in this island.
+		// Now find the start of the next group.
+
+		FirstIndex = BitArray.FindAndSetFirstZeroBit();
+	}
+
+	// Now iterate through the triangle groups, adding each as a convex hull to the AggGeom
+
+	UModel* TempModel = NewObject<UModel>();
+	TempModel->RootOutside = true;
+	TempModel->EmptyModel(true, true);
+	TempModel->Polys->ClearFlags(RF_Transactional);
+
+	FStaticMeshConstAttributes Attributes(MeshDescription);
+	TTriangleAttributesRef<TArrayView<const FVertexID>> TriangleVertices = Attributes.GetTriangleVertexIndices();
+	TVertexAttributesRef<const FVector3f> VertexPositions = Attributes.GetVertexPositions();
+
+	bool bSuccess = true;
+
+	for (const TArray<FTriangleID>& TriangleGroup : TriangleGroups)
+	{
+		// Initialize a new brush
+		TempModel->Polys->Element.Empty();
+
+		// Add each triangle to the brush
+		int32 Index = 0;
+		for (FTriangleID TriangleID : TriangleGroup)
+		{
+			FPoly& Poly = TempModel->Polys->Element.Emplace_GetRef();
+			Poly.Init();
+			Poly.iLink = Index++;
+
+			// For reasons lost in time, BSP poly vertices have the opposite winding order to regular mesh vertices.
+			// So add them backwards (sigh)
+			Poly.Vertices.Emplace(Transform.TransformPosition(FVector(VertexPositions[TriangleVertices[TriangleID][2]])));
+			Poly.Vertices.Emplace(Transform.TransformPosition(FVector(VertexPositions[TriangleVertices[TriangleID][1]])));
+			Poly.Vertices.Emplace(Transform.TransformPosition(FVector(VertexPositions[TriangleVertices[TriangleID][0]])));
+
+			Poly.CalcNormal(true);
+		}
+
+		// Build bounding box
+		TempModel->BuildBound();
+
+		// Build BSP for the brush
+		FBSPOps::bspBuild(TempModel, FBSPOps::BSP_Good, 15, 70, 1, 0);
+		FBSPOps::bspRefresh(TempModel, true);
+		FBSPOps::bspBuildBounds(TempModel);
+
+		bSuccess &= BodySetup->CreateFromModel(TempModel, false);
+	}
+
+	TempModel->ClearInternalFlags(EInternalObjectFlags::Async);
+	TempModel->Polys->ClearInternalFlags(EInternalObjectFlags::Async);
+
+	return bSuccess;
+
+#else // #if WITH_EDITOR
+
+	return false;
+
+#endif
+}
+
+
+static bool AreEqual(float A, float B)
+{
+	constexpr float MeshToPrimTolerance = 0.001f;
+	return FMath::Abs(A - B) < MeshToPrimTolerance;
+}
+
+
+static bool AreParallel(const FVector3f& A, const FVector3f& B)
+{
+	float Dot = FVector3f::DotProduct(A, B);
+
+	return (AreEqual(FMath::Abs(Dot), 1.0f));
+}
+
+
+static FVector3f GetTriangleNormal(const FTransform& Transform, TVertexAttributesRef<const FVector3f> VertexPositions, TArrayView<const FVertexID> VertexIndices)
+{
+	const FVector3f& V0 = VertexPositions[VertexIndices[0]];
+	const FVector3f& V1 = VertexPositions[VertexIndices[1]];
+	const FVector3f& V2 = VertexPositions[VertexIndices[2]];
+	// @todo: LWC conversions everywhere here; surely this can be more elegant?
+	return FVector3f(Transform.TransformVector(FVector(FVector3f::CrossProduct(V1 - V0, V2 - V0).GetSafeNormal())));
+}
+
+
+bool UInterchangeStaticMeshFactory::AddBoxGeomFromTris(const FCreateAssetParams& Arguments, const FMeshDescription& MeshDescription, const FTransform& Transform, FKAggregateGeom& AggGeom)
+{
+	// Maintain an array of the planes we have encountered so far.
+	// We are expecting two instances of three unique plane orientations, one for each side of the box.
+
+	struct FPlaneInfo
+	{
+		FPlaneInfo(const FVector3f InNormal, float InFirstDistance)
+			: Normal(InNormal),
+			  DistCount(1),
+			  PlaneDist{InFirstDistance, 0.0f}
+		{}
+
+		FVector3f Normal;
+		int32 DistCount;
+		float PlaneDist[2];
+	};
+
+	TArray<FPlaneInfo> Planes;
+	FBox Box(ForceInit);
+
+	FStaticMeshConstAttributes Attributes(MeshDescription);
+	TTriangleAttributesRef<TArrayView<const FVertexID>> TriangleVertices = Attributes.GetTriangleVertexIndices();
+	TVertexAttributesRef<const FVector3f> VertexPositions = Attributes.GetVertexPositions();
+
+	for (FTriangleID TriangleID : MeshDescription.Triangles().GetElementIDs())
+	{
+		TArrayView<const FVertexID> VertexIndices = TriangleVertices[TriangleID];
+
+		FVector3f TriangleNormal = GetTriangleNormal(Transform, VertexPositions, VertexIndices);
+		if (TriangleNormal.IsNearlyZero())
+		{
+			continue;
+		}
+
+		bool bFoundPlane = false;
+		for (int32 PlaneIndex = 0; PlaneIndex < Planes.Num() && !bFoundPlane; PlaneIndex++)
+		{
+			// if this triangle plane is already known...
+			if (AreParallel(TriangleNormal, Planes[PlaneIndex].Normal))
+			{
+				// Always use the same normal when comparing distances, to ensure consistent sign.
+				float Dist = FVector3f::DotProduct(VertexPositions[VertexIndices[0]], Planes[PlaneIndex].Normal);
+
+				// we only have one distance, and its not that one, add it.
+				if (Planes[PlaneIndex].DistCount == 1 && !AreEqual(Dist, Planes[PlaneIndex].PlaneDist[0]))
+				{
+					Planes[PlaneIndex].PlaneDist[1] = Dist;
+					Planes[PlaneIndex].DistCount = 2;
+				}
+				// if we have a second distance, and its not that either, something is wrong.
+				else if (Planes[PlaneIndex].DistCount == 2 && !AreEqual(Dist, Planes[PlaneIndex].PlaneDist[1]))
+				{
+					// Error
+//					UE_LOG(LogStaticMeshImportUtils, Log, TEXT("AddBoxGeomFromTris (%s): Found more than 2 planes with different distances."), ObjName);
+					return false;
+				}
+
+				bFoundPlane = true;
+			}
+		}
+
+		// If this triangle does not match an existing plane, add to list.
+		if (!bFoundPlane)
+		{
+			Planes.Emplace(TriangleNormal, FVector3f::DotProduct(VertexPositions[VertexIndices[0]], TriangleNormal));
+		}
+
+		// Maintain an AABB, adding points from each triangle.
+		// We will use this to determine the origin of the box transform.
+
+		Box += Transform.TransformPosition(FVector(VertexPositions[VertexIndices[0]]));
+		Box += Transform.TransformPosition(FVector(VertexPositions[VertexIndices[1]]));
+		Box += Transform.TransformPosition(FVector(VertexPositions[VertexIndices[2]]));
+	}
+
+	// Now we have our candidate planes, see if there are any problems
+
+	// Wrong number of planes.
+	if (Planes.Num() != 3)
+	{
+		// Error
+//		UE_LOG(LogStaticMeshImportUtils, Log, TEXT("AddBoxGeomFromTris (%s): Not very box-like (need 3 sets of planes)."), ObjName);
+		return false;
+	}
+
+	// If we don't have 3 pairs, we can't carry on.
+	if ((Planes[0].DistCount != 2) || (Planes[1].DistCount != 2) || (Planes[2].DistCount != 2))
+	{
+		// Error
+//		UE_LOG(LogStaticMeshImportUtils, Log, TEXT("AddBoxGeomFromTris (%s): Incomplete set of planes (need 2 per axis)."), ObjName);
+		return false;
+	}
+
+	// ensure valid TM by cross-product
+	if (!AreParallel(FVector3f::CrossProduct(Planes[0].Normal, Planes[1].Normal), Planes[2].Normal))
+	{
+		// Error
+//		UE_LOG(LogStaticMeshImportUtils, Log, TEXT("AddBoxGeomFromTris (%s): Box axes are not perpendicular."), ObjName);
+		return false;
+	}
+
+	// Allocate box in array
+	FKBoxElem BoxElem;
+	BoxElem.SetTransform(FTransform(FVector(Planes[0].Normal), FVector(Planes[1].Normal), FVector(Planes[2].Normal), Box.GetCenter()));
+
+	// distance between parallel planes is box edge lengths.
+	BoxElem.X = FMath::Abs(Planes[0].PlaneDist[0] - Planes[0].PlaneDist[1]);
+	BoxElem.Y = FMath::Abs(Planes[1].PlaneDist[0] - Planes[1].PlaneDist[1]);
+	BoxElem.Z = FMath::Abs(Planes[2].PlaneDist[0] - Planes[2].PlaneDist[1]);
+
+	AggGeom.BoxElems.Add(BoxElem);
+
+	return true;
+}
+
+
+bool UInterchangeStaticMeshFactory::AddSphereGeomFromVertices(const FCreateAssetParams& Arguments, const FMeshDescription& MeshDescription, const FTransform& Transform, FKAggregateGeom& AggGeom)
+{
+	FStaticMeshConstAttributes Attributes(MeshDescription);
+	TVertexAttributesRef<const FVector3f> VertexPositions = Attributes.GetVertexPositions();
+
+	if (VertexPositions.GetNumElements() == 0)
+	{
+		return false;
+	}
+
+	FBox Box(ForceInit);
+
+	for (const FVector3f& VertexPosition : VertexPositions.GetRawArray())
+	{
+		Box += Transform.TransformPosition(FVector(VertexPosition));
+	}
+
+	FVector Center;
+	FVector Extents;
+	Box.GetCenterAndExtents(Center, Extents);
+	float Longest = 2.0f * Extents.GetMax();
+	float Shortest = 2.0f * Extents.GetMin();
+
+	// check that the AABB is roughly a square (5% tolerance)
+	if ((Longest - Shortest) / Longest > 0.05f)
+	{
+		// Error
+		//UE_LOG(LogStaticMeshImportUtils, Log, TEXT("AddSphereGeomFromVerts (%s): Sphere bounding box not square."), ObjName);
+		return false;
+	}
+
+	float Radius = 0.5f * Longest;
+
+	// Test that all vertices are a similar radius (5%) from the sphere centre.
+	float MaxR = 0;
+	float MinR = BIG_NUMBER;
+
+
+	for (const FVector3f& VertexPosition : VertexPositions.GetRawArray())
+	{
+		FVector3f CToV = VertexPosition - FVector3f(Center);
+		float RSqr = CToV.SizeSquared();
+
+		MaxR = FMath::Max(RSqr, MaxR);
+
+		// Sometimes vertex at centre, so reject it.
+		if (RSqr > KINDA_SMALL_NUMBER)
+		{
+			MinR = FMath::Min(RSqr, MinR);
+		}
+	}
+
+	MaxR = FMath::Sqrt(MaxR);
+	MinR = FMath::Sqrt(MinR);
+
+	if ((MaxR - MinR) / Radius > 0.05f)
+	{
+		// Error
+		//UE_LOG(LogStaticMeshImportUtils, Log, TEXT("AddSphereGeomFromVerts (%s): Vertices not at constant radius."), ObjName);
+		return false;
+	}
+
+	// Allocate sphere in array
+	FKSphereElem SphereElem;
+	SphereElem.Center = Center;
+	SphereElem.Radius = Radius;
+	AggGeom.SphereElems.Add(SphereElem);
+
+	return true;
+}
+
+
+bool UInterchangeStaticMeshFactory::AddCapsuleGeomFromVertices(const FCreateAssetParams& Arguments, const FMeshDescription& MeshDescription, const FTransform& Transform, FKAggregateGeom& AggGeom)
+{
+	FStaticMeshConstAttributes Attributes(MeshDescription);
+	TVertexAttributesRef<const FVector3f> VertexPositions = Attributes.GetVertexPositions();
+
+	if (VertexPositions.GetNumElements() == 0)
+	{
+		return false;
+	}
+
+	FVector AxisStart;
+	FVector AxisEnd;
+	float MaxDistSqr = 0.f;
+
+	for (int32 IndexA = 0; IndexA < VertexPositions.GetNumElements() - 1; IndexA++)
+	{
+		for (int32 IndexB = IndexA + 1; IndexB < VertexPositions.GetNumElements(); IndexB++)
+		{
+			FVector TransformedA = Transform.TransformPosition(FVector(VertexPositions[IndexA]));
+			FVector TransformedB = Transform.TransformPosition(FVector(VertexPositions[IndexB]));
+
+			float DistSqr = (TransformedA - TransformedB).SizeSquared();
+			if (DistSqr > MaxDistSqr)
+			{
+				AxisStart = TransformedA;
+				AxisEnd = TransformedB;
+				MaxDistSqr = DistSqr;
+			}
+		}
+	}
+
+	// If we got a valid axis, find vertex furthest from it
+	if (MaxDistSqr > SMALL_NUMBER)
+	{
+		float MaxRadius = 0.0f;
+
+		const FVector LineOrigin = AxisStart;
+		const FVector LineDir = (AxisEnd - AxisStart).GetSafeNormal();
+
+		for (int32 IndexA = 0; IndexA < VertexPositions.GetNumElements(); IndexA++)
+		{
+			FVector TransformedA = Transform.TransformPosition(FVector(VertexPositions[IndexA]));
+
+			float DistToAxis = FMath::PointDistToLine(TransformedA, LineDir, LineOrigin);
+			if (DistToAxis > MaxRadius)
+			{
+				MaxRadius = DistToAxis;
+			}
+		}
+
+		if (MaxRadius > SMALL_NUMBER)
+		{
+			// Allocate capsule in array
+			FKSphylElem SphylElem;
+			SphylElem.Center = 0.5f * (AxisStart + AxisEnd);
+			SphylElem.Rotation = FQuat::FindBetweenVectors(FVector::ZAxisVector, LineDir).Rotator(); // Get quat that takes you from z axis to desired axis
+			SphylElem.Radius = MaxRadius;
+			SphylElem.Length = FMath::Max(FMath::Sqrt(MaxDistSqr) - (2.0f * MaxRadius), 0.0f); // subtract two radii from total length to get segment length (ensure > 0)
+			AggGeom.SphylElems.Add(SphylElem);
+			return true;
+		}
+	}
+
+	return false;
+
+}
+
+
+bool UInterchangeStaticMeshFactory::ImportBoxCollision(const FCreateAssetParams& Arguments, UStaticMesh* StaticMesh, const UInterchangeStaticMeshLodDataNode* LodDataNode)
+{
+	using namespace UE::Interchange;
+
+	bool bResult = false;
+
+	TArray<FString> BoxCollisionMeshUids;
+	LodDataNode->GetBoxCollisionMeshUids(BoxCollisionMeshUids);
+
+	FKAggregateGeom& AggGeo = StaticMesh->GetBodySetup()->AggGeom;
+
+	TArray<FMeshPayload> MeshPayloads = GetMeshPayloads(Arguments, BoxCollisionMeshUids);
+	for (const FMeshPayload& MeshPayload : MeshPayloads)
+	{
+		const FTransform& Transform = MeshPayload.Transform;
+		const TOptional<FStaticMeshPayloadData>& PayloadData = MeshPayload.PayloadData.Get();
+
+		if (!PayloadData.IsSet())
+		{
+			// warning here
+			continue;
+		}
+
+		if (AddBoxGeomFromTris(Arguments, PayloadData->MeshDescription, Transform, AggGeo))
+		{
+			bResult = true;
+			FKBoxElem& NewElem = AggGeo.BoxElems.Last();
+
+			// Now test the last element in the AggGeo list and remove it if its a duplicate
+			// @TODO: determine why we have to do this. Was it to prevent duplicate boxes accumulating when reimporting?
+			for (int32 ElementIndex = 0; ElementIndex < AggGeo.BoxElems.Num() - 1; ++ElementIndex)
+			{
+				FKBoxElem& CurrentElem = AggGeo.BoxElems[ElementIndex];
+
+				if (CurrentElem == NewElem)
+				{
+					// The new element is a duplicate, remove it
+					AggGeo.BoxElems.RemoveAt(AggGeo.BoxElems.Num() - 1);
+					break;
+				}
+			}
+		}
+	}
+
+	return bResult;
+}
+
+
+bool UInterchangeStaticMeshFactory::ImportCapsuleCollision(const FCreateAssetParams& Arguments, UStaticMesh* StaticMesh, const UInterchangeStaticMeshLodDataNode* LodDataNode)
+{
+	using namespace UE::Interchange;
+
+	bool bResult = false;
+
+	TArray<FString> CapsuleCollisionMeshUids;
+	LodDataNode->GetCapsuleCollisionMeshUids(CapsuleCollisionMeshUids);
+
+	FKAggregateGeom& AggGeo = StaticMesh->GetBodySetup()->AggGeom;
+
+	TArray<FMeshPayload> MeshPayloads = GetMeshPayloads(Arguments, CapsuleCollisionMeshUids);
+	for (const FMeshPayload& MeshPayload : MeshPayloads)
+	{
+		const FTransform& Transform = MeshPayload.Transform;
+		const TOptional<FStaticMeshPayloadData>& PayloadData = MeshPayload.PayloadData.Get();
+
+		if (!PayloadData.IsSet())
+		{
+			// warning here
+			continue;
+		}
+
+		if (AddCapsuleGeomFromVertices(Arguments, PayloadData->MeshDescription, Transform, AggGeo))
+		{
+			bResult = true;
+
+			FKSphylElem& NewElem = AggGeo.SphylElems.Last();
+
+			// Now test the late element in the AggGeo list and remove it if its a duplicate
+			// @TODO: determine why we have to do this. Was it to prevent duplicate boxes accumulating when reimporting?
+			for (int32 ElementIndex = 0; ElementIndex < AggGeo.SphylElems.Num() - 1; ++ElementIndex)
+			{
+				FKSphylElem& CurrentElem = AggGeo.SphylElems[ElementIndex];
+				if (CurrentElem == NewElem)
+				{
+					// The new element is a duplicate, remove it
+					AggGeo.SphylElems.RemoveAt(AggGeo.SphylElems.Num() - 1);
+					break;
+				}
+			}
+		}
+	}
+
+	return bResult;
+}
+
+
+bool UInterchangeStaticMeshFactory::ImportSphereCollision(const FCreateAssetParams& Arguments, UStaticMesh* StaticMesh, const UInterchangeStaticMeshLodDataNode* LodDataNode)
+{
+	using namespace UE::Interchange;
+
+	bool bResult = false;
+
+	TArray<FString> SphereCollisionMeshUids;
+	LodDataNode->GetSphereCollisionMeshUids(SphereCollisionMeshUids);
+
+	FKAggregateGeom& AggGeo = StaticMesh->GetBodySetup()->AggGeom;
+
+	TArray<FMeshPayload> MeshPayloads = GetMeshPayloads(Arguments, SphereCollisionMeshUids);
+	for (const FMeshPayload& MeshPayload : MeshPayloads)
+	{
+		const FTransform& Transform = MeshPayload.Transform;
+		const TOptional<FStaticMeshPayloadData>& PayloadData = MeshPayload.PayloadData.Get();
+
+		if (!PayloadData.IsSet())
+		{
+			// warning here
+			continue;
+		}
+
+		if (AddSphereGeomFromVertices(Arguments, PayloadData->MeshDescription, Transform, AggGeo))
+		{
+			bResult = true;
+
+			FKSphereElem& NewElem = AggGeo.SphereElems.Last();
+
+			// Now test the last element in the AggGeo list and remove it if its a duplicate
+			// @TODO: determine why we have to do this. Was it to prevent duplicate boxes accumulating when reimporting?
+			for (int32 ElementIndex = 0; ElementIndex < AggGeo.SphereElems.Num() - 1; ++ElementIndex)
+			{
+				FKSphereElem& CurrentElem = AggGeo.SphereElems[ElementIndex];
+
+				if (CurrentElem == NewElem)
+				{
+					// The new element is a duplicate, remove it
+					AggGeo.SphereElems.RemoveAt(AggGeo.SphereElems.Num() - 1);
+					break;
+				}
+			}
+		}
+	}
+
+	return bResult;
+}
+
+
+bool UInterchangeStaticMeshFactory::ImportConvexCollision(const FCreateAssetParams& Arguments, UStaticMesh* StaticMesh, const UInterchangeStaticMeshLodDataNode* LodDataNode)
+{
+	using namespace UE::Interchange;
+
+	bool bResult = false;
+
+	TArray<FString> ConvexCollisionMeshUids;
+	LodDataNode->GetConvexCollisionMeshUids(ConvexCollisionMeshUids);
+
+	TArray<FMeshPayload> MeshPayloads = GetMeshPayloads(Arguments, ConvexCollisionMeshUids);
+
+	bool bOneConvexHullPerUCX;
+	if (!LodDataNode->GetOneConvexHullPerUCX(bOneConvexHullPerUCX) || !bOneConvexHullPerUCX)
+	{
+		for (const FMeshPayload& MeshPayload : MeshPayloads)
+		{
+			const FTransform& Transform = MeshPayload.Transform;
+			const TOptional<FStaticMeshPayloadData>& PayloadData = MeshPayload.PayloadData.Get();
+
+			if (!PayloadData.IsSet())
+			{
+				// warning here
+				continue;
+			}
+
+			if (!DecomposeConvexMesh(Arguments, PayloadData->MeshDescription, Transform, StaticMesh->GetBodySetup()))
+			{
+				// error: could not decompose mesh
+			}
+			else
+			{
+				bResult = true;
+			}
+		}
+	}
+	else
+	{
+		FKAggregateGeom& AggGeo = StaticMesh->GetBodySetup()->AggGeom;
+
+		for (const FMeshPayload& MeshPayload : MeshPayloads)
+		{
+			const FTransform& Transform = MeshPayload.Transform;
+			TOptional<FStaticMeshPayloadData> PayloadData = MeshPayload.PayloadData.Get();
+
+			if (!PayloadData.IsSet())
+			{
+				// warning here
+				continue;
+			}
+
+			if (AddConvexGeomFromVertices(Arguments, PayloadData->MeshDescription, Transform, AggGeo))
+			{
+				bResult = true;
+
+				FKConvexElem& NewElem = AggGeo.ConvexElems.Last();
+
+				// Now test the late element in the AggGeo list and remove it if its a duplicate
+				// @TODO: determine why the importer used to do this. Was it something to do with reimport not adding extra collision or something?
+				for (int32 ElementIndex = 0; ElementIndex < AggGeo.ConvexElems.Num() - 1; ++ElementIndex)
+				{
+					FKConvexElem& CurrentElem = AggGeo.ConvexElems[ElementIndex];
+
+					if (CurrentElem.VertexData.Num() == NewElem.VertexData.Num())
+					{
+						bool bFoundDifference = false;
+						for (int32 VertexIndex = 0; VertexIndex < NewElem.VertexData.Num(); ++VertexIndex)
+						{
+							if (CurrentElem.VertexData[VertexIndex] != NewElem.VertexData[VertexIndex])
+							{
+								bFoundDifference = true;
+								break;
+							}
+						}
+
+						if (!bFoundDifference)
+						{
+							// The new collision geo is a duplicate, delete it
+							AggGeo.ConvexElems.RemoveAt(AggGeo.ConvexElems.Num() - 1);
+							break;
+						}
+					}
+				}
+			}
+		}
+	}
+
+	return bResult;
+}
+
+
+bool UInterchangeStaticMeshFactory::GenerateKDopCollision(const FCreateAssetParams& Arguments, UStaticMesh* StaticMesh)
+{
+#if WITH_EDITOR
+
+	constexpr float RecipSqrt2 = UE_INV_SQRT_2;
+	constexpr int32 KCount = 18;
+	static const FVector3f KDopDir18[KCount] =
+	{
+		FVector3f( 1.0f,  0.0f,  0.0f),
+		FVector3f(-1.0f,  0.0f,  0.0f),
+		FVector3f( 0.0f,  1.0f,  0.0f),
+		FVector3f( 0.0f, -1.0f,  0.0f),
+		FVector3f( 0.0f,  0.0f,  1.0f),
+		FVector3f( 0.0f,  0.0f, -1.0f),
+		FVector3f( 0.0f,  RecipSqrt2,  RecipSqrt2),
+		FVector3f( 0.0f, -RecipSqrt2, -RecipSqrt2),
+		FVector3f( 0.0f,  RecipSqrt2, -RecipSqrt2),
+		FVector3f( 0.0f, -RecipSqrt2,  RecipSqrt2),
+		FVector3f( RecipSqrt2,  0.0f,  RecipSqrt2),
+		FVector3f(-RecipSqrt2,  0.0f, -RecipSqrt2),
+		FVector3f( RecipSqrt2,  0.0f, -RecipSqrt2),
+		FVector3f(-RecipSqrt2,  0.0f,  RecipSqrt2),
+		FVector3f( RecipSqrt2,  RecipSqrt2,  0.0f),
+		FVector3f(-RecipSqrt2, -RecipSqrt2,  0.0f),
+		FVector3f( RecipSqrt2, -RecipSqrt2,  0.0f),
+		FVector3f(-RecipSqrt2,  RecipSqrt2,  0.0f)
+	};
+
+	UBodySetup* BodySetup = StaticMesh->GetBodySetup();
+
+	// Initialize maximum distances from each kdop plane to the minimum value possible
+	TArray<float> MaxDist;
+	MaxDist.Reserve(KCount);
+	for (int32 Count = 0; Count < KCount; Count++)
+	{
+		MaxDist.Add(-MAX_FLT);
+	}
+
+	// Construct temporary UModel for kdop creation.
+	UModel* TempModel = NewObject<UModel>();
+	TempModel->RootOutside = true;
+	TempModel->EmptyModel(true, true);
+	TempModel->Polys->ClearFlags(RF_Transactional);
+
+	// Get the vertex positions for the final LOD0 mesh
+	FMeshDescription* MeshDescription = StaticMesh->GetMeshDescription(0);
+	check(MeshDescription);
+
+	FStaticMeshConstAttributes Attributes(*MeshDescription);
+	TVertexAttributesRef<const FVector3f> VertexPositions = Attributes.GetVertexPositions();
+
+	// For each vertex, project along each kdop direction, to find the max in that direction.
+	for (FVector3f VertexPosition : VertexPositions.GetRawArray())
+	{
+		for (int32 Index = 0; Index < KCount; Index++)
+		{
+			float Dist = FVector3f::DotProduct(VertexPosition, KDopDir18[Index]);
+			MaxDist[Index] = FMath::Max(Dist, MaxDist[Index]);
+		}
+	}
+
+	// Inflate kdop to ensure it is not degenerate
+	constexpr float MinSize = 0.1f;
+	for (int32 Index = 0; Index < KCount; Index++)
+	{
+		MaxDist[Index] += MinSize;
+	}
+
+	// Now we have the planes of the kdop, we work out the face polygons.
+	TArray<FPlane4f> Planes;
+	Planes.Reserve(KCount);
+	for (int32 Index = 0; Index < KCount; Index++)
+	{
+		Planes.Add(FPlane4f(KDopDir18[Index], MaxDist[Index]));
+	}
+
+	for (int32 PlaneIndex = 0; PlaneIndex < Planes.Num(); PlaneIndex++)
+	{
+		FPoly& Polygon = TempModel->Polys->Element.Emplace_GetRef();
+
+		Polygon.Init();
+		Polygon.Normal = Planes[PlaneIndex];
+
+		FVector3f AxisX, AxisY;
+		Polygon.Normal.FindBestAxisVectors(AxisX, AxisY);
+
+		FVector3f Base = Planes[PlaneIndex] * Planes[PlaneIndex].W;
+
+		Polygon.Vertices.Emplace(Base + AxisX * HALF_WORLD_MAX + AxisY * HALF_WORLD_MAX);
+		Polygon.Vertices.Emplace(Base + AxisX * HALF_WORLD_MAX - AxisY * HALF_WORLD_MAX);
+		Polygon.Vertices.Emplace(Base - AxisX * HALF_WORLD_MAX - AxisY * HALF_WORLD_MAX);
+		Polygon.Vertices.Emplace(Base - AxisX * HALF_WORLD_MAX + AxisY * HALF_WORLD_MAX);
+
+		for (int32 OtherPlaneIndex = 0; OtherPlaneIndex < Planes.Num(); OtherPlaneIndex++)
+		{
+			if (PlaneIndex != OtherPlaneIndex)
+			{
+				if (!Polygon.Split(-FVector3f(Planes[OtherPlaneIndex]), Planes[OtherPlaneIndex] * Planes[OtherPlaneIndex].W))
+				{
+					Polygon.Vertices.Empty();
+					break;
+				}
+			}
+		}
+
+		if (Polygon.Vertices.Num() < 3)
+		{
+			// If poly resulted in no verts, remove from array
+			TempModel->Polys->Element.RemoveAt(TempModel->Polys->Element.Num() - 1);
+		}
+		else
+		{
+			Polygon.iLink = PlaneIndex;
+
+			constexpr bool bSilent = true;
+			Polygon.CalcNormal(bSilent);
+		}
+	}
+
+	if (TempModel->Polys->Element.Num() < 4)
+	{
+		return false;
+	}
+
+	// Build bounding box.
+	TempModel->BuildBound();
+
+	// Build BSP for the brush.
+	FBSPOps::bspBuild(TempModel, FBSPOps::BSP_Good, 15, 70, 1, 0);
+
+	constexpr bool bNoRemapSurfs = true;
+	FBSPOps::bspRefresh(TempModel, bNoRemapSurfs);
+	FBSPOps::bspBuildBounds(TempModel);
+
+	bool bRemoveExisting = true;
+	BodySetup->CreateFromModel(TempModel, bRemoveExisting);
+
+	TempModel->ClearInternalFlags(EInternalObjectFlags::Async);
+	TempModel->Polys->ClearInternalFlags(EInternalObjectFlags::Async);
+
+	return true;
+
+#else // #if WITH_EDITOR
+
+	return false;
+
+#endif
+}
diff --git a/Engine/Plugins/Experimental/Interchange/Runtime/Source/Import/Public/Mesh/InterchangeStaticMeshFactory.h b/Engine/Plugins/Experimental/Interchange/Runtime/Source/Import/Public/Mesh/InterchangeStaticMeshFactory.h
--- a/Engine/Plugins/Experimental/Interchange/Runtime/Source/Import/Public/Mesh/InterchangeStaticMeshFactory.h
+++ b/Engine/Plugins/Experimental/Interchange/Runtime/Source/Import/Public/Mesh/InterchangeStaticMeshFactory.h
@@ -9,6 +9,16 @@
 
 #include "InterchangeStaticMeshFactory.generated.h"
 
+
+namespace UE::Interchange
+{
+	struct FStaticMeshPayloadData;
+}
+
+class UStaticMesh;
+class UInterchangeStaticMeshLodDataNode;
+
+
 UCLASS(BlueprintType, Experimental)
 class INTERCHANGEIMPORT_API UInterchangeStaticMeshFactory : public UInterchangeFactoryBase
 {
@@ -22,8 +32,30 @@ public:
 	virtual UObject* CreateEmptyAsset(const FCreateAssetParams& Arguments) override;
 	virtual UObject* CreateAsset(const FCreateAssetParams& Arguments) override;
 	virtual void PreImportPreCompletedCallback(const FImportPreCompletedCallbackParams& Arguments) override;
-	//virtual void PostImportPreCompletedCallback(const FImportPreCompletedCallbackParams& Arguments) const override;
 
 	// Interchange factory base interface end
 	//////////////////////////////////////////////////////////////////////////
+
+private:
+
+	struct FMeshPayload
+	{
+		FString MeshName;
+		TFuture<TOptional<UE::Interchange::FStaticMeshPayloadData>> PayloadData;
+		FTransform Transform = FTransform::Identity;
+	};
+
+	TArray<FMeshPayload> GetMeshPayloads(const FCreateAssetParams& Arguments, const TArray<FString>& MeshUids) const;
+
+	bool AddConvexGeomFromVertices(const FCreateAssetParams& Arguments, const FMeshDescription& MeshDescription, const FTransform& Transform, FKAggregateGeom& AggGeom);
+	bool DecomposeConvexMesh(const FCreateAssetParams& Arguments, const FMeshDescription& MeshDescription, const FTransform& Transform, UBodySetup* BodySetup);
+	bool AddBoxGeomFromTris(const FCreateAssetParams& Arguments, const FMeshDescription& MeshDescription, const FTransform& Transform, FKAggregateGeom& AggGeom);
+	bool AddSphereGeomFromVertices(const FCreateAssetParams& Arguments, const FMeshDescription& MeshDescription, const FTransform& Transform, FKAggregateGeom& AggGeom);
+	bool AddCapsuleGeomFromVertices(const FCreateAssetParams& Arguments, const FMeshDescription& MeshDescription, const FTransform& Transform, FKAggregateGeom& AggGeom);
+
+	bool ImportBoxCollision(const FCreateAssetParams& Arguments, UStaticMesh* StaticMesh, const UInterchangeStaticMeshLodDataNode* LodDataNode);
+	bool ImportCapsuleCollision(const FCreateAssetParams& Arguments, UStaticMesh* StaticMesh, const UInterchangeStaticMeshLodDataNode* LodDataNode);
+	bool ImportSphereCollision(const FCreateAssetParams& Arguments, UStaticMesh* StaticMesh, const UInterchangeStaticMeshLodDataNode* LodDataNode);
+	bool ImportConvexCollision(const FCreateAssetParams& Arguments, UStaticMesh* StaticMesh, const UInterchangeStaticMeshLodDataNode* LodDataNode);
+	bool GenerateKDopCollision(const FCreateAssetParams& Arguments, UStaticMesh* StaticMesh);
 };
diff --git a/Engine/Plugins/Experimental/Interchange/Runtime/Source/Pipelines/Private/InterchangeGenericStaticMeshPipeline.cpp b/Engine/Plugins/Experimental/Interchange/Runtime/Source/Pipelines/Private/InterchangeGenericStaticMeshPipeline.cpp
--- a/Engine/Plugins/Experimental/Interchange/Runtime/Source/Pipelines/Private/InterchangeGenericStaticMeshPipeline.cpp
+++ b/Engine/Plugins/Experimental/Interchange/Runtime/Source/Pipelines/Private/InterchangeGenericStaticMeshPipeline.cpp
@@ -19,6 +19,120 @@
 #include "UObject/Object.h"
 #include "UObject/ObjectMacros.h"
 
+
+enum class EMeshCollisionType
+{
+	None,
+	Box,
+	Sphere,
+	Capsule,
+	Convex
+};
+
+
+static FStringView GetMeshNameFromUid(const FString& NodeUid)
+{
+	FStringView MeshName = NodeUid;
+	int32 FinalDot = INDEX_NONE;
+	if (MeshName.FindLastChar(TEXT('.'), FinalDot))
+	{
+		MeshName.RightChopInline(FinalDot + 1);
+	}
+
+	return MeshName;
+}
+
+
+static TTuple<EMeshCollisionType, FStringView> GetCollisionMeshType(const FString& NodeUid, const TArray<FString>& AllNodeUids)
+{
+	FStringView MeshName = GetMeshNameFromUid(NodeUid);
+	EMeshCollisionType CollisionType = EMeshCollisionType::None;
+
+	// Determine if the mesh name is a potential collision mesh
+
+	if (MeshName.StartsWith(TEXT("UBX_")))
+	{
+		CollisionType = EMeshCollisionType::Box;
+	}
+	else if (MeshName.StartsWith(TEXT("UCX_")) || MeshName.StartsWith(TEXT("MCDCX_")))
+	{
+		CollisionType = EMeshCollisionType::Convex;
+	}
+	else if (MeshName.StartsWith(TEXT("USP_")))
+	{
+		CollisionType = EMeshCollisionType::Sphere;
+	}
+	else if (MeshName.StartsWith(TEXT("UCP_")))
+	{
+		CollisionType = EMeshCollisionType::Capsule;
+	}
+	else
+	{
+		return { EMeshCollisionType::None, FStringView() };
+	}
+
+	// We have a mesh name with a collision type suffix.
+	// However it should only be treated as a collision mesh if its body name corresponds to one of the other meshes.
+	// If we get here, we know there is at least one underscore, so we never expect either of these character searches to fail.
+
+	int32 FirstUnderscore = INDEX_NONE;
+	verify(MeshName.FindChar(TEXT('_'), FirstUnderscore));
+
+	int32 LastUnderscore = INDEX_NONE;
+	verify(MeshName.FindLastChar(TEXT('_'), LastUnderscore));
+
+	auto MatchPredicate = [](FStringView Body)
+	{
+		// Generate a predicate to be used by the below Finds.
+		return [Body](const FString& ToCompare) { return Body == GetMeshNameFromUid(ToCompare); };
+	};
+
+	// If we find a mesh named the same as the collision mesh (following the collision prefix), we have a match
+	// e.g. this will match 'UBX_House' with a mesh called 'House'
+
+	const FString* CorrespondingMeshUid = AllNodeUids.FindByPredicate(MatchPredicate(MeshName.RightChop(FirstUnderscore + 1)));
+	if (CorrespondingMeshUid)
+	{
+		return { CollisionType, *CorrespondingMeshUid };
+	}
+
+	// Otherwise strip the final underscore suffix from the collision mesh name and look again
+	// e.g. this will match 'UBX_House_01' with a mesh called 'House'
+
+	if (FirstUnderscore != LastUnderscore)
+	{
+		CorrespondingMeshUid = AllNodeUids.FindByPredicate(MatchPredicate(MeshName.Mid(FirstUnderscore + 1, LastUnderscore - FirstUnderscore - 1)));
+		if (CorrespondingMeshUid)
+		{
+			return { CollisionType, *CorrespondingMeshUid };
+		}
+	}
+
+	// Mesh had a collision type prefix, but no corresponding mesh, so don't treat it as a collision mesh
+
+	return { EMeshCollisionType::None, FStringView() };
+}
+
+
+static bool IsCollisionMeshUid(const FString& MeshUid, const TArray<FString>& MeshUids)
+{
+	return GetCollisionMeshType(MeshUid, MeshUids).Get<0>() != EMeshCollisionType::None;
+}
+
+
+static void BuildMeshToCollisionMeshMap(const TArray<FString>& MeshUids, TMap<FString, TArray<FString>>& MeshToCollisionMeshMap)
+{
+	for (const FString& MeshUid : MeshUids)
+	{
+		TTuple<EMeshCollisionType, FStringView> CollisionType = GetCollisionMeshType(MeshUid, MeshUids);
+		if (CollisionType.Get<0>() != EMeshCollisionType::None)
+		{
+			MeshToCollisionMeshMap.FindOrAdd(FString(CollisionType.Get<1>())).Emplace(MeshUid);
+		}
+	}
+}
+
+
 void UInterchangeGenericMeshPipeline::ExecutePreImportPipelineStaticMesh()
 {
 	if (bImportStaticMeshes && (ForceAllMeshAsType == EInterchangeForceMeshType::IFMT_None || ForceAllMeshAsType == EInterchangeForceMeshType::IFMT_StaticMesh))
@@ -92,14 +206,28 @@ void UInterchangeGenericMeshPipeline::ExecutePreImportPipelineStaticMesh()
 			// Do not combine static meshes
 
 			bool bFoundMeshes = false;
+
 			if (bBakeMeshes)
 			{
 				TArray<FString> MeshUids;
 				PipelineMeshesUtilities->GetAllStaticMeshInstance(MeshUids, bConvertSkeletalMeshToStaticMesh);
 
+				// Work out which meshes are collision meshes which correspond to another mesh
+				TMap<FString, TArray<FString>> MeshToCollisionMeshMap;
+				if (bImportCollisionAccordingToMeshName)
+				{
+					BuildMeshToCollisionMeshMap(MeshUids, MeshToCollisionMeshMap);
+				}
+
+				// Now iterate through each mesh UID, creating a new factory for each one
 				for (const FString& MeshUid : MeshUids)
 				{
-					// @TODO: check this is correct. What happens when we explicitly don't combine scenes with a LOD group node and multiple transform nodes?
+					if (bImportCollisionAccordingToMeshName && IsCollisionMeshUid(MeshUid, MeshUids))
+					{
+						// If this is a collision mesh, don't add a factory; it will be added as part of another factory
+						continue;
+					}
+
 					TMap<int32, TArray<FString>> MeshUidsPerLodIndex;
 
 					const FInterchangeMeshInstance& MeshInstance = PipelineMeshesUtilities->GetMeshInstanceByUid(MeshUid);
@@ -117,6 +245,14 @@ void UInterchangeGenericMeshPipeline::ExecutePreImportPipelineStaticMesh()
 
 					if (MeshUidsPerLodIndex.Num() > 0)
 					{
+						if (bImportCollisionAccordingToMeshName)
+						{
+							if (const TArray<FString>* CorrespondingCollisionMeshes = MeshToCollisionMeshMap.Find(MeshUid))
+							{
+								MeshUidsPerLodIndex.FindOrAdd(0).Append(*CorrespondingCollisionMeshes);
+							}
+						}
+
 						UInterchangeStaticMeshFactoryNode* StaticMeshFactoryNode = CreateStaticMeshFactoryNode(MeshUidsPerLodIndex);
 						StaticMeshFactoryNodes.Add(StaticMeshFactoryNode);
 						bFoundMeshes = true;
@@ -129,8 +265,21 @@ void UInterchangeGenericMeshPipeline::ExecutePreImportPipelineStaticMesh()
 				TArray<FString> MeshUids;
 				PipelineMeshesUtilities->GetAllStaticMeshGeometry(MeshUids, bConvertSkeletalMeshToStaticMesh);
 
+				// Work out which meshes are collision meshes which correspond to another mesh
+				TMap<FString, TArray<FString>> MeshToCollisionMeshMap;
+				if (bImportCollisionAccordingToMeshName)
+				{
+					BuildMeshToCollisionMeshMap(MeshUids, MeshToCollisionMeshMap);
+				}
+
 				for (const FString& MeshUid : MeshUids)
 				{
+					if (bImportCollisionAccordingToMeshName && IsCollisionMeshUid(MeshUid, MeshUids))
+					{
+						// If this is a collision mesh, don't add a factory; it will be added as part of another factory
+						continue;
+					}
+
 					TMap<int32, TArray<FString>> MeshUidsPerLodIndex;
 
 					const FInterchangeMeshGeometry& MeshGeometry = PipelineMeshesUtilities->GetMeshGeometryByUid(MeshUid);
@@ -140,6 +289,14 @@ void UInterchangeGenericMeshPipeline::ExecutePreImportPipelineStaticMesh()
 
 					if (MeshUidsPerLodIndex.Num() > 0)
 					{
+						if (bImportCollisionAccordingToMeshName)
+						{
+							if (const TArray<FString>* CorrespondingCollisionMeshes = MeshToCollisionMeshMap.Find(MeshUid))
+							{
+								MeshUidsPerLodIndex.FindOrAdd(0).Append(*CorrespondingCollisionMeshes);
+							}
+						}
+
 						UInterchangeStaticMeshFactoryNode* StaticMeshFactoryNode = CreateStaticMeshFactoryNode(MeshUidsPerLodIndex);
 						StaticMeshFactoryNodes.Add(StaticMeshFactoryNode);
 					}
@@ -270,6 +427,7 @@ UInterchangeStaticMeshLodDataNode* UInterchangeGenericMeshPipeline::CreateStatic
 	}
 
 	StaticMeshLodDataNode->InitializeNode(NodeUID, DisplayLabel, EInterchangeNodeContainerType::FactoryData);
+	StaticMeshLodDataNode->SetOneConvexHullPerUCX(bOneConvexHullPerUCX);
 	BaseNodeContainer->AddNode(StaticMeshLodDataNode);
 	return StaticMeshLodDataNode;
 }
@@ -348,7 +506,36 @@ void UInterchangeGenericMeshPipeline::AddLodDataToStaticMesh(UInterchangeStaticM
 				}
 			}
 
-			LodDataNode->AddMeshUid(NodeUid);
+			if (bImportCollisionAccordingToMeshName)
+			{
+				TTuple<EMeshCollisionType, FStringView> MeshType = GetCollisionMeshType(NodeUid, NodeUids);
+				switch (MeshType.Get<0>())
+				{
+				case EMeshCollisionType::None:
+					LodDataNode->AddMeshUid(NodeUid);
+					break;
+
+				case EMeshCollisionType::Box:
+					LodDataNode->AddBoxCollisionMeshUid(NodeUid);
+					break;
+
+				case EMeshCollisionType::Sphere:
+					LodDataNode->AddSphereCollisionMeshUid(NodeUid);
+					break;
+
+				case EMeshCollisionType::Capsule:
+					LodDataNode->AddCapsuleCollisionMeshUid(NodeUid);
+					break;
+
+				case EMeshCollisionType::Convex:
+					LodDataNode->AddConvexCollisionMeshUid(NodeUid);
+					break;
+				}
+			}
+			else
+			{
+				LodDataNode->AddMeshUid(NodeUid);
+			}
 		}
 	}
 }
diff --git a/Engine/Plugins/Experimental/Interchange/Runtime/Source/Pipelines/Public/InterchangeGenericMeshPipeline.h b/Engine/Plugins/Experimental/Interchange/Runtime/Source/Pipelines/Public/InterchangeGenericMeshPipeline.h
--- a/Engine/Plugins/Experimental/Interchange/Runtime/Source/Pipelines/Public/InterchangeGenericMeshPipeline.h
+++ b/Engine/Plugins/Experimental/Interchange/Runtime/Source/Pipelines/Public/InterchangeGenericMeshPipeline.h
@@ -78,14 +78,31 @@ public:
 	
 	//////	STATIC_MESHES_CATEGORY Properties //////
 
-	/** If enable, import the animation asset find in the sources. */
+	/** If enabled, import the static mesh asset found in the sources. */
 	UPROPERTY(EditAnywhere, BlueprintReadWrite, Category = "Static Meshes")
 	bool bImportStaticMeshes = true;
 
-	/** If enable all translated static mesh node will be imported has a one static mesh. */
+	/** If enabled, all translated static mesh nodes will be imported as a single static mesh. */
 	UPROPERTY(EditAnywhere, BlueprintReadWrite, Category = "Static Meshes")
 	bool bCombineStaticMeshes = false;
 
+	/**
+	 * If enabled, meshes with certain prefixes will be imported as collision primitives for the mesh with the corresponding unprefixed name.
+	 * 
+	 * Supported prefixes are:
+	 * UBX_ Box collision
+	 * UCP_ Capsule collision
+	 * USP_ Sphere collision
+	 * UCX_ Convex collision
+	 */
+
+	UPROPERTY(EditAnywhere, BlueprintReadWrite, Category = "Static Meshes")
+	bool bImportCollisionAccordingToMeshName = true;
+
+	/** If enabled, each UCX collision mesh will be imported as a single convex hull. If disabled, a UCX mesh will be decomposed into its separate pieces and a convex hull generated for each. */
+	UPROPERTY(EditAnywhere, BlueprintReadWrite, Category = "Static Meshes")
+	bool bOneConvexHullPerUCX = true;
+
 	
 	//////  COMMON_SKELETAL_ANIMATIONS_CATEGORY //////
 
diff --git a/Engine/Plugins/Tests/InterchangeTests/Source/InterchangeTests/Private/ImportTestFunctions/StaticMeshImportTestFunctions.cpp b/Engine/Plugins/Tests/InterchangeTests/Source/InterchangeTests/Private/ImportTestFunctions/StaticMeshImportTestFunctions.cpp
--- a/Engine/Plugins/Tests/InterchangeTests/Source/InterchangeTests/Private/ImportTestFunctions/StaticMeshImportTestFunctions.cpp
+++ b/Engine/Plugins/Tests/InterchangeTests/Source/InterchangeTests/Private/ImportTestFunctions/StaticMeshImportTestFunctions.cpp
@@ -413,10 +413,14 @@ FInterchangeTestFunctionResult UStaticMeshImportTestFunctions::CheckSimpleCollis
 	}
 	else
 	{
+		// If there was no collision imported, collision will be automatically generated as a kdop-18 convex primitive.
+		// We should not count this when checking imported simple collision, hence we treat the convex element count as zero if there was no customized collision.
+		bool bImportedCollision = Mesh->bCustomizedCollision;
+
 		int32 SphereElementCount = BodySetup->AggGeom.SphereElems.Num();
 		int32 BoxElementCount = BodySetup->AggGeom.BoxElems.Num();
 		int32 CapsuleElementCount = BodySetup->AggGeom.SphylElems.Num();
-		int32 ConvexElementCount = BodySetup->AggGeom.ConvexElems.Num();
+		int32 ConvexElementCount = bImportedCollision ? BodySetup->AggGeom.ConvexElems.Num() : 0;
 		int32 TaperedCapsuleElementCount = BodySetup->AggGeom.TaperedCapsuleElems.Num();
 
 		if (SphereElementCount != ExpectedSphereElementCount)
