Reduced log spam in .obj import, and fixed a race condition issue.

diff --git a/Engine/Plugins/Experimental/Interchange/Runtime/Source/Import/Private/Mesh/InterchangeOBJTranslator.cpp b/Engine/Plugins/Experimental/Interchange/Runtime/Source/Import/Private/Mesh/InterchangeOBJTranslator.cpp
--- a/Engine/Plugins/Experimental/Interchange/Runtime/Source/Import/Private/Mesh/InterchangeOBJTranslator.cpp
+++ b/Engine/Plugins/Experimental/Interchange/Runtime/Source/Import/Private/Mesh/InterchangeOBJTranslator.cpp
@@ -403,6 +403,13 @@ namespace ObjParser
 	{
 		SkipWhitespace(Line);
 
+		// Skip comments
+		if (!Line.IsEmpty() && Line[0] == TEXT('#'))
+		{
+			Line = FStringView();
+			return FStringView();
+		}
+
 		for (int32 Index = 0; Index < Line.Len(); Index++)
 		{
 			if (IsTerminatorChar(Line[Index], Terminator))
@@ -471,13 +478,6 @@ namespace ObjParser
 	}
 
 
-	static bool ParseComment(FObjData& ObjData, FStringView Line)
-	{
-		// Nothing else to do
-		return true;
-	}
-
-
 	static bool ParseVertexPosition(FObjData& ObjData, FStringView Line)
 	{
 		FVector3f Position;
@@ -738,7 +738,6 @@ namespace ObjParser
 	{
 		static FKeywordMap MtlKeywordMap =
 		{
-			{ TEXT("#"),       ParseComment },
 			{ TEXT("newmtl"),  ParseNewMaterial },
 			{ TEXT("illum"),   ParseMaterialProperty<int32,     &FObjData::FMaterialData::IlluminationModel> },
 			{ TEXT("Kd"),      ParseMaterialProperty<FVector3f, &FObjData::FMaterialData::DiffuseColor> },
@@ -793,7 +792,7 @@ static FString MakeMeshNodeUid(const FString& Name)
  */
 static FString MakeShaderGraphNodeUid(const FString& Name)
 {
-	return FString(TEXT("\\Material\\")) + (Name.IsEmpty() ? FString(TEXT("Null")) : Name);
+	return FString(TEXT("\\Material\\MAT_")) + (Name.IsEmpty() ? FString(TEXT("Null")) : Name);
 }
 
 
@@ -852,7 +851,6 @@ bool UInterchangeOBJTranslator::Translate(UInterchangeBaseNodeContainer& BaseNod
 
 	static ObjParser::FKeywordMap ObjKeywordMap =
 	{
-		{ TEXT("#"),      ObjParser::ParseComment },
 		{ TEXT("v"),      ObjParser::ParseVertexPosition },
 		{ TEXT("vt"),     ObjParser::ParseTextureCoordinate },
 		{ TEXT("vn"),     ObjParser::ParseNormalVector },
@@ -930,7 +928,7 @@ bool UInterchangeOBJTranslator::Translate(UInterchangeBaseNodeContainer& BaseNod
 
 TFuture<TOptional<UE::Interchange::FStaticMeshPayloadData>> UInterchangeOBJTranslator::GetStaticMeshPayloadData(const FString& PayloadKey) const
 {
-	return Async(EAsyncExecution::TaskGraph, [this, &PayloadKey]
+	return Async(EAsyncExecution::TaskGraph, [this, PayloadKey]
 		{
 			using namespace UE::Interchange;
 
