Added forward declaration to fix non-unity build.

diff --git a/Engine/Plugins/Experimental/Interchange/Runtime/Source/Import/Public/Mesh/InterchangeStaticMeshFactory.h b/Engine/Plugins/Experimental/Interchange/Runtime/Source/Import/Public/Mesh/InterchangeStaticMeshFactory.h
--- a/Engine/Plugins/Experimental/Interchange/Runtime/Source/Import/Public/Mesh/InterchangeStaticMeshFactory.h
+++ b/Engine/Plugins/Experimental/Interchange/Runtime/Source/Import/Public/Mesh/InterchangeStaticMeshFactory.h
@@ -17,6 +17,7 @@ namespace UE::Interchange
 
 class UStaticMesh;
 class UInterchangeStaticMeshLodDataNode;
+struct FKAggregateGeom;
 
 
 UCLASS(BlueprintType, Experimental)
