Add intechange animation pipeline and support joint animation
#rnx

diff --git a/Engine/Plugins/Experimental/Interchange/Editor/Source/Pipelines/Private/InterchangeEditorPipelinesModule.cpp b/Engine/Plugins/Experimental/Interchange/Editor/Source/Pipelines/Private/InterchangeEditorPipelinesModule.cpp
--- a/Engine/Plugins/Experimental/Interchange/Editor/Source/Pipelines/Private/InterchangeEditorPipelinesModule.cpp
+++ b/Engine/Plugins/Experimental/Interchange/Editor/Source/Pipelines/Private/InterchangeEditorPipelinesModule.cpp
@@ -110,6 +110,7 @@ void FInterchangeEditorPipelinesModule::RegisterPropertySectionMappings()
 	{
 		TSharedRef<FPropertySection> Section = RegisterPropertySection(PropertyModule, "InterchangeGenericMeshPipeline", "SkeletalMeshes", LOCTEXT("Skeletal Meshes", "Skeletal Meshes"));
 		Section->AddCategory("Common Meshes");
+		Section->AddCategory("Common Skeletal Meshes and Animations");
 		Section->AddCategory("Skeletal Meshes");
 	}
 
@@ -120,6 +121,13 @@ void FInterchangeEditorPipelinesModule::RegisterPropertySectionMappings()
 		Section->AddCategory("Static Meshes");
 	}
 
+	// Animation
+	{
+		TSharedRef<FPropertySection> Section = RegisterPropertySection(PropertyModule, "InterchangeGenericAnimationPipeline", "AnimationSequences", LOCTEXT("Animation Sequences", "Animation Sequences"));
+		Section->AddCategory("Common Skeletal Meshes and Animations");
+		Section->AddCategory("Animations");
+	}
+
 	// Textures
 	{
 		TSharedRef<FPropertySection> Section = RegisterPropertySection(PropertyModule, "InterchangeGenericTexturePipeline", "Textures", LOCTEXT("Textures", "Textures"));
diff --git a/Engine/Plugins/Experimental/Interchange/Editor/Source/Pipelines/Private/SInterchangePipelineConfigurationDialog.cpp b/Engine/Plugins/Experimental/Interchange/Editor/Source/Pipelines/Private/SInterchangePipelineConfigurationDialog.cpp
--- a/Engine/Plugins/Experimental/Interchange/Editor/Source/Pipelines/Private/SInterchangePipelineConfigurationDialog.cpp
+++ b/Engine/Plugins/Experimental/Interchange/Editor/Source/Pipelines/Private/SInterchangePipelineConfigurationDialog.cpp
@@ -447,6 +447,7 @@ void SInterchangePipelineConfigurationDialog::Construct(const FArguments& InArgs
 					SNew(SButton)
 					.HAlign(HAlign_Center)
 					.Text(LOCTEXT("InspectorGraphWindow_ImportAll", "Import All"))
+					.IsEnabled(this, &SInterchangePipelineConfigurationDialog::IsImportButtonEnabled)
 					.OnClicked(this, &SInterchangePipelineConfigurationDialog::OnCloseDialog, ECloseEventType::ImportAll)
 				]
 				+ SUniformGridPanel::Slot(3, 0)
@@ -454,6 +455,7 @@ void SInterchangePipelineConfigurationDialog::Construct(const FArguments& InArgs
 					SNew(SButton)
 					.HAlign(HAlign_Center)
 					.Text(LOCTEXT("InspectorGraphWindow_Import", "Import"))
+					.IsEnabled(this, &SInterchangePipelineConfigurationDialog::IsImportButtonEnabled)
 					.OnClicked(this, &SInterchangePipelineConfigurationDialog::OnCloseDialog, ECloseEventType::Import)
 				]
 			]
@@ -468,6 +470,41 @@ void SInterchangePipelineConfigurationDialog::OnSelectionChanged(TSharedPtr<FInt
 	PipelineConfigurationDetailsView->SetObject(Pipeline);
 }
 
+bool SInterchangePipelineConfigurationDialog::RecursiveValidatePipelineSettings(const TSharedPtr<FInterchangePipelineStacksTreeNodeItem>& ParentNode) const
+{
+	if (ParentNode->Pipeline)
+	{
+		if (!ParentNode->Pipeline->IsSettingsAreValid())
+		{
+			return false;
+		}
+	}
+	for (int32 ChildIndex = 0; ChildIndex < ParentNode->Childrens.Num(); ++ChildIndex)
+	{
+		if (!RecursiveValidatePipelineSettings(ParentNode->Childrens[ChildIndex]))
+		{
+			return false;
+		}
+	}
+	return true;
+}
+
+bool SInterchangePipelineConfigurationDialog::IsImportButtonEnabled() const
+{
+	const FName DefaultPipelineStackName = GetDefault<UInterchangeProjectSettings>()->DefaultPipelineStack;
+	const TArray<TSharedPtr<FInterchangePipelineStacksTreeNodeItem>>& RootNodeArray = PipelineConfigurationTreeView->GetRootNodeArray();
+	for (const TSharedPtr<FInterchangePipelineStacksTreeNodeItem>& RootNode : RootNodeArray)
+	{
+		//Reimport have only one stack and the name wont match DefaultPipelineStackName
+		if (bReimport || RootNode->StackName == DefaultPipelineStackName)
+		{
+			return RecursiveValidatePipelineSettings(RootNode);
+		}
+	}
+
+	return true;
+}
+
 void SInterchangePipelineConfigurationDialog::RecursiveSavePipelineSettings(const TSharedPtr<FInterchangePipelineStacksTreeNodeItem>& ParentNode, const int32 PipelineIndex) const
 {
 	if (ParentNode->Pipeline)
diff --git a/Engine/Plugins/Experimental/Interchange/Editor/Source/Pipelines/Private/SInterchangePipelineConfigurationDialog.h b/Engine/Plugins/Experimental/Interchange/Editor/Source/Pipelines/Private/SInterchangePipelineConfigurationDialog.h
--- a/Engine/Plugins/Experimental/Interchange/Editor/Source/Pipelines/Private/SInterchangePipelineConfigurationDialog.h
+++ b/Engine/Plugins/Experimental/Interchange/Editor/Source/Pipelines/Private/SInterchangePipelineConfigurationDialog.h
@@ -127,6 +127,9 @@ private:
 	TSharedRef<SBox> SpawnPipelineConfiguration();
 	void OnSelectionChanged(TSharedPtr<FInterchangePipelineStacksTreeNodeItem> Item, ESelectInfo::Type SelectionType);
 
+	bool RecursiveValidatePipelineSettings(const TSharedPtr<FInterchangePipelineStacksTreeNodeItem>& ParentNode) const;
+	bool IsImportButtonEnabled() const;
+
 	void RecursiveSavePipelineSettings(const TSharedPtr<FInterchangePipelineStacksTreeNodeItem>& ParentNode, const int32 PipelineIndex) const;
 	void RecursiveLoadPipelineSettings(const TSharedPtr<FInterchangePipelineStacksTreeNodeItem>& ParentNode, const int32 PipelineIndex) const;
 
diff --git a/Engine/Plugins/Experimental/Interchange/Runtime/Interchange.uplugin b/Engine/Plugins/Experimental/Interchange/Runtime/Interchange.uplugin
--- a/Engine/Plugins/Experimental/Interchange/Runtime/Interchange.uplugin
+++ b/Engine/Plugins/Experimental/Interchange/Runtime/Interchange.uplugin
@@ -51,6 +51,11 @@
 			"Type": "Runtime",
 			"LoadingPhase": "Default"
 		},
+		{
+			"Name": "InterchangeCommonParser",
+			"Type": "Runtime",
+			"LoadingPhase": "Default"
+		},
 		{
 			"Name": "InterchangeFbxParser",
 			"Type": "Runtime",
diff --git a/Engine/Plugins/Experimental/Interchange/Runtime/Source/Dispatcher/Public/InterchangeCommands.h b/Engine/Plugins/Experimental/Interchange/Runtime/Source/Dispatcher/Public/InterchangeCommands.h
--- a/Engine/Plugins/Experimental/Interchange/Runtime/Source/Dispatcher/Public/InterchangeCommands.h
+++ b/Engine/Plugins/Experimental/Interchange/Runtime/Source/Dispatcher/Public/InterchangeCommands.h
@@ -17,7 +17,7 @@ namespace UE
 		{
 		public:
 			//Major version should be updated when there is an existing API that has been change.
-			static int32 GetMajor() { return 5; }
+			static int32 GetMajor() { return 6; }
 			//Minor version must be updated when there is an addition to the API.
 			static int32 GetMinor() { return 0; }
 			//Patch version should be update if there is some bug fixes in the private code.
diff --git a/Engine/Plugins/Experimental/Interchange/Runtime/Source/FactoryNodes/Private/InterchangeAnimSequenceFactoryNode.cpp b/Engine/Plugins/Experimental/Interchange/Runtime/Source/FactoryNodes/Private/InterchangeAnimSequenceFactoryNode.cpp
new file mode 100644
--- /dev/null
+++ b/Engine/Plugins/Experimental/Interchange/Runtime/Source/FactoryNodes/Private/InterchangeAnimSequenceFactoryNode.cpp
@@ -0,0 +1,102 @@
+// Copyright Epic Games, Inc. All Rights Reserved.
+
+#include "InterchangeAnimSequenceFactoryNode.h"
+#include "Animation/AnimSequence.h"
+
+UInterchangeAnimSequenceFactoryNode::UInterchangeAnimSequenceFactoryNode()
+{
+}
+
+void UInterchangeAnimSequenceFactoryNode::InitializeAnimSequenceNode(const FString& UniqueID, const FString& DisplayLabel)
+{
+	InitializeNode(UniqueID, DisplayLabel, EInterchangeNodeContainerType::FactoryData);
+}
+
+FString UInterchangeAnimSequenceFactoryNode::GetTypeName() const
+{
+	const FString TypeName = TEXT("AnimSequenceNode");
+	return TypeName;
+}
+
+FString UInterchangeAnimSequenceFactoryNode::GetKeyDisplayName(const UE::Interchange::FAttributeKey& NodeAttributeKey) const
+{
+	FString KeyDisplayName = NodeAttributeKey.ToString();
+	if (NodeAttributeKey == Macro_CustomSkeletonFactoryNodeUidKey)
+	{
+		KeyDisplayName = TEXT("Skeleton Uid");
+	}
+	else if (NodeAttributeKey == Macro_CustomSkeletonSoftObjectPathKey)
+	{
+		KeyDisplayName = TEXT("Specified Existing Skeleton");
+	}
+	else
+	{
+		KeyDisplayName = Super::GetKeyDisplayName(NodeAttributeKey);
+	}
+	return KeyDisplayName;
+}
+
+UClass* UInterchangeAnimSequenceFactoryNode::GetObjectClass() const
+{
+	return UAnimSequence::StaticClass();
+}
+
+bool UInterchangeAnimSequenceFactoryNode::GetCustomSkeletonFactoryNodeUid(FString& AttributeValue) const
+{
+	IMPLEMENT_NODE_ATTRIBUTE_GETTER(SkeletonFactoryNodeUid, FString);
+}
+
+bool UInterchangeAnimSequenceFactoryNode::SetCustomSkeletonFactoryNodeUid(const FString& AttributeValue)
+{
+	IMPLEMENT_NODE_ATTRIBUTE_SETTER_NODELEGATE(SkeletonFactoryNodeUid, FString);
+}
+
+bool UInterchangeAnimSequenceFactoryNode::GetCustomImportBoneTracks(bool& AttributeValue) const
+{
+	IMPLEMENT_NODE_ATTRIBUTE_GETTER(ImportBoneTracks, bool);
+}
+
+bool UInterchangeAnimSequenceFactoryNode::SetCustomImportBoneTracks(const bool& AttributeValue)
+{
+	IMPLEMENT_NODE_ATTRIBUTE_SETTER_NODELEGATE(ImportBoneTracks, bool);
+}
+
+bool UInterchangeAnimSequenceFactoryNode::GetCustomImportBoneTracksSampleRate(double& AttributeValue) const
+{
+	IMPLEMENT_NODE_ATTRIBUTE_GETTER(ImportBoneTracksSampleRate, double);
+}
+
+bool UInterchangeAnimSequenceFactoryNode::SetCustomImportBoneTracksSampleRate(const double& AttributeValue)
+{
+	IMPLEMENT_NODE_ATTRIBUTE_SETTER_NODELEGATE(ImportBoneTracksSampleRate, double);
+}
+
+bool UInterchangeAnimSequenceFactoryNode::GetCustomImportBoneTracksRangeStart(double& AttributeValue) const
+{
+	IMPLEMENT_NODE_ATTRIBUTE_GETTER(ImportBoneTracksRangeStart, double);
+}
+
+bool UInterchangeAnimSequenceFactoryNode::SetCustomImportBoneTracksRangeStart(const double& AttributeValue)
+{
+	IMPLEMENT_NODE_ATTRIBUTE_SETTER_NODELEGATE(ImportBoneTracksRangeStart, double);
+}
+
+bool UInterchangeAnimSequenceFactoryNode::GetCustomImportBoneTracksRangeStop(double& AttributeValue) const
+{
+	IMPLEMENT_NODE_ATTRIBUTE_GETTER(ImportBoneTracksRangeStop, double);
+}
+
+bool UInterchangeAnimSequenceFactoryNode::SetCustomImportBoneTracksRangeStop(const double& AttributeValue)
+{
+	IMPLEMENT_NODE_ATTRIBUTE_SETTER_NODELEGATE(ImportBoneTracksRangeStop, double);
+}
+
+bool UInterchangeAnimSequenceFactoryNode::GetCustomSkeletonSoftObjectPath(FSoftObjectPath& AttributeValue) const
+{
+	IMPLEMENT_NODE_ATTRIBUTE_GETTER(SkeletonSoftObjectPath, FSoftObjectPath)
+}
+
+bool UInterchangeAnimSequenceFactoryNode::SetCustomSkeletonSoftObjectPath(const FSoftObjectPath& AttributeValue)
+{
+	IMPLEMENT_NODE_ATTRIBUTE_SETTER_NODELEGATE(SkeletonSoftObjectPath, FSoftObjectPath)
+}
diff --git a/Engine/Plugins/Experimental/Interchange/Runtime/Source/FactoryNodes/Public/InterchangeAnimSequenceFactoryNode.h b/Engine/Plugins/Experimental/Interchange/Runtime/Source/FactoryNodes/Public/InterchangeAnimSequenceFactoryNode.h
new file mode 100644
--- /dev/null
+++ b/Engine/Plugins/Experimental/Interchange/Runtime/Source/FactoryNodes/Public/InterchangeAnimSequenceFactoryNode.h
@@ -0,0 +1,104 @@
+// Copyright Epic Games, Inc. All Rights Reserved.
+
+#pragma once
+
+#include "CoreMinimal.h"
+#include "Nodes/InterchangeFactoryBaseNode.h"
+#include "Animation/AnimSequence.h"
+
+#include "InterchangeAnimSequenceFactoryNode.generated.h"
+
+UCLASS(BlueprintType, Experimental)
+class INTERCHANGEFACTORYNODES_API UInterchangeAnimSequenceFactoryNode : public UInterchangeFactoryBaseNode
+{
+	GENERATED_BODY()
+
+public:
+	UInterchangeAnimSequenceFactoryNode();
+
+	/**
+	 * Initialize node data
+	 * @param UniqueID - The uniqueId for this node
+	 * @param DisplayLabel - The name of the node
+	 */
+	UFUNCTION(BlueprintCallable, Category = "Interchange | Node | AnimSequence")
+	void InitializeAnimSequenceNode(const FString& UniqueID, const FString& DisplayLabel);
+
+	/**
+	 * Return the node type name of the class, we use this when reporting error
+	 */
+	virtual FString GetTypeName() const override;
+
+	virtual FString GetKeyDisplayName(const UE::Interchange::FAttributeKey& NodeAttributeKey) const override;
+
+	/** Get the class this node want to create */
+	UFUNCTION(BlueprintCallable, Category = "Interchange | Node | AnimSequence")
+	virtual class UClass* GetObjectClass() const override;
+
+public:
+	/** Return false if the Attribute was not set previously. */
+	UFUNCTION(BlueprintCallable, Category = "Interchange | Node | AnimSequence")
+	bool GetCustomSkeletonFactoryNodeUid(FString& AttributeValue) const;
+
+	UFUNCTION(BlueprintCallable, Category = "Interchange | Node | AnimSequence")
+	bool SetCustomSkeletonFactoryNodeUid(const FString& AttributeValue);
+
+	/** Query whether we must import the bone tracks. Return false if the Attribute was not set previously. */
+	UFUNCTION(BlueprintCallable, Category = "Interchange | Node | AnimSequence")
+	bool GetCustomImportBoneTracks(bool& AttributeValue) const;
+
+	/** Set true if we must import the bone tracks, otherwise set it to false. Return false if the Attribute was not set previously. */
+	UFUNCTION(BlueprintCallable, Category = "Interchange | Node | AnimSequence")
+	bool SetCustomImportBoneTracks(const bool& AttributeValue);
+
+	/** Query weather we must import the bone tracks. Return false if the Attribute was not set previously. */
+	UFUNCTION(BlueprintCallable, Category = "Interchange | Node | AnimSequence")
+	bool GetCustomImportBoneTracksSampleRate(double& AttributeValue) const;
+
+	/** Set true if we must import the bone tracks, otherwise set it to false. Return false if the Attribute was not set previously. */
+	UFUNCTION(BlueprintCallable, Category = "Interchange | Node | AnimSequence")
+	bool SetCustomImportBoneTracksSampleRate(const double& AttributeValue);
+
+	/** Query weather we must import the bone tracks. Return false if the Attribute was not set previously. */
+	UFUNCTION(BlueprintCallable, Category = "Interchange | Node | AnimSequence")
+	bool GetCustomImportBoneTracksRangeStart(double& AttributeValue) const;
+
+	/** Set true if we must import the bone tracks, otherwise set it to false. Return false if the Attribute was not set previously. */
+	UFUNCTION(BlueprintCallable, Category = "Interchange | Node | AnimSequence")
+	bool SetCustomImportBoneTracksRangeStart(const double& AttributeValue);
+
+	/** Query weather we must import the bone tracks. Return false if the Attribute was not set previously. */
+	UFUNCTION(BlueprintCallable, Category = "Interchange | Node | AnimSequence")
+	bool GetCustomImportBoneTracksRangeStop(double& AttributeValue) const;
+
+	/** Set true if we must import the bone tracks, otherwise set it to false. Return false if the Attribute was not set previously. */
+	UFUNCTION(BlueprintCallable, Category = "Interchange | Node | AnimSequence")
+	bool SetCustomImportBoneTracksRangeStop(const double& AttributeValue);
+
+	/**
+	 * Query the optional existing USkeleton this anim must use. The animsequence factory will use this skeleton instead of the imported one
+	 * (from GetCustomSkeletonFactoryNodeUid) if this attribute is set and the skeleton pointer is valid.
+	 * Pipeline set this attribute in case the user want to specify an existing skeleton.
+	 * Return false if the attribute was not set.
+	 */
+	UFUNCTION(BlueprintCallable, Category = "Interchange | Node | AnimSequence")
+	bool GetCustomSkeletonSoftObjectPath(FSoftObjectPath& AttributeValue) const;
+
+	/**
+	 * Set the optional existing USkeleton this anim must use. The AnimSequence factory will use this skeleton instead of the imported one
+	 * (from GetCustomSkeletonFactoryNodeUid) if this attribute is set and the skeleton pointer is valid.
+	 * Pipeline set this attribute in case the user want to specify an existing skeleton.
+	 */
+	UFUNCTION(BlueprintCallable, Category = "Interchange | Node | AnimSequence")
+	bool SetCustomSkeletonSoftObjectPath(const FSoftObjectPath& AttributeValue);
+
+private:
+	const UE::Interchange::FAttributeKey Macro_CustomSkeletonFactoryNodeUidKey = UE::Interchange::FAttributeKey(TEXT("SkeletonFactoryNodeUid"));
+	
+	const UE::Interchange::FAttributeKey Macro_CustomImportBoneTracksKey = UE::Interchange::FAttributeKey(TEXT("ImportBoneTracks"));
+	const UE::Interchange::FAttributeKey Macro_CustomImportBoneTracksSampleRateKey = UE::Interchange::FAttributeKey(TEXT("ImportBoneTracksSampleRate"));
+	const UE::Interchange::FAttributeKey Macro_CustomImportBoneTracksRangeStartKey = UE::Interchange::FAttributeKey(TEXT("ImportBoneTracksRangeStart"));
+	const UE::Interchange::FAttributeKey Macro_CustomImportBoneTracksRangeStopKey = UE::Interchange::FAttributeKey(TEXT("ImportBoneTracksRangeStop"));
+
+	const UE::Interchange::FAttributeKey Macro_CustomSkeletonSoftObjectPathKey = UE::Interchange::FAttributeKey(TEXT("SkeletonSoftObjectPath"));
+};
diff --git a/Engine/Plugins/Experimental/Interchange/Runtime/Source/Import/InterchangeImport.Build.cs b/Engine/Plugins/Experimental/Interchange/Runtime/Source/Import/InterchangeImport.Build.cs
--- a/Engine/Plugins/Experimental/Interchange/Runtime/Source/Import/InterchangeImport.Build.cs
+++ b/Engine/Plugins/Experimental/Interchange/Runtime/Source/Import/InterchangeImport.Build.cs
@@ -30,6 +30,7 @@ namespace UnrealBuildTool.Rules
 					"GLTFCore",
 					"IESFile",
 					"ImageWrapper",
+					"InterchangeCommonParser",
 					"InterchangeDispatcher",
 					"InterchangeMessages",
 					"Json",
diff --git a/Engine/Plugins/Experimental/Interchange/Runtime/Source/Import/Private/Animation/InterchangeAnimSequenceFactory.cpp b/Engine/Plugins/Experimental/Interchange/Runtime/Source/Import/Private/Animation/InterchangeAnimSequenceFactory.cpp
new file mode 100644
--- /dev/null
+++ b/Engine/Plugins/Experimental/Interchange/Runtime/Source/Import/Private/Animation/InterchangeAnimSequenceFactory.cpp
@@ -0,0 +1,529 @@
+// Copyright Epic Games, Inc. All Rights Reserved. 
+#include "Animation/InterchangeAnimSequenceFactory.h"
+
+#include "Animation/AnimSequence.h"
+#include "Animation/InterchangeAnimationPayload.h"
+#include "Animation/InterchangeAnimationPayloadInterface.h"
+#include "InterchangeAssetImportData.h"
+#include "InterchangeCommonAnimationPayload.h"
+#include "InterchangeImportCommon.h"
+#include "InterchangeImportLog.h"
+#include "InterchangeSceneNode.h"
+#include "InterchangeAnimSequenceFactoryNode.h"
+#include "InterchangeSkeletonFactoryNode.h"
+#include "InterchangeSourceData.h"
+#include "InterchangeTranslatorBase.h"
+#include "Nodes/InterchangeBaseNode.h"
+#include "Nodes/InterchangeBaseNodeContainer.h"
+#include "Nodes/InterchangeSourceNode.h"
+
+#if WITH_EDITORONLY_DATA
+
+#include "EditorFramework/AssetImportData.h"
+
+#endif //WITH_EDITORONLY_DATA
+
+#if WITH_EDITOR
+namespace UE::Interchange::Private
+{
+	void GetSkeletonSceneNodeFlatListRecursive(const UInterchangeBaseNodeContainer* NodeContainer, const FString& NodeUid, TArray<FString>& SkeletonSceneNodeUids)
+	{
+		SkeletonSceneNodeUids.Add(NodeUid);
+		TArray<FString> Children = NodeContainer->GetNodeChildrenUids(NodeUid);
+		for (const FString& ChildUid : Children)
+		{
+			GetSkeletonSceneNodeFlatListRecursive(NodeContainer, ChildUid, SkeletonSceneNodeUids);
+		}
+	}
+
+	void RetrieveAnimationPayloads(UAnimSequence* AnimSequence
+		, const UInterchangeAnimSequenceFactoryNode* AnimSequenceFactoryNode
+		, const UInterchangeBaseNodeContainer* NodeContainer
+		, const UInterchangeSkeletonFactoryNode* SkeletonFactoryNode
+		, const USkeleton* Skeleton
+		, const IInterchangeAnimationPayloadInterface* AnimSequenceTranslatorPayloadInterface
+		, const FString& AssetName)
+	{
+		TMap<FString, TFuture<TOptional<UE::Interchange::FAnimationTransformPayloadData>>> AnimationPayloads;
+
+		FString SkeletonRootUid;
+		if (!SkeletonFactoryNode->GetCustomRootJointUid(SkeletonRootUid))
+		{
+			//Cannot import animation without a skeleton
+			return;
+		}
+		bool bImportBoneTracks = false;
+		AnimSequenceFactoryNode->GetCustomImportBoneTracks(bImportBoneTracks);
+		if (bImportBoneTracks)
+		{
+			TArray<FString> SkeletonNodes;
+			GetSkeletonSceneNodeFlatListRecursive(NodeContainer, SkeletonRootUid, SkeletonNodes);
+			for (const FString& NodeUid : SkeletonNodes)
+			{
+				if (const UInterchangeSceneNode* SkeletonSceneNode = Cast<UInterchangeSceneNode>(NodeContainer->GetNode(NodeUid)))
+				{
+					FString PayloadKey;
+					if (SkeletonSceneNode->GetCustomTransformCurvePayloadKey(PayloadKey))
+					{
+						AnimationPayloads.Add(PayloadKey, AnimSequenceTranslatorPayloadInterface->GetAnimationTransformPayloadData(PayloadKey));
+					}
+				}
+			}
+			//Get the sample rate, default to 30Hz in case the attribute is missing
+			double SampleRate = 30.0;
+			AnimSequenceFactoryNode->GetCustomImportBoneTracksSampleRate(SampleRate);
+
+			double RangeStart = 0.0;
+			AnimSequenceFactoryNode->GetCustomImportBoneTracksRangeStart(RangeStart);
+
+			double RangeStop = 1.0 / SampleRate; //One frame duration per default
+			AnimSequenceFactoryNode->GetCustomImportBoneTracksRangeStop(RangeStop);
+
+			const double BakeInterval = 1.0 / SampleRate;
+			const double SequenceLength = RangeStop - RangeStart;
+			int32 BakeKeyCount = (SequenceLength / BakeInterval)+1;
+
+			IAnimationDataController& Controller = AnimSequence->GetController();
+			//This destroy all previously imported animation raw data
+			Controller.RemoveAllBoneTracks();
+			Controller.SetPlayLength(FGenericPlatformMath::Max<float>(SequenceLength, MINIMUM_ANIMATION_LENGTH));
+
+			for (const FString& NodeUid : SkeletonNodes)
+			{
+				if (const UInterchangeSceneNode* SkeletonSceneNode = Cast<UInterchangeSceneNode>(NodeContainer->GetNode(NodeUid)))
+				{
+					const FName BoneName = FName(*(SkeletonSceneNode->GetDisplayLabel()));
+					const int32 BoneIndex = Skeleton->GetReferenceSkeleton().FindBoneIndex(BoneName);
+					if (BoneIndex == INDEX_NONE)
+					{
+						//Skip this bone, we did not found it in the skeleton
+						continue;
+					}
+					FString PayloadKey;
+					if (SkeletonSceneNode->GetCustomTransformCurvePayloadKey(PayloadKey))
+					{
+						TOptional<UE::Interchange::FAnimationTransformPayloadData> AnimationTransformPayload = AnimationPayloads.FindChecked(PayloadKey).Get();
+						if (!AnimationTransformPayload.IsSet())
+						{
+							UE_LOG(LogInterchangeImport, Warning, TEXT("Invalid animation transform payload key [%s] AnimSequence asset %s"), *PayloadKey, *AssetName);
+							continue;
+						}
+
+						FRawAnimSequenceTrack RawTrack;
+						RawTrack.PosKeys.Reserve(BakeKeyCount);
+						RawTrack.RotKeys.Reserve(BakeKeyCount);
+						RawTrack.ScaleKeys.Reserve(BakeKeyCount);
+						TArray<float> TimeKeys;
+						TimeKeys.Reserve(BakeKeyCount);
+
+						for (double CurrentTime = RangeStart; CurrentTime <= RangeStop + SMALL_NUMBER; CurrentTime += BakeInterval)
+						{
+							//Default value to identity
+							FVector3f Position(0.0f);
+							FVector3f Euler(0.0f);
+							FVector3f Scale(1.0f);
+							//Fill the anim sequence raw data
+							for (int32 CurveIndex = 0; CurveIndex < AnimationTransformPayload->TransformCurves.Num(); ++CurveIndex)
+							{
+								const UE::Interchange::FAnimationCurveTransformPayloadData& CurveData = AnimationTransformPayload->TransformCurves[CurveIndex];
+								if (CurveData.TransformChannel == EInterchangeTransformCurveChannel::TranslationX)
+								{
+									Position.X = CurveData.Curve.Eval(CurrentTime);
+								}
+								else if (CurveData.TransformChannel == EInterchangeTransformCurveChannel::TranslationY)
+								{
+									Position.Y = CurveData.Curve.Eval(CurrentTime);
+								}
+								else if (CurveData.TransformChannel == EInterchangeTransformCurveChannel::TranslationZ)
+								{
+									Position.Z = CurveData.Curve.Eval(CurrentTime);
+								}
+								else if (CurveData.TransformChannel == EInterchangeTransformCurveChannel::EulerX)
+								{
+									Euler.X = CurveData.Curve.Eval(CurrentTime);
+								}
+								else if (CurveData.TransformChannel == EInterchangeTransformCurveChannel::EulerY)
+								{
+									Euler.Y = CurveData.Curve.Eval(CurrentTime);
+								}
+								else if (CurveData.TransformChannel == EInterchangeTransformCurveChannel::EulerZ)
+								{
+									Euler.Z = CurveData.Curve.Eval(CurrentTime);
+								}
+								else if (CurveData.TransformChannel == EInterchangeTransformCurveChannel::ScaleX)
+								{
+									Scale.X = CurveData.Curve.Eval(CurrentTime);
+								}
+								else if (CurveData.TransformChannel == EInterchangeTransformCurveChannel::ScaleY)
+								{
+									Scale.Y = CurveData.Curve.Eval(CurrentTime);
+								}
+								else if (CurveData.TransformChannel == EInterchangeTransformCurveChannel::ScaleZ)
+								{
+									Scale.Z = CurveData.Curve.Eval(CurrentTime);
+								}
+							}
+							FQuat4f Quaternion = FQuat4f::MakeFromEuler(Euler);
+							RawTrack.ScaleKeys.Add(Scale);
+							RawTrack.PosKeys.Add(Position);
+							RawTrack.RotKeys.Add(Quaternion);
+							//Animation are always translated to zero
+							TimeKeys.Add(CurrentTime - RangeStart);
+						}
+
+						//Make sure we create the correct amount of keys
+						check(RawTrack.ScaleKeys.Num() == BakeKeyCount
+							&& RawTrack.PosKeys.Num() == BakeKeyCount
+							&& RawTrack.RotKeys.Num() == BakeKeyCount
+							&& TimeKeys.Num() == BakeKeyCount);
+
+						//add new track
+						Controller.AddBoneTrack(BoneName);
+						Controller.SetBoneTrackKeys(BoneName, RawTrack.PosKeys, RawTrack.RotKeys, RawTrack.ScaleKeys);
+					}
+				}
+			}
+		}
+	}
+
+	FFrameRate ConvertSampleRatetoFrameRate(double SampleRate)
+	{
+		double IntegralPart = FMath::FloorToDouble(SampleRate);
+		double FractionalPart = SampleRate - IntegralPart;
+		const int32 Tolerance = 1000000;
+		double Divisor = static_cast<double>(FMath::GreatestCommonDivisor(FMath::RoundToInt(FractionalPart * Tolerance), Tolerance));
+		int32 Denominator = static_cast<int32>(Tolerance / Divisor);
+		int32 Numerator = static_cast<int32>(IntegralPart * Denominator + FMath::RoundToDouble(FractionalPart * Tolerance) / Divisor);
+		check(Denominator != 0);
+		return FFrameRate(Numerator, Denominator);
+	}
+
+} //namespace UE::Interchange::Private
+#endif //WITH_EDITOR
+
+UClass* UInterchangeAnimSequenceFactory::GetFactoryClass() const
+{
+	return UAnimSequence::StaticClass();
+}
+
+UObject* UInterchangeAnimSequenceFactory::CreateEmptyAsset(const FCreateAssetParams& Arguments)
+{
+#if !WITH_EDITOR || !WITH_EDITORONLY_DATA
+
+	UE_LOG(LogInterchangeImport, Error, TEXT("Cannot import animsequence asset in runtime, this is an editor only feature."));
+	return nullptr;
+#else
+	UAnimSequence* AnimSequence = nullptr;
+	if (!Arguments.AssetNode || !Arguments.AssetNode->GetObjectClass()->IsChildOf(GetFactoryClass()))
+	{
+		return nullptr;
+	}
+
+	const UInterchangeAnimSequenceFactoryNode* AnimSequenceFactoryNode = Cast<UInterchangeAnimSequenceFactoryNode>(Arguments.AssetNode);
+	if (AnimSequenceFactoryNode == nullptr)
+	{
+		return nullptr;
+	}
+
+	// create an asset if it doesn't exist
+	UObject* ExistingAsset = StaticFindObject(nullptr, Arguments.Parent, *Arguments.AssetName);
+
+	// create a new material or overwrite existing asset, if possible
+	if (!ExistingAsset)
+	{
+		AnimSequence = NewObject<UAnimSequence>(Arguments.Parent, *Arguments.AssetName, RF_Public | RF_Standalone);
+	}
+	else if (ExistingAsset->GetClass()->IsChildOf(UAnimSequence::StaticClass()))
+	{
+		//This is a reimport, we are just re-updating the source data
+		AnimSequence = Cast<UAnimSequence>(ExistingAsset);
+	}
+
+	if (!AnimSequence)
+	{
+		UE_LOG(LogInterchangeImport, Warning, TEXT("Could not create AnimSequence asset %s"), *Arguments.AssetName);
+		return nullptr;
+	}
+
+	AnimSequence->PreEditChange(nullptr);
+
+	return AnimSequence;
+#endif //else !WITH_EDITOR || !WITH_EDITORONLY_DATA
+}
+
+UObject* UInterchangeAnimSequenceFactory::CreateAsset(const FCreateAssetParams& Arguments)
+{
+#if !WITH_EDITOR || !WITH_EDITORONLY_DATA
+
+	UE_LOG(LogInterchangeImport, Error, TEXT("Cannot import AnimSequence asset in runtime, this is an editor only feature."));
+	return nullptr;
+
+#else
+
+	if (!Arguments.AssetNode || !Arguments.AssetNode->GetObjectClass()->IsChildOf(GetFactoryClass()))
+	{
+		return nullptr;
+	}
+
+	UInterchangeAnimSequenceFactoryNode* AnimSequenceFactoryNode = Cast<UInterchangeAnimSequenceFactoryNode>(Arguments.AssetNode);
+	if (AnimSequenceFactoryNode == nullptr)
+	{
+		return nullptr;
+	}
+
+	FString SkeletonUid;
+	if (!AnimSequenceFactoryNode->GetCustomSkeletonFactoryNodeUid(SkeletonUid))
+	{
+		//Do not create a empty anim sequence, we need skeleton that contain animation
+		return nullptr;
+	}
+
+	const UInterchangeSkeletonFactoryNode* SkeletonFactoryNode = Cast<UInterchangeSkeletonFactoryNode>(Arguments.NodeContainer->GetNode(SkeletonUid));
+	if (!SkeletonFactoryNode)
+	{
+		UE_LOG(LogInterchangeImport, Warning, TEXT("Invalid skeleton factory node, the skeleton factory node is obligatory to import this animsequence [%s]!"), *Arguments.AssetName);
+		return nullptr;
+	}
+
+	USkeleton* Skeleton = nullptr;
+
+	FSoftObjectPath SpecifiedSkeleton;
+	AnimSequenceFactoryNode->GetCustomSkeletonSoftObjectPath(SpecifiedSkeleton);
+	if (Skeleton == nullptr)
+	{
+		UObject* SkeletonObject = nullptr;
+
+		if (SpecifiedSkeleton.IsValid())
+		{
+			SkeletonObject = SpecifiedSkeleton.TryLoad();
+		}
+		else if (SkeletonFactoryNode->ReferenceObject.IsValid())
+		{
+			SkeletonObject = SkeletonFactoryNode->ReferenceObject.TryLoad();
+		}
+
+		if (SkeletonObject)
+		{
+			Skeleton = Cast<USkeleton>(SkeletonObject);
+
+		}
+
+		if (!ensure(Skeleton))
+		{
+			UE_LOG(LogInterchangeImport, Warning, TEXT("Invalid Skeleton when importing animation sequence asset %s"), *Arguments.AssetName);
+			return nullptr;
+		}
+	}
+
+	const IInterchangeAnimationPayloadInterface* AnimSequenceTranslatorPayloadInterface = Cast<IInterchangeAnimationPayloadInterface>(Arguments.Translator);
+	if (!AnimSequenceTranslatorPayloadInterface)
+	{
+		UE_LOG(LogInterchangeImport, Error, TEXT("Cannot import AnimSequence, the translator do not implement the IInterchangeAnimationPayloadInterface."));
+		return nullptr;
+	}
+
+	const UClass* AnimSequenceClass = AnimSequenceFactoryNode->GetObjectClass();
+	check(AnimSequenceClass && AnimSequenceClass->IsChildOf(GetFactoryClass()));
+
+	// create an asset if it doesn't exist
+	UObject* ExistingAsset = StaticFindObject(nullptr, Arguments.Parent, *Arguments.AssetName);
+
+	UObject* AnimSequenceObject = nullptr;
+	// create a new material or overwrite existing asset, if possible
+	if (!ExistingAsset)
+	{
+		//NewObject is not thread safe, the asset registry directory watcher tick on the main thread can trig before we finish initializing the UObject and will crash
+		//The UObject should have been create by calling CreateEmptyAsset on the main thread.
+		check(IsInGameThread());
+		AnimSequenceObject = NewObject<UObject>(Arguments.Parent, AnimSequenceClass, *Arguments.AssetName, RF_Public | RF_Standalone);
+	}
+	else if (ExistingAsset->GetClass()->IsChildOf(AnimSequenceClass))
+	{
+		//This is a reimport, we are just re-updating the source data
+		AnimSequenceObject = ExistingAsset;
+	}
+
+	if (!AnimSequenceObject)
+	{
+		UE_LOG(LogInterchangeImport, Error, TEXT("Could not create AnimSequence asset %s"), *Arguments.AssetName);
+		return nullptr;
+	}
+
+	UAnimSequence* AnimSequence = Cast<UAnimSequence>(AnimSequenceObject);
+
+	const bool bIsReImport = (Arguments.ReimportObject != nullptr);
+
+	if (!ensure(AnimSequence))
+	{
+		if (!bIsReImport)
+		{
+			UE_LOG(LogInterchangeImport, Error, TEXT("Could not create AnimSequence asset %s"), *Arguments.AssetName);
+		}
+		else
+		{
+			UE_LOG(LogInterchangeImport, Error, TEXT("Could not find reimported AnimSequence asset %s"), *Arguments.AssetName);
+		}
+		return nullptr;
+	}
+
+	//Fill the animsequence data, we need to retrieve the skeleton and then ask the payload for every joint
+	{
+		FFrameRate FrameRate(30, 1);
+		double SampleRate = 30.0;
+		
+		//If we import bone track we will set the framerate according to the bone track settings
+		//The default is 30Hz when there is no bone tracks, i.e. like when we just have morphtarget attribute curve animation.
+		//Curve do not need framerate, framerate is needed if we want to bake curve animation data.
+		//TODO: discuss with animation team so we can pass curves for bones transform animation instead of baked data.
+		bool bImportBoneTracks = false;
+		if (AnimSequenceFactoryNode->GetCustomImportBoneTracks(bImportBoneTracks) && bImportBoneTracks)
+		{
+			if (AnimSequenceFactoryNode->GetCustomImportBoneTracksSampleRate(SampleRate))
+			{
+				FrameRate = UE::Interchange::Private::ConvertSampleRatetoFrameRate(SampleRate);
+			}
+		}
+
+		IAnimationDataController& Controller = AnimSequence->GetController();
+		Controller.OpenBracket(NSLOCTEXT("InterchangeAnimSequenceFactory", "ImportAnimationInterchange_Bracket", "Importing Animation (Interchange)"));
+		AnimSequence->SetSkeleton(Skeleton);
+		AnimSequence->ImportFileFramerate = SampleRate;
+		AnimSequence->ImportResampleFramerate = SampleRate;
+		Controller.SetFrameRate(FrameRate);
+
+		UE::Interchange::Private::RetrieveAnimationPayloads(AnimSequence
+			, AnimSequenceFactoryNode
+			, Arguments.NodeContainer
+			, SkeletonFactoryNode
+			, Skeleton
+			, AnimSequenceTranslatorPayloadInterface
+			, Arguments.AssetName);
+
+		Controller.NotifyPopulated();
+		Controller.CloseBracket(false);
+	}
+
+	if (!bIsReImport)
+	{
+		/** Apply all AnimSequenceFactoryNode custom attributes to the skeletal mesh asset */
+		AnimSequenceFactoryNode->ApplyAllCustomAttributeToObject(AnimSequence);
+	}
+	else
+	{
+		//Apply the re import strategy 
+		UInterchangeAssetImportData* InterchangeAssetImportData = Cast<UInterchangeAssetImportData>(AnimSequence->AssetImportData);
+		UInterchangeBaseNode* PreviousNode = nullptr;
+		if (InterchangeAssetImportData)
+		{
+			PreviousNode = InterchangeAssetImportData->NodeContainer->GetNode(InterchangeAssetImportData->NodeUniqueID);
+		}
+		UInterchangeBaseNode* CurrentNode = NewObject<UInterchangeBaseNode>(GetTransientPackage(), UInterchangeAnimSequenceFactoryNode::StaticClass());
+		UInterchangeBaseNode::CopyStorage(AnimSequenceFactoryNode, CurrentNode);
+		CurrentNode->FillAllCustomAttributeFromObject(AnimSequence);
+		UE::Interchange::FFactoryCommon::ApplyReimportStrategyToAsset(AnimSequence, PreviousNode, CurrentNode, AnimSequenceFactoryNode);
+	}
+
+	//Getting the file Hash will cache it into the source data
+	Arguments.SourceData->GetFileContentHash();
+
+	//The interchange completion task (call in the GameThread after the factories pass), will call PostEditChange which will trig another asynchronous system that will build all material in parallel
+	return AnimSequenceObject;
+
+#endif //else !WITH_EDITOR || !WITH_EDITORONLY_DATA
+}
+
+/* This function is call in the completion task on the main thread, use it to call main thread post creation step for your assets*/
+void UInterchangeAnimSequenceFactory::PreImportPreCompletedCallback(const FImportPreCompletedCallbackParams& Arguments)
+{
+	check(IsInGameThread());
+	Super::PreImportPreCompletedCallback(Arguments);
+
+	// TODO: make sure this works at runtime
+#if WITH_EDITORONLY_DATA
+	if (ensure(Arguments.ImportedObject && Arguments.SourceData))
+	{
+		// We must call the Update of the asset source file in the main thread because UAssetImportData::Update execute some delegate we do not control
+		UAnimSequence* AnimSequence = CastChecked<UAnimSequence>(Arguments.ImportedObject);
+
+		UAssetImportData* ImportDataPtr = AnimSequence->AssetImportData;
+		UE::Interchange::FFactoryCommon::FUpdateImportAssetDataParameters UpdateImportAssetDataParameters(AnimSequence, ImportDataPtr, Arguments.SourceData, Arguments.NodeUniqueID, Arguments.NodeContainer, Arguments.Pipelines);
+		ImportDataPtr = UE::Interchange::FFactoryCommon::UpdateImportAssetData(UpdateImportAssetDataParameters);
+		AnimSequence->AssetImportData = ImportDataPtr;
+	}
+#endif
+}
+
+/************************************************************************/
+/* Automation tests                                                     */
+/************************************************************************/
+#if WITH_EDITOR
+#if WITH_DEV_AUTOMATION_TESTS
+
+#include "Misc/AutomationTest.h"
+
+IMPLEMENT_SIMPLE_AUTOMATION_TEST(FInterchangeAnimSequenceTest, "System.Runtime.Interchange.ConvertSampleRatetoFrameRate", EAutomationTestFlags::EditorContext | EAutomationTestFlags::EngineFilter)
+
+bool FInterchangeAnimSequenceTest::RunTest(const FString& Parameters)
+{
+	FFrameRate FrameRate;
+	FrameRate = UE::Interchange::Private::ConvertSampleRatetoFrameRate(120.0);
+	TestEqual(TEXT("`ConvertSampleRatetoFrameRate` Error converting 120.0 to FFrameRate"), FrameRate.Numerator, 120);
+	TestEqual(TEXT("`ConvertSampleRatetoFrameRate` Error converting 120.0 to FFrameRate"), FrameRate.Denominator, 1);
+	
+	FrameRate = UE::Interchange::Private::ConvertSampleRatetoFrameRate(100.0);
+	TestEqual(TEXT("`ConvertSampleRatetoFrameRate` Error converting 100.0 to FFrameRate"), FrameRate.Numerator, 100);
+	TestEqual(TEXT("`ConvertSampleRatetoFrameRate` Error converting 100.0 to FFrameRate"), FrameRate.Denominator, 1);
+
+	FrameRate = UE::Interchange::Private::ConvertSampleRatetoFrameRate(60.0);
+	TestEqual(TEXT("`ConvertSampleRatetoFrameRate` Error converting 60.0 to FFrameRate"), FrameRate.Numerator, 60);
+	TestEqual(TEXT("`ConvertSampleRatetoFrameRate` Error converting 60.0 to FFrameRate"), FrameRate.Denominator, 1);
+
+	FrameRate = UE::Interchange::Private::ConvertSampleRatetoFrameRate(50.0);
+	TestEqual(TEXT("`ConvertSampleRatetoFrameRate` Error converting 50.0 to FFrameRate"), FrameRate.Numerator, 50);
+	TestEqual(TEXT("`ConvertSampleRatetoFrameRate` Error converting 50.0 to FFrameRate"), FrameRate.Denominator, 1);
+
+	FrameRate = UE::Interchange::Private::ConvertSampleRatetoFrameRate(48.0);
+	TestEqual(TEXT("`ConvertSampleRatetoFrameRate` Error converting 48.0 to FFrameRate"), FrameRate.Numerator, 48);
+	TestEqual(TEXT("`ConvertSampleRatetoFrameRate` Error converting 48.0 to FFrameRate"), FrameRate.Denominator, 1);
+
+	FrameRate = UE::Interchange::Private::ConvertSampleRatetoFrameRate(30.0);
+	TestEqual(TEXT("`ConvertSampleRatetoFrameRate` Error converting 30.0 to FFrameRate"), FrameRate.Numerator, 30);
+	TestEqual(TEXT("`ConvertSampleRatetoFrameRate` Error converting 30.0 to FFrameRate"), FrameRate.Denominator, 1);
+
+	FrameRate = UE::Interchange::Private::ConvertSampleRatetoFrameRate(29.97);
+	TestEqual(TEXT("`ConvertSampleRatetoFrameRate` Error converting 29.97 to FFrameRate"), FrameRate.Numerator, 2997);
+	TestEqual(TEXT("`ConvertSampleRatetoFrameRate` Error converting 29.97 to FFrameRate"), FrameRate.Denominator, 100);
+
+	FrameRate = UE::Interchange::Private::ConvertSampleRatetoFrameRate(25.0);
+	TestEqual(TEXT("`ConvertSampleRatetoFrameRate` Error converting 25.0 to FFrameRate"), FrameRate.Numerator, 25);
+	TestEqual(TEXT("`ConvertSampleRatetoFrameRate` Error converting 25.0 to FFrameRate"), FrameRate.Denominator, 1);
+
+	FrameRate = UE::Interchange::Private::ConvertSampleRatetoFrameRate(24.0);
+	TestEqual(TEXT("`ConvertSampleRatetoFrameRate` Error converting 24.0 to FFrameRate"), FrameRate.Numerator, 24);
+	TestEqual(TEXT("`ConvertSampleRatetoFrameRate` Error converting 24.0 to FFrameRate"), FrameRate.Denominator, 1);
+
+	FrameRate = UE::Interchange::Private::ConvertSampleRatetoFrameRate(23.976);
+	TestEqual(TEXT("`ConvertSampleRatetoFrameRate` Error converting 23.976 to FFrameRate"), FrameRate.Numerator, 2997);
+	TestEqual(TEXT("`ConvertSampleRatetoFrameRate` Error converting 23.976 to FFrameRate"), FrameRate.Denominator, 125);
+
+	FrameRate = UE::Interchange::Private::ConvertSampleRatetoFrameRate(96.0);
+	TestEqual(TEXT("`ConvertSampleRatetoFrameRate` Error converting 96.0 to FFrameRate"), FrameRate.Numerator, 96);
+	TestEqual(TEXT("`ConvertSampleRatetoFrameRate` Error converting 96.0 to FFrameRate"), FrameRate.Denominator, 1);
+
+	FrameRate = UE::Interchange::Private::ConvertSampleRatetoFrameRate(72.0);
+	TestEqual(TEXT("`ConvertSampleRatetoFrameRate` Error converting 72.0 to FFrameRate"), FrameRate.Numerator, 72);
+	TestEqual(TEXT("`ConvertSampleRatetoFrameRate` Error converting 72.0 to FFrameRate"), FrameRate.Denominator, 1);
+
+	FrameRate = UE::Interchange::Private::ConvertSampleRatetoFrameRate(59.94);
+	TestEqual(TEXT("`ConvertSampleRatetoFrameRate` Error converting 59.94 to FFrameRate"), FrameRate.Numerator, 2997);
+	TestEqual(TEXT("`ConvertSampleRatetoFrameRate` Error converting 59.94 to FFrameRate"), FrameRate.Denominator, 50);
+
+	FrameRate = UE::Interchange::Private::ConvertSampleRatetoFrameRate(119.88);
+	TestEqual(TEXT("`ConvertSampleRatetoFrameRate` Error converting 119.88 to FFrameRate"), FrameRate.Numerator, 2997);
+	TestEqual(TEXT("`ConvertSampleRatetoFrameRate` Error converting 119.88 to FFrameRate"), FrameRate.Denominator, 25);
+	
+	return true;
+}
+
+#endif //WITH_DEV_AUTOMATION_TESTS
+#endif //WITH_EDITOR
diff --git a/Engine/Plugins/Experimental/Interchange/Runtime/Source/Import/Private/Fbx/InterchangeFbxTranslator.cpp b/Engine/Plugins/Experimental/Interchange/Runtime/Source/Import/Private/Fbx/InterchangeFbxTranslator.cpp
--- a/Engine/Plugins/Experimental/Interchange/Runtime/Source/Import/Private/Fbx/InterchangeFbxTranslator.cpp
+++ b/Engine/Plugins/Experimental/Interchange/Runtime/Source/Import/Private/Fbx/InterchangeFbxTranslator.cpp
@@ -5,6 +5,7 @@
 #include "Dom/JsonValue.h"
 #include "Dom/JsonObject.h"
 #include "HAL/PlatformFileManager.h"
+#include "InterchangeCommonAnimationPayload.h"
 #include "InterchangeDispatcher.h"
 #include "InterchangeDispatcherTask.h"
 #include "InterchangeManager.h"
@@ -390,6 +391,90 @@ TFuture<TOptional<UE::Interchange::FSkeletalMeshBlendShapePayloadData>> UInterch
 	return Promise->GetFuture();
 }
 
+TFuture<TOptional<UE::Interchange::FAnimationTransformPayloadData>> UInterchangeFbxTranslator::GetAnimationTransformPayloadData(const FString& PayLoadKey) const
+{
+	TSharedPtr<TPromise<TOptional<UE::Interchange::FAnimationTransformPayloadData>>> Promise = MakeShared<TPromise<TOptional<UE::Interchange::FAnimationTransformPayloadData>>>();
+
+	if (!Dispatcher.IsValid())
+	{
+		Promise->SetValue(TOptional<UE::Interchange::FAnimationTransformPayloadData>());
+		return Promise->GetFuture();
+	}
+
+	// Create a json command to read the fbx file
+	FString JsonCommand = CreateFetchPayloadFbxCommand(PayLoadKey);
+	const int32 CreatedTaskIndex = Dispatcher->AddTask(JsonCommand, FInterchangeDispatcherTaskCompleted::CreateLambda([this, Promise](const int32 TaskIndex)
+		{
+			UE::Interchange::ETaskState TaskState;
+			FString JsonResult;
+			TArray<FString> JsonMessages;
+			Dispatcher->GetTaskState(TaskIndex, TaskState, JsonResult, JsonMessages);
+
+			// Parse the Json messages into UInterchangeResults
+			for (const FString& JsonMessage : JsonMessages)
+			{
+				AddMessage(UInterchangeResult::FromJson(JsonMessage));
+			}
+
+			if (TaskState != UE::Interchange::ETaskState::ProcessOk)
+			{
+				Promise->SetValue(TOptional<UE::Interchange::FAnimationTransformPayloadData>());
+				return;
+			}
+
+			// Grab the result file and fill the BaseNodeContainer
+			UE::Interchange::FJsonFetchPayloadCmd::JsonResultParser ResultParser;
+			ResultParser.FromJson(JsonResult);
+			FString AnimationTransformPayloadFilename = ResultParser.GetResultFilename();
+
+			if (!ensure(FPaths::FileExists(AnimationTransformPayloadFilename)))
+			{
+				// TODO log an error saying the payload file does not exist even if the get payload command succeeded
+				Promise->SetValue(TOptional<UE::Interchange::FAnimationTransformPayloadData>());
+				return;
+			}
+
+			UE::Interchange::FAnimationTransformPayloadData AnimationTransformPayload;
+
+			// All sub object should be gone with the reset
+			TArray64<uint8> Buffer;
+			FFileHelper::LoadFileToArray(Buffer, *AnimationTransformPayloadFilename);
+			uint8* FileData = Buffer.GetData();
+			int64 FileDataSize = Buffer.Num();
+			if (FileDataSize < 1)
+			{
+				// Nothing to load from this file
+				Promise->SetValue(TOptional<UE::Interchange::FAnimationTransformPayloadData>());
+				return;
+			}
+
+			// Buffer keeps the ownership of the data, the large memory reader is use to serialize the TMap
+			FLargeMemoryReader Ar(FileData, FileDataSize);
+			TArray<FInterchangeCurve> InterchangeCurves;
+			Ar << InterchangeCurves;
+			AnimationTransformPayload.TransformCurves.AddDefaulted(static_cast<int32>(EInterchangeTransformCurveChannel::TransformChannelCount));
+			for (const FInterchangeCurve& InterchangeCurve : InterchangeCurves)
+			{
+				int32 TransformChannelIndex = static_cast<int32>(InterchangeCurve.TransformChannel);
+				if (AnimationTransformPayload.TransformCurves.IsValidIndex(TransformChannelIndex))
+				{
+					UE::Interchange::FAnimationCurveTransformPayloadData& CurveTransformData = AnimationTransformPayload.TransformCurves[static_cast<int32>(InterchangeCurve.TransformChannel)];
+					CurveTransformData.TransformChannel = InterchangeCurve.TransformChannel;
+					InterchangeCurve.ToRichCurve(CurveTransformData.Curve);
+				}
+			}
+			Promise->SetValue(MoveTemp(AnimationTransformPayload));
+		}));
+
+	// The task was not added to the dispatcher
+	if (CreatedTaskIndex == INDEX_NONE)
+	{
+		Promise->SetValue(TOptional<UE::Interchange::FAnimationTransformPayloadData>{});
+	}
+
+	return Promise->GetFuture();
+}
+
 FString UInterchangeFbxTranslator::CreateLoadFbxFileCommand(const FString& FbxFilePath) const
 {
 	UE::Interchange::FJsonLoadSourceCmd LoadSourceCommand(TEXT("FBX"), FbxFilePath);
diff --git a/Engine/Plugins/Experimental/Interchange/Runtime/Source/Import/Private/InterchangeImportModule.cpp b/Engine/Plugins/Experimental/Interchange/Runtime/Source/Import/Private/InterchangeImportModule.cpp
--- a/Engine/Plugins/Experimental/Interchange/Runtime/Source/Import/Private/InterchangeImportModule.cpp
+++ b/Engine/Plugins/Experimental/Interchange/Runtime/Source/Import/Private/InterchangeImportModule.cpp
@@ -2,6 +2,7 @@
 
 #include "InterchangeImportModule.h"
 
+#include "Animation/InterchangeAnimSequenceFactory.h"
 #include "CoreMinimal.h"
 #include "Fbx/InterchangeFbxTranslator.h"
 #include "Gltf/InterchangeGltfTranslator.h"
@@ -88,6 +89,7 @@ void FInterchangeImportModule::StartupModule()
 		InterchangeManager.RegisterFactory(UInterchangeStaticMeshFactory::StaticClass());
 		InterchangeManager.RegisterFactory(UInterchangePhysicsAssetFactory::StaticClass());
 		InterchangeManager.RegisterFactory(UInterchangeActorFactory::StaticClass());
+		InterchangeManager.RegisterFactory(UInterchangeAnimSequenceFactory::StaticClass());
 	};
 
 	if (GEngine)
diff --git a/Engine/Plugins/Experimental/Interchange/Runtime/Source/Import/Private/Mesh/InterchangeSkeletalMeshFactory.cpp b/Engine/Plugins/Experimental/Interchange/Runtime/Source/Import/Private/Mesh/InterchangeSkeletalMeshFactory.cpp
--- a/Engine/Plugins/Experimental/Interchange/Runtime/Source/Import/Private/Mesh/InterchangeSkeletalMeshFactory.cpp
+++ b/Engine/Plugins/Experimental/Interchange/Runtime/Source/Import/Private/Mesh/InterchangeSkeletalMeshFactory.cpp
@@ -14,6 +14,7 @@
 #include "InterchangeSkeletalMeshLodDataNode.h"
 #include "InterchangeSkeletalMeshFactoryNode.h"
 #include "InterchangeSkeletonFactoryNode.h"
+#include "InterchangeSkeletonHelper.h"
 #include "InterchangeSourceData.h"
 #include "InterchangeTranslatorBase.h"
 #include "Math/GenericOctree.h"
@@ -50,135 +51,6 @@ namespace UE
 				FString TranslatorPayloadKey;
 			};
 
-			struct FJointInfo
-			{
-				FString Name;
-				int32 ParentIndex;  // 0 if this is the root bone.  
-				FTransform	LocalTransform; // local transform
-			};
-
-			void RecursiveAddBones(const UInterchangeBaseNodeContainer* NodeContainer, const FString& JointNodeId, TArray <FJointInfo>& JointInfos, int32 ParentIndex, TArray<SkeletalMeshImportData::FBone>& RefBonesBinary, const bool bUseTimeZeroAsBindPose, bool & bOutDiffPose)
-			{
-				const UInterchangeSceneNode* JointNode = Cast<UInterchangeSceneNode>(NodeContainer->GetNode(JointNodeId));
-				if (!JointNode)
-				{
-					UE_LOG(LogInterchangeImport, Warning, TEXT("Invalid Skeleton Joint"));
-					return;
-				}
-
-				int32 JointInfoIndex = JointInfos.Num();
-				FJointInfo& Info = JointInfos.AddZeroed_GetRef();
-				Info.Name = JointNode->GetDisplayLabel();
-
-				FTransform LocalTransform;
-				ensure(JointNode->GetCustomLocalTransform(LocalTransform));
-
-				FTransform TimeZeroLocalTransform;
-				const bool bHasTimeZeroTransform = JointNode->GetCustomTimeZeroLocalTransform(TimeZeroLocalTransform);
-				FTransform BindPoseLocalTransform;
-				const bool bHasBindPoseTransform = JointNode->GetCustomBindPoseLocalTransform(BindPoseLocalTransform);
-
-				Info.LocalTransform = bHasBindPoseTransform ? BindPoseLocalTransform : LocalTransform;
-				//If user want to bind the mesh at time zero try to get the time zero transform
-				if (bUseTimeZeroAsBindPose && bHasTimeZeroTransform)
-				{
-					if (bHasBindPoseTransform)
-					{
-						if(!TimeZeroLocalTransform.Equals(Info.LocalTransform))
-						{
-							bOutDiffPose = true;
-						}
-					}
-					Info.LocalTransform = TimeZeroLocalTransform;
-				}
-				else if (bHasBindPoseTransform)
-				{
-					Info.LocalTransform = BindPoseLocalTransform;
-				}
-
-				Info.ParentIndex = ParentIndex;
-
-				SkeletalMeshImportData::FBone& Bone = RefBonesBinary.AddZeroed_GetRef();
-				Bone.Name = Info.Name;
-				Bone.BonePos.Transform = FTransform3f(Info.LocalTransform);
-				Bone.ParentIndex = ParentIndex;
-				//Fill the scrap we do not need
-				Bone.BonePos.Length = 0.0f;
-				Bone.BonePos.XSize = 1.0f;
-				Bone.BonePos.YSize = 1.0f;
-				Bone.BonePos.ZSize = 1.0f;
-				
-				const TArray<FString> ChildrenIds = NodeContainer->GetNodeChildrenUids(JointNodeId);
-				Bone.NumChildren = ChildrenIds.Num();
-				for (int32 ChildIndex = 0; ChildIndex < ChildrenIds.Num(); ++ChildIndex)
-				{
-					RecursiveAddBones(NodeContainer, ChildrenIds[ChildIndex], JointInfos, JointInfoIndex, RefBonesBinary, bUseTimeZeroAsBindPose, bOutDiffPose);
-				}
-			}
-
-			bool ProcessImportMeshSkeleton(const USkeleton* SkeletonAsset, FReferenceSkeleton& RefSkeleton, int32& SkeletalDepth, const UInterchangeBaseNodeContainer* NodeContainer, const FString& RootJointNodeId, TArray<SkeletalMeshImportData::FBone>& RefBonesBinary, const bool bUseTimeZeroAsBindPose, bool& bOutDiffPose)
-			{
-				auto FixupBoneName = [](FString BoneName)
-				{
-					BoneName.TrimStartAndEndInline();
-					BoneName.ReplaceInline(TEXT(" "), TEXT("-"), ESearchCase::IgnoreCase);
-					return BoneName;
-				};
-
-				RefBonesBinary.Empty();
-				// Setup skeletal hierarchy + names structure.
-				RefSkeleton.Empty();
-
-				FReferenceSkeletonModifier RefSkelModifier(RefSkeleton, SkeletonAsset);
-				TArray <FJointInfo> JointInfos;
-				RecursiveAddBones(NodeContainer, RootJointNodeId, JointInfos, INDEX_NONE, RefBonesBinary, bUseTimeZeroAsBindPose, bOutDiffPose);
-				if (bOutDiffPose)
-				{
-					//bOutDiffPose can only be true if the user ask to bind with time zero transform.
-					ensure(bUseTimeZeroAsBindPose);
-				}
-				// Digest bones to the serializable format.
-				for (int32 b = 0; b < JointInfos.Num(); b++)
-				{
-					const FJointInfo& BinaryBone = JointInfos[b];
-
-					const FString BoneName = FixupBoneName(BinaryBone.Name);
-					const FMeshBoneInfo BoneInfo(FName(*BoneName, FNAME_Add), BinaryBone.Name, BinaryBone.ParentIndex);
-					const FTransform BoneTransform(BinaryBone.LocalTransform);
-					if (RefSkeleton.FindRawBoneIndex(BoneInfo.Name) != INDEX_NONE)
-					{
-						UE_LOG(LogInterchangeImport, Error, TEXT("Invalid Skeleton because of non-unique bone names [%s]"), *BoneInfo.Name.ToString());
-						return false;
-					}
-					RefSkelModifier.Add(BoneInfo, BoneTransform);
-				}
-
-				// Add hierarchy index to each bone and detect max depth.
-				SkeletalDepth = 0;
-
-				TArray<int32> SkeletalDepths;
-				SkeletalDepths.Empty(JointInfos.Num());
-				SkeletalDepths.AddZeroed(JointInfos.Num());
-				for (int32 BoneIndex = 0; BoneIndex < RefSkeleton.GetRawBoneNum(); ++BoneIndex)
-				{
-					int32 Parent = RefSkeleton.GetRawParentIndex(BoneIndex);
-					int32 Depth = 1.0f;
-
-					SkeletalDepths[BoneIndex] = 1.0f;
-					if (Parent != INDEX_NONE)
-					{
-						Depth += SkeletalDepths[Parent];
-					}
-					if (SkeletalDepth < Depth)
-					{
-						SkeletalDepth = Depth;
-					}
-					SkeletalDepths[BoneIndex] = Depth;
-				}
-
-				return true;
-			}
-
 			void FillBlendShapeMeshDescriptionsPerBlendShapeName(const FMeshNodeContext& MeshNodeContext
 																 , TMap<FString, TOptional<UE::Interchange::FSkeletalMeshBlendShapePayloadData>>& BlendShapeMeshDescriptionsPerBlendShapeName
 																 , const IInterchangeSkeletalMeshPayloadInterface* SkeletalMeshTranslatorPayloadInterface
@@ -1240,7 +1112,7 @@ UObject* UInterchangeSkeletalMeshFactory::CreateAsset(const FCreateAssetParams&
 		bool bUseTimeZeroAsBindPose = false;
 		SkeletonNode->GetCustomUseTimeZeroForBindPose(bUseTimeZeroAsBindPose);
 		bool bDiffPose = false;
-		UE::Interchange::Private::ProcessImportMeshSkeleton(SkeletonReference, SkeletalMesh->GetRefSkeleton(), SkeletonDepth, Arguments.NodeContainer, RootJointNodeId, RefBonesBinary, bUseTimeZeroAsBindPose, bDiffPose);
+		UE::Interchange::Private::FSkeletonHelper::ProcessImportMeshSkeleton(SkeletonReference, SkeletalMesh->GetRefSkeleton(), SkeletonDepth, Arguments.NodeContainer, RootJointNodeId, RefBonesBinary, bUseTimeZeroAsBindPose, bDiffPose);
 		if (bSpecifiedSkeleton && !SkeletonReference->IsCompatibleMesh(SkeletalMesh))
 		{
 			UE_LOG(LogInterchangeImport, Warning, TEXT("The skeleton %s is incompatible with the imported skeletalmesh asset %s"), *SkeletonReference->GetName(), *Arguments.AssetName);
diff --git a/Engine/Plugins/Experimental/Interchange/Runtime/Source/Import/Private/Mesh/InterchangeSkeletonHelper.cpp b/Engine/Plugins/Experimental/Interchange/Runtime/Source/Import/Private/Mesh/InterchangeSkeletonHelper.cpp
new file mode 100644
--- /dev/null
+++ b/Engine/Plugins/Experimental/Interchange/Runtime/Source/Import/Private/Mesh/InterchangeSkeletonHelper.cpp
@@ -0,0 +1,314 @@
+// Copyright Epic Games, Inc. All Rights Reserved. 
+#include "InterchangeSkeletonHelper.h"
+
+#if WITH_EDITOR
+#include "Animation/Skeleton.h"
+#include "InterchangeImportLog.h"
+#include "InterchangeSceneNode.h"
+#include "ReferenceSkeleton.h"
+#include "Rendering/SkeletalMeshLODImporterData.h"
+#include "Nodes/InterchangeBaseNode.h"
+#include "Nodes/InterchangeBaseNodeContainer.h"
+
+
+namespace UE::Interchange::Private
+{
+
+	bool FSkeletonHelper::ProcessImportMeshSkeleton(const USkeleton* SkeletonAsset, FReferenceSkeleton& RefSkeleton, int32& SkeletalDepth, const UInterchangeBaseNodeContainer* NodeContainer, const FString& RootJointNodeId, TArray<SkeletalMeshImportData::FBone>& RefBonesBinary, const bool bUseTimeZeroAsBindPose, bool& bOutDiffPose)
+	{
+		auto FixupBoneName = [](FString BoneName)
+		{
+			BoneName.TrimStartAndEndInline();
+			BoneName.ReplaceInline(TEXT(" "), TEXT("-"), ESearchCase::IgnoreCase);
+			return BoneName;
+		};
+
+		RefBonesBinary.Empty();
+		// Setup skeletal hierarchy + names structure.
+		RefSkeleton.Empty();
+
+		FReferenceSkeletonModifier RefSkelModifier(RefSkeleton, SkeletonAsset);
+		TArray <FJointInfo> JointInfos;
+		RecursiveAddBones(NodeContainer, RootJointNodeId, JointInfos, INDEX_NONE, RefBonesBinary, bUseTimeZeroAsBindPose, bOutDiffPose);
+		if (bOutDiffPose)
+		{
+			//bOutDiffPose can only be true if the user ask to bind with time zero transform.
+			ensure(bUseTimeZeroAsBindPose);
+		}
+		// Digest bones to the serializable format.
+		for (int32 b = 0; b < JointInfos.Num(); b++)
+		{
+			const FJointInfo& BinaryBone = JointInfos[b];
+
+			const FString BoneName = FixupBoneName(BinaryBone.Name);
+			const FMeshBoneInfo BoneInfo(FName(*BoneName, FNAME_Add), BinaryBone.Name, BinaryBone.ParentIndex);
+			const FTransform BoneTransform(BinaryBone.LocalTransform);
+			if (RefSkeleton.FindRawBoneIndex(BoneInfo.Name) != INDEX_NONE)
+			{
+				UE_LOG(LogInterchangeImport, Error, TEXT("Invalid Skeleton because of non-unique bone names [%s]"), *BoneInfo.Name.ToString());
+				return false;
+			}
+			RefSkelModifier.Add(BoneInfo, BoneTransform);
+		}
+
+		// Add hierarchy index to each bone and detect max depth.
+		SkeletalDepth = 0;
+
+		TArray<int32> SkeletalDepths;
+		SkeletalDepths.Empty(JointInfos.Num());
+		SkeletalDepths.AddZeroed(JointInfos.Num());
+		for (int32 BoneIndex = 0; BoneIndex < RefSkeleton.GetRawBoneNum(); ++BoneIndex)
+		{
+			int32 Parent = RefSkeleton.GetRawParentIndex(BoneIndex);
+			int32 Depth = 1.0f;
+
+			SkeletalDepths[BoneIndex] = 1.0f;
+			if (Parent != INDEX_NONE)
+			{
+				Depth += SkeletalDepths[Parent];
+			}
+			if (SkeletalDepth < Depth)
+			{
+				SkeletalDepth = Depth;
+			}
+			SkeletalDepths[BoneIndex] = Depth;
+		}
+
+		return true;
+	}
+
+	bool FSkeletonHelper::IsCompatibleSkeleton(const USkeleton* Skeleton, const FString RootJoinUid, const UInterchangeBaseNodeContainer* BaseNodeContainer)
+	{
+		// at least % of bone should match 
+		int32 NumOfBoneMatches = 0;
+		//Make sure the specified Skeleton fit this skeletal mesh
+		const FReferenceSkeleton& SkeletonRef = Skeleton->GetReferenceSkeleton();
+		const int32 SkeletonBoneCount = SkeletonRef.GetRawBoneNum();
+
+		TArray<FMeshBoneInfo> SkeletalLodRawInfos;
+		SkeletalLodRawInfos.Reserve(SkeletonBoneCount);
+		RecursiveBuildSkeletalSkeleton(RootJoinUid, INDEX_NONE, BaseNodeContainer, SkeletalLodRawInfos);
+		const int32 SkeletalLodBoneCount = SkeletalLodRawInfos.Num();
+
+		// first ensure the parent exists for each bone
+		for (int32 MeshBoneIndex = 0; MeshBoneIndex < SkeletalLodBoneCount; MeshBoneIndex++)
+		{
+			FName MeshBoneName = SkeletalLodRawInfos[MeshBoneIndex].Name;
+			// See if Mesh bone exists in Skeleton.
+			int32 SkeletonBoneIndex = SkeletonRef.FindBoneIndex(MeshBoneName);
+
+			// if found, increase num of bone matches count
+			if (SkeletonBoneIndex != INDEX_NONE)
+			{
+				++NumOfBoneMatches;
+
+				// follow the parent chain to verify the chain is same
+				if (!DoesParentChainMatch(SkeletonBoneIndex, SkeletonRef, SkeletalLodRawInfos))
+				{
+					//Not compatible
+					return false;
+				}
+			}
+			else
+			{
+				int32 CurrentBoneId = MeshBoneIndex;
+				// if not look for parents that matches
+				while (SkeletonBoneIndex == INDEX_NONE && CurrentBoneId != INDEX_NONE)
+				{
+					// find Parent one see exists
+					const int32 ParentMeshBoneIndex = SkeletalLodGetParentIndex(SkeletalLodRawInfos, CurrentBoneId);
+					if (ParentMeshBoneIndex != INDEX_NONE)
+					{
+						// @TODO: make sure RefSkeleton's root ParentIndex < 0 if not, I'll need to fix this by checking TreeBoneIdx
+						FName ParentBoneName = SkeletalLodGetBoneName(SkeletalLodRawInfos, ParentMeshBoneIndex);
+						SkeletonBoneIndex = SkeletonRef.FindBoneIndex(ParentBoneName);
+					}
+
+					// root is reached
+					if (ParentMeshBoneIndex == 0)
+					{
+						break;
+					}
+					else
+					{
+						CurrentBoneId = ParentMeshBoneIndex;
+					}
+				}
+
+				// still no match, return false, no parent to look for
+				if (SkeletonBoneIndex == INDEX_NONE)
+				{
+					return false;
+				}
+
+				// second follow the parent chain to verify the chain is same
+				if (!DoesParentChainMatch(SkeletonBoneIndex, SkeletonRef, SkeletalLodRawInfos))
+				{
+					return false;
+				}
+			}
+		}
+
+		// originally we made sure at least matches more than 50% 
+		// but then slave components can't play since they're only partial
+		// if the hierarchy matches, and if it's more then 1 bone, we allow
+		return (NumOfBoneMatches > 0);
+	}
+
+	void FSkeletonHelper::RecursiveAddBones(const UInterchangeBaseNodeContainer* NodeContainer, const FString& JointNodeId, TArray <FJointInfo>& JointInfos, int32 ParentIndex, TArray<SkeletalMeshImportData::FBone>& RefBonesBinary, const bool bUseTimeZeroAsBindPose, bool& bOutDiffPose)
+	{
+		const UInterchangeSceneNode* JointNode = Cast<UInterchangeSceneNode>(NodeContainer->GetNode(JointNodeId));
+		if (!JointNode)
+		{
+			UE_LOG(LogInterchangeImport, Warning, TEXT("Invalid Skeleton Joint"));
+			return;
+		}
+
+		int32 JointInfoIndex = JointInfos.Num();
+		FJointInfo& Info = JointInfos.AddZeroed_GetRef();
+		Info.Name = JointNode->GetDisplayLabel();
+
+		FTransform LocalTransform;
+		ensure(JointNode->GetCustomLocalTransform(LocalTransform));
+
+		FTransform TimeZeroLocalTransform;
+		const bool bHasTimeZeroTransform = JointNode->GetCustomTimeZeroLocalTransform(TimeZeroLocalTransform);
+		FTransform BindPoseLocalTransform;
+		const bool bHasBindPoseTransform = JointNode->GetCustomBindPoseLocalTransform(BindPoseLocalTransform);
+
+		Info.LocalTransform = bHasBindPoseTransform ? BindPoseLocalTransform : LocalTransform;
+		//If user want to bind the mesh at time zero try to get the time zero transform
+		if (bUseTimeZeroAsBindPose && bHasTimeZeroTransform)
+		{
+			if (bHasBindPoseTransform)
+			{
+				if (!TimeZeroLocalTransform.Equals(Info.LocalTransform))
+				{
+					bOutDiffPose = true;
+				}
+			}
+			Info.LocalTransform = TimeZeroLocalTransform;
+		}
+		else if (bHasBindPoseTransform)
+		{
+			Info.LocalTransform = BindPoseLocalTransform;
+		}
+
+		Info.ParentIndex = ParentIndex;
+
+		SkeletalMeshImportData::FBone& Bone = RefBonesBinary.AddZeroed_GetRef();
+		Bone.Name = Info.Name;
+		Bone.BonePos.Transform = FTransform3f(Info.LocalTransform);
+		Bone.ParentIndex = ParentIndex;
+		//Fill the scrap we do not need
+		Bone.BonePos.Length = 0.0f;
+		Bone.BonePos.XSize = 1.0f;
+		Bone.BonePos.YSize = 1.0f;
+		Bone.BonePos.ZSize = 1.0f;
+
+		const TArray<FString> ChildrenIds = NodeContainer->GetNodeChildrenUids(JointNodeId);
+		Bone.NumChildren = ChildrenIds.Num();
+		for (int32 ChildIndex = 0; ChildIndex < ChildrenIds.Num(); ++ChildIndex)
+		{
+			RecursiveAddBones(NodeContainer, ChildrenIds[ChildIndex], JointInfos, JointInfoIndex, RefBonesBinary, bUseTimeZeroAsBindPose, bOutDiffPose);
+		}
+	}
+
+	FName FSkeletonHelper::SkeletalLodGetBoneName(const TArray<FMeshBoneInfo>& SkeletalLodRawInfos, int32 BoneIndex)
+	{
+		if (SkeletalLodRawInfos.IsValidIndex(BoneIndex))
+		{
+			return SkeletalLodRawInfos[BoneIndex].Name;
+		}
+		return NAME_None;
+	}
+
+	int32 FSkeletonHelper::SkeletalLodFindBoneIndex(const TArray<FMeshBoneInfo>& SkeletalLodRawInfos, FName BoneName)
+	{
+		const int32 BoneCount = SkeletalLodRawInfos.Num();
+		for (int32 BoneIndex = 0; BoneIndex < BoneCount; ++BoneIndex)
+		{
+			if (SkeletalLodRawInfos[BoneIndex].Name == BoneName)
+			{
+				return BoneIndex;
+			}
+		}
+		return INDEX_NONE;
+	}
+
+	int32 FSkeletonHelper::SkeletalLodGetParentIndex(const TArray<FMeshBoneInfo>& SkeletalLodRawInfos, int32 BoneIndex)
+	{
+		if (SkeletalLodRawInfos.IsValidIndex(BoneIndex))
+		{
+			return SkeletalLodRawInfos[BoneIndex].ParentIndex;
+		}
+		return INDEX_NONE;
+	}
+
+	bool FSkeletonHelper::DoesParentChainMatch(int32 StartBoneIndex, const FReferenceSkeleton& SkeletonRef, const TArray<FMeshBoneInfo>& SkeletalLodRawInfos)
+	{
+		// if start is root bone
+		if (StartBoneIndex == 0)
+		{
+			// verify name of root bone matches
+			return (SkeletonRef.GetBoneName(0) == SkeletalLodGetBoneName(SkeletalLodRawInfos, 0));
+		}
+
+		int32 SkeletonBoneIndex = StartBoneIndex;
+		// If skeleton bone is not found in mesh, fail.
+		int32 MeshBoneIndex = SkeletalLodFindBoneIndex(SkeletalLodRawInfos, SkeletonRef.GetBoneName(SkeletonBoneIndex));
+		if (MeshBoneIndex == INDEX_NONE)
+		{
+			return false;
+		}
+		do
+		{
+			// verify if parent name matches
+			int32 ParentSkeletonBoneIndex = SkeletonRef.GetParentIndex(SkeletonBoneIndex);
+			int32 ParentMeshBoneIndex = SkeletalLodGetParentIndex(SkeletalLodRawInfos, MeshBoneIndex);
+
+			// if one of the parents doesn't exist, make sure both end. Otherwise fail.
+			if ((ParentSkeletonBoneIndex == INDEX_NONE) || (ParentMeshBoneIndex == INDEX_NONE))
+			{
+				return (ParentSkeletonBoneIndex == ParentMeshBoneIndex);
+			}
+
+			// If parents are not named the same, fail.
+			if (SkeletonRef.GetBoneName(ParentSkeletonBoneIndex) != SkeletalLodGetBoneName(SkeletalLodRawInfos, ParentMeshBoneIndex))
+			{
+				return false;
+			}
+
+			// move up
+			SkeletonBoneIndex = ParentSkeletonBoneIndex;
+			MeshBoneIndex = ParentMeshBoneIndex;
+		} while (true);
+
+		return true;
+	}
+
+	void FSkeletonHelper::RecursiveBuildSkeletalSkeleton(const FString JoinToAddUid, const int32 ParentIndex, const UInterchangeBaseNodeContainer* BaseNodeContainer, TArray<FMeshBoneInfo>& SkeletalLodRawInfos)
+	{
+		const UInterchangeSceneNode* SceneNode = Cast<UInterchangeSceneNode>(BaseNodeContainer->GetNode(JoinToAddUid));
+		if (!SceneNode || !SceneNode->IsSpecializedTypeContains(FSceneNodeStaticData::GetJointSpecializeTypeString()))
+		{
+			return;
+		}
+
+		int32 JoinIndex = SkeletalLodRawInfos.Num();
+		FMeshBoneInfo& Info = SkeletalLodRawInfos.AddZeroed_GetRef();
+		Info.Name = *SceneNode->GetDisplayLabel();
+		Info.ParentIndex = ParentIndex;
+#if WITH_EDITORONLY_DATA
+		Info.ExportName = Info.Name.ToString();
+#endif
+		//Iterate childrens
+		const TArray<FString> ChildrenIds = BaseNodeContainer->GetNodeChildrenUids(JoinToAddUid);
+		for (int32 ChildIndex = 0; ChildIndex < ChildrenIds.Num(); ++ChildIndex)
+		{
+			RecursiveBuildSkeletalSkeleton(ChildrenIds[ChildIndex], JoinIndex, BaseNodeContainer, SkeletalLodRawInfos);
+		}
+	}
+} //namespace UE::Interchange::Private
+
+#endif //WITH_EDITOR
\ No newline at end of file
diff --git a/Engine/Plugins/Experimental/Interchange/Runtime/Source/Import/Public/Animation/InterchangeAnimSequenceFactory.h b/Engine/Plugins/Experimental/Interchange/Runtime/Source/Import/Public/Animation/InterchangeAnimSequenceFactory.h
new file mode 100644
--- /dev/null
+++ b/Engine/Plugins/Experimental/Interchange/Runtime/Source/Import/Public/Animation/InterchangeAnimSequenceFactory.h
@@ -0,0 +1,38 @@
+// Copyright Epic Games, Inc. All Rights Reserved.
+
+#pragma once
+
+#include "CoreMinimal.h"
+#include "InterchangeFactoryBase.h"
+#include "UObject/Object.h"
+#include "UObject/ObjectMacros.h"
+
+#include "InterchangeAnimSequenceFactory.generated.h"
+
+class UAnimSequence;
+
+UCLASS(BlueprintType, Experimental)
+class INTERCHANGEIMPORT_API UInterchangeAnimSequenceFactory : public UInterchangeFactoryBase
+{
+	GENERATED_BODY()
+public:
+
+	//////////////////////////////////////////////////////////////////////////
+	// Interchange factory base interface begin
+
+	virtual UClass* GetFactoryClass() const override;
+	virtual UObject* CreateEmptyAsset(const FCreateAssetParams& Arguments) override;
+	virtual UObject* CreateAsset(const FCreateAssetParams& Arguments) override;
+	virtual void PreImportPreCompletedCallback(const FImportPreCompletedCallbackParams& Arguments) override;
+	//virtual void PostImportPreCompletedCallback(const FImportPreCompletedCallbackParams& Arguments) const override;
+	virtual bool CanExecuteOnAnyThread() const override
+	{
+		//Currently we cannot use the anim sequence controller outside of the game thread
+		return false;
+	}
+
+	// Interchange factory base interface end
+	//////////////////////////////////////////////////////////////////////////
+};
+
+
diff --git a/Engine/Plugins/Experimental/Interchange/Runtime/Source/Import/Public/Animation/InterchangeAnimationPayload.h b/Engine/Plugins/Experimental/Interchange/Runtime/Source/Import/Public/Animation/InterchangeAnimationPayload.h
new file mode 100644
--- /dev/null
+++ b/Engine/Plugins/Experimental/Interchange/Runtime/Source/Import/Public/Animation/InterchangeAnimationPayload.h
@@ -0,0 +1,20 @@
+// Copyright Epic Games, Inc. All Rights Reserved.
+
+#pragma once
+
+#include "CoreMinimal.h"
+#include "Curves/RichCurve.h"
+#include "InterchangeCommonAnimationPayload.h"
+
+namespace UE::Interchange
+{
+	struct FAnimationCurveTransformPayloadData
+	{
+		EInterchangeTransformCurveChannel TransformChannel;
+		FRichCurve Curve;
+	};
+	struct FAnimationTransformPayloadData
+	{
+		TArray<FAnimationCurveTransformPayloadData> TransformCurves;
+	};
+}//ns UE::Interchange
diff --git a/Engine/Plugins/Experimental/Interchange/Runtime/Source/Import/Public/Animation/InterchangeAnimationPayloadInterface.h b/Engine/Plugins/Experimental/Interchange/Runtime/Source/Import/Public/Animation/InterchangeAnimationPayloadInterface.h
new file mode 100644
--- /dev/null
+++ b/Engine/Plugins/Experimental/Interchange/Runtime/Source/Import/Public/Animation/InterchangeAnimationPayloadInterface.h
@@ -0,0 +1,36 @@
+// Copyright Epic Games, Inc. All Rights Reserved.
+#pragma once
+
+#include "CoreMinimal.h"
+#include "Animation/InterchangeAnimationPayload.h"
+#include "UObject/Interface.h"
+#include "UObject/Object.h"
+#include "UObject/ObjectMacros.h"
+
+#include "InterchangeAnimationPayloadInterface.generated.h"
+
+UINTERFACE()
+class INTERCHANGEIMPORT_API UInterchangeAnimationPayloadInterface : public UInterface
+{
+	GENERATED_BODY()
+};
+
+/**
+ * Animation payload interface. Derive from this interface if your payload can import skeletal mesh
+ */
+class INTERCHANGEIMPORT_API IInterchangeAnimationPayloadInterface
+{
+	GENERATED_BODY()
+public:
+
+	/**
+	 * Get animation payload data for the specified payload key
+	 *
+	 * @param PayLoadKey - The key to retrieve the a particular payload contain into the specified source data.
+	 * @return - The resulting PayloadData as a TFuture point by the PayloadKey. The TOptional will not be set if there is an error retrieving the payload.
+	 * 
+	 */
+	virtual TFuture<TOptional<UE::Interchange::FAnimationTransformPayloadData>> GetAnimationTransformPayloadData(const FString& PayLoadKey) const = 0;
+};
+
+
diff --git a/Engine/Plugins/Experimental/Interchange/Runtime/Source/Import/Public/Fbx/InterchangeFbxTranslator.h b/Engine/Plugins/Experimental/Interchange/Runtime/Source/Import/Public/Fbx/InterchangeFbxTranslator.h
--- a/Engine/Plugins/Experimental/Interchange/Runtime/Source/Import/Public/Fbx/InterchangeFbxTranslator.h
+++ b/Engine/Plugins/Experimental/Interchange/Runtime/Source/Import/Public/Fbx/InterchangeFbxTranslator.h
@@ -1,6 +1,7 @@
 // Copyright Epic Games, Inc. All Rights Reserved.
 #pragma once
 
+#include "Animation/InterchangeAnimationPayloadInterface.h"
 #include "CoreMinimal.h"
 #include "InterchangeTranslatorBase.h"
 #include "InterchangeDispatcher.h"
@@ -22,6 +23,7 @@ class INTERCHANGEIMPORT_API UInterchangeFbxTranslator : public UInterchangeTrans
 , public IInterchangeTexturePayloadInterface
 , public IInterchangeStaticMeshPayloadInterface
 , public IInterchangeSkeletalMeshPayloadInterface
+, public IInterchangeAnimationPayloadInterface
 {
 	GENERATED_BODY()
 public:
@@ -75,6 +77,13 @@ public:
 	virtual TFuture<TOptional<UE::Interchange::FSkeletalMeshBlendShapePayloadData>> GetSkeletalMeshBlendShapePayloadData(const FString& PayLoadKey) const override;
 	
 	/* IInterchangeSkeletalMeshPayloadInterface End */
+
+	//////////////////////////////////////////////////////////////////////////
+	/* IInterchangeAnimationPayloadInterface Begin */
+
+	virtual TFuture<TOptional<UE::Interchange::FAnimationTransformPayloadData>> GetAnimationTransformPayloadData(const FString& PayLoadKey) const override;
+
+	/* IInterchangeAnimationPayloadInterface End */
 private:
 	FString CreateLoadFbxFileCommand(const FString& FbxFilePath) const;
 
diff --git a/Engine/Plugins/Experimental/Interchange/Runtime/Source/Import/Public/InterchangeSkeletonHelper.h b/Engine/Plugins/Experimental/Interchange/Runtime/Source/Import/Public/InterchangeSkeletonHelper.h
new file mode 100644
--- /dev/null
+++ b/Engine/Plugins/Experimental/Interchange/Runtime/Source/Import/Public/InterchangeSkeletonHelper.h
@@ -0,0 +1,42 @@
+// Copyright Epic Games, Inc. All Rights Reserved. 
+#pragma once
+
+#if WITH_EDITOR
+
+#include "Containers/Array.h"
+#include "Containers/UnrealString.h"
+#include "Rendering/SkeletalMeshLODImporterData.h"
+//#include "ReferenceSkeleton.h"
+
+struct FReferenceSkeleton;
+class UInterchangeBaseNodeContainer;
+class USkeleton;
+struct FMeshBoneInfo;
+
+namespace UE::Interchange::Private
+{
+
+	struct FJointInfo
+	{
+		FString Name;
+		int32 ParentIndex;  // 0 if this is the root bone.  
+		FTransform	LocalTransform; // local transform
+	};
+	
+	struct INTERCHANGEIMPORT_API FSkeletonHelper
+	{
+	public:
+		static bool ProcessImportMeshSkeleton(const USkeleton* SkeletonAsset, FReferenceSkeleton& RefSkeleton, int32& SkeletalDepth, const UInterchangeBaseNodeContainer* NodeContainer, const FString& RootJointNodeId, TArray<SkeletalMeshImportData::FBone>& RefBonesBinary, const bool bUseTimeZeroAsBindPose, bool& bOutDiffPose);
+		static bool IsCompatibleSkeleton(const USkeleton* Skeleton, const FString RootJoinUid, const UInterchangeBaseNodeContainer* BaseNodeContainer);
+			
+	private:
+		static void RecursiveAddBones(const UInterchangeBaseNodeContainer* NodeContainer, const FString& JointNodeId, TArray <FJointInfo>& JointInfos, int32 ParentIndex, TArray<SkeletalMeshImportData::FBone>& RefBonesBinary, const bool bUseTimeZeroAsBindPose, bool& bOutDiffPose);
+		static FName SkeletalLodGetBoneName(const TArray<FMeshBoneInfo>& SkeletalLodRawInfos, int32 BoneIndex);
+		static int32 SkeletalLodFindBoneIndex(const TArray<FMeshBoneInfo>& SkeletalLodRawInfos, FName BoneName);
+		static int32 SkeletalLodGetParentIndex(const TArray<FMeshBoneInfo>& SkeletalLodRawInfos, int32 BoneIndex);
+		static bool DoesParentChainMatch(int32 StartBoneIndex, const FReferenceSkeleton& SkeletonRef, const TArray<FMeshBoneInfo>& SkeletalLodRawInfos);
+		static void RecursiveBuildSkeletalSkeleton(const FString JoinToAddUid, const int32 ParentIndex, const UInterchangeBaseNodeContainer* BaseNodeContainer, TArray<FMeshBoneInfo>& SkeletalLodRawInfos);
+	};
+} //namespace UE::Interchange::Private
+
+#endif //WITH_EDITOR
diff --git a/Engine/Plugins/Experimental/Interchange/Runtime/Source/Nodes/Private/InterchangeSceneNode.cpp b/Engine/Plugins/Experimental/Interchange/Runtime/Source/Nodes/Private/InterchangeSceneNode.cpp
--- a/Engine/Plugins/Experimental/Interchange/Runtime/Source/Nodes/Private/InterchangeSceneNode.cpp
+++ b/Engine/Plugins/Experimental/Interchange/Runtime/Source/Nodes/Private/InterchangeSceneNode.cpp
@@ -91,6 +91,10 @@ FString UInterchangeSceneNode::GetKeyDisplayName(const UE::Interchange::FAttribu
 		}
 		return KeyDisplayName;
 	}
+	else if (NodeAttributeKey.Key.Equals(Macro_CustomTransformCurvePayloadKeyKey.Key))
+	{
+		return FString(TEXT("Transform Curve Payload Key"));
+	}
 	return Super::GetKeyDisplayName(NodeAttributeKey);
 }
 
@@ -109,6 +113,12 @@ FString UInterchangeSceneNode::GetAttributeCategory(const UE::Interchange::FAttr
 	{
 		return FString(TEXT("Scene"));
 	}
+	else if (NodeAttributeKey == Macro_CustomBindPoseLocalTransformKey
+		|| NodeAttributeKey == Macro_CustomTimeZeroLocalTransformKey
+		|| NodeAttributeKey == Macro_CustomTransformCurvePayloadKeyKey)
+	{
+		return FString(TEXT("Joint"));
+	}
 	return Super::GetAttributeCategory(NodeAttributeKey);
 }
 
@@ -256,6 +266,56 @@ bool UInterchangeSceneNode::SetCustomAssetInstanceUid(const FString& AttributeVa
 	IMPLEMENT_NODE_ATTRIBUTE_SETTER_NODELEGATE(AssetInstanceUid, FString);
 }
 
+bool UInterchangeSceneNode::GetCustomIsNodeTransformAnimated(bool& AttributeValue) const
+{
+	IMPLEMENT_NODE_ATTRIBUTE_GETTER(IsNodeTransformAnimated, bool);
+}
+
+bool UInterchangeSceneNode::SetCustomIsNodeTransformAnimated(const bool& AttributeValue)
+{
+	IMPLEMENT_NODE_ATTRIBUTE_SETTER_NODELEGATE(IsNodeTransformAnimated, bool);
+}
+
+bool UInterchangeSceneNode::GetCustomNodeTransformAnimationKeyCount(int32& AttributeValue) const
+{
+	IMPLEMENT_NODE_ATTRIBUTE_GETTER(NodeTransformAnimationKeyCount, int32);
+}
+
+bool UInterchangeSceneNode::SetCustomNodeTransformAnimationKeyCount(const int32& AttributeValue)
+{
+	IMPLEMENT_NODE_ATTRIBUTE_SETTER_NODELEGATE(NodeTransformAnimationKeyCount, int32);
+}
+
+bool UInterchangeSceneNode::GetCustomNodeTransformAnimationStartTime(double& AttributeValue) const
+{
+	IMPLEMENT_NODE_ATTRIBUTE_GETTER(NodeTransformAnimationStartTime, double);
+}
+
+bool UInterchangeSceneNode::SetCustomNodeTransformAnimationStartTime(const double& AttributeValue)
+{
+	IMPLEMENT_NODE_ATTRIBUTE_SETTER_NODELEGATE(NodeTransformAnimationStartTime, double);
+}
+
+bool UInterchangeSceneNode::GetCustomNodeTransformAnimationEndTime(double& AttributeValue) const
+{
+	IMPLEMENT_NODE_ATTRIBUTE_GETTER(NodeTransformAnimationEndTime, double);
+}
+
+bool UInterchangeSceneNode::SetCustomNodeTransformAnimationEndTime(const double& AttributeValue)
+{
+	IMPLEMENT_NODE_ATTRIBUTE_SETTER_NODELEGATE(NodeTransformAnimationEndTime, double);
+}
+
+bool UInterchangeSceneNode::GetCustomTransformCurvePayloadKey(FString& AttributeValue) const
+{
+	IMPLEMENT_NODE_ATTRIBUTE_GETTER(TransformCurvePayloadKey, FString);
+}
+
+bool UInterchangeSceneNode::SetCustomTransformCurvePayloadKey(const FString& AttributeValue)
+{
+	IMPLEMENT_NODE_ATTRIBUTE_SETTER_NODELEGATE(TransformCurvePayloadKey, FString);
+}
+
 void UInterchangeSceneNode::ResetAllGlobalTransformCaches(const UInterchangeBaseNodeContainer* BaseNodeContainer)
 {
 	BaseNodeContainer->IterateNodes([](const FString& NodeUid, UInterchangeBaseNode* Node)
diff --git a/Engine/Plugins/Experimental/Interchange/Runtime/Source/Nodes/Public/InterchangeSceneNode.h b/Engine/Plugins/Experimental/Interchange/Runtime/Source/Nodes/Public/InterchangeSceneNode.h
--- a/Engine/Plugins/Experimental/Interchange/Runtime/Source/Nodes/Public/InterchangeSceneNode.h
+++ b/Engine/Plugins/Experimental/Interchange/Runtime/Source/Nodes/Public/InterchangeSceneNode.h
@@ -144,6 +144,46 @@ public:
 	UFUNCTION(BlueprintCallable, Category = "Interchange | Node | Scene")
 	bool SetCustomAssetInstanceUid(const FString& AttributeValue);
 
+	/** Return true if the node transform have animation, false if the attribute is not set or there is no animation for the transform. */
+	UFUNCTION(BlueprintCallable, Category = "Interchange | Node | Scene | animation")
+	bool GetCustomIsNodeTransformAnimated(bool& AttributeValue) const;
+
+	/** Setter to be able to query if the node transform has animation. */
+	UFUNCTION(BlueprintCallable, Category = "Interchange | Node | Scene | animation")
+	bool SetCustomIsNodeTransformAnimated(const bool& AttributeValue);
+
+	/** Return false if the Attribute was not set previously. This attribute store the payload key needed to retrieve the animation from the source*/
+	UFUNCTION(BlueprintCallable, Category = "Interchange | Node | Scene | animation")
+	bool GetCustomNodeTransformAnimationKeyCount(int32& AttributeValue) const;
+
+	/** If the node transform is animated, this set the payload key needed to retrieve the animation from the source. */
+	UFUNCTION(BlueprintCallable, Category = "Interchange | Node | Scene | animation")
+	bool SetCustomNodeTransformAnimationKeyCount(const int32& AttributeValue);
+
+	/** If the node transform is animated, this should return the starting frame of the animation for this node transform*/
+	UFUNCTION(BlueprintCallable, Category = "Interchange | Node | Scene | animation")
+	bool GetCustomNodeTransformAnimationStartTime(double& AttributeValue) const;
+
+	/** If the node transform is animated, this set the starting frame of the animation for this node transform*/
+	UFUNCTION(BlueprintCallable, Category = "Interchange | Node | Scene | animation")
+	bool SetCustomNodeTransformAnimationStartTime(const double& AttributeValue);
+
+	/** If the node transform is animated, this should return the ending frame of the animation for this node transform*/
+	UFUNCTION(BlueprintCallable, Category = "Interchange | Node | Scene | animation")
+	bool GetCustomNodeTransformAnimationEndTime(double& AttributeValue) const;
+
+	/** If the node transform is animated, this set the ending frame of the animation for this node transform*/
+	UFUNCTION(BlueprintCallable, Category = "Interchange | Node | Scene | animation")
+	bool SetCustomNodeTransformAnimationEndTime(const double& AttributeValue);
+
+	/** Return false if the Attribute was not set previously. This attribute store the payload key needed to retrieve the animation from the source*/
+	UFUNCTION(BlueprintCallable, Category = "Interchange | Node | Scene | animation")
+	bool GetCustomTransformCurvePayloadKey(FString& AttributeValue) const;
+
+	/** If the node transform is animated, this is the where we store the payload key needed to retrieve the animation from the source. */
+	UFUNCTION(BlueprintCallable, Category = "Interchange | Node | Scene | animation")
+	bool SetCustomTransformCurvePayloadKey(const FString& AttributeValue);
+
 	/** This static function make sure all the global transform caches are reset for all the UInterchangeSceneNode nodes in the UInterchangeBaseNodeContainer */
 	static void ResetAllGlobalTransformCaches(const UInterchangeBaseNodeContainer* BaseNodeContainer);
 
@@ -160,6 +200,12 @@ private:
 	const UE::Interchange::FAttributeKey Macro_CustomTimeZeroLocalTransformKey = UE::Interchange::FAttributeKey(TEXT("TimeZeroLocalTransform"));
 	const UE::Interchange::FAttributeKey Macro_CustomGeometricTransformKey = UE::Interchange::FAttributeKey(TEXT("GeometricTransform"));
 	const UE::Interchange::FAttributeKey Macro_CustomAssetInstanceUidKey = UE::Interchange::FAttributeKey(TEXT("AssetInstanceUid"));
+	//Animation data
+	const UE::Interchange::FAttributeKey Macro_CustomIsNodeTransformAnimatedKey = UE::Interchange::FAttributeKey(TEXT("IsNodeTransformAnimated"));
+	const UE::Interchange::FAttributeKey Macro_CustomNodeTransformAnimationKeyCountKey = UE::Interchange::FAttributeKey(TEXT("NodeTransformAnimationKeyCount"));
+	const UE::Interchange::FAttributeKey Macro_CustomNodeTransformAnimationStartTimeKey = UE::Interchange::FAttributeKey(TEXT("NodeTransformAnimationStartTime"));
+	const UE::Interchange::FAttributeKey Macro_CustomNodeTransformAnimationEndTimeKey = UE::Interchange::FAttributeKey(TEXT("NodeTransformAnimationEndTime"));
+	const UE::Interchange::FAttributeKey Macro_CustomTransformCurvePayloadKeyKey = UE::Interchange::FAttributeKey(TEXT("TransformCurvePayloadKey"));
 
 	UE::Interchange::TArrayAttributeHelper<FString> NodeSpecializeTypes;
 	UE::Interchange::TArrayAttributeHelper<FString> MaterialDependencyUids;
diff --git a/Engine/Plugins/Experimental/Interchange/Runtime/Source/Parsers/CommonParser/InterchangeCommonParser.Build.cs b/Engine/Plugins/Experimental/Interchange/Runtime/Source/Parsers/CommonParser/InterchangeCommonParser.Build.cs
new file mode 100644
--- /dev/null
+++ b/Engine/Plugins/Experimental/Interchange/Runtime/Source/Parsers/CommonParser/InterchangeCommonParser.Build.cs
@@ -0,0 +1,30 @@
+// Copyright Epic Games, Inc. All Rights Reserved.
+
+using System.IO;
+
+namespace UnrealBuildTool.Rules
+{
+	public class InterchangeCommonParser : ModuleRules
+	{
+		public InterchangeCommonParser(ReadOnlyTargetRules Target) : base(Target)
+		{
+			PublicDependencyModuleNames.AddRange(
+				new string[]
+				{
+					"Core",
+					"CoreUObject",
+					"InterchangeCore",
+				}
+			);
+
+			if (Target.bCompileAgainstEngine)
+			{
+				PrivateDependencyModuleNames.AddRange(
+					new string[] {
+						"Engine",
+					}
+				);
+			}
+		}
+	}
+}
diff --git a/Engine/Plugins/Experimental/Interchange/Runtime/Source/Parsers/CommonParser/Private/InterchangeCommonAnimationPayload.cpp b/Engine/Plugins/Experimental/Interchange/Runtime/Source/Parsers/CommonParser/Private/InterchangeCommonAnimationPayload.cpp
new file mode 100644
--- /dev/null
+++ b/Engine/Plugins/Experimental/Interchange/Runtime/Source/Parsers/CommonParser/Private/InterchangeCommonAnimationPayload.cpp
@@ -0,0 +1,77 @@
+// Copyright Epic Games, Inc. All Rights Reserved.
+
+#include "InterchangeCommonAnimationPayload.h"
+
+#include "CoreMinimal.h"
+
+#if WITH_ENGINE
+#include "Curves/RichCurve.h"
+
+void FInterchangeCurveKey::ToRichCurveKey(FRichCurveKey& RichCurveKey) const
+{
+	RichCurveKey.Time = Time;
+	RichCurveKey.Value = Value;
+	switch (InterpMode)
+	{
+		case EInterchangeCurveInterpMode::Constant:
+			RichCurveKey.InterpMode = ERichCurveInterpMode::RCIM_Constant;
+			break;
+		case EInterchangeCurveInterpMode::Cubic:
+			RichCurveKey.InterpMode = ERichCurveInterpMode::RCIM_Cubic;
+			break;
+		case EInterchangeCurveInterpMode::Linear:
+			RichCurveKey.InterpMode = ERichCurveInterpMode::RCIM_Linear;
+			break;
+		default:
+			RichCurveKey.InterpMode = ERichCurveInterpMode::RCIM_None;
+			break;
+	}
+	switch (TangentMode)
+	{
+		case EInterchangeCurveTangentMode::Auto:
+			RichCurveKey.TangentMode = ERichCurveTangentMode::RCTM_Auto;
+			break;
+		case EInterchangeCurveTangentMode::Break:
+			RichCurveKey.TangentMode = ERichCurveTangentMode::RCTM_Break;
+			break;
+		case EInterchangeCurveTangentMode::User:
+			RichCurveKey.TangentMode = ERichCurveTangentMode::RCTM_User;
+			break;
+		default:
+			RichCurveKey.TangentMode = ERichCurveTangentMode::RCTM_None;
+			break;
+	}
+	switch (TangentWeightMode)
+	{
+		case EInterchangeCurveTangentWeightMode::WeightedArrive:
+			RichCurveKey.TangentWeightMode = ERichCurveTangentWeightMode::RCTWM_WeightedArrive;
+			break;
+		case EInterchangeCurveTangentWeightMode::WeightedBoth:
+			RichCurveKey.TangentWeightMode = ERichCurveTangentWeightMode::RCTWM_WeightedBoth;
+			break;
+		case EInterchangeCurveTangentWeightMode::WeightedLeave:
+			RichCurveKey.TangentWeightMode = ERichCurveTangentWeightMode::RCTWM_WeightedLeave;
+			break;
+		default:
+			RichCurveKey.TangentWeightMode = ERichCurveTangentWeightMode::RCTWM_WeightedNone;
+			break;
+	}
+	RichCurveKey.ArriveTangent = ArriveTangent;
+	RichCurveKey.ArriveTangentWeight = ArriveTangentWeight;
+	RichCurveKey.LeaveTangent = LeaveTangent;
+	RichCurveKey.LeaveTangentWeight = LeaveTangentWeight;
+
+}
+
+void FInterchangeCurve::ToRichCurve(FRichCurve& OutRichCurve) const
+{
+	OutRichCurve.Keys.Reserve(Keys.Num());
+	for (const FInterchangeCurveKey& CurveKey : Keys)
+	{
+		FKeyHandle RichCurveKeyHandle = OutRichCurve.AddKey(CurveKey.Time, CurveKey.Value);
+		CurveKey.ToRichCurveKey(OutRichCurve.GetKey(RichCurveKeyHandle));
+	}
+	OutRichCurve.AutoSetTangents();
+}
+
+#endif //WITH_ENGINE
\ No newline at end of file
diff --git a/Engine/Plugins/Experimental/Interchange/Runtime/Source/Parsers/CommonParser/Private/InterchangeCommonParserModule.cpp b/Engine/Plugins/Experimental/Interchange/Runtime/Source/Parsers/CommonParser/Private/InterchangeCommonParserModule.cpp
new file mode 100644
--- /dev/null
+++ b/Engine/Plugins/Experimental/Interchange/Runtime/Source/Parsers/CommonParser/Private/InterchangeCommonParserModule.cpp
@@ -0,0 +1,23 @@
+// Copyright Epic Games, Inc. All Rights Reserved.
+#include "InterchangeCommonParserModule.h"
+
+#include "Modules/ModuleInterface.h"
+#include "Modules/ModuleManager.h"
+
+
+
+FInterchangeCommonParserModule& FInterchangeCommonParserModule::Get()
+{
+	return FModuleManager::LoadModuleChecked< FInterchangeCommonParserModule >(INTERCHANGECOMMONPARSER_MODULE_NAME);
+}
+
+bool FInterchangeCommonParserModule::IsAvailable()
+{
+	return FModuleManager::Get().IsModuleLoaded(INTERCHANGECOMMONPARSER_MODULE_NAME);
+}
+
+void FInterchangeCommonParserModule::StartupModule()
+{
+}
+
+IMPLEMENT_MODULE(FInterchangeCommonParserModule, InterchangeCommonParser);
diff --git a/Engine/Plugins/Experimental/Interchange/Runtime/Source/Parsers/CommonParser/Public/InterchangeCommonAnimationPayload.h b/Engine/Plugins/Experimental/Interchange/Runtime/Source/Parsers/CommonParser/Public/InterchangeCommonAnimationPayload.h
new file mode 100644
--- /dev/null
+++ b/Engine/Plugins/Experimental/Interchange/Runtime/Source/Parsers/CommonParser/Public/InterchangeCommonAnimationPayload.h
@@ -0,0 +1,174 @@
+// Copyright Epic Games, Inc. All Rights Reserved.
+
+#pragma once
+
+#include "UObject/Class.h"
+#include "UObject/Object.h"
+#include "UObject/ObjectMacros.h"
+
+#if WITH_ENGINE
+#include "Curves/RichCurve.h"
+#endif
+
+#include "InterchangeCommonAnimationPayload.generated.h"
+
+UENUM()
+enum class EInterchangeTransformCurveChannel
+{
+	TranslationX = 0,
+	TranslationY = 1,
+	TranslationZ = 2,
+	EulerX = 3,
+	EulerY = 4,
+	EulerZ = 5,
+	ScaleX = 6,
+	ScaleY = 7,
+	ScaleZ = 8,
+	TransformChannelCount = 9,
+	None
+};
+
+/** If using Cubic, this enum describes how the tangents should be controlled. */
+UENUM()
+enum class EInterchangeCurveInterpMode : uint8
+{
+	/** Use linear interpolation between values. */
+	Linear,
+	/** Use a constant value. Represents stepped values. */
+	Constant,
+	/** Cubic interpolation. See TangentMode for different cubic interpolation options. */
+	Cubic,
+	/** No interpolation. */
+	None
+};
+
+/** If using Cubic interpolation mode, this enum describes how the tangents should be controlled. */
+UENUM()
+enum class EInterchangeCurveTangentMode : uint8
+{
+	/** Automatically calculates tangents to create smooth curves between values. */
+	Auto,
+	/** User specifies the tangent as a unified tangent where the two tangents are locked to each other, presenting a consistent curve before and after. */
+	User,
+	/** User specifies the tangent as two separate broken tangents on each side of the key which can allow a sharp change in evaluation before or after. */
+	Break,
+	/** No tangents. */
+	None
+};
+
+
+/** Enumerates tangent weight modes. */
+UENUM()
+enum class EInterchangeCurveTangentWeightMode : uint8
+{
+	/** Don't take tangent weights into account. */
+	WeightedNone,
+	/** Only take the arrival tangent weight into account for evaluation. */
+	WeightedArrive,
+	/** Only take the leaving tangent weight into account for evaluation. */
+	WeightedLeave,
+	/** Take both the arrival and leaving tangent weights into account for evaluation. */
+	WeightedBoth
+};
+
+/**
+* This struct contains only the key data, this is only used to pass animation data from translators to factories
+*/
+USTRUCT()
+struct INTERCHANGECOMMONPARSER_API FInterchangeCurveKey
+{
+	GENERATED_BODY()
+
+	/** Interpolation mode between this key and the next */
+	UPROPERTY()
+	EInterchangeCurveInterpMode InterpMode = EInterchangeCurveInterpMode::None;
+
+	/** Mode for tangents at this key */
+	UPROPERTY()
+	EInterchangeCurveTangentMode TangentMode = EInterchangeCurveTangentMode::None;
+
+	/** If either tangent at this key is 'weighted' */
+	UPROPERTY()
+	EInterchangeCurveTangentWeightMode TangentWeightMode = EInterchangeCurveTangentWeightMode::WeightedNone;
+
+	/** Time at this key */
+	UPROPERTY()
+	float Time = 0.0f;
+
+	/** Value at this key */
+	UPROPERTY()
+	float Value = 0.0f;
+
+	/** If RCIM_Cubic, the arriving tangent at this key */
+	UPROPERTY()
+	float ArriveTangent = 0.0f;
+
+	/** If RCTWM_WeightedArrive or RCTWM_WeightedBoth, the weight of the left tangent */
+	UPROPERTY()
+	float ArriveTangentWeight = 0.0f;
+
+	/** If RCIM_Cubic, the leaving tangent at this key */
+	UPROPERTY()
+	float LeaveTangent = 0.0f;
+
+	/** If RCTWM_WeightedLeave or RCTWM_WeightedBoth, the weight of the right tangent */
+	UPROPERTY()
+	float LeaveTangentWeight = 0.0f;
+
+#if WITH_ENGINE
+	/** Conversion to FRichCurve */
+	void ToRichCurveKey(FRichCurveKey& OutKey) const;
+#endif
+
+	void Serialize(FArchive& Ar)
+	{
+		Ar << InterpMode;
+		Ar << TangentMode;
+		Ar << TangentWeightMode;
+		Ar << Time;
+		Ar << Value;
+		Ar << ArriveTangent;
+		Ar << ArriveTangentWeight;
+		Ar << LeaveTangent;
+		Ar << LeaveTangentWeight;
+	}
+
+	friend FArchive& operator<<(FArchive& Ar, FInterchangeCurveKey& P)
+	{
+		P.Serialize(Ar);
+		return Ar;
+	}
+};
+
+/**
+* This struct contains only the key data, this is only used to pass animation data from translators to factories.
+* You cannot evaluate a curve with this struct.
+*/
+USTRUCT()
+struct INTERCHANGECOMMONPARSER_API FInterchangeCurve
+{
+	GENERATED_BODY()
+
+	UPROPERTY()
+	EInterchangeTransformCurveChannel TransformChannel = EInterchangeTransformCurveChannel::None;
+	
+	UPROPERTY()
+	TArray<FInterchangeCurveKey> Keys;
+
+#if WITH_ENGINE
+	/** Conversion to FRichCurve */
+	void ToRichCurve(FRichCurve& OutKey) const;
+#endif
+
+	void Serialize(FArchive& Ar)
+	{
+		Ar << TransformChannel;
+		Ar << Keys;
+	}
+
+	friend FArchive& operator<<(FArchive& Ar, FInterchangeCurve& P)
+	{
+		P.Serialize(Ar);
+		return Ar;
+	}
+};
diff --git a/Engine/Plugins/Experimental/Interchange/Runtime/Source/Parsers/CommonParser/Public/InterchangeCommonParserModule.h b/Engine/Plugins/Experimental/Interchange/Runtime/Source/Parsers/CommonParser/Public/InterchangeCommonParserModule.h
new file mode 100644
--- /dev/null
+++ b/Engine/Plugins/Experimental/Interchange/Runtime/Source/Parsers/CommonParser/Public/InterchangeCommonParserModule.h
@@ -0,0 +1,19 @@
+// Copyright Epic Games, Inc. All Rights Reserved.
+#pragma once
+
+#include "Modules/ModuleInterface.h"
+
+#define INTERCHANGECOMMONPARSER_MODULE_NAME TEXT("InterchangeCommonParser")
+
+/**
+ * This module is used to add common structure between the translators and the parsers (i.e. payload temporary structure)
+ */
+class FInterchangeCommonParserModule : public IModuleInterface
+{
+public:
+	static FInterchangeCommonParserModule& Get();
+	static bool IsAvailable();
+
+private:
+	virtual void StartupModule() override;
+};
diff --git a/Engine/Plugins/Experimental/Interchange/Runtime/Source/Parsers/Fbx/InterchangeFbxParser.Build.cs b/Engine/Plugins/Experimental/Interchange/Runtime/Source/Parsers/Fbx/InterchangeFbxParser.Build.cs
--- a/Engine/Plugins/Experimental/Interchange/Runtime/Source/Parsers/Fbx/InterchangeFbxParser.Build.cs
+++ b/Engine/Plugins/Experimental/Interchange/Runtime/Source/Parsers/Fbx/InterchangeFbxParser.Build.cs
@@ -34,6 +34,7 @@ namespace UnrealBuildTool.Rules
 			PrivateDependencyModuleNames.AddRange(
 				new string[]
 				{
+					"InterchangeCommonParser",
 					"InterchangeMessages",
 					"InterchangeNodes",
 					"MeshDescription",
diff --git a/Engine/Plugins/Experimental/Interchange/Runtime/Source/Parsers/Fbx/Private/FbxAPI.cpp b/Engine/Plugins/Experimental/Interchange/Runtime/Source/Parsers/Fbx/Private/FbxAPI.cpp
--- a/Engine/Plugins/Experimental/Interchange/Runtime/Source/Parsers/Fbx/Private/FbxAPI.cpp
+++ b/Engine/Plugins/Experimental/Interchange/Runtime/Source/Parsers/Fbx/Private/FbxAPI.cpp
@@ -117,7 +117,7 @@ namespace UE
 				FbxCamera.AddAllCameras(SDKScene, NodeContainer);
 
 				FFbxScene FbxScene(*this);
-				FbxScene.AddHierarchy(SDKScene, NodeContainer);
+				FbxScene.AddHierarchy(SDKScene, NodeContainer, PayloadContexts);
 			}
 
 			bool FFbxParser::FetchPayloadData(const FString& PayloadKey, const FString& PayloadFilepath)
diff --git a/Engine/Plugins/Experimental/Interchange/Runtime/Source/Parsers/Fbx/Private/FbxAPI.h b/Engine/Plugins/Experimental/Interchange/Runtime/Source/Parsers/Fbx/Private/FbxAPI.h
--- a/Engine/Plugins/Experimental/Interchange/Runtime/Source/Parsers/Fbx/Private/FbxAPI.h
+++ b/Engine/Plugins/Experimental/Interchange/Runtime/Source/Parsers/Fbx/Private/FbxAPI.h
@@ -4,6 +4,7 @@
 
 #include "CoreMinimal.h"
 #include "FbxInclude.h"
+#include "HAL/CriticalSection.h"
 #include "InterchangeResultsContainer.h"
 
 
@@ -65,6 +66,14 @@ namespace UE
 					ResultsContainer->Add(Item);
 					Item->SourceAssetName = SourceFilename;
 				}
+
+				/**
+				 * Critical section to avoid getting animation payload curves asynchronously.
+				 * The FBX evaluator use a cache mechanism for evaluating global transform that is not thread safe.
+				 * We avoid evaluating global transform asynchronously by using this CriticalSection when we import curves payload.
+				 * Each scene nodes representing a joint have is own payload for the transform animation.
+				 */
+				FCriticalSection AnimationTransformPayloadCriticalSection;
 			private:
 
 				void CleanupFbxData();
@@ -75,9 +84,7 @@ namespace UE
 				FbxImporter* SDKImporter = nullptr;
 				FbxGeometryConverter* SDKGeometryConverter = nullptr;
 				FString SourceFilename;
-
 				TMap<FString, TSharedPtr<FPayloadContextBase>> PayloadContexts;
-				
 			};
 		}//ns Private
 	}//ns Interchange
diff --git a/Engine/Plugins/Experimental/Interchange/Runtime/Source/Parsers/Fbx/Private/FbxAnimation.cpp b/Engine/Plugins/Experimental/Interchange/Runtime/Source/Parsers/Fbx/Private/FbxAnimation.cpp
new file mode 100644
--- /dev/null
+++ b/Engine/Plugins/Experimental/Interchange/Runtime/Source/Parsers/Fbx/Private/FbxAnimation.cpp
@@ -0,0 +1,595 @@
+// Copyright Epic Games, Inc. All Rights Reserved.
+
+#include "FbxAnimation.h"
+
+#include "CoreMinimal.h"
+#include "FbxAPI.h"
+#include "FbxConvert.h"
+#include "FbxHelper.h"
+#include "FbxInclude.h"
+#include "Fbx/InterchangeFbxMessages.h"
+#include "InterchangeCommonAnimationPayload.h"
+#include "InterchangeMeshNode.h"
+#include "InterchangeResultsContainer.h"
+#include "InterchangeSceneNode.h"
+#include "MeshDescription.h"
+#include "Misc/FileHelper.h"
+#include "Misc/ScopeLock.h"
+#include "Nodes/InterchangeBaseNodeContainer.h"
+#include "Serialization/LargeMemoryWriter.h"
+#include "SkeletalMeshAttributes.h"
+#include "StaticMeshAttributes.h"
+
+
+#define LOCTEXT_NAMESPACE "InterchangeFbxMesh"
+
+namespace UE::Interchange::Private
+{
+	float GetTransformChannelValue(const EInterchangeTransformCurveChannel Channel, const FTransform& Transform)
+	{
+		switch (Channel)
+		{
+			case EInterchangeTransformCurveChannel::TranslationX:
+			{
+				return Transform.GetLocation().X;
+			}
+			break;
+			case EInterchangeTransformCurveChannel::TranslationY:
+			{
+				return Transform.GetLocation().Y;
+			}
+			break;
+			case EInterchangeTransformCurveChannel::TranslationZ:
+			{
+				return Transform.GetLocation().Z;
+			}
+			break;
+			case EInterchangeTransformCurveChannel::EulerX:
+			{
+				return Transform.GetRotation().Euler().X;
+			}
+			break;
+			case EInterchangeTransformCurveChannel::EulerY:
+			{
+				return Transform.GetRotation().Euler().Y;
+			}
+			break;
+			case EInterchangeTransformCurveChannel::EulerZ:
+			{
+				return Transform.GetRotation().Euler().Z;
+			}
+			break;
+			case EInterchangeTransformCurveChannel::ScaleX:
+			{
+				return Transform.GetScale3D().X;
+			}
+			break;
+			case EInterchangeTransformCurveChannel::ScaleY:
+			{
+				return Transform.GetScale3D().Y;
+			}
+			break;
+			case EInterchangeTransformCurveChannel::ScaleZ:
+			{
+				return Transform.GetScale3D().Z;
+			}
+			break;
+		}
+		return 0.0f;
+	}
+
+	bool ImportCurve(FbxNode* Node, const EInterchangeTransformCurveChannel TransformChannel, const FbxAnimCurve* SourceFloatCurves, const float ScaleValue, TArray<FInterchangeCurveKey>& DestinationFloatCurve)
+	{
+		bool bUseNodeTransform = false;
+		if (TransformChannel != EInterchangeTransformCurveChannel::None)
+		{
+			bUseNodeTransform = Node != nullptr;
+		}
+
+		if (!SourceFloatCurves)
+		{
+			return true;
+		}
+		const float DefaultCurveWeight = FbxAnimCurveDef::sDEFAULT_WEIGHT;
+		//We use the non const to query the left and right derivative of the key, for whatever reason those FBX API functions are not const
+		FbxAnimCurve* NonConstSourceFloatCurves = const_cast<FbxAnimCurve*>(SourceFloatCurves);
+		int32 KeyCount = SourceFloatCurves->KeyGetCount();
+		for (int32 KeyIndex = 0; KeyIndex < KeyCount; ++KeyIndex)
+		{
+			FbxAnimCurveKey Key = SourceFloatCurves->KeyGet(KeyIndex);
+			FbxTime KeyTime = Key.GetTime();
+			const float KeyTimeValue = static_cast<float>(KeyTime.GetSecondDouble());
+			float ValueAtTime = Key.GetValue();
+			if (bUseNodeTransform)
+			{
+				FbxAMatrix NodeTransform = Node->EvaluateGlobalTransform(KeyTime);
+				FbxNode* ParentNode = Node->GetParent();
+				if (ParentNode)
+				{
+					FbxAMatrix ParentTransform = ParentNode->EvaluateGlobalTransform(KeyTime);
+					NodeTransform = ParentTransform.Inverse() * NodeTransform;
+				}
+				FTransform Transform = UE::Interchange::Private::FFbxConvert::ConvertTransform(NodeTransform);
+				ValueAtTime = GetTransformChannelValue(TransformChannel, Transform);
+			}
+			float Value = ValueAtTime * ScaleValue;
+			FInterchangeCurveKey& InterchangeCurveKey = DestinationFloatCurve.AddDefaulted_GetRef();
+			InterchangeCurveKey.Time = KeyTimeValue;
+			InterchangeCurveKey.Value = Value;
+
+			const bool bIncludeOverrides = true;
+			FbxAnimCurveDef::ETangentMode KeyTangentMode = Key.GetTangentMode(bIncludeOverrides);
+			FbxAnimCurveDef::EInterpolationType KeyInterpMode = Key.GetInterpolation();
+			FbxAnimCurveDef::EWeightedMode KeyTangentWeightMode = Key.GetTangentWeightMode();
+
+			EInterchangeCurveInterpMode NewInterpMode = EInterchangeCurveInterpMode::Linear;
+			EInterchangeCurveTangentMode NewTangentMode = EInterchangeCurveTangentMode::Auto;
+			EInterchangeCurveTangentWeightMode NewTangentWeightMode = EInterchangeCurveTangentWeightMode::WeightedNone;
+
+			float RightTangent = NonConstSourceFloatCurves->KeyGetRightDerivative(KeyIndex) * ScaleValue;
+			float LeftTangent = NonConstSourceFloatCurves->KeyGetLeftDerivative(KeyIndex) * ScaleValue;
+			float RightTangentWeight = 0.0f;
+			float LeftTangentWeight = 0.0f; //This one is dependent on the previous key.
+			bool bLeftWeightActive = false;
+			bool bRightWeightActive = false;
+
+			const bool bPreviousKeyValid = KeyIndex > 0;
+			const bool bNextKeyValid = KeyIndex < KeyCount - 1;
+			float PreviousValue = 0.0f;
+			float PreviousKeyTimeValue = 0.0f;
+			float NextValue = 0.0f;
+			float NextKeyTimeValue = 0.0f;
+			if (bPreviousKeyValid)
+			{
+				FbxAnimCurveKey PreviousKey = SourceFloatCurves->KeyGet(KeyIndex - 1);
+				FbxTime PreviousKeyTime = PreviousKey.GetTime();
+				PreviousKeyTimeValue = static_cast<float>(PreviousKeyTime.GetSecondDouble());
+				PreviousValue = PreviousKey.GetValue() * ScaleValue;
+				//The left tangent is driven by the previous key. If the previous key have a the NextLeftweight or both flag weighted mode, it mean the next key is weighted on the left side
+				bLeftWeightActive = (PreviousKey.GetTangentWeightMode() & FbxAnimCurveDef::eWeightedNextLeft) > 0;
+				if (bLeftWeightActive)
+				{
+					LeftTangentWeight = PreviousKey.GetDataFloat(FbxAnimCurveDef::eNextLeftWeight);
+				}
+			}
+			if (bNextKeyValid)
+			{
+				FbxAnimCurveKey NextKey = SourceFloatCurves->KeyGet(KeyIndex + 1);
+				FbxTime NextKeyTime = NextKey.GetTime();
+				NextKeyTimeValue = static_cast<float>(NextKeyTime.GetSecondDouble());
+				NextValue = NextKey.GetValue() * ScaleValue;
+
+				bRightWeightActive = (KeyTangentWeightMode & FbxAnimCurveDef::eWeightedRight) > 0;
+				if (bRightWeightActive)
+				{
+					//The right tangent weight should be use only if we are not the last key since the last key do not have a right tangent.
+					//Use the current key to gather the right tangent weight
+					RightTangentWeight = Key.GetDataFloat(FbxAnimCurveDef::eRightWeight);
+				}
+			}
+
+			// When this flag is true, the tangent is flat if the value has the same value as the previous or next key.
+			const bool bTangentGenericClamp = (KeyTangentMode & FbxAnimCurveDef::eTangentGenericClamp);
+
+			//Time independent tangent this is consider has a spline tangent key
+			const bool bTangentGenericTimeIndependent = (KeyTangentMode & FbxAnimCurveDef::ETangentMode::eTangentGenericTimeIndependent);
+
+			// When this flag is true, the tangent is flat if the value is outside of the [previous key, next key] value range.
+			//Clamp progressive is (eTangentGenericClampProgressive |eTangentGenericTimeIndependent)
+			const bool bTangentGenericClampProgressive = (KeyTangentMode & FbxAnimCurveDef::ETangentMode::eTangentGenericClampProgressive) == FbxAnimCurveDef::ETangentMode::eTangentGenericClampProgressive;
+
+			if (KeyTangentMode & FbxAnimCurveDef::eTangentGenericBreak)
+			{
+				NewTangentMode = EInterchangeCurveTangentMode::Break;
+			}
+			else if (KeyTangentMode & FbxAnimCurveDef::eTangentUser)
+			{
+				NewTangentMode = EInterchangeCurveTangentMode::User;
+			}
+
+			switch (KeyInterpMode)
+			{
+			case FbxAnimCurveDef::eInterpolationConstant://! Constant value until next key.
+				NewInterpMode = EInterchangeCurveInterpMode::Constant;
+				break;
+			case FbxAnimCurveDef::eInterpolationLinear://! Linear progression to next key.
+				NewInterpMode = EInterchangeCurveInterpMode::Linear;
+				break;
+			case FbxAnimCurveDef::eInterpolationCubic://! Cubic progression to next key.
+				NewInterpMode = EInterchangeCurveInterpMode::Cubic;
+				// get tangents
+				{
+					bool bIsFlatTangent = false;
+					bool bIsComputedTangent = false;
+					if (bTangentGenericClampProgressive)
+					{
+						if (bPreviousKeyValid && bNextKeyValid)
+						{
+							const float PreviousNextHalfDelta = (NextValue - PreviousValue) * 0.5f;
+							const float PreviousNextAverage = PreviousValue + PreviousNextHalfDelta;
+							// If the value is outside of the previous-next value range, the tangent is flat.
+							bIsFlatTangent = FMath::Abs(Value - PreviousNextAverage) >= FMath::Abs(PreviousNextHalfDelta);
+						}
+						else
+						{
+							//Start/End tangent with the ClampProgressive flag are flat.
+							bIsFlatTangent = true;
+						}
+					}
+					else if (bTangentGenericClamp && (bPreviousKeyValid || bNextKeyValid))
+					{
+						if (bPreviousKeyValid && PreviousValue == Value)
+						{
+							bIsFlatTangent = true;
+						}
+						if (bNextKeyValid)
+						{
+							bIsFlatTangent |= Value == NextValue;
+						}
+					}
+					else if (bTangentGenericTimeIndependent)
+					{
+						//Spline tangent key, because bTangentGenericClampProgressive include bTangentGenericTimeIndependent, we must treat this case after bTangentGenericClampProgressive
+						if (KeyCount == 1)
+						{
+							bIsFlatTangent = true;
+						}
+						else
+						{
+							//Spline tangent key must be User mode since we want to keep the tangents provide by the fbx key left and right derivatives
+							NewTangentMode = EInterchangeCurveTangentMode::User;
+						}
+					}
+
+					if (bIsFlatTangent)
+					{
+						RightTangent = 0;
+						LeftTangent = 0;
+						//To force flat tangent we need to set the tangent mode to user
+						NewTangentMode = EInterchangeCurveTangentMode::User;
+					}
+
+				}
+				break;
+			}
+
+			//auto with weighted give the wrong result, so when auto is weighted we set user mode and set the Right tangent equal to the left tangent.
+			//Auto has only the left tangent set
+			if (NewTangentMode == EInterchangeCurveTangentMode::Auto && (bLeftWeightActive || bRightWeightActive))
+			{
+
+				NewTangentMode = EInterchangeCurveTangentMode::User;
+				RightTangent = LeftTangent;
+			}
+
+			if (NewTangentMode != EInterchangeCurveTangentMode::Auto)
+			{
+				const bool bEqualTangents = FMath::IsNearlyEqual(LeftTangent, RightTangent);
+				//If tangents are different then broken.
+				if (bEqualTangents)
+				{
+					NewTangentMode = EInterchangeCurveTangentMode::User;
+				}
+				else
+				{
+					NewTangentMode = EInterchangeCurveTangentMode::Break;
+				}
+			}
+
+			//Only cubic interpolation allow weighted tangents
+			if (KeyInterpMode == FbxAnimCurveDef::eInterpolationCubic)
+			{
+				if (bLeftWeightActive && bRightWeightActive)
+				{
+					NewTangentWeightMode = EInterchangeCurveTangentWeightMode::WeightedBoth;
+				}
+				else if (bLeftWeightActive)
+				{
+					NewTangentWeightMode = EInterchangeCurveTangentWeightMode::WeightedArrive;
+					RightTangentWeight = DefaultCurveWeight;
+				}
+				else if (bRightWeightActive)
+				{
+					NewTangentWeightMode = EInterchangeCurveTangentWeightMode::WeightedLeave;
+					LeftTangentWeight = DefaultCurveWeight;
+				}
+				else
+				{
+					NewTangentWeightMode = EInterchangeCurveTangentWeightMode::WeightedNone;
+					LeftTangentWeight = DefaultCurveWeight;
+					RightTangentWeight = DefaultCurveWeight;
+				}
+
+				auto ComputeWeightInternal = [](float TimeA, float TimeB, const float TangentSlope, const float TangentWeight)
+				{
+					const float X = TimeA - TimeB;
+					const float Y = TangentSlope * X;
+					return FMath::Sqrt(X * X + Y * Y) * TangentWeight;
+				};
+
+				if (!FMath::IsNearlyZero(LeftTangentWeight))
+				{
+					if (bPreviousKeyValid)
+					{
+						LeftTangentWeight = ComputeWeightInternal(KeyTimeValue, PreviousKeyTimeValue, LeftTangent, LeftTangentWeight);
+					}
+					else
+					{
+						LeftTangentWeight = 0.0f;
+					}
+				}
+
+				if (!FMath::IsNearlyZero(RightTangentWeight))
+				{
+					if (bNextKeyValid)
+					{
+						RightTangentWeight = ComputeWeightInternal(NextKeyTimeValue, KeyTimeValue, RightTangent, RightTangentWeight);
+					}
+					else
+					{
+						RightTangentWeight = 0.0f;
+					}
+				}
+			}
+
+			const bool bForceDisableTangentRecompute = false; //No need to recompute all the tangents of the curve every time we change de key.
+			InterchangeCurveKey.InterpMode = NewInterpMode;
+			InterchangeCurveKey.TangentMode = NewTangentMode;
+			InterchangeCurveKey.TangentWeightMode = NewTangentWeightMode;
+
+			InterchangeCurveKey.ArriveTangent = LeftTangent;
+			InterchangeCurveKey.LeaveTangent = RightTangent;
+			InterchangeCurveKey.ArriveTangentWeight = LeftTangentWeight;
+			InterchangeCurveKey.LeaveTangentWeight = RightTangentWeight;
+		}
+		return true;
+	}
+
+	bool ImportCurve(const FbxAnimCurve* SourceFloatCurves, const float ScaleValue, TArray<FInterchangeCurveKey>& DestinationFloatCurve)
+	{
+		return ImportCurve(nullptr, EInterchangeTransformCurveChannel::None, SourceFloatCurves, ScaleValue, DestinationFloatCurve);
+	}
+
+	bool ImportTranslationCurves(FbxNode* Node, TArray<FbxAnimCurve*>& SourceTransformChannelCurves, TArray<FInterchangeCurve>& TransformChannelCurves)
+	{
+		bool bResult = true;
+		//Translation X
+		bResult &= ImportCurve(Node, EInterchangeTransformCurveChannel::TranslationX, SourceTransformChannelCurves[0], 1.0f, TransformChannelCurves[0].Keys);
+		TransformChannelCurves[0].TransformChannel = EInterchangeTransformCurveChannel::TranslationX;
+		//Translation Y
+		bResult &= ImportCurve(Node, EInterchangeTransformCurveChannel::TranslationY, SourceTransformChannelCurves[1], 1.0f, TransformChannelCurves[1].Keys);
+		TransformChannelCurves[1].TransformChannel = EInterchangeTransformCurveChannel::TranslationY;
+		//Translation Z
+		bResult &= ImportCurve(Node, EInterchangeTransformCurveChannel::TranslationZ, SourceTransformChannelCurves[2], 1.0f, TransformChannelCurves[2].Keys);
+		TransformChannelCurves[2].TransformChannel = EInterchangeTransformCurveChannel::TranslationZ;
+
+		return bResult;
+	}
+	
+	bool ImportRotationCurves(FbxNode* Node, TArray<FbxAnimCurve*>& SourceTransformChannelCurves, TArray<FInterchangeCurve>& TransformChannelCurves)
+	{
+		bool bResult = true;
+		//Euler X
+		bResult &= ImportCurve(Node, EInterchangeTransformCurveChannel::EulerX, SourceTransformChannelCurves[3], 1.0f, TransformChannelCurves[3].Keys);
+		TransformChannelCurves[3].TransformChannel = EInterchangeTransformCurveChannel::EulerX;
+		//Euler Y
+		bResult &= ImportCurve(Node, EInterchangeTransformCurveChannel::EulerY, SourceTransformChannelCurves[4], 1.0f, TransformChannelCurves[4].Keys);
+		TransformChannelCurves[4].TransformChannel = EInterchangeTransformCurveChannel::EulerY;
+		//Euler Z
+		bResult &= ImportCurve(Node, EInterchangeTransformCurveChannel::EulerZ, SourceTransformChannelCurves[5], 1.0f, TransformChannelCurves[5].Keys);
+		TransformChannelCurves[5].TransformChannel = EInterchangeTransformCurveChannel::EulerZ;
+
+		return bResult;
+	}
+	
+	bool ImportScaleCurves(FbxNode* Node, TArray<FbxAnimCurve*>& SourceTransformChannelCurves, TArray<FInterchangeCurve>& TransformChannelCurves)
+	{
+		bool bResult = true;
+		//Scale X
+		bResult &= ImportCurve(Node, EInterchangeTransformCurveChannel::ScaleX, SourceTransformChannelCurves[6], 1.0f, TransformChannelCurves[6].Keys);
+		TransformChannelCurves[6].TransformChannel = EInterchangeTransformCurveChannel::ScaleX;
+		//Scale Y
+		bResult &= ImportCurve(Node, EInterchangeTransformCurveChannel::ScaleY, SourceTransformChannelCurves[7], 1.0f, TransformChannelCurves[7].Keys);
+		TransformChannelCurves[7].TransformChannel = EInterchangeTransformCurveChannel::ScaleY;
+		//Scale Z
+		bResult &= ImportCurve(Node, EInterchangeTransformCurveChannel::ScaleZ, SourceTransformChannelCurves[8], 1.0f, TransformChannelCurves[8].Keys);
+		TransformChannelCurves[8].TransformChannel = EInterchangeTransformCurveChannel::ScaleZ;
+
+		return bResult;
+	}
+
+	bool ImportTransformCurves(FbxNode* Node
+		, TArray<FbxAnimCurve*>& SourceTransformChannelCurves
+		, TArray<FInterchangeCurve>& TransformChannelCurves)
+	{
+		bool bResult = true;
+		bResult &= ImportTranslationCurves(Node, SourceTransformChannelCurves, TransformChannelCurves);
+		bResult &= ImportRotationCurves(Node, SourceTransformChannelCurves, TransformChannelCurves);
+		bResult &= ImportScaleCurves(Node, SourceTransformChannelCurves, TransformChannelCurves);
+		return bResult;
+	}
+
+	struct FGetFbxTransformCurvesParameters
+	{
+		FGetFbxTransformCurvesParameters(FbxScene* InSDKScene, FbxNode* InNode)
+		{
+			SDKScene = InSDKScene;
+			ensure(SDKScene);
+			Node = InNode;
+			ensure(Node);
+		}
+
+		FbxScene* SDKScene = nullptr;
+		FbxNode* Node = nullptr;
+		TArray<FbxAnimCurve*> TransformChannelCurves;
+		int32 TransformChannelCount = 0;
+		bool IsNodeAnimated = false;
+		FbxTimeSpan ExportedTimeSpan = FbxTimeSpan(FBXSDK_TIME_INFINITE, FBXSDK_TIME_MINUS_INFINITE);
+		FbxTime StartTime = FBXSDK_TIME_INFINITE;
+		FbxTime EndTime = FBXSDK_TIME_MINUS_INFINITE;
+		int32 KeyCount = 0;
+	};
+
+	void GetFbxTransformCurves(FGetFbxTransformCurvesParameters& Parameters)
+	{
+		ensure(Parameters.SDKScene);
+		//Get the node transform curve keys, the transform components are separate into float curve
+		//Translation X
+		//Translation Y
+		//Translation Z
+		//Euler X
+		//Euler Y
+		//Euler Z
+		//Scale X
+		//Scale Y
+		//Scale Z
+		int32 NumAnimations = Parameters.SDKScene->GetSrcObjectCount<FbxAnimStack>();
+		//Anim stack should be merge so we expect to have only one stack here
+		ensure(NumAnimations <= 1);
+
+		Parameters.IsNodeAnimated = false;
+
+		Parameters.TransformChannelCount = 9;
+		Parameters.TransformChannelCurves.Reset(Parameters.TransformChannelCount);
+		for (int32 AnimationIndex = 0; AnimationIndex < NumAnimations && !Parameters.IsNodeAnimated; AnimationIndex++)
+		{
+			FbxAnimStack* AnimStack = (FbxAnimStack*)Parameters.SDKScene->GetSrcObject<FbxAnimStack>(AnimationIndex);
+			int32 NumLayers = AnimStack->GetMemberCount();
+			for (int LayerIndex = 0; LayerIndex < NumLayers && !Parameters.IsNodeAnimated; LayerIndex++)
+			{
+				FbxAnimLayer* AnimLayer = (FbxAnimLayer*)AnimStack->GetMember(LayerIndex);
+				// Display curves specific to properties
+				Parameters.TransformChannelCurves.Add(Parameters.Node->LclTranslation.GetCurve(AnimLayer, FBXSDK_CURVENODE_COMPONENT_X, false));
+				Parameters.TransformChannelCurves.Add(Parameters.Node->LclTranslation.GetCurve(AnimLayer, FBXSDK_CURVENODE_COMPONENT_Y, false));
+				Parameters.TransformChannelCurves.Add(Parameters.Node->LclTranslation.GetCurve(AnimLayer, FBXSDK_CURVENODE_COMPONENT_Z, false));
+
+				Parameters.TransformChannelCurves.Add(Parameters.Node->LclRotation.GetCurve(AnimLayer, FBXSDK_CURVENODE_COMPONENT_X, false));
+				Parameters.TransformChannelCurves.Add(Parameters.Node->LclRotation.GetCurve(AnimLayer, FBXSDK_CURVENODE_COMPONENT_Y, false));
+				Parameters.TransformChannelCurves.Add(Parameters.Node->LclRotation.GetCurve(AnimLayer, FBXSDK_CURVENODE_COMPONENT_Z, false));
+
+				Parameters.TransformChannelCurves.Add(Parameters.Node->LclScaling.GetCurve(AnimLayer, FBXSDK_CURVENODE_COMPONENT_X, false));
+				Parameters.TransformChannelCurves.Add(Parameters.Node->LclScaling.GetCurve(AnimLayer, FBXSDK_CURVENODE_COMPONENT_Y, false));
+				Parameters.TransformChannelCurves.Add(Parameters.Node->LclScaling.GetCurve(AnimLayer, FBXSDK_CURVENODE_COMPONENT_Z, false));
+
+				if (!Parameters.IsNodeAnimated)
+				{
+					for (int32 ChannelIndex = 0; ChannelIndex < Parameters.TransformChannelCount; ++ChannelIndex)
+					{
+						if (Parameters.TransformChannelCurves[ChannelIndex])
+						{
+							Parameters.IsNodeAnimated = true;
+							const int32 LocalKeycount = Parameters.TransformChannelCurves[ChannelIndex]->KeyGetCount();
+							if (LocalKeycount > Parameters.KeyCount)
+							{
+								Parameters.KeyCount = LocalKeycount;
+							}
+							if (LocalKeycount > 1)
+							{
+								FbxTimeSpan LocalAnimatedTimeSpan(FBXSDK_TIME_INFINITE, FBXSDK_TIME_MINUS_INFINITE);
+								Parameters.TransformChannelCurves[ChannelIndex]->GetTimeInterval(LocalAnimatedTimeSpan);
+								if (Parameters.StartTime > LocalAnimatedTimeSpan.GetStart())
+								{
+									Parameters.StartTime = LocalAnimatedTimeSpan.GetStart();
+								}
+								if (Parameters.EndTime < LocalAnimatedTimeSpan.GetStop())
+								{
+									Parameters.EndTime = LocalAnimatedTimeSpan.GetStop();
+								}
+							}
+							else if (Parameters.KeyCount == 1)
+							{
+								//When there is only one key there is no interval
+								FbxAnimCurveKey FbxKey = Parameters.TransformChannelCurves[ChannelIndex]->KeyGet(0);
+								FbxTime KeyTime = FbxKey.GetTime();
+								if (Parameters.StartTime > KeyTime)
+								{
+									Parameters.StartTime = KeyTime;
+								}
+								if (Parameters.EndTime < KeyTime)
+								{
+									Parameters.EndTime = KeyTime;
+								}
+							}
+						}
+					}
+				}
+			}
+		}
+	}
+
+	bool FAnimationPayloadContextTransform::FetchPayloadToFile(FFbxParser& Parser, const FString& PayloadFilepath)
+	{
+		if (!ensure(SDKScene != nullptr))
+		{
+			UInterchangeResultError_Generic* Message = Parser.AddMessage<UInterchangeResultError_Generic>();
+			Message->InterchangeKey = FFbxHelper::GetFbxNodeHierarchyName(Node);
+			Message->Text = LOCTEXT("FBXSceneNull", "Cannot fetch FBX animation transform payload because the FBX scene is null.");
+			return false;
+		}
+
+		if (!ensure(Node != nullptr))
+		{
+			UInterchangeResultError_Generic* Message = Parser.AddMessage<UInterchangeResultError_Generic>();
+			Message->InterchangeKey = FFbxHelper::GetFbxNodeHierarchyName(Node);
+			Message->Text = LOCTEXT("FBXNodeNull", "Cannot fetch FBX animation transform payload because the FBX node is null.");
+			return false;
+		}
+
+		FGetFbxTransformCurvesParameters Parameters(SDKScene, Node);
+		GetFbxTransformCurves(Parameters);
+
+		TArray<FInterchangeCurve> Curves;
+		Curves.AddDefaulted(Parameters.TransformChannelCount);
+
+		if (Parameters.IsNodeAnimated)
+		{
+			//Fbx node global transform evaluator is not thread safe, so we acquire a lock to import the curves
+			FScopeLock Lock(&Parser.AnimationTransformPayloadCriticalSection);
+			ImportTransformCurves(Node, Parameters.TransformChannelCurves, Curves);
+		}
+
+		//Dump all curves to a file
+		{
+			FLargeMemoryWriter Ar;
+			Ar << Curves;
+			uint8* ArchiveData = Ar.GetData();
+			int64 ArchiveSize = Ar.TotalSize();
+			TArray64<uint8> Buffer(ArchiveData, ArchiveSize);
+			FFileHelper::SaveArrayToFile(Buffer, *PayloadFilepath);
+		}
+		return true;
+	}
+
+	void FFbxAnimation::AddJointAnimation(FbxScene* SDKScene, FbxNode* JointNode, UInterchangeBaseNodeContainer& NodeContainer, UInterchangeSceneNode* SceneNode, TMap<FString, TSharedPtr<FPayloadContextBase>>& PayloadContexts)
+	{
+		int32 NumAnimations = SDKScene->GetSrcObjectCount<FbxAnimStack>();
+		//Anim stack should be merge so we expect to have only one stack here
+		ensure(NumAnimations <= 1);
+
+		double FrameRate = FbxTime::GetFrameRate(SDKScene->GetGlobalSettings().GetTimeMode());
+
+		FGetFbxTransformCurvesParameters Parameters(SDKScene, JointNode);
+		GetFbxTransformCurves(Parameters);
+
+		if (Parameters.IsNodeAnimated)
+		{
+			FString PayLoadKey = FFbxHelper::GetFbxNodeHierarchyName(JointNode) + TEXT("_AnimationPayloadKey");
+			if (ensure(!PayloadContexts.Contains(PayLoadKey)))
+			{
+				TSharedPtr<FAnimationPayloadContextTransform> AnimPayload = MakeShared<FAnimationPayloadContextTransform>();
+				AnimPayload->Node = JointNode;
+				AnimPayload->SDKScene = SDKScene;
+				
+				PayloadContexts.Add(PayLoadKey, AnimPayload);
+			}
+			SceneNode->SetCustomTransformCurvePayloadKey(PayLoadKey);
+			SceneNode->SetCustomIsNodeTransformAnimated(Parameters.IsNodeAnimated);
+			SceneNode->SetCustomNodeTransformAnimationKeyCount(Parameters.KeyCount);
+			SceneNode->SetCustomNodeTransformAnimationStartTime(Parameters.StartTime.GetSecondDouble());
+			SceneNode->SetCustomNodeTransformAnimationEndTime(Parameters.EndTime.GetSecondDouble());
+		}
+	}
+
+}//ns UE::Interchange::Private
+
+#undef LOCTEXT_NAMESPACE
diff --git a/Engine/Plugins/Experimental/Interchange/Runtime/Source/Parsers/Fbx/Private/FbxAnimation.h b/Engine/Plugins/Experimental/Interchange/Runtime/Source/Parsers/Fbx/Private/FbxAnimation.h
new file mode 100644
--- /dev/null
+++ b/Engine/Plugins/Experimental/Interchange/Runtime/Source/Parsers/Fbx/Private/FbxAnimation.h
@@ -0,0 +1,32 @@
+// Copyright Epic Games, Inc. All Rights Reserved.
+
+#pragma once
+
+#include "CoreMinimal.h"
+#include "FbxAPI.h"
+#include "FbxHelper.h"
+#include "FbxInclude.h"
+
+/** Forward declarations */
+struct FMeshDescription;
+class UInterchangeBaseNodeContainer;
+class UInterchangeSceneNode;
+
+namespace UE::Interchange::Private
+{
+	class FAnimationPayloadContextTransform : public FPayloadContextBase
+	{
+	public:
+		virtual ~FAnimationPayloadContextTransform() {}
+		virtual FString GetPayloadType() const override { return TEXT("JointAnimation-PayloadContext"); }
+		virtual bool FetchPayloadToFile(FFbxParser& Parser, const FString& PayloadFilepath) override;
+		FbxNode* Node = nullptr;
+		FbxScene* SDKScene = nullptr;
+	};
+
+	class FFbxAnimation
+	{
+	public:
+		static void AddJointAnimation(FbxScene* SDKScene, FbxNode* JointNode, UInterchangeBaseNodeContainer& NodeContainer, UInterchangeSceneNode* UnrealNode, TMap<FString, TSharedPtr<FPayloadContextBase>>& PayloadContexts);
+	};
+}//ns UE::Interchange::Private
diff --git a/Engine/Plugins/Experimental/Interchange/Runtime/Source/Parsers/Fbx/Private/FbxScene.cpp b/Engine/Plugins/Experimental/Interchange/Runtime/Source/Parsers/Fbx/Private/FbxScene.cpp
--- a/Engine/Plugins/Experimental/Interchange/Runtime/Source/Parsers/Fbx/Private/FbxScene.cpp
+++ b/Engine/Plugins/Experimental/Interchange/Runtime/Source/Parsers/Fbx/Private/FbxScene.cpp
@@ -3,6 +3,7 @@
 #include "FbxScene.h"
 
 #include "CoreMinimal.h"
+#include "FbxAnimation.h"
 #include "FbxAPI.h"
 #include "FbxConvert.h"
 #include "FbxHelper.h"
@@ -15,6 +16,7 @@
 #include "InterchangeResultsContainer.h"
 #include "InterchangeSceneNode.h"
 #include "Nodes/InterchangeBaseNodeContainer.h"
+#include "Nodes/InterchangeSourceNode.h"
 #include "Nodes/InterchangeUserDefinedAttribute.h"
 
 #define LOCTEXT_NAMESPACE "InterchangeFbxScene"
@@ -89,7 +91,11 @@ namespace UE
 				return DoesTheParentHierarchyContainJoints(Node->GetParent());
 			}
 
-			void FFbxScene::AddHierarchyRecursively(UInterchangeSceneNode* UnrealParentNode, FbxNode* Node, FbxScene* SDKScene, UInterchangeBaseNodeContainer& NodeContainer)
+			void FFbxScene::AddHierarchyRecursively(UInterchangeSceneNode* UnrealParentNode
+				, FbxNode* Node
+				, FbxScene* SDKScene
+				, UInterchangeBaseNodeContainer& NodeContainer
+				, TMap<FString, TSharedPtr<FPayloadContextBase, ESPMode::ThreadSafe>>& PayloadContexts)
 			{
 				FString NodeName = FFbxHelper::GetFbxObjectName(Node);
 				FString NodeUniqueID = FFbxHelper::GetFbxNodeHierarchyName(Node);
@@ -224,7 +230,7 @@ namespace UE
 									UnrealNode->SetCustomTimeZeroLocalTransform(&NodeContainer, GlobalTransform);
 								}
 							}
-
+							FFbxAnimation::AddJointAnimation(SDKScene, Node, NodeContainer, UnrealNode, PayloadContexts);
 							break;
 						}
 
@@ -309,6 +315,7 @@ namespace UE
 						if (CurveNode && CurveNode->IsAnimated())
 						{
 							//Set the optional payload key
+							//Create an attribute animation payload
 						}
 						switch (Property.GetPropertyDataType().GetType())
 						{
@@ -431,7 +438,7 @@ namespace UE
 				for (int32 ChildIndex = 0; ChildIndex < ChildCount; ++ChildIndex)
 				{
 					FbxNode* ChildNode = Node->GetChild(ChildIndex);
-					AddHierarchyRecursively(UnrealNode, ChildNode, SDKScene, NodeContainer);
+					AddHierarchyRecursively(UnrealNode, ChildNode, SDKScene, NodeContainer, PayloadContexts);
 				}
 			}
 
@@ -452,10 +459,58 @@ namespace UE
 				return TransformNode;
 			}
 
-			void FFbxScene::AddHierarchy(FbxScene* SDKScene, UInterchangeBaseNodeContainer& NodeContainer)
+			void FFbxScene::AddHierarchy(FbxScene* SDKScene, UInterchangeBaseNodeContainer& NodeContainer, TMap<FString, TSharedPtr<FPayloadContextBase, ESPMode::ThreadSafe>>& PayloadContexts)
 			{
 				 FbxNode* RootNode = SDKScene->GetRootNode();
-				 AddHierarchyRecursively(nullptr, RootNode, SDKScene, NodeContainer);
+				 //Create a source node where we can store any general file info
+				 {
+					 FString NodeName = FFbxHelper::GetFbxObjectName(RootNode) + TEXT("_SourceInfo");
+					 FString NodeUniqueID = TEXT("\\SourceInfo\\") + FFbxHelper::GetFbxNodeHierarchyName(RootNode);
+					 UInterchangeSourceNode* SourceNode = NewObject<UInterchangeSourceNode>(&NodeContainer, NAME_None);
+					 if (!ensure(SourceNode))
+					 {
+						 UInterchangeResultError_Generic* Message = Parser.AddMessage<UInterchangeResultError_Generic>();
+						 Message->Text = LOCTEXT("NodeAllocationError", "Unable to allocate a node when importing FBX.");
+						 return;
+					 }
+					 SourceNode->InitializeNode(NodeUniqueID, NodeName, EInterchangeNodeContainerType::TranslatedAsset);
+
+					 //Store the fbx frame rate
+					 {
+						 double FrameRate = FbxTime::GetFrameRate(SDKScene->GetGlobalSettings().GetTimeMode());
+						 SourceNode->SetCustomSourceFrameRateNumerator(FrameRate);
+						 constexpr double Denominator = 1.0;
+						 SourceNode->SetCustomSourceFrameRateDenominator(Denominator);
+					 }
+
+					 //Store the fbx timeline
+					 {
+						 //Timeline time span
+						 FbxTimeSpan TimelineTimeSpan(FBXSDK_TIME_INFINITE, FBXSDK_TIME_MINUS_INFINITE);
+						 SDKScene->GetGlobalSettings().GetTimelineDefaultTimeSpan(TimelineTimeSpan);
+						 SourceNode->SetCustomSourceTimelineStart(TimelineTimeSpan.GetStart().GetSecondDouble());
+						 SourceNode->SetCustomSourceTimelineEnd(TimelineTimeSpan.GetStop().GetSecondDouble());
+
+						 //Animated time span
+						 FbxTimeSpan AnimatedTimeSpan(FBXSDK_TIME_INFINITE, FBXSDK_TIME_MINUS_INFINITE);
+						 int32 AnimCurveNodeCount = SDKScene->GetSrcObjectCount<FbxAnimCurveNode>();
+						 for (int32 AnimCurveNodeIndex = 0; AnimCurveNodeIndex < AnimCurveNodeCount; AnimCurveNodeIndex++)
+						 {
+							 FbxAnimCurveNode* CurAnimCruveNode = SDKScene->GetSrcObject<FbxAnimCurveNode>(AnimCurveNodeIndex);
+							 if (CurAnimCruveNode->IsAnimated(true))
+							 {
+								 FbxTimeSpan CurveTimeSpan(FBXSDK_TIME_INFINITE, FBXSDK_TIME_MINUS_INFINITE);
+								 CurAnimCruveNode->GetAnimationInterval(CurveTimeSpan);
+								 AnimatedTimeSpan.UnionAssignment(CurveTimeSpan);
+							 }
+						 }
+						 SourceNode->SetCustomAnimatedTimeStart(AnimatedTimeSpan.GetStart().GetSecondDouble());
+						 SourceNode->SetCustomAnimatedTimeEnd(AnimatedTimeSpan.GetStop().GetSecondDouble());
+					 }
+
+					 NodeContainer.AddNode(SourceNode);
+				 }
+				 AddHierarchyRecursively(nullptr, RootNode, SDKScene, NodeContainer, PayloadContexts);
 			}
 		} //ns Private
 	} //ns Interchange
diff --git a/Engine/Plugins/Experimental/Interchange/Runtime/Source/Parsers/Fbx/Private/FbxScene.h b/Engine/Plugins/Experimental/Interchange/Runtime/Source/Parsers/Fbx/Private/FbxScene.h
--- a/Engine/Plugins/Experimental/Interchange/Runtime/Source/Parsers/Fbx/Private/FbxScene.h
+++ b/Engine/Plugins/Experimental/Interchange/Runtime/Source/Parsers/Fbx/Private/FbxScene.h
@@ -6,6 +6,10 @@
 #include "FbxInclude.h"
 
 /** Forward declarations */
+namespace UE::Interchange::Private
+{
+	class FPayloadContextBase;
+}
 class UInterchangeBaseNodeContainer;
 class UInterchangeSceneNode;
 
@@ -24,14 +28,18 @@ namespace UE
 					: Parser(InParser)
 				{}
 
-				void AddHierarchy(FbxScene* SDKScene, UInterchangeBaseNodeContainer& NodeContainer);
+				void AddHierarchy(FbxScene* SDKScene, UInterchangeBaseNodeContainer& NodeContainer, TMap<FString, TSharedPtr<FPayloadContextBase, ESPMode::ThreadSafe>>& PayloadContexts);
 				UInterchangeSceneNode* CreateTransformNode(UInterchangeBaseNodeContainer& NodeContainer, const FString& NodeName, const FString& NodeUniqueID);
 
 			protected:
 				void CreateMeshNodeReference(UInterchangeSceneNode* UnrealSceneNode, FbxNodeAttribute* NodeAttribute, UInterchangeBaseNodeContainer& NodeContainer, const FTransform& GeometricTransform);
 				void CreateCameraNodeReference(UInterchangeSceneNode* UnrealSceneNode, FbxNodeAttribute* NodeAttribute, UInterchangeBaseNodeContainer& NodeContainer);
 				void CreateLightNodeReference(UInterchangeSceneNode* UnrealSceneNode, FbxNodeAttribute* NodeAttribute, UInterchangeBaseNodeContainer& NodeContainer);
-				void AddHierarchyRecursively(UInterchangeSceneNode* UnrealParentNode, FbxNode* Node, FbxScene* SDKScene, UInterchangeBaseNodeContainer& NodeContainer);
+				void AddHierarchyRecursively(UInterchangeSceneNode* UnrealParentNode
+					, FbxNode* Node
+					, FbxScene* SDKScene
+					, UInterchangeBaseNodeContainer& NodeContainer
+					, TMap<FString, TSharedPtr<FPayloadContextBase, ESPMode::ThreadSafe>>& PayloadContexts);
 
 			private:
 				FFbxParser& Parser;
diff --git a/Engine/Plugins/Experimental/Interchange/Runtime/Source/Pipelines/InterchangePipelines.Build.cs b/Engine/Plugins/Experimental/Interchange/Runtime/Source/Pipelines/InterchangePipelines.Build.cs
--- a/Engine/Plugins/Experimental/Interchange/Runtime/Source/Pipelines/InterchangePipelines.Build.cs
+++ b/Engine/Plugins/Experimental/Interchange/Runtime/Source/Pipelines/InterchangePipelines.Build.cs
@@ -23,6 +23,7 @@ namespace UnrealBuildTool.Rules
 				new string[]
 				{
 					"CinematicCamera",
+					"InterchangeImport",
 				}
 			);
 
diff --git a/Engine/Plugins/Experimental/Interchange/Runtime/Source/Pipelines/Private/InterchangeGenericAnimationPipeline.cpp b/Engine/Plugins/Experimental/Interchange/Runtime/Source/Pipelines/Private/InterchangeGenericAnimationPipeline.cpp
new file mode 100644
--- /dev/null
+++ b/Engine/Plugins/Experimental/Interchange/Runtime/Source/Pipelines/Private/InterchangeGenericAnimationPipeline.cpp
@@ -0,0 +1,276 @@
+// Copyright Epic Games, Inc. All Rights Reserved. 
+#include "InterchangeGenericAnimationPipeline.h"
+
+#include "CoreMinimal.h"
+#include "InterchangeAnimSequenceFactoryNode.h"
+#include "InterchangePipelineLog.h"
+#include "InterchangeSceneNode.h"
+#include "InterchangeSkeletalMeshFactoryNode.h"
+#include "InterchangeSkeletonFactoryNode.h"
+#include "InterchangeSkeletonHelper.h"
+#include "InterchangeSourceData.h"
+#include "Nodes/InterchangeBaseNode.h"
+#include "Nodes/InterchangeBaseNodeContainer.h"
+#include "Nodes/InterchangeSourceNode.h"
+
+
+namespace UE::Interchange::Private
+{
+	/**
+	 * Return true if there is one animated node in the scene node hierarchy under NodeUid
+	 */
+	bool IsSkeletonAnimatedRecursive(const FString& NodeUid, UInterchangeBaseNodeContainer* BaseNodeContainer)
+	{
+		if (const UInterchangeSceneNode* SceneNode = Cast<UInterchangeSceneNode>(BaseNodeContainer->GetNode(NodeUid)))
+		{
+			bool bIsAnimated = false;
+			if (SceneNode->GetCustomIsNodeTransformAnimated(bIsAnimated) && bIsAnimated)
+			{
+				return true;
+			}
+		}
+		TArray<FString> Children = BaseNodeContainer->GetNodeChildrenUids(NodeUid);
+		for (const FString& ChildUid : Children)
+		{
+			if (IsSkeletonAnimatedRecursive(ChildUid, BaseNodeContainer))
+			{
+				return true;
+			}
+		}
+		return false;
+	}
+}
+
+void UInterchangeGenericAnimationPipeline::SetupReimportData(TObjectPtr<UObject> ReimportObject)
+{
+	check(!CommonSkeletalMeshesAndAnimationsProperties.IsNull());
+	if (UAnimSequence* AnimSequence = Cast<UAnimSequence>(ReimportObject))
+	{
+		//Set the skeleton to the current asset skeleton and re-import only the animation
+		CommonSkeletalMeshesAndAnimationsProperties->Skeleton = AnimSequence->GetSkeleton();
+		CommonSkeletalMeshesAndAnimationsProperties->bImportOnlyAnimations = true;
+	}
+}
+
+void UInterchangeGenericAnimationPipeline::ExecutePreImportPipeline(UInterchangeBaseNodeContainer* InBaseNodeContainer, const TArray<UInterchangeSourceData*>& InSourceDatas)
+{
+	check(!CommonSkeletalMeshesAndAnimationsProperties.IsNull());
+	if (!InBaseNodeContainer)
+	{
+		UE_LOG(LogInterchangePipeline, Warning, TEXT("UInterchangeGenericAnimationPipeline: Cannot execute pre-import pipeline because InBaseNodeContrainer is null"));
+		return;
+	}
+
+	if (CommonSkeletalMeshesAndAnimationsProperties->bImportOnlyAnimations && CommonSkeletalMeshesAndAnimationsProperties->Skeleton.IsNull())
+	{
+		UE_LOG(LogInterchangePipeline, Warning, TEXT("UInterchangeGenericAnimationPipeline: Cannot execute pre-import pipeline because we cannot import animation only but not specify any valid skeleton"));
+		return;
+	}
+
+	BaseNodeContainer = InBaseNodeContainer;
+	SourceDatas.Empty(InSourceDatas.Num());
+	for (const UInterchangeSourceData* SourceData : InSourceDatas)
+	{
+		SourceDatas.Add(SourceData);
+	}
+
+	double SampleRate = 30.0;
+	double RangeStart = 0;
+	double RangeStop = 0;
+	bool bRangeIsValid = false;
+	TArray<FString> SourceNodeUids;
+	BaseNodeContainer->GetNodes(UInterchangeSourceNode::StaticClass(), SourceNodeUids);
+	if (SourceNodeUids.Num() > 0)
+	{
+		if (UInterchangeSourceNode* SourceNode = Cast<UInterchangeSourceNode>(BaseNodeContainer->GetNode(SourceNodeUids[0])))
+		{
+			if (bImportBoneTracks)
+			{
+				int32 Numerator, Denominator;
+				if (!bUse30HzToBakeBoneAnimation && CustomBoneAnimationSampleRate == 0 && SourceNode->GetCustomSourceFrameRateNumerator(Numerator))
+				{
+					if (SourceNode->GetCustomSourceFrameRateDenominator(Denominator) && Denominator > 0 && Numerator > 0)
+					{
+						SampleRate = static_cast<double>(Numerator) / static_cast<double>(Denominator);
+					}
+				}
+				else if ((!bUse30HzToBakeBoneAnimation && CustomBoneAnimationSampleRate > 0))
+				{
+					SampleRate = static_cast<double>(CustomBoneAnimationSampleRate);
+				}
+			
+				if (AnimationRange == EInterchangeAnimationRange::Timeline)
+				{
+					if (SourceNode->GetCustomSourceTimelineStart(RangeStart))
+					{
+						if (SourceNode->GetCustomSourceTimelineEnd(RangeStop))
+						{
+							bRangeIsValid = true;
+						}
+					}
+				}
+				else if (AnimationRange == EInterchangeAnimationRange::Animated)
+				{
+					if (SourceNode->GetCustomAnimatedTimeStart(RangeStart))
+					{
+						if (SourceNode->GetCustomAnimatedTimeEnd(RangeStop))
+						{
+							bRangeIsValid = true;
+						}
+					}
+				}
+				else if (AnimationRange == EInterchangeAnimationRange::SetRange)
+				{
+					RangeStart = static_cast<double>(FrameImportRange.Min) / SampleRate;
+					RangeStop = static_cast<double>(FrameImportRange.Max) / SampleRate;
+					bRangeIsValid = true;
+				}
+			}
+		}
+	}
+	else
+	{
+		if (bImportBoneTracks)
+		{
+			if ((!bUse30HzToBakeBoneAnimation && CustomBoneAnimationSampleRate > 0))
+			{
+				SampleRate = CustomBoneAnimationSampleRate;
+			}
+			//Find the range by iteration the scene node
+			TArray<FString> SceneNodeUids;
+			BaseNodeContainer->GetNodes(UInterchangeSceneNode::StaticClass(), SceneNodeUids);
+			for (const FString& SceneNodeUid : SceneNodeUids)
+			{
+				if (UInterchangeSceneNode* SceneNode = Cast<UInterchangeSceneNode>(BaseNodeContainer->GetNode(SceneNodeUid)))
+				{
+					double SceneNodeAnimStart;
+					double SceneNodeAnimStop;
+					if (SceneNode->GetCustomNodeTransformAnimationStartTime(SceneNodeAnimStart))
+					{
+						if (SceneNode->GetCustomNodeTransformAnimationEndTime(SceneNodeAnimStop))
+						{
+							if (RangeStart > SceneNodeAnimStart)
+							{
+								RangeStart = SceneNodeAnimStart;
+							}
+							if (RangeStop < SceneNodeAnimStop)
+							{
+								RangeStop = SceneNodeAnimStop;
+							}
+							bRangeIsValid = true;
+						}
+					}
+				}
+			}
+		}
+	}
+
+
+	//Retrieve all skeletons
+	TArray<UInterchangeSkeletonFactoryNode*> Skeletons;
+	BaseNodeContainer->IterateNodes([&Skeletons, LocalNodeContainer = BaseNodeContainer](const FString& NodeUid, UInterchangeBaseNode* Node)
+		{
+			if (UInterchangeSkeletonFactoryNode* SkeletonFactoryNode = Cast<UInterchangeSkeletonFactoryNode>(Node))
+			{
+				FString RootSceneNodeUid;
+				SkeletonFactoryNode->GetCustomRootJointUid(RootSceneNodeUid);
+				if(UE::Interchange::Private::IsSkeletonAnimatedRecursive(RootSceneNodeUid, LocalNodeContainer))
+				{
+					Skeletons.Add(SkeletonFactoryNode);
+				}
+			}
+		});
+	//for each animated skeleton create one anim sequence factory node
+	for (UInterchangeSkeletonFactoryNode* SkeletonFactoryNode : Skeletons)
+	{
+		const FString AnimSequenceUid = TEXT("\\AnimSequence") + SkeletonFactoryNode->GetUniqueID();
+		FString AnimSequenceName = SkeletonFactoryNode->GetDisplayLabel();
+		if (AnimSequenceName.EndsWith(TEXT("_Skeleton")))
+		{
+			AnimSequenceName.LeftChopInline(9);
+		}
+		AnimSequenceName += TEXT("_Anim");
+		UInterchangeAnimSequenceFactoryNode* AnimSequenceFactoryNode = NewObject<UInterchangeAnimSequenceFactoryNode>(BaseNodeContainer, NAME_None);
+		AnimSequenceFactoryNode->InitializeAnimSequenceNode(AnimSequenceUid, AnimSequenceName);
+		
+		AnimSequenceFactoryNode->SetCustomSkeletonFactoryNodeUid(SkeletonFactoryNode->GetUniqueID());
+		AnimSequenceFactoryNode->SetCustomImportBoneTracks(bImportBoneTracks);
+		AnimSequenceFactoryNode->SetCustomImportBoneTracksSampleRate(SampleRate);
+		if (bRangeIsValid)
+		{
+			AnimSequenceFactoryNode->SetCustomImportBoneTracksRangeStart(RangeStart);
+			AnimSequenceFactoryNode->SetCustomImportBoneTracksRangeStop(RangeStop);
+		}
+
+		//USkeleton cannot be created without a valid skeletalmesh
+		const FString SkeletonUid = SkeletonFactoryNode->GetUniqueID();
+		AnimSequenceFactoryNode->AddFactoryDependencyUid(SkeletonUid);
+		const bool bImportOnlyAnimation = CommonSkeletalMeshesAndAnimationsProperties->bImportOnlyAnimations;
+		
+		//Iterate dependencies
+		{
+			TArray<FString> SkeletalMeshNodeUids;
+			BaseNodeContainer->GetNodes(UInterchangeSkeletalMeshFactoryNode::StaticClass(), SkeletalMeshNodeUids);
+			for (const FString& SkelMeshFactoryNodeUid : SkeletalMeshNodeUids)
+			{
+				if (UInterchangeSkeletalMeshFactoryNode* SkeletalMeshFactoryNode = Cast<UInterchangeSkeletalMeshFactoryNode>(BaseNodeContainer->GetNode(SkelMeshFactoryNodeUid)))
+				{
+					TArray<FString> SkeletalMeshDependencies;
+					SkeletalMeshFactoryNode->GetFactoryDependencies(SkeletalMeshDependencies);
+					for (const FString& SkeletalMeshDependencyUid : SkeletalMeshDependencies)
+					{
+						if (SkeletonUid.Equals(SkeletalMeshDependencyUid))
+						{
+							AnimSequenceFactoryNode->AddFactoryDependencyUid(SkelMeshFactoryNodeUid);
+							break;
+						}
+					}
+				}
+			}
+		}
+
+		if (!CommonSkeletalMeshesAndAnimationsProperties->Skeleton.IsNull())
+		{
+			FString RootJointUid;
+			SkeletonFactoryNode->GetCustomRootJointUid(RootJointUid);
+			bool bSkeletonCompatible = true;
+
+			//TODO: support skeleton helper in runtime
+#if WITH_EDITOR
+			bSkeletonCompatible = UE::Interchange::Private::FSkeletonHelper::IsCompatibleSkeleton(CommonSkeletalMeshesAndAnimationsProperties->Skeleton, RootJointUid, BaseNodeContainer);
+#endif
+			if(bSkeletonCompatible)
+			{
+				FSoftObjectPath SkeletonSoftObjectPath(CommonSkeletalMeshesAndAnimationsProperties->Skeleton.Get());
+				AnimSequenceFactoryNode->SetCustomSkeletonSoftObjectPath(SkeletonSoftObjectPath);
+			}
+			else
+			{
+				UInterchangeResultError_Generic* Message = AddMessage<UInterchangeResultError_Generic>();
+				Message->Text = FText::Format(NSLOCTEXT("UInterchangeGenericAnimationPipeline", "IncompatibleSkeleton", "Incompatible skeleton {0} when importing AnimSequence {1}."),
+					FText::FromString(CommonSkeletalMeshesAndAnimationsProperties->Skeleton->GetName()),
+					FText::FromString(AnimSequenceName));
+			}
+		}
+		BaseNodeContainer->AddNode(AnimSequenceFactoryNode);
+	}
+}
+
+void UInterchangeGenericAnimationPipeline::ExecutePostImportPipeline(const UInterchangeBaseNodeContainer* InBaseNodeContainer, const FString& NodeKey, UObject* CreatedAsset, bool bIsAReimport)
+{
+	//We do not use the provided base container since ExecutePreImportPipeline cache it
+	//We just make sure the same one is pass in parameter
+	if (!InBaseNodeContainer || !ensure(BaseNodeContainer == InBaseNodeContainer) || !CreatedAsset)
+	{
+		return;
+	}
+
+	UInterchangeBaseNode* Node = BaseNodeContainer->GetNode(NodeKey);
+	if (!Node)
+	{
+		return;
+	}
+}
+
+
+
diff --git a/Engine/Plugins/Experimental/Interchange/Runtime/Source/Pipelines/Private/InterchangeGenericAssetsPipeline.cpp b/Engine/Plugins/Experimental/Interchange/Runtime/Source/Pipelines/Private/InterchangeGenericAssetsPipeline.cpp
--- a/Engine/Plugins/Experimental/Interchange/Runtime/Source/Pipelines/Private/InterchangeGenericAssetsPipeline.cpp
+++ b/Engine/Plugins/Experimental/Interchange/Runtime/Source/Pipelines/Private/InterchangeGenericAssetsPipeline.cpp
@@ -4,6 +4,8 @@
 
 #include "Animation/Skeleton.h"
 #include "CoreMinimal.h"
+#include "InterchangeAnimSequenceFactoryNode.h"
+#include "InterchangeGenericAnimationPipeline.h"
 #include "InterchangeGenericMaterialPipeline.h"
 #include "InterchangeGenericMeshPipeline.h"
 #include "InterchangeGenericTexturePipeline.h"
@@ -29,11 +31,21 @@ UInterchangeGenericAssetsPipeline::UInterchangeGenericAssetsPipeline()
 {
 	TexturePipeline = CreateDefaultSubobject<UInterchangeGenericTexturePipeline>("TexturePipeline");
 	MaterialPipeline = CreateDefaultSubobject<UInterchangeGenericMaterialPipeline>("MaterialPipeline");
+	CommonMeshesProperties = CreateDefaultSubobject<UInterchangeGenericCommonMeshesProperties>("CommonMeshesProperties");
+	CommonSkeletalMeshesAndAnimationsProperties = CreateDefaultSubobject<UInterchangeGenericCommonSkeletalMeshesAndAnimationsProperties>("CommonSkeletalMeshesAndAnimationsProperties");
 	MeshPipeline = CreateDefaultSubobject<UInterchangeGenericMeshPipeline>("MeshPipeline");
+	MeshPipeline->CommonMeshesProperties = CommonMeshesProperties;
+	MeshPipeline->CommonSkeletalMeshesAndAnimationsProperties = CommonSkeletalMeshesAndAnimationsProperties;
+	AnimationPipeline = CreateDefaultSubobject<UInterchangeGenericAnimationPipeline>("AnimationPipeline");
+	AnimationPipeline->CommonSkeletalMeshesAndAnimationsProperties = CommonSkeletalMeshesAndAnimationsProperties;
 }
 
 void UInterchangeGenericAssetsPipeline::PreDialogCleanup(const FName PipelineStackName)
 {
+	check(!CommonSkeletalMeshesAndAnimationsProperties.IsNull())
+	//We always clean the pipeline skeleton when showing the dialog
+	CommonSkeletalMeshesAndAnimationsProperties->Skeleton = nullptr;
+
 	if (TexturePipeline)
 	{
 		TexturePipeline->PreDialogCleanup(PipelineStackName);
@@ -48,18 +60,99 @@ void UInterchangeGenericAssetsPipeline::PreDialogCleanup(const FName PipelineSta
 	{
 		MeshPipeline->PreDialogCleanup(PipelineStackName);
 	}
+
+	if (AnimationPipeline)
+	{
+		AnimationPipeline->PreDialogCleanup(PipelineStackName);
+	}
 	
 	SaveSettings(PipelineStackName);
 }
 
+bool UInterchangeGenericAssetsPipeline::IsSettingsAreValid() const
+{
+	if (TexturePipeline && !TexturePipeline->IsSettingsAreValid())
+	{
+		return false;
+	}
+
+	if (MaterialPipeline && !MaterialPipeline->IsSettingsAreValid())
+	{
+		return false;
+	}
+
+	if (CommonMeshesProperties && !CommonMeshesProperties->IsSettingsAreValid())
+	{
+		return false;
+	}
+
+	if (CommonSkeletalMeshesAndAnimationsProperties && !CommonSkeletalMeshesAndAnimationsProperties->IsSettingsAreValid())
+	{
+		return false;
+	}
+
+	if (MeshPipeline && !MeshPipeline->IsSettingsAreValid())
+	{
+		return false;
+	}
+
+	if (AnimationPipeline && !AnimationPipeline->IsSettingsAreValid())
+	{
+		return false;
+	}
+
+	return Super::IsSettingsAreValid();
+}
+
+void UInterchangeGenericAssetsPipeline::SetupReimportData(TObjectPtr<UObject> ReimportObject)
+{
+	if (TexturePipeline)
+	{
+		TexturePipeline->SetupReimportData(ReimportObject);
+	}
+
+	if (MaterialPipeline)
+	{
+		MaterialPipeline->SetupReimportData(ReimportObject);
+	}
+
+	if (MeshPipeline)
+	{
+		MeshPipeline->SetupReimportData(ReimportObject);
+	}
+
+	if (AnimationPipeline)
+	{
+		AnimationPipeline->SetupReimportData(ReimportObject);
+	}
+}
+
 void UInterchangeGenericAssetsPipeline::ExecutePreImportPipeline(UInterchangeBaseNodeContainer* InBaseNodeContainer, const TArray<UInterchangeSourceData*>& InSourceDatas)
 {
+	check(!CommonSkeletalMeshesAndAnimationsProperties.IsNull());
+
 	if (!InBaseNodeContainer)
 	{
 		UE_LOG(LogInterchangePipeline, Warning, TEXT("UInterchangeGenericAssetsPipeline: Cannot execute pre-import pipeline because InBaseNodeContrainer is null"));
 		return;
 	}
 
+	//////////////////////////////////////////////////////////////////////////
+	//Make sure all options go together
+	
+	//When we import only animation we need to prevent material and physic asset to be created
+	if (CommonSkeletalMeshesAndAnimationsProperties->bImportOnlyAnimations)
+	{
+		MaterialPipeline->MaterialImport = EInterchangeMaterialImportOption::DoNotImport;
+		MeshPipeline->bImportStaticMeshes = false;
+		MeshPipeline->bCreatePhysicsAsset = false;
+		MeshPipeline->PhysicsAsset = nullptr;
+		TexturePipeline->bImportTextures = false;
+	}
+
+	//////////////////////////////////////////////////////////////////////////
+
+
 	BaseNodeContainer = InBaseNodeContainer;
 	SourceDatas.Empty(InSourceDatas.Num());
 	for (const UInterchangeSourceData* SourceData : InSourceDatas)
@@ -82,8 +175,12 @@ void UInterchangeGenericAssetsPipeline::ExecutePreImportPipeline(UInterchangeBas
 		MeshPipeline->ScriptedExecutePreImportPipeline(InBaseNodeContainer, InSourceDatas);
 	}
 
-	ImplementUseSourceNameForAssetOption();
+	if (AnimationPipeline)
+	{
+		AnimationPipeline->ScriptedExecutePreImportPipeline(InBaseNodeContainer, InSourceDatas);
+	}
 
+	ImplementUseSourceNameForAssetOption();
 	//Make sure all factory nodes have the specified strategy
 	BaseNodeContainer->IterateNodes([ReimportStrategyClosure = ReimportStrategy](const FString& NodeUid, UInterchangeBaseNode* Node)
 		{
@@ -110,6 +207,11 @@ void UInterchangeGenericAssetsPipeline::ExecutePostImportPipeline(const UInterch
 	{
 		MeshPipeline->ScriptedExecutePostImportPipeline(InBaseNodeContainer, NodeKey, CreatedAsset, bIsAReimport);
 	}
+
+	if (AnimationPipeline)
+	{
+		AnimationPipeline->ScriptedExecutePostImportPipeline(InBaseNodeContainer, NodeKey, CreatedAsset, bIsAReimport);
+	}
 }
 
 void UInterchangeGenericAssetsPipeline::SetReimportSourceIndex(UClass* ReimportObjectClass, const int32 SourceFileIndex)
@@ -128,35 +230,50 @@ void UInterchangeGenericAssetsPipeline::SetReimportSourceIndex(UClass* ReimportO
 	{
 		MeshPipeline->ScriptedSetReimportSourceIndex(ReimportObjectClass, SourceFileIndex);
 	}
+
+	if (AnimationPipeline)
+	{
+		AnimationPipeline->ScriptedSetReimportSourceIndex(ReimportObjectClass, SourceFileIndex);
+	}
 }
 
 void UInterchangeGenericAssetsPipeline::ImplementUseSourceNameForAssetOption()
 {
-	const UClass* SkeletalMeshFactoryNodeClass = UInterchangeSkeletalMeshFactoryNode::StaticClass();
-	TArray<FString> SkeletalMeshNodeUids;
-	BaseNodeContainer->GetNodes(SkeletalMeshFactoryNodeClass, SkeletalMeshNodeUids);
+	if (bUseSourceNameForAsset)
+	{
+		const UClass* SkeletalMeshFactoryNodeClass = UInterchangeSkeletalMeshFactoryNode::StaticClass();
+		TArray<FString> SkeletalMeshNodeUids;
+		BaseNodeContainer->GetNodes(SkeletalMeshFactoryNodeClass, SkeletalMeshNodeUids);
 
-	const UClass* StaticMeshFactoryNodeClass = UInterchangeStaticMeshFactoryNode::StaticClass();
-	TArray<FString> StaticMeshNodeUids;
-	BaseNodeContainer->GetNodes(StaticMeshFactoryNodeClass, StaticMeshNodeUids);
+		const UClass* StaticMeshFactoryNodeClass = UInterchangeStaticMeshFactoryNode::StaticClass();
+		TArray<FString> StaticMeshNodeUids;
+		BaseNodeContainer->GetNodes(StaticMeshFactoryNodeClass, StaticMeshNodeUids);
 
-	//TODO count also the imported animations
+		const UClass* AnimSequenceFactoryNodeClass = UInterchangeAnimSequenceFactoryNode::StaticClass();
+		TArray<FString> AnimSequenceNodeUids;
+		BaseNodeContainer->GetNodes(AnimSequenceFactoryNodeClass, AnimSequenceNodeUids);
 
-	//If we import only one asset, and bUseSourceNameForAsset is true, we want to rename the asset using the file name.
-	const int32 MeshesAndAnimsImportedNodeCount = SkeletalMeshNodeUids.Num() + StaticMeshNodeUids.Num();
+		//If we import only one mesh, we want to rename the mesh using the file name.
+		const int32 MeshesImportedNodeCount = SkeletalMeshNodeUids.Num() + StaticMeshNodeUids.Num();
 
-	//Skeletalmesh must always rename created skeleton and physics asset to (SKNAME_Skeleton of SKNAME_PhysicsAsset). So the pipeline option
-	// bUseSourceNameForAsset will change or not the SKNAME.
-	MeshPipeline->ImplementUseSourceNameForAssetOptionSkeletalMesh(MeshesAndAnimsImportedNodeCount, bUseSourceNameForAsset);
+		//SkeletalMesh
+		MeshPipeline->ImplementUseSourceNameForAssetOptionSkeletalMesh(MeshesImportedNodeCount, bUseSourceNameForAsset);
 
-	if (bUseSourceNameForAsset && MeshesAndAnimsImportedNodeCount == 1)
-	{
-		if (StaticMeshNodeUids.Num() > 0)
+		//StaticMesh
+		if (MeshesImportedNodeCount == 1 && StaticMeshNodeUids.Num() > 0)
 		{
 			UInterchangeStaticMeshFactoryNode* StaticMeshNode = Cast<UInterchangeStaticMeshFactoryNode>(BaseNodeContainer->GetNode(StaticMeshNodeUids[0]));
 			const FString DisplayLabelName = FPaths::GetBaseFilename(SourceDatas[0]->GetFilename());
 			StaticMeshNode->SetDisplayLabel(DisplayLabelName);
 		}
+
+		//Animation, simply look if we import only 1 animation before applying the option to animation
+		if (AnimSequenceNodeUids.Num() == 1)
+		{
+			UInterchangeAnimSequenceFactoryNode* AnimSequenceNode = Cast<UInterchangeAnimSequenceFactoryNode>(BaseNodeContainer->GetNode(AnimSequenceNodeUids[0]));
+			const FString DisplayLabelName = FPaths::GetBaseFilename(SourceDatas[0]->GetFilename()) + TEXT("_Anim");
+			AnimSequenceNode->SetDisplayLabel(DisplayLabelName);
+		}
 	}
 
 }
diff --git a/Engine/Plugins/Experimental/Interchange/Runtime/Source/Pipelines/Private/InterchangeGenericMeshPipeline.cpp b/Engine/Plugins/Experimental/Interchange/Runtime/Source/Pipelines/Private/InterchangeGenericMeshPipeline.cpp
--- a/Engine/Plugins/Experimental/Interchange/Runtime/Source/Pipelines/Private/InterchangeGenericMeshPipeline.cpp
+++ b/Engine/Plugins/Experimental/Interchange/Runtime/Source/Pipelines/Private/InterchangeGenericMeshPipeline.cpp
@@ -2,6 +2,7 @@
 #include "InterchangeGenericMeshPipeline.h"
 
 #include "CoreMinimal.h"
+#include "InterchangeGenericAssetsPipeline.h"
 #include "InterchangeMeshNode.h"
 #include "InterchangePipelineLog.h"
 #include "InterchangePipelineMeshesUtilities.h"
@@ -29,6 +30,8 @@ void UInterchangeGenericMeshPipeline::ExecutePreImportPipeline(UInterchangeBaseN
 
 	PipelineMeshesUtilities = UInterchangePipelineMeshesUtilities::CreateInterchangePipelineMeshesUtilities(BaseNodeContainer);
 
+
+
 	//Create skeletalmesh factory nodes
 	ExecutePreImportPipelineSkeletalMesh();
 
@@ -96,10 +99,3 @@ void UInterchangeGenericMeshPipeline::SetReimportSourceIndex(UClass* ReimportObj
 	}
 }
 
-void UInterchangeGenericMeshPipeline::PreDialogCleanup(const FName PipelineStackName)
-{
-	//We always clean the pipeline skeleton when showing the dialog
-	Skeleton = nullptr;
-}
-
-
diff --git a/Engine/Plugins/Experimental/Interchange/Runtime/Source/Pipelines/Private/InterchangeGenericSkeletalMeshPipeline.cpp b/Engine/Plugins/Experimental/Interchange/Runtime/Source/Pipelines/Private/InterchangeGenericSkeletalMeshPipeline.cpp
--- a/Engine/Plugins/Experimental/Interchange/Runtime/Source/Pipelines/Private/InterchangeGenericSkeletalMeshPipeline.cpp
+++ b/Engine/Plugins/Experimental/Interchange/Runtime/Source/Pipelines/Private/InterchangeGenericSkeletalMeshPipeline.cpp
@@ -5,6 +5,7 @@
 #include "AssetRegistry/AssetRegistryModule.h"
 #include "Async/Async.h"
 #include "CoreMinimal.h"
+#include "InterchangeGenericAssetsPipeline.h"
 #include "InterchangeMaterialFactoryNode.h"
 #include "InterchangeMeshNode.h"
 #include "InterchangePhysicsAssetFactoryNode.h"
@@ -14,14 +15,15 @@
 #include "InterchangeSkeletalMeshFactoryNode.h"
 #include "InterchangeSkeletalMeshLodDataNode.h"
 #include "InterchangeSkeletonFactoryNode.h"
+#include "InterchangeSkeletonHelper.h"
 #include "InterchangeSourceData.h"
 #include "Misc/Paths.h"
 #include "Nodes/InterchangeBaseNode.h"
 #include "Nodes/InterchangeBaseNodeContainer.h"
 #if WITH_EDITOR
 #include "PhysicsAssetUtils.h"
+#endif //WITH_EDITOR
 #include "PhysicsEngine/PhysicsAsset.h"
-#endif
 #include "ReferenceSkeleton.h"
 #include "Tasks/Task.h"
 #include "UObject/Object.h"
@@ -29,180 +31,6 @@
 
 namespace UE::Interchange::SkeletalMeshGenericPipeline
 {
-	FName SkeletalLodGetBoneName(const TArray<FMeshBoneInfo>& SkeletalLodRawInfos, int32 BoneIndex)
-	{
-		if (SkeletalLodRawInfos.IsValidIndex(BoneIndex))
-		{
-			return SkeletalLodRawInfos[BoneIndex].Name;
-		}
-		return NAME_None;
-	}
-
-	int32 SkeletalLodFindBoneIndex(const TArray<FMeshBoneInfo>& SkeletalLodRawInfos, FName BoneName)
-	{
-		const int32 BoneCount = SkeletalLodRawInfos.Num();
-		for (int32 BoneIndex = 0; BoneIndex < BoneCount; ++BoneIndex)
-		{
-			if (SkeletalLodRawInfos[BoneIndex].Name == BoneName)
-			{
-				return BoneIndex;
-			}
-		}
-		return INDEX_NONE;
-	}
-
-	int32 SkeletalLodGetParentIndex(const TArray<FMeshBoneInfo>& SkeletalLodRawInfos, int32 BoneIndex)
-	{
-		if (SkeletalLodRawInfos.IsValidIndex(BoneIndex))
-		{
-			return SkeletalLodRawInfos[BoneIndex].ParentIndex;
-		}
-		return INDEX_NONE;
-	}
-
-	bool DoesParentChainMatch(int32 StartBoneIndex, const FReferenceSkeleton& SkeletonRef, const TArray<FMeshBoneInfo>& SkeletalLodRawInfos)
-	{
-		// if start is root bone
-		if (StartBoneIndex == 0)
-		{
-			// verify name of root bone matches
-			return (SkeletonRef.GetBoneName(0) == SkeletalLodGetBoneName(SkeletalLodRawInfos, 0));
-		}
-
-		int32 SkeletonBoneIndex = StartBoneIndex;
-		// If skeleton bone is not found in mesh, fail.
-		int32 MeshBoneIndex = SkeletalLodFindBoneIndex(SkeletalLodRawInfos, SkeletonRef.GetBoneName(SkeletonBoneIndex));
-		if (MeshBoneIndex == INDEX_NONE)
-		{
-			return false;
-		}
-		do
-		{
-			// verify if parent name matches
-			int32 ParentSkeletonBoneIndex = SkeletonRef.GetParentIndex(SkeletonBoneIndex);
-			int32 ParentMeshBoneIndex = SkeletalLodGetParentIndex(SkeletalLodRawInfos, MeshBoneIndex);
-
-			// if one of the parents doesn't exist, make sure both end. Otherwise fail.
-			if ((ParentSkeletonBoneIndex == INDEX_NONE) || (ParentMeshBoneIndex == INDEX_NONE))
-			{
-				return (ParentSkeletonBoneIndex == ParentMeshBoneIndex);
-			}
-
-			// If parents are not named the same, fail.
-			if (SkeletonRef.GetBoneName(ParentSkeletonBoneIndex) != SkeletalLodGetBoneName(SkeletalLodRawInfos, ParentMeshBoneIndex))
-			{
-				return false;
-			}
-
-			// move up
-			SkeletonBoneIndex = ParentSkeletonBoneIndex;
-			MeshBoneIndex = ParentMeshBoneIndex;
-		} while (true);
-
-		return true;
-	}
-
-	void RecursiveBuildSkeletalSkeleton(const FString JoinToAddUid, const int32 ParentIndex, const UInterchangeBaseNodeContainer* BaseNodeContainer, TArray<FMeshBoneInfo>& SkeletalLodRawInfos)
-	{
-		const UInterchangeSceneNode* SceneNode = Cast<UInterchangeSceneNode>(BaseNodeContainer->GetNode(JoinToAddUid));
-		if (!SceneNode || !SceneNode->IsSpecializedTypeContains(FSceneNodeStaticData::GetJointSpecializeTypeString()))
-		{
-			return;
-		}
-
-		int32 JoinIndex = SkeletalLodRawInfos.Num();
-		FMeshBoneInfo& Info = SkeletalLodRawInfos.AddZeroed_GetRef();
-		Info.Name = *SceneNode->GetDisplayLabel();
-		Info.ParentIndex = ParentIndex;
-#if WITH_EDITORONLY_DATA
-		Info.ExportName = Info.Name.ToString();
-#endif
-		//Iterate childrens
-		const TArray<FString> ChildrenIds = BaseNodeContainer->GetNodeChildrenUids(JoinToAddUid);
-		for (int32 ChildIndex = 0; ChildIndex < ChildrenIds.Num(); ++ChildIndex)
-		{
-			RecursiveBuildSkeletalSkeleton(ChildrenIds[ChildIndex], JoinIndex, BaseNodeContainer, SkeletalLodRawInfos);
-		}
-	}
-
-	bool IsCompatibleSkeleton(const USkeleton* Skeleton, const FString RootJoinUid, const UInterchangeBaseNodeContainer* BaseNodeContainer)
-	{
-		// at least % of bone should match 
-		int32 NumOfBoneMatches = 0;
-		//Make sure the specified Skeleton fit this skeletal mesh
-		const FReferenceSkeleton& SkeletonRef = Skeleton->GetReferenceSkeleton();
-		const int32 SkeletonBoneCount = SkeletonRef.GetRawBoneNum();
-
-		TArray<FMeshBoneInfo> SkeletalLodRawInfos;
-		SkeletalLodRawInfos.Reserve(SkeletonBoneCount);
-		RecursiveBuildSkeletalSkeleton(RootJoinUid, INDEX_NONE, BaseNodeContainer, SkeletalLodRawInfos);
-		const int32 SkeletalLodBoneCount = SkeletalLodRawInfos.Num();
-
-		// first ensure the parent exists for each bone
-		for (int32 MeshBoneIndex = 0; MeshBoneIndex < SkeletalLodBoneCount; MeshBoneIndex++)
-		{
-			FName MeshBoneName = SkeletalLodRawInfos[MeshBoneIndex].Name;
-			// See if Mesh bone exists in Skeleton.
-			int32 SkeletonBoneIndex = SkeletonRef.FindBoneIndex(MeshBoneName);
-
-			// if found, increase num of bone matches count
-			if (SkeletonBoneIndex != INDEX_NONE)
-			{
-				++NumOfBoneMatches;
-
-				// follow the parent chain to verify the chain is same
-				if (!DoesParentChainMatch(SkeletonBoneIndex, SkeletonRef, SkeletalLodRawInfos))
-				{
-					//Not compatible
-					return false;
-				}
-			}
-			else
-			{
-				int32 CurrentBoneId = MeshBoneIndex;
-				// if not look for parents that matches
-				while (SkeletonBoneIndex == INDEX_NONE && CurrentBoneId != INDEX_NONE)
-				{
-					// find Parent one see exists
-					const int32 ParentMeshBoneIndex = SkeletalLodGetParentIndex(SkeletalLodRawInfos, CurrentBoneId);
-					if (ParentMeshBoneIndex != INDEX_NONE)
-					{
-						// @TODO: make sure RefSkeleton's root ParentIndex < 0 if not, I'll need to fix this by checking TreeBoneIdx
-						FName ParentBoneName = SkeletalLodGetBoneName(SkeletalLodRawInfos, ParentMeshBoneIndex);
-						SkeletonBoneIndex = SkeletonRef.FindBoneIndex(ParentBoneName);
-					}
-
-					// root is reached
-					if (ParentMeshBoneIndex == 0)
-					{
-						break;
-					}
-					else
-					{
-						CurrentBoneId = ParentMeshBoneIndex;
-					}
-				}
-
-				// still no match, return false, no parent to look for
-				if (SkeletonBoneIndex == INDEX_NONE)
-				{
-					return false;
-				}
-
-				// second follow the parent chain to verify the chain is same
-				if (!DoesParentChainMatch(SkeletonBoneIndex, SkeletonRef, SkeletalLodRawInfos))
-				{
-					return false;
-				}
-			}
-		}
-
-		// originally we made sure at least matches more than 50% 
-		// but then slave components can't play since they're only partial
-		// if the hierarchy matches, and if it's more then 1 bone, we allow
-		return (NumOfBoneMatches > 0);
-	}
-
 	bool RecursiveFindChildUid(const UInterchangeBaseNodeContainer* BaseNodeContainer, const FString& ParentUid, const FString& SearchUid)
 	{
 		if (ParentUid == SearchUid)
@@ -242,18 +70,19 @@ namespace UE::Interchange::SkeletalMeshGenericPipeline
 
 void UInterchangeGenericMeshPipeline::ExecutePreImportPipelineSkeletalMesh()
 {
+	check(!CommonMeshesProperties.IsNull());
 	if (!bImportSkeletalMeshes)
 	{
 		//Nothing to import
 		return;
 	}
 
-	if (ForceAllMeshAsType != EInterchangeForceMeshType::IFMT_None && ForceAllMeshAsType != EInterchangeForceMeshType::IFMT_SkeletalMesh)
+	if (CommonMeshesProperties->ForceAllMeshAsType != EInterchangeForceMeshType::IFMT_None && CommonMeshesProperties->ForceAllMeshAsType != EInterchangeForceMeshType::IFMT_SkeletalMesh)
 	{
 		//Nothing to import
 		return;
 	}
-	const bool bConvertStaticMeshToSkeletalMesh = (ForceAllMeshAsType == EInterchangeForceMeshType::IFMT_SkeletalMesh);
+	const bool bConvertStaticMeshToSkeletalMesh = (CommonMeshesProperties->ForceAllMeshAsType == EInterchangeForceMeshType::IFMT_SkeletalMesh);
 	TMap<FString, TArray<FString>> SkeletalMeshFactoryDependencyOrderPerSkeletonRootNodeUid;
 
 	auto SetSkeletalMeshDependencies = [&SkeletalMeshFactoryDependencyOrderPerSkeletonRootNodeUid](const FString& JointNodeUid, UInterchangeSkeletalMeshFactoryNode* SkeletalMeshFactoryNode)
@@ -324,7 +153,7 @@ void UInterchangeGenericMeshPipeline::ExecutePreImportPipelineSkeletalMesh()
 		};
 
 		bool bFoundMeshes = false;
-		if (bBakeMeshes)
+		if (CommonMeshesProperties->bBakeMeshes)
 		{
 			PipelineMeshesUtilities->GetCombinedSkinnedMeshInstances(MeshUidsPerSkeletonRootUid, bConvertStaticMeshToSkeletalMesh);
 			const bool bUseMeshInstance = true;
@@ -398,7 +227,7 @@ void UInterchangeGenericMeshPipeline::ExecutePreImportPipelineSkeletalMesh()
 		};
 			
 		bool bFoundMeshes = false;
-		if (bBakeMeshes)
+		if (CommonMeshesProperties->bBakeMeshes)
 		{
 			PipelineMeshesUtilities->GetAllSkinnedMeshInstance(MeshUids, bConvertStaticMeshToSkeletalMesh);
 			const bool bUseMeshInstance = true;
@@ -418,11 +247,13 @@ void UInterchangeGenericMeshPipeline::ExecutePreImportPipelineSkeletalMesh()
 
 UInterchangeSkeletonFactoryNode* UInterchangeGenericMeshPipeline::CreateSkeletonFactoryNode(const FString& RootJointUid)
 {
+	check(!CommonSkeletalMeshesAndAnimationsProperties.IsNull());
 	const UInterchangeBaseNode* RootJointNode = BaseNodeContainer->GetNode(RootJointUid);
 	if (!RootJointNode)
 	{
 		return nullptr;
 	}
+
 	FString DisplayLabel = RootJointNode->GetDisplayLabel() + TEXT("_Skeleton");
 	FString SkeletonUid = TEXT("\\Skeleton\\") + RootJointNode->GetUniqueID();
 
@@ -453,21 +284,23 @@ UInterchangeSkeletonFactoryNode* UInterchangeGenericMeshPipeline::CreateSkeleton
 		}
 		SkeletonFactoryNode->InitializeSkeletonNode(SkeletonUid, DisplayLabel, USkeleton::StaticClass()->GetName());
 		SkeletonFactoryNode->SetCustomRootJointUid(RootJointNode->GetUniqueID());
-		SkeletonFactoryNode->SetCustomUseTimeZeroForBindPose(bUseT0AsRefPose);
+		SkeletonFactoryNode->SetCustomUseTimeZeroForBindPose(CommonSkeletalMeshesAndAnimationsProperties->bUseT0AsRefPose);
 		BaseNodeContainer->AddNode(SkeletonFactoryNode);
 	}
 
 	//If we have a specified skeleton
-	if (Skeleton)
+	if (CommonSkeletalMeshesAndAnimationsProperties->Skeleton)
 	{
 		SkeletonFactoryNode->SetEnabled(false);
-		SkeletonFactoryNode->ReferenceObject = Skeleton;
+		SkeletonFactoryNode->ReferenceObject = CommonSkeletalMeshesAndAnimationsProperties->Skeleton;
 	}
 	return SkeletonFactoryNode;
 }
 
 UInterchangeSkeletalMeshFactoryNode* UInterchangeGenericMeshPipeline::CreateSkeletalMeshFactoryNode(const FString& RootJointUid, const TMap<int32, TArray<FString>>& MeshUidsPerLodIndex)
 {
+	check(!CommonMeshesProperties.IsNull());
+	check(!CommonSkeletalMeshesAndAnimationsProperties.IsNull());
 	//Get the skeleton factory node
 	const UInterchangeBaseNode* RootJointNode = BaseNodeContainer->GetNode(RootJointUid);
 	if (!RootJointNode)
@@ -569,25 +402,28 @@ UInterchangeSkeletalMeshFactoryNode* UInterchangeGenericMeshPipeline::CreateSkel
 	SkeletalMeshFactoryNode->SetCustomImportContentType(SkeletalMeshImportContentType);
 
 	//If we have a specified skeleton
-	if (Skeleton)
+	if (CommonSkeletalMeshesAndAnimationsProperties->Skeleton)
 	{
-		if (UE::Interchange::SkeletalMeshGenericPipeline::IsCompatibleSkeleton(Skeleton, RootJointNode->GetUniqueID(), BaseNodeContainer))
+		bool bSkeletonCompatible = false;
+
+		//TODO: support skeleton helper in runtime
+#if WITH_EDITOR
+		bSkeletonCompatible = UE::Interchange::Private::FSkeletonHelper::IsCompatibleSkeleton(CommonSkeletalMeshesAndAnimationsProperties->Skeleton, RootJointNode->GetUniqueID(), BaseNodeContainer);
+#endif
+		if (bSkeletonCompatible)
 		{
-			FSoftObjectPath SkeletonSoftObjectPath(Skeleton.Get());
+			FSoftObjectPath SkeletonSoftObjectPath(CommonSkeletalMeshesAndAnimationsProperties->Skeleton.Get());
 			SkeletalMeshFactoryNode->SetCustomSkeletonSoftObjectPath(SkeletonSoftObjectPath);
 		}
 		else
 		{
-			//Log an error, saying we will create a skeleton instead of using the specified one
-			//Should we show a dialog so we are sure the user understand he choose the wrong skeleton
-
-			//Make sure we enable the skeleton factory node
-			SkeletonFactoryNode->SetEnabled(true);
-
+			UInterchangeResultError_Generic* Message = AddMessage<UInterchangeResultError_Generic>();
+			Message->Text = FText::Format(NSLOCTEXT("UInterchangeGenericMeshPipeline", "IncompatibleSkeleton", "Incompatible skeleton {0} when importing skeletalmesh {1}."),
+				FText::FromString(CommonSkeletalMeshesAndAnimationsProperties->Skeleton->GetName()),
+				FText::FromString(DisplayLabel));
 		}
 	}
 
-#if WITH_EDITOR
 	//Physic asset dependency, if we must create or use a specialize physic asset let create
 	//a PhysicsAsset factory node, so the asset will exist when we will setup the skeletalmesh
 	if (bCreatePhysicsAsset)
@@ -608,10 +444,9 @@ UInterchangeSkeletalMeshFactoryNode* UInterchangeGenericMeshPipeline::CreateSkel
 		FSoftObjectPath PhysicSoftObjectPath(PhysicsAsset.Get());
 		SkeletalMeshFactoryNode->SetCustomPhysicAssetSoftObjectPath(PhysicSoftObjectPath);
 	}
-#endif
 
 	const bool bTrueValue = true;
-	switch (VertexColorImportOption)
+	switch (CommonMeshesProperties->VertexColorImportOption)
 	{
 		case EInterchangeVertexColorImportOption::IVCIO_Replace:
 		{
@@ -625,11 +460,17 @@ UInterchangeSkeletalMeshFactoryNode* UInterchangeGenericMeshPipeline::CreateSkel
 		break;
 		case EInterchangeVertexColorImportOption::IVCIO_Override:
 		{
-			SkeletalMeshFactoryNode->SetCustomVertexColorOverride(VertexOverrideColor);
+			SkeletalMeshFactoryNode->SetCustomVertexColorOverride(CommonMeshesProperties->VertexOverrideColor);
 		}
 		break;
 	}
 
+	//Avoid importing skeletalmesh if we want to only import animation
+	if (CommonSkeletalMeshesAndAnimationsProperties->bImportOnlyAnimations)
+	{
+		SkeletonFactoryNode->SetEnabled(false);
+		SkeletalMeshFactoryNode->SetEnabled(false);
+	}
 	return SkeletalMeshFactoryNode;
 }
 
@@ -651,12 +492,15 @@ UInterchangeSkeletalMeshLodDataNode* UInterchangeGenericMeshPipeline::CreateSkel
 
 void UInterchangeGenericMeshPipeline::AddLodDataToSkeletalMesh(const UInterchangeSkeletonFactoryNode* SkeletonFactoryNode, UInterchangeSkeletalMeshFactoryNode* SkeletalMeshFactoryNode, const TMap<int32, TArray<FString>>& NodeUidsPerLodIndex)
 {
+	check(!CommonMeshesProperties.IsNull());
+	check(!CommonSkeletalMeshesAndAnimationsProperties.IsNull());
+
 	const FString SkeletalMeshUid = SkeletalMeshFactoryNode->GetUniqueID();
 	const FString SkeletonUid = SkeletonFactoryNode->GetUniqueID();
 	for (const TPair<int32, TArray<FString>>& LodIndexAndNodeUids : NodeUidsPerLodIndex)
 	{
 		const int32 LodIndex = LodIndexAndNodeUids.Key;
-		if (!bImportLods && LodIndex > 0)
+		if (!CommonMeshesProperties->bImportLods && LodIndex > 0)
 		{
 			//If the pipeline should not import lods, skip any lod over base lod
 			continue;
@@ -664,7 +508,7 @@ void UInterchangeGenericMeshPipeline::AddLodDataToSkeletalMesh(const UInterchang
 
 		//Copy the nodes unique id because we need to remove nested mesh if the option is to not import them
 		TArray<FString> NodeUids = LodIndexAndNodeUids.Value;
-		if (!bImportMeshesInBoneHierarchy)
+		if (!CommonSkeletalMeshesAndAnimationsProperties->bImportMeshesInBoneHierarchy)
 		{
 			UE::Interchange::SkeletalMeshGenericPipeline::RemoveNestedMeshNodes(BaseNodeContainer, SkeletonFactoryNode, NodeUids);
 		}
@@ -731,6 +575,8 @@ void UInterchangeGenericMeshPipeline::AddLodDataToSkeletalMesh(const UInterchang
 
 void UInterchangeGenericMeshPipeline::PostImportSkeletalMesh(UObject* CreatedAsset, UInterchangeBaseNode* Node)
 {
+	check(!CommonSkeletalMeshesAndAnimationsProperties.IsNull());
+
 	if (!BaseNodeContainer)
 	{
 		return;
@@ -744,7 +590,7 @@ void UInterchangeGenericMeshPipeline::PostImportSkeletalMesh(UObject* CreatedAss
 
 	//If we import only the geometry we do not want to update the skeleton reference pose.
 	const bool bImportGeometryOnlyContent = SkeletalMeshImportContentType == EInterchangeSkeletalMeshContentType::Geometry;
-	if (!bImportGeometryOnlyContent && bUpdateSkeletonReferencePose && !Skeleton.IsNull() && SkeletalMesh->GetSkeleton() == Skeleton.Get())
+	if (!bImportGeometryOnlyContent && bUpdateSkeletonReferencePose && !CommonSkeletalMeshesAndAnimationsProperties->Skeleton.IsNull() && SkeletalMesh->GetSkeleton() == CommonSkeletalMeshesAndAnimationsProperties->Skeleton.Get())
 	{
 		SkeletalMesh->GetSkeleton()->UpdateReferencePoseFromMesh(SkeletalMesh);
 		//TODO: notify editor the skeleton has change
@@ -805,8 +651,10 @@ void UInterchangeGenericMeshPipeline::PostImportPhysicsAssetImport(UObject* Crea
 #endif //WITH_EDITOR
 }
 
-void UInterchangeGenericMeshPipeline::ImplementUseSourceNameForAssetOptionSkeletalMesh(const int32 MeshesAndAnimsImportedNodeCount, const bool bUseSourceNameForAsset)
+void UInterchangeGenericMeshPipeline::ImplementUseSourceNameForAssetOptionSkeletalMesh(const int32 MeshesImportedNodeCount, const bool bUseSourceNameForAsset)
 {
+	check(!CommonSkeletalMeshesAndAnimationsProperties.IsNull());
+
 	const UClass* SkeletalMeshFactoryNodeClass = UInterchangeSkeletalMeshFactoryNode::StaticClass();
 	TArray<FString> SkeletalMeshNodeUids;
 	BaseNodeContainer->GetNodes(SkeletalMeshFactoryNodeClass, SkeletalMeshNodeUids);
@@ -815,7 +663,7 @@ void UInterchangeGenericMeshPipeline::ImplementUseSourceNameForAssetOptionSkelet
 		return;
 	}
 	//If we import only one asset, and bUseSourceNameForAsset is true, we want to rename the asset using the file name.
-	const bool bShouldChangeAssetName = (bUseSourceNameForAsset && MeshesAndAnimsImportedNodeCount == 1);
+	const bool bShouldChangeAssetName = (bUseSourceNameForAsset && MeshesImportedNodeCount == 1);
 	const FString SkeletalMeshUid = SkeletalMeshNodeUids[0];
 	UInterchangeSkeletalMeshFactoryNode* SkeletalMeshNode = Cast<UInterchangeSkeletalMeshFactoryNode>(BaseNodeContainer->GetNode(SkeletalMeshUid));
 	if (!SkeletalMeshNode)
@@ -840,7 +688,7 @@ void UInterchangeGenericMeshPipeline::ImplementUseSourceNameForAssetOptionSkelet
 		if (UInterchangeSkeletalMeshLodDataNode* SkeletalMeshLodDataNode = Cast<UInterchangeSkeletalMeshLodDataNode>(BaseNodeContainer->GetNode(LodDataUids[0])))
 		{
 			//If the user did not specify any skeleton
-			if (Skeleton.IsNull())
+			if (CommonSkeletalMeshesAndAnimationsProperties->Skeleton.IsNull())
 			{
 				FString SkeletalMeshSkeletonUid;
 				SkeletalMeshLodDataNode->GetCustomSkeletonUid(SkeletalMeshSkeletonUid);
diff --git a/Engine/Plugins/Experimental/Interchange/Runtime/Source/Pipelines/Private/InterchangeGenericStaticMeshPipeline.cpp b/Engine/Plugins/Experimental/Interchange/Runtime/Source/Pipelines/Private/InterchangeGenericStaticMeshPipeline.cpp
--- a/Engine/Plugins/Experimental/Interchange/Runtime/Source/Pipelines/Private/InterchangeGenericStaticMeshPipeline.cpp
+++ b/Engine/Plugins/Experimental/Interchange/Runtime/Source/Pipelines/Private/InterchangeGenericStaticMeshPipeline.cpp
@@ -135,15 +135,17 @@ static void BuildMeshToCollisionMeshMap(const TArray<FString>& MeshUids, TMap<FS
 
 void UInterchangeGenericMeshPipeline::ExecutePreImportPipelineStaticMesh()
 {
-	if (bImportStaticMeshes && (ForceAllMeshAsType == EInterchangeForceMeshType::IFMT_None || ForceAllMeshAsType == EInterchangeForceMeshType::IFMT_StaticMesh))
+	check(!CommonMeshesProperties.IsNull());
+
+	if (bImportStaticMeshes && (CommonMeshesProperties->ForceAllMeshAsType == EInterchangeForceMeshType::IFMT_None || CommonMeshesProperties->ForceAllMeshAsType == EInterchangeForceMeshType::IFMT_StaticMesh))
 	{
-		const bool bConvertSkeletalMeshToStaticMesh = (ForceAllMeshAsType == EInterchangeForceMeshType::IFMT_StaticMesh);
+		const bool bConvertSkeletalMeshToStaticMesh = (CommonMeshesProperties->ForceAllMeshAsType == EInterchangeForceMeshType::IFMT_StaticMesh);
 		if (bCombineStaticMeshes)
 		{
 			// Combine all the static meshes
 
 			bool bFoundMeshes = false;
-			if (bBakeMeshes)
+			if (CommonMeshesProperties->bBakeMeshes)
 			{
 				// If baking transforms, get all the static mesh instance nodes, and group them by LOD
 				TArray<FString> MeshUids;
@@ -207,7 +209,7 @@ void UInterchangeGenericMeshPipeline::ExecutePreImportPipelineStaticMesh()
 
 			bool bFoundMeshes = false;
 
-			if (bBakeMeshes)
+			if (CommonMeshesProperties->bBakeMeshes)
 			{
 				TArray<FString> MeshUids;
 				PipelineMeshesUtilities->GetAllStaticMeshInstance(MeshUids, bConvertSkeletalMeshToStaticMesh);
@@ -366,6 +368,7 @@ bool UInterchangeGenericMeshPipeline::MakeMeshFactoryNodeUidAndDisplayLabel(cons
 
 UInterchangeStaticMeshFactoryNode* UInterchangeGenericMeshPipeline::CreateStaticMeshFactoryNode(const TMap<int32, TArray<FString>>& MeshUidsPerLodIndex)
 {
+	check(!CommonMeshesProperties.IsNull());
 	if (MeshUidsPerLodIndex.Num() == 0)
 	{
 		return nullptr;
@@ -393,7 +396,7 @@ UInterchangeStaticMeshFactoryNode* UInterchangeGenericMeshPipeline::CreateStatic
 
 	AddLodDataToStaticMesh(StaticMeshFactoryNode, MeshUidsPerLodIndex);
 
-	switch (VertexColorImportOption)
+	switch (CommonMeshesProperties->VertexColorImportOption)
 	{
 	case EInterchangeVertexColorImportOption::IVCIO_Replace:
 	{
@@ -407,7 +410,7 @@ UInterchangeStaticMeshFactoryNode* UInterchangeGenericMeshPipeline::CreateStatic
 	break;
 	case EInterchangeVertexColorImportOption::IVCIO_Override:
 	{
-		StaticMeshFactoryNode->SetCustomVertexColorOverride(VertexOverrideColor);
+		StaticMeshFactoryNode->SetCustomVertexColorOverride(CommonMeshesProperties->VertexOverrideColor);
 	}
 	break;
 	}
@@ -435,12 +438,13 @@ UInterchangeStaticMeshLodDataNode* UInterchangeGenericMeshPipeline::CreateStatic
 
 void UInterchangeGenericMeshPipeline::AddLodDataToStaticMesh(UInterchangeStaticMeshFactoryNode* StaticMeshFactoryNode, const TMap<int32, TArray<FString>>& NodeUidsPerLodIndex)
 {
+	check(!CommonMeshesProperties.IsNull());
 	const FString StaticMeshFactoryUid = StaticMeshFactoryNode->GetUniqueID();
 
 	for (const TPair<int32, TArray<FString>>& LodIndexAndNodeUids : NodeUidsPerLodIndex)
 	{
 		const int32 LodIndex = LodIndexAndNodeUids.Key;
-		if (!bImportLods && LodIndex > 0)
+		if (!CommonMeshesProperties->bImportLods && LodIndex > 0)
 		{
 			// If the pipeline should not import lods, skip any lod over base lod
 			continue;
diff --git a/Engine/Plugins/Experimental/Interchange/Runtime/Source/Pipelines/Public/InterchangeGenericAnimationPipeline.h b/Engine/Plugins/Experimental/Interchange/Runtime/Source/Pipelines/Public/InterchangeGenericAnimationPipeline.h
new file mode 100644
--- /dev/null
+++ b/Engine/Plugins/Experimental/Interchange/Runtime/Source/Pipelines/Public/InterchangeGenericAnimationPipeline.h
@@ -0,0 +1,123 @@
+// Copyright Epic Games, Inc. All Rights Reserved.
+
+#pragma once
+
+#include "CoreMinimal.h"
+#include "InterchangeGenericAssetsPipelineSharedSettings.h"
+#include "InterchangePipelineBase.h"
+#include "InterchangeSourceData.h"
+#include "Nodes/InterchangeBaseNodeContainer.h"
+#include "UObject/Object.h"
+#include "UObject/ObjectMacros.h"
+
+#include "InterchangeGenericAnimationPipeline.generated.h"
+
+class UInterchangeMeshNode;
+
+/** Animation length type when importing */
+UENUM(BlueprintType)
+enum class EInterchangeAnimationRange : uint8
+{
+	/** This option imports the range of frames based on timeline definition in the source. */
+	Timeline				UMETA(DisplayName = "Source Timeline"),
+	/** This option imports the range of frames that have animation. */
+	Animated				UMETA(DisplayName = "Animated Time"),
+	/** This option imports the range of frames specified by "FrameImportRange". */
+	SetRange				UMETA(DisplayName = "Set Range"),
+	MAX,
+};
+
+
+UCLASS(BlueprintType, Experimental)
+class INTERCHANGEPIPELINES_API UInterchangeGenericAnimationPipeline : public UInterchangePipelineBase
+{
+	GENERATED_BODY()
+
+public:
+
+	//Common SkeletalMeshes And Animations Properties Settings Pointer
+	UPROPERTY(Transient)
+	TObjectPtr<UInterchangeGenericCommonSkeletalMeshesAndAnimationsProperties> CommonSkeletalMeshesAndAnimationsProperties;
+
+	//////	ANIMATION_CATEGORY Properties //////
+	/** If enable, import all animation assets find in the sources. */
+	UPROPERTY(EditAnywhere, BlueprintReadWrite, Category = "Animations")
+	bool bImportAnimations = true;
+
+	/** Import bone transform tracks. If false, this will discard any bone transform tracks. (useful for curves only animations)*/
+	UPROPERTY(EditAnywhere, BlueprintReadWrite, Category = "Animations")
+	bool bImportBoneTracks = true;
+
+	/** Which animation range to import. The one defined at Exported, at Animated time or define a range manually */
+	UPROPERTY(EditAnywhere, BlueprintReadWrite, Category = "Animations", meta = (EditCondition = "bImportBoneTracks", DisplayName = "Animation Length"))
+	EInterchangeAnimationRange AnimationRange = EInterchangeAnimationRange::Timeline;
+
+	/** Frame range used when Set Range is used in Animation Length */
+	UPROPERTY(EditAnywhere, BlueprintReadWrite, Category = "Animations", meta = (EditCondition = "bImportBoneTracks", UIMin = 0, ClampMin = 0))
+	FInt32Interval FrameImportRange = FInt32Interval(0, 0);
+
+	/** Enable this option to use default sample rate for the imported animation at 30 frames per second */
+	UPROPERTY(EditAnywhere, BlueprintReadWrite, Category = "Animations", meta = (EditCondition = "bImportBoneTracks", ToolTip = "If enabled, samples all animation curves to 30 FPS"))
+	bool bUse30HzToBakeBoneAnimation = false;
+
+	/** Use this option to specify a sample rate for the imported animation, a value of 0 use the best matching sample rate. */
+	UPROPERTY(EditAnywhere, BlueprintReadWrite, Category = "Animations", meta = (EditCondition = "bImportBoneTracks && !bUse30HzToBakeBoneAnimation", ToolTip = "Sample fbx animation data at the specified sample rate, 0 find automaticaly the best sample rate", ClampMin = 0, UIMin = 0, ClampMax = 48000, UIMax = 60))
+	int32 CustomBoneAnimationSampleRate = 0;
+
+	/** Import if custom attribute as a curve within the animation */
+	UPROPERTY(EditAnywhere, BlueprintReadWrite, Category = "Animations")
+	bool bImportCustomAttribute = true;
+
+	/** Set Material Curve Type for all custom attributes that exists */
+	UPROPERTY(EditAnywhere, BlueprintReadWrite, Category = "Animations", meta = (EditCondition = "bImportCustomAttribute", DisplayName = "Set Material Curve Type"))
+	bool bSetMaterialDriveParameterOnCustomAttribute = false;
+
+	/** Set Material Curve Type for the custom attribute with the following suffixes. This doesn't matter if Set Material Curve Type is true  */
+	UPROPERTY(EditAnywhere, BlueprintReadWrite, Category = "Animations", meta = (EditCondition = "bImportCustomAttribute", DisplayName = "Material Curve Suffixes"))
+	TArray<FString> MaterialCurveSuffixes = {TEXT("_mat")};
+
+	/** When importing custom attribute as curve, remove redundant keys */
+	UPROPERTY(EditAnywhere, BlueprintReadWrite, Category = "Animations", meta = (EditCondition = "bImportCustomAttribute", DisplayName = "Remove Redundant Keys"))
+	bool bRemoveRedundantKeys = false;
+
+	/** When importing custom attribute or morphtarget as curve, do not import if it doens't have any value other than zero. This is to avoid adding extra curves to evaluate */
+	UPROPERTY(EditAnywhere, BlueprintReadWrite, Category = "Animations", meta = (EditCondition = "bImportCustomAttribute", DisplayName = "Do not import curves with only 0 values"))
+	bool bDoNotImportCurveWithZero = false;
+
+	/** If true, all previous non-curve custom attributes will be deleted when doing a re-import. */
+	UPROPERTY(EditAnywhere, BlueprintReadWrite, Category = "Animations")
+	bool bDeleteExistingNonCurveCustomAttributes = false;
+
+	/** If true, all previous custom attribute curves will be deleted when doing a re-import. */
+	UPROPERTY(EditAnywhere, BlueprintReadWrite, Category = "Animations")
+	bool bDeleteExistingCustomAttributeCurves = false;
+
+	/** If enabled, this will delete this type of asset from the FBX */
+	UPROPERTY(EditAnywhere, BlueprintReadWrite, Category = "Animations")
+	bool bDeleteExistingMorphTargetCurves = false;
+
+	/** Name of source animation that was imported, used to reimport correct animation from the translated source */
+	UPROPERTY()
+	FString SourceAnimationName;
+
+	virtual void SetupReimportData(TObjectPtr<UObject> ReimportObject) override;
+	
+protected:
+
+	virtual void ExecutePreImportPipeline(UInterchangeBaseNodeContainer* InBaseNodeContainer, const TArray<UInterchangeSourceData*>& InSourceDatas) override;
+
+	virtual void ExecutePostImportPipeline(const UInterchangeBaseNodeContainer* InBaseNodeContainer, const FString& NodeKey, UObject* CreatedAsset, bool bIsAReimport) override;
+
+	virtual bool CanExecuteOnAnyThread(EInterchangePipelineTask PipelineTask) override
+	{
+		return true;
+	}
+
+private:
+
+	UInterchangeBaseNodeContainer* BaseNodeContainer = nullptr;
+	TArray<const UInterchangeSourceData*> SourceDatas;
+
+};
+
+
diff --git a/Engine/Plugins/Experimental/Interchange/Runtime/Source/Pipelines/Public/InterchangeGenericAssetsPipeline.h b/Engine/Plugins/Experimental/Interchange/Runtime/Source/Pipelines/Public/InterchangeGenericAssetsPipeline.h
--- a/Engine/Plugins/Experimental/Interchange/Runtime/Source/Pipelines/Public/InterchangeGenericAssetsPipeline.h
+++ b/Engine/Plugins/Experimental/Interchange/Runtime/Source/Pipelines/Public/InterchangeGenericAssetsPipeline.h
@@ -3,6 +3,7 @@
 #pragma once
 
 #include "CoreMinimal.h"
+#include "InterchangeGenericAssetsPipelineSharedSettings.h"
 #include "InterchangePipelineBase.h"
 #include "InterchangeSourceData.h"
 #include "Nodes/InterchangeBaseNode.h"
@@ -12,6 +13,7 @@
 
 #include "InterchangeGenericAssetsPipeline.generated.h"
 
+class UInterchangeGenericAnimationPipeline;
 class UInterchangeGenericMaterialPipeline;
 class UInterchangeGenericMeshPipeline;
 class UInterchangeGenericTexturePipeline;
@@ -28,12 +30,8 @@ class INTERCHANGEPIPELINES_API UInterchangeGenericAssetsPipeline : public UInter
 
 public:
 	UInterchangeGenericAssetsPipeline();
-
-	//////////////////////////////////////////////////////////////////////////
-	// BEGIN Pre import pipeline properties, please keep by per category order for properties declaration
-
+	
 	//////	COMMON_CATEGORY Properties //////
-
 	/* Allow user to choose the re-import strategy. */
 	UPROPERTY(EditAnywhere, BlueprintReadWrite, Category = "Common")
 	EReimportStrategyFlags ReimportStrategy = EReimportStrategyFlags::ApplyNoProperties;
@@ -42,37 +40,48 @@ public:
 	UPROPERTY(EditAnywhere, BlueprintReadWrite, Category = "Common")
 	bool bUseSourceNameForAsset = true;
 
-	//////	MESHES_CATEGORY Properties //////
+	/** Translation offset applied to meshes and animations. */
+	UPROPERTY(EditAnywhere, BlueprintReadWrite, Category = "Common", meta = (DisplayName = "Offset Translation"))
+	FVector ImportOffsetTranslation;
 
-	UPROPERTY(VisibleAnywhere, Instanced, Category = "Meshes")
-	TObjectPtr<UInterchangeGenericMeshPipeline> MeshPipeline;
-	
-	//////	ANIMATIONS_CATEGORY Properties //////
+	/** Rotation offset applied to meshes and animations. */
+	UPROPERTY(EditAnywhere, BlueprintReadWrite, Category = "Common", meta = (DisplayName = "Offset Rotation"))
+	FRotator ImportOffsetRotation;
 
+	/** Uniform scale offset applied to meshes and animations. */
+	UPROPERTY(EditAnywhere, BlueprintReadWrite, Category = "Common", meta = (DisplayName = "Offset Uniform Scale"))
+	float ImportOffsetUniformScale;
 
-	/** If enable, import the animation asset find in the sources. */
-// 	UPROPERTY(EditAnywhere, BlueprintReadWrite, Category = "Animations")
-// 	bool bImportAnimations = true;
+	//////	COMMON_MESHES_CATEGORY Properties //////
+	UPROPERTY(VisibleAnywhere, Instanced, Category = "Common Meshes")
+	TObjectPtr<UInterchangeGenericCommonMeshesProperties> CommonMeshesProperties;
+		
+	//////  COMMON_SKELETAL_ANIMATIONS_CATEGORY //////
+	UPROPERTY(VisibleAnywhere, Instanced, Category = "Common Skeletal Meshes and Animations")
+	TObjectPtr<UInterchangeGenericCommonSkeletalMeshesAndAnimationsProperties> CommonSkeletalMeshesAndAnimationsProperties;
 
+	//////	MESHES_CATEGORY Properties //////
+	UPROPERTY(VisibleAnywhere, Instanced, Category = "Meshes")
+	TObjectPtr<UInterchangeGenericMeshPipeline> MeshPipeline;
 
-	//////	MATERIALS_CATEGORY Properties //////
-
+	//////	ANIMATIONS_CATEGORY Properties //////
+	UPROPERTY(VisibleAnywhere, Instanced, Category = "Animation")
+	TObjectPtr<UInterchangeGenericAnimationPipeline> AnimationPipeline;
 
+	//////	MATERIALS_CATEGORY Properties //////
 	UPROPERTY(VisibleAnywhere, Instanced, Category = "Materials")
 	TObjectPtr<UInterchangeGenericMaterialPipeline> MaterialPipeline;
 
-
 	//////	TEXTURES_CATEGORY Properties //////
-
-
 	UPROPERTY(VisibleAnywhere, Instanced, Category = "Textures")
 	TObjectPtr<UInterchangeGenericTexturePipeline> TexturePipeline;
 
-	// END Pre import pipeline properties
-	//////////////////////////////////////////////////////////////////////////
-
 	virtual void PreDialogCleanup(const FName PipelineStackName) override;
 
+	virtual bool IsSettingsAreValid() const override;
+
+	virtual void SetupReimportData(TObjectPtr<UObject> ReimportObject) override;
+
 protected:
 
 	virtual void ExecutePreImportPipeline(UInterchangeBaseNodeContainer* InBaseNodeContainer, const TArray<UInterchangeSourceData*>& InSourceDatas) override;
diff --git a/Engine/Plugins/Experimental/Interchange/Runtime/Source/Pipelines/Public/InterchangeGenericAssetsPipelineSharedSettings.h b/Engine/Plugins/Experimental/Interchange/Runtime/Source/Pipelines/Public/InterchangeGenericAssetsPipelineSharedSettings.h
new file mode 100644
--- /dev/null
+++ b/Engine/Plugins/Experimental/Interchange/Runtime/Source/Pipelines/Public/InterchangeGenericAssetsPipelineSharedSettings.h
@@ -0,0 +1,104 @@
+// Copyright Epic Games, Inc. All Rights Reserved.
+
+#pragma once
+
+#include "InterchangePipelineBase.h"
+#include "Math/Color.h"
+#include "Math/Rotator.h"
+#include "Math/Vector.h"
+#include "Nodes/InterchangeBaseNode.h"
+#include "UObject/Class.h"
+#include "UObject/Object.h"
+#include "UObject/ObjectMacros.h"
+#include "UObject/ObjectPtr.h"
+
+#include "InterchangeGenericAssetsPipelineSharedSettings.generated.h"
+
+class USkeleton;
+
+/** Force mesh type, if user want to import all meshes as one type*/
+UENUM(BlueprintType)
+enum class EInterchangeForceMeshType : uint8
+{
+	/** Will import from the source type, no conversion */
+	IFMT_None UMETA(DisplayName = "None"),
+	/** Will import any mesh to static mesh. */
+	IFMT_StaticMesh UMETA(DisplayName = "Static Mesh"),
+	/** Will import any mesh to skeletal mesh. */
+	IFMT_SkeletalMesh UMETA(DisplayName = "Skeletal Mesh"),
+
+	IFMT_MAX
+};
+
+UENUM(BlueprintType)
+enum class EInterchangeVertexColorImportOption : uint8
+{
+	/** Import the mesh using the vertex colors from the translated source. */
+	IVCIO_Replace UMETA(DisplayName = "Replace"),
+	/** Ignore vertex colors from the translated source. In case of a re-import keep the existing mesh vertex colors. */
+	IVCIO_Ignore UMETA(DisplayName = "Ignore"),
+	/** Override all vertex colors with the specified color. */
+	IVCIO_Override UMETA(DisplayName = "Override"),
+
+	IVCIO_MAX
+};
+
+
+UCLASS(BlueprintType, Experimental)
+class INTERCHANGEPIPELINES_API UInterchangeGenericCommonMeshesProperties : public UInterchangePipelineBase
+{
+	GENERATED_BODY()
+public:
+	//////	COMMON_MESHES_CATEGORY Properties //////
+
+	/** Allow to convert mesh to a particular type */
+ 	UPROPERTY(EditAnywhere, BlueprintReadWrite, Category = "Common Meshes")
+ 	EInterchangeForceMeshType ForceAllMeshAsType = EInterchangeForceMeshType::IFMT_None;
+
+	/** If enable, meshes LODs will be imported. Note that it required the advanced bBakeMesh property to be enabled. */
+	UPROPERTY(EditAnywhere, BlueprintReadWrite, Category = "Common Meshes", meta = (editcondition = "bBakeMeshes"))
+	bool bImportLods = true;
+
+	/** If enable, meshes will be baked with the scene instance hierarchy transform. */
+	UPROPERTY(EditAnywhere, BlueprintReadWrite, Category = "Common Meshes")
+	bool bBakeMeshes = true;
+
+	/** Specify how vertex colors should be imported */
+	UPROPERTY(EditAnywhere, BlueprintReadWrite, Category = "Common Meshes")
+	EInterchangeVertexColorImportOption VertexColorImportOption = EInterchangeVertexColorImportOption::IVCIO_Replace;
+
+	/** Specify override color in the case that VertexColorImportOption is set to Override */
+	UPROPERTY(EditAnywhere, BlueprintReadWrite, Category = "Common Meshes")
+	FColor VertexOverrideColor;
+};
+
+UCLASS(BlueprintType, Experimental)
+class INTERCHANGEPIPELINES_API UInterchangeGenericCommonSkeletalMeshesAndAnimationsProperties : public UInterchangePipelineBase
+{
+	GENERATED_BODY()
+public:
+	/** Enable this option to only import animation, a valid skeleton must be set to import only the animations. */
+	UPROPERTY(EditAnywhere, BlueprintReadWrite, Category = "Common Skeletal Meshes and Animations")
+	bool bImportOnlyAnimations = false;
+
+	/** Skeleton to use for imported asset. When importing a skeletal mesh, leaving this as "None" will create a new skeleton. When importing an animation this MUST be specified to import the asset. */
+	UPROPERTY(EditAnywhere, BlueprintReadWrite, Category = "Common Skeletal Meshes and Animations")
+	TObjectPtr<USkeleton> Skeleton;
+
+	/** If checked, meshes nested in bone hierarchies will be imported instead of being converted to bones. */
+	UPROPERTY(EditAnywhere, BlueprintReadWrite, Category = "Common Skeletal Meshes and Animations")
+	bool bImportMeshesInBoneHierarchy = true;
+
+	/** Enable this option to use frame 0 as reference pose */
+	UPROPERTY(EditAnywhere, BlueprintReadWrite, Category = "Common Skeletal Meshes and Animations")
+	bool bUseT0AsRefPose = false;
+
+	virtual bool IsSettingsAreValid() const override
+	{
+		if (bImportOnlyAnimations && Skeleton.IsNull())
+		{
+			return false;
+		}
+		return Super::IsSettingsAreValid();
+	}
+};
diff --git a/Engine/Plugins/Experimental/Interchange/Runtime/Source/Pipelines/Public/InterchangeGenericMeshPipeline.h b/Engine/Plugins/Experimental/Interchange/Runtime/Source/Pipelines/Public/InterchangeGenericMeshPipeline.h
--- a/Engine/Plugins/Experimental/Interchange/Runtime/Source/Pipelines/Public/InterchangeGenericMeshPipeline.h
+++ b/Engine/Plugins/Experimental/Interchange/Runtime/Source/Pipelines/Public/InterchangeGenericMeshPipeline.h
@@ -3,6 +3,7 @@
 #pragma once
 
 #include "CoreMinimal.h"
+#include "InterchangeGenericAssetsPipelineSharedSettings.h"
 #include "InterchangePipelineBase.h"
 #include "InterchangeSourceData.h"
 #include "Nodes/InterchangeBaseNodeContainer.h"
@@ -12,6 +13,7 @@
 
 #include "InterchangeGenericMeshPipeline.generated.h"
 
+class UInterchangeGenericAssetsPipeline;
 class UInterchangeMeshNode;
 class UInterchangePipelineMeshesUtilities;
 class UInterchangeSkeletalMeshFactoryNode;
@@ -19,33 +21,8 @@ class UInterchangeSkeletalMeshLodDataNode;
 class UInterchangeSkeletonFactoryNode;
 class UInterchangeStaticMeshFactoryNode;
 class UInterchangeStaticMeshLodDataNode;
+class UPhysicsAsset;
 
-/** Force mesh type, if user want to import all meshes as one type*/
-UENUM(BlueprintType)
-enum class EInterchangeForceMeshType : uint8
-{
-	/** Will import from the source type, no conversion */
-	IFMT_None UMETA(DisplayName = "None"),
-	/** Will import any mesh to static mesh. */
-	IFMT_StaticMesh UMETA(DisplayName = "Static Mesh"),
-	/** Will import any mesh to skeletal mesh. */
-	IFMT_SkeletalMesh UMETA(DisplayName = "Skeletal Mesh"),
-
-	IFMT_MAX
-};
-
-UENUM(BlueprintType)
-enum class EInterchangeVertexColorImportOption : uint8
-{
-	/** Import the mesh using the vertex colors from the translated source. */
-	IVCIO_Replace UMETA(DisplayName = "Replace"),
-	/** Ignore vertex colors from the translated source. In case of a re-import keep the existing mesh vertex colors. */
-	IVCIO_Ignore UMETA(DisplayName = "Ignore"),
-	/** Override all vertex colors with the specified color. */
-	IVCIO_Override UMETA(DisplayName = "Override"),
-
-	IVCIO_MAX
-};
 
 UCLASS(BlueprintType, Experimental)
 class INTERCHANGEPIPELINES_API UInterchangeGenericMeshPipeline : public UInterchangePipelineBase
@@ -54,27 +31,13 @@ class INTERCHANGEPIPELINES_API UInterchangeGenericMeshPipeline : public UInterch
 
 public:
 
-	//////	COMMON_MESHES_CATEGORY Properties //////
+	//Common Meshes Properties Settings Pointer
+	UPROPERTY(Transient)
+	TObjectPtr<UInterchangeGenericCommonMeshesProperties> CommonMeshesProperties;
 
-	/** Allow to convert mesh to a particular type */
- 	UPROPERTY(EditAnywhere, BlueprintReadWrite, Category = "Common Meshes")
- 	EInterchangeForceMeshType ForceAllMeshAsType = EInterchangeForceMeshType::IFMT_None;
-
-	/** If enable, meshes LODs will be imported. Note that it required the advanced bBakeMesh property to be enabled. */
-	UPROPERTY(EditAnywhere, BlueprintReadWrite, Category = "Common Meshes", meta = (editcondition = "bBakeMeshes"))
-	bool bImportLods = true;
-
-	/** If enable, meshes will be baked with the scene instance hierarchy transform. */
-	UPROPERTY(EditAnywhere, BlueprintReadWrite, Category = "Common Meshes")
-	bool bBakeMeshes = true;
-
-	/** Specify how vertex colors should be imported */
-	UPROPERTY(EditAnywhere, BlueprintReadWrite, Category = "Common Meshes")
-	EInterchangeVertexColorImportOption VertexColorImportOption = EInterchangeVertexColorImportOption::IVCIO_Replace;
-
-	/** Specify override color in the case that VertexColorImportOption is set to Override */
-	UPROPERTY(EditAnywhere, BlueprintReadWrite, Category = "Common Meshes")
-	FColor VertexOverrideColor;
+	//Common SkeletalMeshes And Animations Properties Settings Pointer
+	UPROPERTY(Transient)
+	TObjectPtr<UInterchangeGenericCommonSkeletalMeshesAndAnimationsProperties> CommonSkeletalMeshesAndAnimationsProperties;
 	
 	//////	STATIC_MESHES_CATEGORY Properties //////
 
@@ -104,14 +67,6 @@ public:
 	bool bOneConvexHullPerUCX = true;
 
 	
-	//////  COMMON_SKELETAL_ANIMATIONS_CATEGORY //////
-
-
-	/** Skeleton to use for imported asset. When importing a skeletal mesh, leaving this as "None" will create a new skeleton. When importing an animation this MUST be specified to import the asset. */
- 	UPROPERTY(EditAnywhere, BlueprintReadWrite, Category = "Common Skeletal Mesh and Animations")
- 	TObjectPtr<class USkeleton> Skeleton;
-
-	
 	//////	SKELETAL_MESHES_CATEGORY Properties //////
 
 	/** If enable, import the animation asset find in the sources. */
@@ -138,23 +93,13 @@ public:
 	UPROPERTY(EditAnywhere, BlueprintReadWrite, Category = "Skeletal Meshes")
 	bool bUpdateSkeletonReferencePose = false;
 
-	/** If checked, meshes nested in bone hierarchies will be imported instead of being converted to bones. */
-	UPROPERTY(EditAnywhere, BlueprintReadWrite, Category = "Skeletal Meshes")
-	bool bImportMeshesInBoneHierarchy = true;
-
-	/** Enable this option to use frame 0 as reference pose */
-	UPROPERTY(EditAnywhere, BlueprintReadWrite, Category = "Skeletal Meshes")
-	bool bUseT0AsRefPose = false;
-
 	/** If checked, create new PhysicsAsset if it doesn't have it */
 	UPROPERTY(EditAnywhere, BlueprintReadWrite, Category = "Skeletal Meshes")
 	bool bCreatePhysicsAsset = true;
 
 	/** If this is set, use this specified PhysicsAsset. If its not set and bCreatePhysicsAsset is false, the importer will not generate or set any physic asset. */
 	UPROPERTY(EditAnywhere, BlueprintReadWrite, Category = "Skeletal Meshes", meta = (editcondition = "!bCreatePhysicsAsset"))
-	TObjectPtr<class UPhysicsAsset> PhysicsAsset;
-
-	virtual void PreDialogCleanup(const FName PipelineStackName) override;
+	TObjectPtr<UPhysicsAsset> PhysicsAsset;
 
 protected:
 	virtual void ExecutePreImportPipeline(UInterchangeBaseNodeContainer* InBaseNodeContainer, const TArray<UInterchangeSourceData*>& InSourceDatas) override;
@@ -171,8 +116,7 @@ protected:
 private:
 
 	/* Meshes utilities, to parse the translated graph and extract the meshes informations. */
-	UInterchangePipelineMeshesUtilities* PipelineMeshesUtilities;
-
+	TObjectPtr<UInterchangePipelineMeshesUtilities> PipelineMeshesUtilities = nullptr;
 
 	/************************************************************************/
 	/* Skeletal mesh API BEGIN                                              */
@@ -218,7 +162,7 @@ private:
 public:
 	
 	/** Specialize for skeletalmesh */
-	void ImplementUseSourceNameForAssetOptionSkeletalMesh(const int32 MeshesAndAnimsImportedNodeCount, const bool bUseSourceNameForAsset);
+	void ImplementUseSourceNameForAssetOptionSkeletalMesh(const int32 MeshesImportedNodeCount, const bool bUseSourceNameForAsset);
 
 private:
 	/* Skeletal mesh API END                                                */
diff --git a/Engine/Plugins/MeshPainting/Source/MeshPaintEditorMode/Private/MeshPaintModeHelpers.cpp b/Engine/Plugins/MeshPainting/Source/MeshPaintEditorMode/Private/MeshPaintModeHelpers.cpp
--- a/Engine/Plugins/MeshPainting/Source/MeshPaintEditorMode/Private/MeshPaintModeHelpers.cpp
+++ b/Engine/Plugins/MeshPainting/Source/MeshPaintEditorMode/Private/MeshPaintModeHelpers.cpp
@@ -322,11 +322,11 @@ void UMeshPaintModeSubsystem::ImportVertexColorsToSkeletalMesh(USkeletalMesh* Sk
 			{
 				if (UInterchangeGenericAssetsPipeline* GenericAssetPipeline = Cast<UInterchangeGenericAssetsPipeline>(PipelineBase.Get()))
 				{
-					if (GenericAssetPipeline->MeshPipeline && GenericAssetPipeline->MeshPipeline->VertexColorImportOption != EInterchangeVertexColorImportOption::IVCIO_Ignore)
+					if (!GenericAssetPipeline->CommonMeshesProperties.IsNull() && GenericAssetPipeline->CommonMeshesProperties->VertexColorImportOption != EInterchangeVertexColorImportOption::IVCIO_Ignore)
 					{
 						InterchangeAssetImportData->SetFlags(RF_Transactional);
 						InterchangeAssetImportData->Modify();
-						GenericAssetPipeline->MeshPipeline->VertexColorImportOption = EInterchangeVertexColorImportOption::IVCIO_Ignore;
+						GenericAssetPipeline->CommonMeshesProperties->VertexColorImportOption = EInterchangeVertexColorImportOption::IVCIO_Ignore;
 					}
 				}
 			}
diff --git a/Engine/Plugins/MeshPainting/Source/MeshPaintingToolset/Private/MeshPaintSkeletalMeshAdapter.cpp b/Engine/Plugins/MeshPainting/Source/MeshPaintingToolset/Private/MeshPaintSkeletalMeshAdapter.cpp
--- a/Engine/Plugins/MeshPainting/Source/MeshPaintingToolset/Private/MeshPaintSkeletalMeshAdapter.cpp
+++ b/Engine/Plugins/MeshPainting/Source/MeshPaintingToolset/Private/MeshPaintSkeletalMeshAdapter.cpp
@@ -402,11 +402,11 @@ void FMeshPaintSkeletalMeshComponentAdapter::PreEdit()
 			{
 				if (UInterchangeGenericAssetsPipeline* GenericAssetPipeline = Cast<UInterchangeGenericAssetsPipeline>(PipelineBase.Get()))
 				{
-					if (GenericAssetPipeline->MeshPipeline && GenericAssetPipeline->MeshPipeline->VertexColorImportOption != EInterchangeVertexColorImportOption::IVCIO_Ignore)
+					if (!GenericAssetPipeline->CommonMeshesProperties.IsNull() && GenericAssetPipeline->CommonMeshesProperties->VertexColorImportOption != EInterchangeVertexColorImportOption::IVCIO_Ignore)
 					{
 						GenericAssetPipeline->SetFlags(RF_Transactional);
 						GenericAssetPipeline->Modify();
-						GenericAssetPipeline->MeshPipeline->VertexColorImportOption = EInterchangeVertexColorImportOption::IVCIO_Ignore;
+						GenericAssetPipeline->CommonMeshesProperties->VertexColorImportOption = EInterchangeVertexColorImportOption::IVCIO_Ignore;
 					}
 				}
 			}
diff --git a/Engine/Source/Editor/UnrealEd/Private/Factories/InterchangeReimportHandler.cpp b/Engine/Source/Editor/UnrealEd/Private/Factories/InterchangeReimportHandler.cpp
--- a/Engine/Source/Editor/UnrealEd/Private/Factories/InterchangeReimportHandler.cpp
+++ b/Engine/Source/Editor/UnrealEd/Private/Factories/InterchangeReimportHandler.cpp
@@ -2,6 +2,7 @@
 
 #include "Factories/InterchangeReimportHandler.h"
 
+#include "Animation/AnimSequence.h"
 #include "CoreMinimal.h"
 #include "EngineDefines.h"
 #include "Engine/SkeletalMesh.h"
@@ -46,7 +47,7 @@ bool UInterchangeReimportHandler::CanReimport(UObject* Obj, TArray<FString>& Out
 
 	USkeletalMesh* SkeletalMesh = Cast<USkeletalMesh>(Obj);
 	UStaticMesh* StaticMesh = Cast<UStaticMesh>(Obj);
-	//UAnimSequence* AnimSequence = Cast<UAnimSequence>(Obj);
+	UAnimSequence* AnimSequence = Cast<UAnimSequence>(Obj);
 	UTexture* Texture = Cast<UTexture>(Obj);
 	if (bUseInterchangeFramework && SkeletalMesh)
 	{
@@ -62,13 +63,13 @@ bool UInterchangeReimportHandler::CanReimport(UObject* Obj, TArray<FString>& Out
 			return CanReimportAsset(AssetImportData);
 		}
 	}
-// 	else if (bUseInterchangeFramework && AnimSequence)
-// 	{
-// 		if (UAssetImportData* AssetImportData = AnimSequence->AssetImportData)
-// 		{
-// 			return CanReimportAsset(AssetImportData);
-// 		}
-// 	}
+	else if (bUseInterchangeFramework && AnimSequence)
+	{
+		if (UAssetImportData* AssetImportData = AnimSequence->AssetImportData)
+		{
+			return CanReimportAsset(AssetImportData);
+		}
+	}
 	else if ((bUseInterchangeFramework || bUseInterchangeFrameworkForTextureOnly) && Texture)
 	{
 		if (UAssetImportData* AssetImportData = Texture->AssetImportData)
@@ -83,7 +84,7 @@ void UInterchangeReimportHandler::SetReimportPaths(UObject* Obj, const FString&
 {
 	USkeletalMesh* SkeletalMesh = Cast<USkeletalMesh>(Obj);
 	UStaticMesh* StaticMesh = Cast<UStaticMesh>(Obj);
-//	UAnimSequence* AnimSequence = Cast<UAnimSequence>(Obj);
+	UAnimSequence* AnimSequence = Cast<UAnimSequence>(Obj);
 	UTexture* Texture = Cast<UTexture>(Obj);
 
 	UAssetImportData* AssetImportData = nullptr;
@@ -95,10 +96,10 @@ void UInterchangeReimportHandler::SetReimportPaths(UObject* Obj, const FString&
 	{
 		AssetImportData = StaticMesh->GetAssetImportData();
 	}
-// 	else if (AnimSequence)
-// 	{
-// 		AssetImportData = AnimSequence->AssetImportData;
-// 	}
+	else if (AnimSequence)
+	{
+		AssetImportData = AnimSequence->AssetImportData;
+	}
 	else if (Texture)
 	{
 		AssetImportData = Texture->AssetImportData;
diff --git a/Engine/Source/Runtime/Experimental/Interchange/Core/Private/InterchangePipelineBase.cpp b/Engine/Source/Runtime/Experimental/Interchange/Core/Private/InterchangePipelineBase.cpp
--- a/Engine/Source/Runtime/Experimental/Interchange/Core/Private/InterchangePipelineBase.cpp
+++ b/Engine/Source/Runtime/Experimental/Interchange/Core/Private/InterchangePipelineBase.cpp
@@ -42,7 +42,7 @@ void UInterchangePipelineBase::LoadSettings(const FName PipelineStackName)
 		UObject* SuperClassDefaultObject = Class->GetSuperClass()->GetDefaultObject();
 
 		const FString& PropFileName = GEditorPerProjectIni;
-
+		FObjectProperty* SubObject = CastField<FObjectProperty>(Property);
 		FArrayProperty* Array = CastField<FArrayProperty>(Property);
 		if (Array)
 		{
@@ -95,15 +95,12 @@ void UInterchangePipelineBase::LoadSettings(const FName PipelineStackName)
 				}
 			}
 		}
-		else if (FObjectProperty* SubObject = CastField<FObjectProperty>(Property))
+		else if (UInterchangePipelineBase* SubPipeline = SubObject ? Cast<UInterchangePipelineBase>(SubObject->GetObjectPropertyValue_InContainer(this)) : nullptr)
 		{
-			if (UInterchangePipelineBase* SubPipeline = Cast<UInterchangePipelineBase>(SubObject->GetObjectPropertyValue_InContainer(this)))
+			// Load the settings if the referenced pipeline is a subobject of ours
+			if (SubPipeline->IsInOuter(this))
 			{
-				// Load the settings if the referenced pipeline is a subobject of ours
-				if (SubPipeline->IsInOuter(this))
-				{
-					SubPipeline->LoadSettings(PipelineStackName);
-				}
+				SubPipeline->LoadSettings(PipelineStackName);
 			}
 		}
 		else
@@ -148,7 +145,7 @@ void UInterchangePipelineBase::SaveSettings(const FName PipelineStackName)
 
 		const bool bIsPropertyInherited = Property->GetOwnerClass() != Class;
 		UObject* SuperClassDefaultObject = Class->GetSuperClass()->GetDefaultObject();
-
+		FObjectProperty* SubObject = CastField<FObjectProperty>(Property);
 		FArrayProperty* Array = CastField<FArrayProperty>(Property);
 		if (Array)
 		{
@@ -166,15 +163,12 @@ void UInterchangePipelineBase::SaveSettings(const FName PipelineStackName)
 				Sec->Add(*Key, *Buffer);
 			}
 		}
-		else if (FObjectProperty* SubObject = CastField<FObjectProperty>(Property))
+		else if (UInterchangePipelineBase* SubPipeline = SubObject ? Cast<UInterchangePipelineBase>(SubObject->GetObjectPropertyValue_InContainer(this)) : nullptr)
 		{
-			if (UInterchangePipelineBase* SubPipeline = Cast<UInterchangePipelineBase>(SubObject->GetObjectPropertyValue_InContainer(this)))
+			// Save the settings if the referenced pipeline is a subobject of ours
+			if (SubPipeline->IsInOuter(this))
 			{
-				// Save the settings if the referenced pipeline is a subobject of ours
-				if (SubPipeline->IsInOuter(this))
-				{
-					SubPipeline->SaveSettings(PipelineStackName);
-				}
+				SubPipeline->SaveSettings(PipelineStackName);
 			}
 		}
 		else
diff --git a/Engine/Source/Runtime/Experimental/Interchange/Core/Private/Nodes/InterchangeSourceNode.cpp b/Engine/Source/Runtime/Experimental/Interchange/Core/Private/Nodes/InterchangeSourceNode.cpp
new file mode 100644
--- /dev/null
+++ b/Engine/Source/Runtime/Experimental/Interchange/Core/Private/Nodes/InterchangeSourceNode.cpp
@@ -0,0 +1,83 @@
+// Copyright Epic Games, Inc. All Rights Reserved.
+#include "Nodes/InterchangeSourceNode.h"
+
+#include "CoreMinimal.h"
+#include "Types/AttributeStorage.h"
+#include "UObject/Class.h"
+#include "UObject/Object.h"
+#include "UObject/ObjectMacros.h"
+
+UInterchangeSourceNode::UInterchangeSourceNode()
+{
+}
+
+void UInterchangeSourceNode::InitializeSourceNode(const FString& UniqueID, const FString& DisplayLabel)
+{
+	InitializeNode(UniqueID, DisplayLabel, EInterchangeNodeContainerType::TranslatedAsset);
+}
+
+FString UInterchangeSourceNode::GetTypeName() const
+{
+	const FString TypeName = TEXT("SourceNode");
+	return TypeName;
+}
+
+bool UInterchangeSourceNode::GetCustomSourceFrameRateNumerator(int32& AttributeValue) const
+{
+	IMPLEMENT_NODE_ATTRIBUTE_GETTER(SourceFrameRateNumerator, int32);
+}
+
+bool UInterchangeSourceNode::SetCustomSourceFrameRateNumerator(const int32& AttributeValue)
+{
+	IMPLEMENT_NODE_ATTRIBUTE_SETTER_NODELEGATE(SourceFrameRateNumerator, int32);
+}
+
+bool UInterchangeSourceNode::GetCustomSourceFrameRateDenominator(int32& AttributeValue) const
+{
+	IMPLEMENT_NODE_ATTRIBUTE_GETTER(SourceFrameRateDenominator, int32);
+}
+
+bool UInterchangeSourceNode::SetCustomSourceFrameRateDenominator(const int32& AttributeValue)
+{
+	IMPLEMENT_NODE_ATTRIBUTE_SETTER_NODELEGATE(SourceFrameRateDenominator, int32);
+}
+
+bool UInterchangeSourceNode::GetCustomSourceTimelineStart(double& AttributeValue) const
+{
+	IMPLEMENT_NODE_ATTRIBUTE_GETTER(SourceTimelineStart, double);
+}
+
+bool UInterchangeSourceNode::SetCustomSourceTimelineStart(const double& AttributeValue)
+{
+	IMPLEMENT_NODE_ATTRIBUTE_SETTER_NODELEGATE(SourceTimelineStart, double);
+}
+
+bool UInterchangeSourceNode::GetCustomSourceTimelineEnd(double& AttributeValue) const
+{
+	IMPLEMENT_NODE_ATTRIBUTE_GETTER(SourceTimelineEnd, double);
+}
+
+bool UInterchangeSourceNode::SetCustomSourceTimelineEnd(const double& AttributeValue)
+{
+	IMPLEMENT_NODE_ATTRIBUTE_SETTER_NODELEGATE(SourceTimelineEnd, double);
+}
+
+bool UInterchangeSourceNode::GetCustomAnimatedTimeStart(double& AttributeValue) const
+{
+	IMPLEMENT_NODE_ATTRIBUTE_GETTER(AnimatedTimeStart, double);
+}
+
+bool UInterchangeSourceNode::SetCustomAnimatedTimeStart(const double& AttributeValue)
+{
+	IMPLEMENT_NODE_ATTRIBUTE_SETTER_NODELEGATE(AnimatedTimeStart, double);
+}
+
+bool UInterchangeSourceNode::GetCustomAnimatedTimeEnd(double& AttributeValue) const
+{
+	IMPLEMENT_NODE_ATTRIBUTE_GETTER(AnimatedTimeEnd, double);
+}
+
+bool UInterchangeSourceNode::SetCustomAnimatedTimeEnd(const double& AttributeValue)
+{
+	IMPLEMENT_NODE_ATTRIBUTE_SETTER_NODELEGATE(AnimatedTimeEnd, double);
+}
\ No newline at end of file
diff --git a/Engine/Source/Runtime/Experimental/Interchange/Core/Public/InterchangePipelineBase.h b/Engine/Source/Runtime/Experimental/Interchange/Core/Public/InterchangePipelineBase.h
--- a/Engine/Source/Runtime/Experimental/Interchange/Core/Public/InterchangePipelineBase.h
+++ b/Engine/Source/Runtime/Experimental/Interchange/Core/Public/InterchangePipelineBase.h
@@ -112,8 +112,25 @@ public:
 
 	void SaveSettings(const FName PipelineStackName);
 
+	/**
+	 * This function is called before showing the import dialog it is not called when doing a re-import.
+	 */
 	virtual void PreDialogCleanup(const FName PipelineStackName) {};
 
+	/**
+	 * This function should return true if all the pipeline settings are in a valid state to start the import.
+	 * The Interchange Pipeline Configuration Dialog will call this to know if the Import button can be enable.
+	 */
+	virtual bool IsSettingsAreValid() const
+	{
+		return true;
+	}
+
+	/**
+	 * This function is call before showing the re-import dialog it is not call when importing a new assets.
+	 */
+	virtual void SetupReimportData(TObjectPtr<UObject> ReimportObject) {};
+
 	/**
 	 * This function is used to add the given message object directly into the results for this operation.
 	 */
diff --git a/Engine/Source/Runtime/Experimental/Interchange/Core/Public/Nodes/InterchangeSourceNode.h b/Engine/Source/Runtime/Experimental/Interchange/Core/Public/Nodes/InterchangeSourceNode.h
new file mode 100644
--- /dev/null
+++ b/Engine/Source/Runtime/Experimental/Interchange/Core/Public/Nodes/InterchangeSourceNode.h
@@ -0,0 +1,98 @@
+// Copyright Epic Games, Inc. All Rights Reserved.
+
+#pragma once
+
+#include "CoreMinimal.h"
+#include "Nodes/InterchangeBaseNode.h"
+#include "Nodes/InterchangeBaseNodeUtilities.h"
+#include "Types/AttributeStorage.h"
+#include "UObject/Class.h"
+#include "UObject/Object.h"
+#include "UObject/ObjectMacros.h"
+
+#include "InterchangeSourceNode.generated.h"
+
+
+
+/**
+ * This class allow a translator to add general source data that describe the whole source. Pipeline can use this information.
+ */
+UCLASS(BlueprintType, Experimental)
+class INTERCHANGECORE_API UInterchangeSourceNode : public UInterchangeBaseNode
+{
+	GENERATED_BODY()
+
+public:
+	UInterchangeSourceNode();
+
+	/**
+	 * Initialize the base data of the node
+	 * @param UniqueID - The uniqueId for this node
+	 * @param DisplayLabel - The name of the node
+	 *
+	 */
+	UFUNCTION(BlueprintCallable, Category = "Interchange | Node | source")
+	void InitializeSourceNode(const FString& UniqueID, const FString& DisplayLabel);
+
+	/**
+	 * Return the node type name of the class, we use this when reporting error
+	 */
+	virtual FString GetTypeName() const override;
+
+	/** Query the source frame rate numerator. */
+	UFUNCTION(BlueprintCallable, Category = "Interchange | Node | Source")
+	bool GetCustomSourceFrameRateNumerator(int32& AttributeValue) const;
+
+	/** Store the source frame rate numerator. */
+	UFUNCTION(BlueprintCallable, Category = "Interchange | Node | Source")
+	bool SetCustomSourceFrameRateNumerator(const int32& AttributeValue);
+
+	/** Query the source frame rate denominator. */
+	UFUNCTION(BlueprintCallable, Category = "Interchange | Node | Source")
+	bool GetCustomSourceFrameRateDenominator(int32& AttributeValue) const;
+
+	/** Store the source frame rate denominator. */
+	UFUNCTION(BlueprintCallable, Category = "Interchange | Node | Source")
+	bool SetCustomSourceFrameRateDenominator(const int32& AttributeValue);
+
+	/** Query the source time line start. */
+	UFUNCTION(BlueprintCallable, Category = "Interchange | Node | Source")
+	bool GetCustomSourceTimelineStart(double& AttributeValue) const;
+
+	/** Store the source time line start. */
+	UFUNCTION(BlueprintCallable, Category = "Interchange | Node | Source")
+	bool SetCustomSourceTimelineStart(const double& AttributeValue);
+
+	/** Query the source time line end. */
+	UFUNCTION(BlueprintCallable, Category = "Interchange | Node | Source")
+	bool GetCustomSourceTimelineEnd(double& AttributeValue) const;
+
+	/** Store the source time line end. */
+	UFUNCTION(BlueprintCallable, Category = "Interchange | Node | Source")
+	bool SetCustomSourceTimelineEnd(const double& AttributeValue);
+
+	/** Query the source animated time start. */
+	UFUNCTION(BlueprintCallable, Category = "Interchange | Node | Source")
+	bool GetCustomAnimatedTimeStart(double& AttributeValue) const;
+
+	/** Store the source animated time start. */
+	UFUNCTION(BlueprintCallable, Category = "Interchange | Node | Source")
+	bool SetCustomAnimatedTimeStart(const double& AttributeValue);
+
+	/** Query the source animated time end. */
+	UFUNCTION(BlueprintCallable, Category = "Interchange | Node | Source")
+	bool GetCustomAnimatedTimeEnd(double& AttributeValue) const;
+
+	/** Store the source animated time end. */
+	UFUNCTION(BlueprintCallable, Category = "Interchange | Node | Source")
+	bool SetCustomAnimatedTimeEnd(const double& AttributeValue);
+
+private:
+	const UE::Interchange::FAttributeKey Macro_CustomSourceFrameRateNumeratorKey = UE::Interchange::FAttributeKey(TEXT("SourceFrameRateNumerator"));
+	const UE::Interchange::FAttributeKey Macro_CustomSourceFrameRateDenominatorKey = UE::Interchange::FAttributeKey(TEXT("SourceFrameRateDenominator"));
+	const UE::Interchange::FAttributeKey Macro_CustomSourceTimelineStartKey = UE::Interchange::FAttributeKey(TEXT("SourceTimelineStart"));
+	const UE::Interchange::FAttributeKey Macro_CustomSourceTimelineEndKey = UE::Interchange::FAttributeKey(TEXT("SourceTimelineEnd"));
+	const UE::Interchange::FAttributeKey Macro_CustomAnimatedTimeStartKey = UE::Interchange::FAttributeKey(TEXT("AnimatedTimeStart"));
+	const UE::Interchange::FAttributeKey Macro_CustomAnimatedTimeEndKey = UE::Interchange::FAttributeKey(TEXT("AnimatedTimeEnd"));
+};
+
diff --git a/Engine/Source/Runtime/Experimental/Interchange/Engine/Private/InterchangeManager.cpp b/Engine/Source/Runtime/Experimental/Interchange/Engine/Private/InterchangeManager.cpp
--- a/Engine/Source/Runtime/Experimental/Interchange/Engine/Private/InterchangeManager.cpp
+++ b/Engine/Source/Runtime/Experimental/Interchange/Engine/Private/InterchangeManager.cpp
@@ -711,6 +711,7 @@ UInterchangeManager::ImportInternal(const FString& ContentPath, const UInterchan
 				{
 					//Duplicate the pipeline saved in the asset import data
 					UInterchangePipelineBase* GeneratedPipeline = Cast<UInterchangePipelineBase>(StaticDuplicateObject(SourcePipeline, GetTransientPackage()));
+					GeneratedPipeline->SetupReimportData(ImportAssetParameters.ReimportAsset);
 					PipelineStack.Add(GeneratedPipeline);
 				}
 				else
