Fixed non-unity build.

diff --git a/Engine/Plugins/Experimental/Interchange/Runtime/Source/Import/Private/Mesh/InterchangeStaticMeshFactory.cpp b/Engine/Plugins/Experimental/Interchange/Runtime/Source/Import/Private/Mesh/InterchangeStaticMeshFactory.cpp
--- a/Engine/Plugins/Experimental/Interchange/Runtime/Source/Import/Private/Mesh/InterchangeStaticMeshFactory.cpp
+++ b/Engine/Plugins/Experimental/Interchange/Runtime/Source/Import/Private/Mesh/InterchangeStaticMeshFactory.cpp
@@ -21,7 +21,6 @@
 #include "Materials/Material.h"
 #include "Materials/MaterialInstance.h"
 #include "Materials/MaterialInterface.h"
-#include "Mesh/InterchangeStaticMeshPayload.h"
 #include "Mesh/InterchangeStaticMeshPayloadInterface.h"
 #include "Model.h"
 #include "Nodes/InterchangeBaseNode.h"
diff --git a/Engine/Plugins/Experimental/Interchange/Runtime/Source/Import/Public/Mesh/InterchangeStaticMeshFactory.h b/Engine/Plugins/Experimental/Interchange/Runtime/Source/Import/Public/Mesh/InterchangeStaticMeshFactory.h
--- a/Engine/Plugins/Experimental/Interchange/Runtime/Source/Import/Public/Mesh/InterchangeStaticMeshFactory.h
+++ b/Engine/Plugins/Experimental/Interchange/Runtime/Source/Import/Public/Mesh/InterchangeStaticMeshFactory.h
@@ -6,15 +6,11 @@
 #include "InterchangeFactoryBase.h"
 #include "UObject/Object.h"
 #include "UObject/ObjectMacros.h"
+#include "Mesh/InterchangeStaticMeshPayload.h"
 
 #include "InterchangeStaticMeshFactory.generated.h"
 
 
-namespace UE::Interchange
-{
-	struct FStaticMeshPayloadData;
-}
-
 class UStaticMesh;
 class UInterchangeStaticMeshLodDataNode;
 struct FKAggregateGeom;
