Fix CIS build
#rnx

diff --git a/Engine/Plugins/Experimental/Interchange/Runtime/Source/Parsers/Fbx/Private/FbxAnimation.cpp b/Engine/Plugins/Experimental/Interchange/Runtime/Source/Parsers/Fbx/Private/FbxAnimation.cpp
--- a/Engine/Plugins/Experimental/Interchange/Runtime/Source/Parsers/Fbx/Private/FbxAnimation.cpp
+++ b/Engine/Plugins/Experimental/Interchange/Runtime/Source/Parsers/Fbx/Private/FbxAnimation.cpp
@@ -414,9 +414,9 @@ namespace UE::Interchange::Private
 		FGetFbxTransformCurvesParameters(FbxScene* InSDKScene, FbxNode* InNode)
 		{
 			SDKScene = InSDKScene;
-			ensure(SDKScene);
+			check(SDKScene);
 			Node = InNode;
-			ensure(Node);
+			check(Node);
 		}
 
 		FbxScene* SDKScene = nullptr;
@@ -432,7 +432,10 @@ namespace UE::Interchange::Private
 
 	void GetFbxTransformCurves(FGetFbxTransformCurvesParameters& Parameters)
 	{
-		ensure(Parameters.SDKScene);
+		if (!ensure(Parameters.SDKScene) || !ensure(Parameters.Node))
+		{
+			return;
+		}
 		//Get the node transform curve keys, the transform components are separate into float curve
 		//Translation X
 		//Translation Y
diff --git a/Engine/Source/Runtime/Experimental/Interchange/Core/Public/Nodes/InterchangeBaseNode.h b/Engine/Source/Runtime/Experimental/Interchange/Core/Public/Nodes/InterchangeBaseNode.h
--- a/Engine/Source/Runtime/Experimental/Interchange/Core/Public/Nodes/InterchangeBaseNode.h
+++ b/Engine/Source/Runtime/Experimental/Interchange/Core/Public/Nodes/InterchangeBaseNode.h
@@ -393,7 +393,7 @@ public:
 	template<typename AttributeType>
 	AttributeType GetAttributeChecked(const FString& NodeAttributeKey) const
 	{
-		AttributeType Value;
+		AttributeType Value = AttributeType();
 		check(HasAttribute(UE::Interchange::FAttributeKey(NodeAttributeKey)));
 		UE::Interchange::FAttributeStorage::TAttributeHandle<AttributeType> Handle = GetAttributeHandle<AttributeType>(UE::Interchange::FAttributeKey(NodeAttributeKey));
 		check(Handle.IsValid());
