Fix non-unity build.

diff --git a/Engine/Plugins/Experimental/Interchange/Runtime/Source/Import/Public/Mesh/InterchangeStaticMeshFactory.h b/Engine/Plugins/Experimental/Interchange/Runtime/Source/Import/Public/Mesh/InterchangeStaticMeshFactory.h
--- a/Engine/Plugins/Experimental/Interchange/Runtime/Source/Import/Public/Mesh/InterchangeStaticMeshFactory.h
+++ b/Engine/Plugins/Experimental/Interchange/Runtime/Source/Import/Public/Mesh/InterchangeStaticMeshFactory.h
@@ -13,6 +13,7 @@
 
 class UStaticMesh;
 class UInterchangeStaticMeshLodDataNode;
+struct FMeshDescription;
 struct FKAggregateGeom;
 
 
