Interchange Scene Import - Changing the scene import to from having a single factory in charge of instanciating the whole scene, to using one factory for each scene factory node.

This allows additional types to be imported without modifying the UInterchangeActorFactory (a requirement for implementing custom interchange import plugins), and better matches the rest of interchange design.
Moving away from a single-call scene instanciation will also allow us to further improve of the scene import with changes such as:
 - Distributing the scene instanciation on multiple tick, relying of interchange's dependency system, rather than waiting for all the assets to be imported.
 - Calling the pre/post-import pipeline events per actor.
 - Having inter-actor import dependencies (right now the interchange scene import task only consider asset dependencies)

diff --git a/Engine/Plugins/Experimental/Interchange/Runtime/Source/Import/Private/InterchangeImportModule.cpp b/Engine/Plugins/Experimental/Interchange/Runtime/Source/Import/Private/InterchangeImportModule.cpp
--- a/Engine/Plugins/Experimental/Interchange/Runtime/Source/Import/Private/InterchangeImportModule.cpp
+++ b/Engine/Plugins/Experimental/Interchange/Runtime/Source/Import/Private/InterchangeImportModule.cpp
@@ -17,6 +17,9 @@
 #include "Misc/CoreDelegates.h"
 #include "Modules/ModuleManager.h"
 #include "Scene/InterchangeActorFactory.h"
+#include "Scene/InterchangeCineCameraActorFactory.h"
+#include "Scene/InterchangeStaticMeshActorFactory.h"
+#include "Scene/InterchangeSkeletalMeshActorFactory.h"
 #include "Texture/InterchangeBMPTranslator.h"
 #include "Texture/InterchangeDDSTranslator.h"
 #include "Texture/InterchangeEXRTranslator.h"
@@ -90,6 +93,9 @@ void FInterchangeImportModule::StartupModule()
 		InterchangeManager.RegisterFactory(UInterchangePhysicsAssetFactory::StaticClass());
 		InterchangeManager.RegisterFactory(UInterchangeActorFactory::StaticClass());
 		InterchangeManager.RegisterFactory(UInterchangeAnimSequenceFactory::StaticClass());
+		InterchangeManager.RegisterFactory(UInterchangeCineCameraActorFactory::StaticClass());
+		InterchangeManager.RegisterFactory(UInterchangeStaticMeshActorFactory::StaticClass());
+		InterchangeManager.RegisterFactory(UInterchangeSkeletalMeshActorFactory::StaticClass());
 	};
 
 	if (GEngine)
diff --git a/Engine/Plugins/Experimental/Interchange/Runtime/Source/Import/Private/Scene/InterchangeActorFactory.cpp b/Engine/Plugins/Experimental/Interchange/Runtime/Source/Import/Private/Scene/InterchangeActorFactory.cpp
--- a/Engine/Plugins/Experimental/Interchange/Runtime/Source/Import/Private/Scene/InterchangeActorFactory.cpp
+++ b/Engine/Plugins/Experimental/Interchange/Runtime/Source/Import/Private/Scene/InterchangeActorFactory.cpp
@@ -1,262 +1,30 @@
 // Copyright Epic Games, Inc. All Rights Reserved. 
 #include "Scene/InterchangeActorFactory.h"
 
-#include "InterchangeActorFactoryNode.h"
-#include "InterchangeImportCommon.h"
-#include "InterchangeImportLog.h"
-#include "InterchangeSourceData.h"
-
-#include "CineCameraActor.h"
-#include "CineCameraComponent.h"
-#include "Components/SceneComponent.h"
-#include "Components/SkeletalMeshComponent.h"
-#include "Engine/Engine.h"
-#include "Engine/StaticMesh.h"
-#include "Engine/World.h"
-#include "GameFramework/Actor.h"
-
 #include "Nodes/InterchangeBaseNode.h"
-#include "Nodes/InterchangeBaseNodeContainer.h"
-#include "InterchangeMeshNode.h"
-#include "InterchangeSceneNode.h"
-
-#if WITH_EDITOR
-#include "Editor/EditorEngine.h"
-#endif
-
-namespace UE
-{
-	namespace Interchange
-	{
-		namespace Private
-		{
-			FString FindMeshDependencyUid(const UInterchangeBaseNodeContainer* NodeContainer, const UInterchangeActorFactoryNode* ActorFactoryNode)
-			{
-				FString MeshDependency;
-
-				TArray<FString> ActorTargetNodes;
-				ActorFactoryNode->GetTargetNodeUids(ActorTargetNodes);
-				if (!ActorTargetNodes.IsEmpty())
-				{
-					if (const UInterchangeSceneNode* SceneNode = Cast<UInterchangeSceneNode>(NodeContainer->GetNode(ActorTargetNodes[0])))
-					{
-						FString AssetInstanceUid;
-						SceneNode->GetCustomAssetInstanceUid(AssetInstanceUid);
-
-						if (const UInterchangeMeshNode* MeshNode = Cast<UInterchangeMeshNode>(NodeContainer->GetNode(AssetInstanceUid)))
-						{
-							TArray<FString> MeshTargetNodes;
-							MeshNode->GetTargetNodeUids(MeshTargetNodes);
-
-							if (!MeshTargetNodes.IsEmpty())
-							{
-								MeshDependency = MeshTargetNodes[0];
-							}
-						}
-					}
-				}
-
-				return MeshDependency;
-			}
-
-			void SetupStaticMeshComponent(const UInterchangeBaseNodeContainer* NodeContainer, const UInterchangeActorFactoryNode* ActorFactoryNode, UStaticMeshComponent& StaticMeshComponent)
-			{
-				StaticMeshComponent.UnregisterComponent();
-
-				const FString MeshDependency = FindMeshDependencyUid(NodeContainer, ActorFactoryNode);
-
-				if (const UInterchangeBaseNode* MeshNode = NodeContainer->GetNode(MeshDependency))
-				{
-					if (UStaticMesh* StaticMesh = Cast<UStaticMesh>(MeshNode->ReferenceObject.TryLoad()))
-					{
-						StaticMeshComponent.SetStaticMesh(StaticMesh);
-					}
-				}
-			}
-
-			void SetupSkeletalMeshComponent(const UInterchangeBaseNodeContainer* NodeContainer, const UInterchangeActorFactoryNode* ActorFactoryNode, USkeletalMeshComponent& SkeletalMeshComponent)
-			{
-				SkeletalMeshComponent.UnregisterComponent();
-
-				const FString MeshDependency = FindMeshDependencyUid(NodeContainer, ActorFactoryNode);
-
-				if (const UInterchangeBaseNode* MeshNode = NodeContainer->GetNode(MeshDependency))
-				{
-					if (USkeletalMesh* SkeletalMesh = Cast<USkeletalMesh>(MeshNode->ReferenceObject.TryLoad()))
-					{
-						SkeletalMeshComponent.SetSkeletalMesh(SkeletalMesh);
-					}
-				}
-			}
-
-			AActor* SpawnActor(const FActorSpawnParameters& ActorSpawnParameters, const UInterchangeActorFactoryNode* ActorFactoryNode, AActor* ParentActor)
-			{
-				ensure(ActorFactoryNode);
-
-				if (!ActorFactoryNode)
-				{
-					return nullptr;
-				}
-
-				UWorld* const World = [&ActorSpawnParameters, &ParentActor]()
-				{
-					UWorld* DefaultWorld = nullptr;
-
-					if (ParentActor)
-					{
-						DefaultWorld = ParentActor->GetWorld();
-					}
-#if WITH_EDITOR
-					if (DefaultWorld == nullptr)
-					{
-						UEditorEngine* EditorEngine = Cast<UEditorEngine>(GEngine);
-						if (GIsEditor && EditorEngine != nullptr)
-						{
-							DefaultWorld = EditorEngine->GetEditorWorldContext().World();
-						}
-					}
-#endif
-					if (DefaultWorld == nullptr && GEngine)
-					{
-						DefaultWorld = GEngine->GetWorld();
-					}
-
-					return ActorSpawnParameters.OverrideLevel ? ActorSpawnParameters.OverrideLevel->GetWorld() : DefaultWorld;
-				}();
-
-				if (!World)
-				{
-					return nullptr;
-				}
-
-				FTransform Transform = FTransform::Identity;
-				ActorFactoryNode->GetCustomGlobalTransform(Transform);
+#include "Scene/InterchangeActorHelper.h"
 
-				AActor* SpawnedActor = World->SpawnActor<AActor>(ActorFactoryNode->GetObjectClass(), Transform, ActorSpawnParameters);
-
-				if (SpawnedActor)
-				{
-#if WITH_EDITOR
-					SpawnedActor->SetActorLabel(ActorSpawnParameters.Name.ToString());
-#endif
-					if (!SpawnedActor->GetRootComponent())
-					{
-						USceneComponent* RootComponent = NewObject<USceneComponent>(SpawnedActor, USceneComponent::GetDefaultSceneRootVariableName(), RF_Transactional);
-#if WITH_EDITORONLY_DATA
-						RootComponent->bVisualizeComponent = true;
-#endif
-						RootComponent->SetWorldTransform(Transform);
-
-						SpawnedActor->SetRootComponent(RootComponent);
-						SpawnedActor->AddInstanceComponent(RootComponent);
-					}
-
-					if (USceneComponent* RootComponent = SpawnedActor->GetRootComponent())
-					{
-						uint8 Mobility;
-						if (ActorFactoryNode->GetCustomMobility(Mobility))
-						{
-							//Make sure we don't have a mobility that's more restrictive than our parent mobility, as that wouldn't be a valid setup.
-							EComponentMobility::Type TargetMobility = (EComponentMobility::Type)Mobility;
-
-							if (ParentActor && ParentActor->GetRootComponent())
-							{
-								TargetMobility = (EComponentMobility::Type)FMath::Max((uint8)Mobility, (uint8)ParentActor->GetRootComponent()->Mobility);
-							}
-
-							RootComponent->SetMobility(TargetMobility);
-						}
-					}
-				}
-
-				if (ParentActor)
-				{
-					SpawnedActor->AttachToActor(ParentActor, FAttachmentTransformRules::KeepWorldTransform);
-				}
-
-				return SpawnedActor;
-			}
-
-			void RecursivelySpawnActors(TMap<FString, UObject*>& OutSpawnedActors,
-				const UInterchangeBaseNodeContainer* NodeContainer, const UInterchangeActorFactoryNode* ActorNode, const FActorSpawnParameters& SpawnParameters,
-				AActor* ParentActor, const bool bSpawnChildren)
-			{
-				if (!ActorNode || !NodeContainer)
-				{
-					return;
-				}
-
-				const FString ActorNodeUniqueID = ActorNode->GetUniqueID();
-
-				AActor* SpawnedActor = SpawnActor(SpawnParameters, Cast<UInterchangeActorFactoryNode>(ActorNode), ParentActor);
-
-				if (!SpawnedActor)
-				{
-					return;
-				}
-
-				if (UStaticMeshComponent* StaticMeshComponent = Cast<UStaticMeshComponent>(SpawnedActor->GetRootComponent()))
-				{
-					SetupStaticMeshComponent(NodeContainer, ActorNode, *StaticMeshComponent);
-				}
-				else if (USkeletalMeshComponent* SkeletalMeshComponent = Cast<USkeletalMeshComponent>(SpawnedActor->GetRootComponent()))
-				{
-					SetupSkeletalMeshComponent(NodeContainer, ActorNode, *SkeletalMeshComponent);
-				}
-
-				if (ACineCameraActor* CineCameraActor = Cast<ACineCameraActor>(SpawnedActor))
-				{
-					ActorNode->ApplyAllCustomAttributeToObject(CineCameraActor->GetCineCameraComponent());
-				}
-				else
-				{
-					ActorNode->ApplyAllCustomAttributeToObject(SpawnedActor->GetRootComponent());
-				}
-
-				OutSpawnedActors.Add(ActorNodeUniqueID, SpawnedActor);
-
-				if (bSpawnChildren)
-				{
-					for (int32 ChildIndex = 0; ChildIndex < NodeContainer->GetNodeChildrenCount(ActorNodeUniqueID); ++ChildIndex)
-					{
-						if (const UInterchangeBaseNode* ChildNode = NodeContainer->GetNodeChildren(ActorNodeUniqueID, ChildIndex))
-						{
-							if (const UInterchangeActorFactoryNode* ChildActorFactoryNode = Cast<const UInterchangeActorFactoryNode>(ChildNode))
-							{
-								FActorSpawnParameters ChildActorSpawnParameters;
-								ChildActorSpawnParameters.Name = FName(ChildActorFactoryNode->GetDisplayLabel());
-								ChildActorSpawnParameters.NameMode = FActorSpawnParameters::ESpawnActorNameMode::Requested;
-								ChildActorSpawnParameters.OverrideLevel = SpawnParameters.OverrideLevel;
-								ChildActorSpawnParameters.SpawnCollisionHandlingOverride = ESpawnActorCollisionHandlingMethod::AlwaysSpawn;
-
-								RecursivelySpawnActors(OutSpawnedActors, NodeContainer, ChildActorFactoryNode, ChildActorSpawnParameters, SpawnedActor, bSpawnChildren);
-							}
-						}
-					}
-				}
-			};
-		}
-	}
-}
+#include "GameFramework/Actor.h"
+#include "Components/SceneComponent.h"
 
 UClass* UInterchangeActorFactory::GetFactoryClass() const
 {
 	return AActor::StaticClass();
 }
 
-TMap<FString, UObject*> UInterchangeActorFactory::CreateSceneObjects(const UInterchangeFactoryBase::FCreateSceneObjectsParams& CreateSceneObjectsParams)
+UObject* UInterchangeActorFactory::CreateSceneObject(const UInterchangeFactoryBase::FCreateSceneObjectsParams& CreateSceneObjectsParams)
 {
-	FActorSpawnParameters ActorSpawnParameters;
-	ActorSpawnParameters.Name = FName(*CreateSceneObjectsParams.ObjectName);
-	ActorSpawnParameters.NameMode = FActorSpawnParameters::ESpawnActorNameMode::Requested;
-	ActorSpawnParameters.OverrideLevel = CreateSceneObjectsParams.Level;
-	ActorSpawnParameters.SpawnCollisionHandlingOverride = ESpawnActorCollisionHandlingMethod::AlwaysSpawn;
+	AActor* SpawnedActor = UE::Interchange::ActorHelper::SpawnFactoryActor(CreateSceneObjectsParams);
 
-	TMap<FString, UObject*> SpawnedActors;
-	constexpr AActor* ParentActor = nullptr;
+	if (!SpawnedActor)
+	{
+		return nullptr;
+	}
 
-	UE::Interchange::Private::RecursivelySpawnActors(SpawnedActors, CreateSceneObjectsParams.NodeContainer, Cast<const UInterchangeActorFactoryNode>(CreateSceneObjectsParams.ObjectNode),
-		ActorSpawnParameters, ParentActor, CreateSceneObjectsParams.bCreateSceneObjectsForChildren);
+	if (USceneComponent* RootComponent = SpawnedActor->GetRootComponent())
+	{
+		CreateSceneObjectsParams.FactoryNode->ApplyAllCustomAttributeToObject(RootComponent);
+	}
 
-	return SpawnedActors;
-}
+	return SpawnedActor;
+}
\ No newline at end of file
diff --git a/Engine/Plugins/Experimental/Interchange/Runtime/Source/Import/Private/Scene/InterchangeActorHelper.cpp b/Engine/Plugins/Experimental/Interchange/Runtime/Source/Import/Private/Scene/InterchangeActorHelper.cpp
new file mode 100644
--- /dev/null
+++ b/Engine/Plugins/Experimental/Interchange/Runtime/Source/Import/Private/Scene/InterchangeActorHelper.cpp
@@ -0,0 +1,147 @@
+// Copyright Epic Games, Inc. All Rights Reserved. 
+
+#include "Scene/InterchangeActorHelper.h"
+
+#include "InterchangeActorFactoryNode.h"
+#include "InterchangeMeshNode.h"
+#include "InterchangeSceneNode.h"
+#include "Nodes/InterchangeBaseNode.h"
+#include "Nodes/InterchangeBaseNodeContainer.h"
+
+#include "Engine/World.h"
+
+#if WITH_EDITOR
+#include "Editor/EditorEngine.h"
+#endif
+
+AActor* UE::Interchange::ActorHelper::GetSpawnedParentActor(const UInterchangeBaseNodeContainer* NodeContainer, const UInterchangeActorFactoryNode* FactoryNode)
+{
+	AActor* ParentActor = nullptr;
+
+	if (const UInterchangeFactoryBaseNode* ParentFactoryNode = Cast<UInterchangeFactoryBaseNode>(NodeContainer->GetNode(FactoryNode->GetParentUid())))
+	{
+		ParentActor = Cast<AActor>(ParentFactoryNode->ReferenceObject.TryLoad());
+	}
+
+	return  ParentActor;
+}
+
+AActor* UE::Interchange::ActorHelper::SpawnFactoryActor(const UInterchangeFactoryBase::FCreateSceneObjectsParams& CreateSceneObjectsParams)
+{
+	FActorSpawnParameters SpawnParameters;
+	SpawnParameters.Name = FName(*CreateSceneObjectsParams.ObjectName);
+	SpawnParameters.NameMode = FActorSpawnParameters::ESpawnActorNameMode::Requested;
+	SpawnParameters.OverrideLevel = CreateSceneObjectsParams.Level;
+	SpawnParameters.SpawnCollisionHandlingOverride = ESpawnActorCollisionHandlingMethod::AlwaysSpawn;
+
+	const UInterchangeActorFactoryNode* FactoryNode = Cast<UInterchangeActorFactoryNode>(CreateSceneObjectsParams.FactoryNode);
+	const UInterchangeBaseNodeContainer* NodeContainer = CreateSceneObjectsParams.NodeContainer;
+	AActor* ParentActor = UE::Interchange::ActorHelper::GetSpawnedParentActor(NodeContainer, FactoryNode);
+
+	if (!FactoryNode || !NodeContainer)
+	{
+		return nullptr;
+	}
+
+	UWorld* const World = [&SpawnParameters, &ParentActor]()
+	{
+		UWorld* DefaultWorld = nullptr;
+
+		if (ParentActor)
+		{
+			DefaultWorld = ParentActor->GetWorld();
+		}
+#if WITH_EDITOR
+		if (DefaultWorld == nullptr)
+		{
+			UEditorEngine* EditorEngine = Cast<UEditorEngine>(GEngine);
+			if (GIsEditor && EditorEngine != nullptr)
+			{
+				DefaultWorld = EditorEngine->GetEditorWorldContext().World();
+			}
+		}
+#endif
+		if (DefaultWorld == nullptr && GEngine)
+		{
+			DefaultWorld = GEngine->GetWorld();
+		}
+
+		return SpawnParameters.OverrideLevel ? SpawnParameters.OverrideLevel->GetWorld() : DefaultWorld;
+	}();
+
+	if (!World)
+	{
+		return nullptr;
+	}
+
+	FTransform Transform = FTransform::Identity;
+	FactoryNode->GetCustomGlobalTransform(Transform);
+
+	UClass* ActorClass = FactoryNode->GetObjectClass();
+	AActor* SpawnedActor = World->SpawnActor<AActor>(ActorClass, Transform, SpawnParameters);
+
+	if (SpawnedActor)
+	{
+#if WITH_EDITOR
+		SpawnedActor->SetActorLabel(SpawnParameters.Name.ToString());
+#endif
+		if (!SpawnedActor->GetRootComponent())
+		{
+			USceneComponent* RootComponent = NewObject<USceneComponent>(SpawnedActor, USceneComponent::GetDefaultSceneRootVariableName(), RF_Transactional);
+#if WITH_EDITORONLY_DATA
+			RootComponent->bVisualizeComponent = true;
+#endif
+			RootComponent->SetWorldTransform(Transform);
+
+			SpawnedActor->SetRootComponent(RootComponent);
+			SpawnedActor->AddInstanceComponent(RootComponent);
+		}
+
+		if (USceneComponent* RootComponent = SpawnedActor->GetRootComponent())
+		{
+			uint8 Mobility;
+			if (FactoryNode->GetCustomMobility(Mobility))
+			{
+				//Make sure we don't have a mobility that's more restrictive than our parent mobility, as that wouldn't be a valid setup.
+				EComponentMobility::Type TargetMobility = (EComponentMobility::Type)Mobility;
+
+				if (ParentActor && ParentActor->GetRootComponent())
+				{
+					TargetMobility = (EComponentMobility::Type)FMath::Max((uint8)Mobility, (uint8)ParentActor->GetRootComponent()->Mobility);
+				}
+
+				RootComponent->SetMobility(TargetMobility);
+			}
+		}
+	}
+
+	if (ParentActor)
+	{
+		SpawnedActor->AttachToActor(ParentActor, FAttachmentTransformRules::KeepWorldTransform);
+	}
+
+	return SpawnedActor;
+}
+
+const UInterchangeBaseNode* UE::Interchange::ActorHelper::FindAssetInstanceFactoryNode(const UInterchangeBaseNodeContainer* NodeContainer, const UInterchangeBaseNode* ActorFactoryNode)
+{
+	TArray<FString> ActorTargetNodes;
+	ActorFactoryNode->GetTargetNodeUids(ActorTargetNodes);
+	const UInterchangeSceneNode* SceneNode = ActorTargetNodes.IsEmpty() ? nullptr : Cast<UInterchangeSceneNode>(NodeContainer->GetNode(ActorTargetNodes[0]));
+	if (!SceneNode)
+	{
+		return nullptr;
+	}
+
+	FString AssetInstanceUid;
+	SceneNode->GetCustomAssetInstanceUid(AssetInstanceUid);
+	const UInterchangeBaseNode* AssetNode = NodeContainer->GetNode(AssetInstanceUid);
+	if (!AssetNode)
+	{
+		return nullptr;
+	}
+
+	TArray<FString> AssetTargetNodeIds;
+	AssetNode->GetTargetNodeUids(AssetTargetNodeIds);
+	return AssetTargetNodeIds.IsEmpty() ? nullptr : NodeContainer->GetNode(AssetTargetNodeIds[0]);
+}
\ No newline at end of file
diff --git a/Engine/Plugins/Experimental/Interchange/Runtime/Source/Import/Private/Scene/InterchangeCineCameraActorFactory.cpp b/Engine/Plugins/Experimental/Interchange/Runtime/Source/Import/Private/Scene/InterchangeCineCameraActorFactory.cpp
new file mode 100644
--- /dev/null
+++ b/Engine/Plugins/Experimental/Interchange/Runtime/Source/Import/Private/Scene/InterchangeCineCameraActorFactory.cpp
@@ -0,0 +1,36 @@
+// Copyright Epic Games, Inc. All Rights Reserved.
+
+#include "Scene/InterchangeCineCameraActorFactory.h"
+
+#include "InterchangeCineCameraFactoryNode.h"
+#include "Scene/InterchangeActorFactory.h"
+#include "Scene/InterchangeActorHelper.h"
+
+#include "CineCameraActor.h"
+#include "CineCameraComponent.h"
+#include "CoreMinimal.h"
+#include "UObject/Object.h"
+#include "UObject/ObjectMacros.h"
+
+UClass* UInterchangeCineCameraActorFactory::GetFactoryClass() const
+{
+	return ACineCameraActor::StaticClass();
+}
+
+UObject* UInterchangeCineCameraActorFactory::CreateSceneObject(const UInterchangeFactoryBase::FCreateSceneObjectsParams& CreateSceneObjectsParams)
+{
+	ACineCameraActor* SpawnedActor = Cast<ACineCameraActor>(UE::Interchange::ActorHelper::SpawnFactoryActor(CreateSceneObjectsParams));
+
+	if (!SpawnedActor)
+	{
+		return nullptr;
+	}
+
+	if (UCineCameraComponent* CineCameraComponent = SpawnedActor->GetCineCameraComponent())
+	{
+		CreateSceneObjectsParams.FactoryNode->ApplyAllCustomAttributeToObject(CineCameraComponent);
+	}
+
+	return SpawnedActor;
+}
+
diff --git a/Engine/Plugins/Experimental/Interchange/Runtime/Source/Import/Private/Scene/InterchangeSkeletalMeshActorFactory.cpp b/Engine/Plugins/Experimental/Interchange/Runtime/Source/Import/Private/Scene/InterchangeSkeletalMeshActorFactory.cpp
new file mode 100644
--- /dev/null
+++ b/Engine/Plugins/Experimental/Interchange/Runtime/Source/Import/Private/Scene/InterchangeSkeletalMeshActorFactory.cpp
@@ -0,0 +1,50 @@
+// Copyright Epic Games, Inc. All Rights Reserved. 
+#include "Scene/InterchangeSkeletalMeshActorFactory.h"
+
+#include "InterchangeActorFactoryNode.h"
+#include "Scene/InterchangeActorHelper.h"
+
+#include "Animation/SkeletalMeshActor.h"
+#include "Components/SkeletalMeshComponent.h"
+
+#include "Nodes/InterchangeBaseNodeContainer.h"
+
+
+UObject* UInterchangeSkeletalMeshActorFactory::CreateSceneObject(const UInterchangeFactoryBase::FCreateSceneObjectsParams& CreateSceneObjectsParams)
+{
+	ASkeletalMeshActor* SpawnedActor = Cast<ASkeletalMeshActor>(UE::Interchange::ActorHelper::SpawnFactoryActor(CreateSceneObjectsParams));
+
+	if (!SpawnedActor)
+	{
+		return nullptr;
+	}
+
+	UInterchangeBaseNode* FactoryNode = CreateSceneObjectsParams.FactoryNode;
+	SetupSkeletalMeshActor(CreateSceneObjectsParams.NodeContainer, FactoryNode, SpawnedActor);
+
+	if (USkeletalMeshComponent* SkeletalMeshComponent = SpawnedActor->GetSkeletalMeshComponent())
+	{
+		FactoryNode->ApplyAllCustomAttributeToObject(SkeletalMeshComponent);
+	}
+
+	return SpawnedActor;
+};
+
+UClass* UInterchangeSkeletalMeshActorFactory::GetFactoryClass() const
+{
+	return ASkeletalMeshActor::StaticClass();
+}
+
+void UInterchangeSkeletalMeshActorFactory::SetupSkeletalMeshActor(const UInterchangeBaseNodeContainer* NodeContainer, const UInterchangeBaseNode* ActorFactoryNode, ASkeletalMeshActor* SkeletalMeshActor)
+{
+	USkeletalMeshComponent* SkeletalMeshComponent = SkeletalMeshActor->GetSkeletalMeshComponent();
+	SkeletalMeshComponent->UnregisterComponent();
+
+	if (const UInterchangeBaseNode* MeshNode = UE::Interchange::ActorHelper::FindAssetInstanceFactoryNode(NodeContainer, ActorFactoryNode))
+	{
+		if (USkeletalMesh* SkeletalMesh = Cast<USkeletalMesh>(MeshNode->ReferenceObject.TryLoad()))
+		{
+			SkeletalMeshComponent->SetSkeletalMesh(SkeletalMesh);
+		}
+	}
+}
\ No newline at end of file
diff --git a/Engine/Plugins/Experimental/Interchange/Runtime/Source/Import/Private/Scene/InterchangeStaticMeshActorFactory.cpp b/Engine/Plugins/Experimental/Interchange/Runtime/Source/Import/Private/Scene/InterchangeStaticMeshActorFactory.cpp
new file mode 100644
--- /dev/null
+++ b/Engine/Plugins/Experimental/Interchange/Runtime/Source/Import/Private/Scene/InterchangeStaticMeshActorFactory.cpp
@@ -0,0 +1,49 @@
+// Copyright Epic Games, Inc. All Rights Reserved. 
+#include "Scene/InterchangeStaticMeshActorFactory.h"
+
+#include "InterchangeActorFactoryNode.h"
+#include "Nodes/InterchangeBaseNodeContainer.h"
+#include "Scene/InterchangeActorHelper.h"
+
+#include "Components/StaticMeshComponent.h"
+#include "Engine/StaticMesh.h"
+#include "Engine/StaticMeshActor.h"
+
+UClass* UInterchangeStaticMeshActorFactory::GetFactoryClass() const
+{
+	return AStaticMeshActor::StaticClass();
+}
+
+UObject* UInterchangeStaticMeshActorFactory::CreateSceneObject(const UInterchangeFactoryBase::FCreateSceneObjectsParams& CreateSceneObjectsParams)
+{
+	AStaticMeshActor* SpawnedActor = Cast<AStaticMeshActor>(UE::Interchange::ActorHelper::SpawnFactoryActor(CreateSceneObjectsParams));
+
+	if (!SpawnedActor)
+	{
+		return nullptr;
+	}
+
+	UInterchangeBaseNode* FactoryNode = CreateSceneObjectsParams.FactoryNode;
+	SetupStaticMeshActor(CreateSceneObjectsParams.NodeContainer, FactoryNode, SpawnedActor);
+
+	if (UStaticMeshComponent* StaticMeshComponent = SpawnedActor->GetStaticMeshComponent())
+	{
+		FactoryNode->ApplyAllCustomAttributeToObject(StaticMeshComponent);
+	}
+
+	return SpawnedActor;
+};
+
+void UInterchangeStaticMeshActorFactory::SetupStaticMeshActor(const UInterchangeBaseNodeContainer* NodeContainer, const UInterchangeBaseNode* ActorFactoryNode, AStaticMeshActor* StaticMeshActor)
+{
+	UStaticMeshComponent* StaticMeshComponent = StaticMeshActor->GetStaticMeshComponent();
+	StaticMeshComponent->UnregisterComponent();
+
+	if (const UInterchangeBaseNode* MeshNode = UE::Interchange::ActorHelper::FindAssetInstanceFactoryNode(NodeContainer, ActorFactoryNode))
+	{
+		if (UStaticMesh* StaticMesh = Cast<UStaticMesh>(MeshNode->ReferenceObject.TryLoad()))
+		{
+			StaticMeshComponent->SetStaticMesh(StaticMesh);
+		}
+	}
+}
\ No newline at end of file
diff --git a/Engine/Plugins/Experimental/Interchange/Runtime/Source/Import/Public/Scene/InterchangeActorFactory.h b/Engine/Plugins/Experimental/Interchange/Runtime/Source/Import/Public/Scene/InterchangeActorFactory.h
--- a/Engine/Plugins/Experimental/Interchange/Runtime/Source/Import/Public/Scene/InterchangeActorFactory.h
+++ b/Engine/Plugins/Experimental/Interchange/Runtime/Source/Import/Public/Scene/InterchangeActorFactory.h
@@ -20,9 +20,9 @@ public:
 
 	virtual UClass* GetFactoryClass() const override;
 
-	virtual TMap<FString, UObject*> CreateSceneObjects(const UInterchangeFactoryBase::FCreateSceneObjectsParams& CreateSceneObjectsParams) override;
+	virtual UObject* CreateSceneObject(const UInterchangeFactoryBase::FCreateSceneObjectsParams& CreateSceneObjectsParams) override;
 
-	virtual bool CanExecuteOnAnyThread() const
+	virtual bool CanExecuteOnAnyThread() const override
 	{
 		return false;
 	}
diff --git a/Engine/Plugins/Experimental/Interchange/Runtime/Source/Import/Public/Scene/InterchangeActorHelper.h b/Engine/Plugins/Experimental/Interchange/Runtime/Source/Import/Public/Scene/InterchangeActorHelper.h
new file mode 100644
--- /dev/null
+++ b/Engine/Plugins/Experimental/Interchange/Runtime/Source/Import/Public/Scene/InterchangeActorHelper.h
@@ -0,0 +1,26 @@
+// Copyright Epic Games, Inc. All Rights Reserved. 
+
+#include "InterchangeFactoryBase.h"
+
+class UInterchangeBaseNode;
+class UInterchangeBaseNodeContainer;
+class UInterchangeActorFactoryNode;
+
+namespace UE::Interchange::ActorHelper
+{
+	/**
+	 * Returns the parent actor of the actor described by the FactoryNode using the node hierarchy.
+	 */
+	AActor* GetSpawnedParentActor(const UInterchangeBaseNodeContainer* NodeContainer, const UInterchangeActorFactoryNode* FactoryNode);
+
+	/**
+	 * Spawns and return a AAActor using the given SceneObjectsParam.
+	 * The actor's parent will be determine by the node hierarchy and its type by the factory node's ObjectClass.
+	 */
+	AActor* SpawnFactoryActor(const UInterchangeFactoryBase::FCreateSceneObjectsParams& CreateSceneObjectsParams);
+	
+	/**
+	 * Returns the factory node of the asset instanced by ActorFactoryNode.
+	 */
+	const UInterchangeBaseNode* FindAssetInstanceFactoryNode(const UInterchangeBaseNodeContainer* NodeContainer, const UInterchangeBaseNode* ActorFactoryNode);
+}
diff --git a/Engine/Plugins/Experimental/Interchange/Runtime/Source/Import/Public/Scene/InterchangeCineCameraActorFactory.h b/Engine/Plugins/Experimental/Interchange/Runtime/Source/Import/Public/Scene/InterchangeCineCameraActorFactory.h
new file mode 100644
--- /dev/null
+++ b/Engine/Plugins/Experimental/Interchange/Runtime/Source/Import/Public/Scene/InterchangeCineCameraActorFactory.h
@@ -0,0 +1,34 @@
+// Copyright Epic Games, Inc. All Rights Reserved.
+
+#pragma once
+
+#include "CoreMinimal.h"
+#include "InterchangeActorFactory.h"
+#include "UObject/Object.h"
+#include "UObject/ObjectMacros.h"
+
+#include "InterchangeCineCameraActorFactory.generated.h"
+
+UCLASS(BlueprintType, Experimental)
+class INTERCHANGEIMPORT_API UInterchangeCineCameraActorFactory : public UInterchangeFactoryBase
+{
+	GENERATED_BODY()
+public:
+
+	//////////////////////////////////////////////////////////////////////////
+	// Interchange factory base interface begin
+
+	virtual UClass* GetFactoryClass() const override;
+
+	virtual UObject* CreateSceneObject(const UInterchangeFactoryBase::FCreateSceneObjectsParams& CreateSceneObjectsParams) override;
+
+	virtual bool CanExecuteOnAnyThread() const override
+	{
+		return false; 
+	}
+
+	// Interchange factory base interface end
+	//////////////////////////////////////////////////////////////////////////
+};
+
+
diff --git a/Engine/Plugins/Experimental/Interchange/Runtime/Source/Import/Public/Scene/InterchangeSkeletalMeshActorFactory.h b/Engine/Plugins/Experimental/Interchange/Runtime/Source/Import/Public/Scene/InterchangeSkeletalMeshActorFactory.h
new file mode 100644
--- /dev/null
+++ b/Engine/Plugins/Experimental/Interchange/Runtime/Source/Import/Public/Scene/InterchangeSkeletalMeshActorFactory.h
@@ -0,0 +1,42 @@
+// Copyright Epic Games, Inc. All Rights Reserved.
+
+#pragma once
+
+#include "CoreMinimal.h"
+#include "InterchangeActorFactory.h"
+#include "UObject/Object.h"
+#include "UObject/ObjectMacros.h"
+
+#include "InterchangeSkeletalMeshActorFactory.generated.h"
+
+class AActor;
+class ASkeletalMeshActor;
+class UInterchangeActorFactoryNode;
+
+UCLASS(BlueprintType, Experimental)
+class INTERCHANGEIMPORT_API UInterchangeSkeletalMeshActorFactory : public UInterchangeActorFactory
+{
+	GENERATED_BODY()
+public:
+
+	//////////////////////////////////////////////////////////////////////////
+	// Interchange factory base interface begin
+
+	virtual UClass* GetFactoryClass() const override;
+
+	virtual UObject* CreateSceneObject(const UInterchangeFactoryBase::FCreateSceneObjectsParams& CreateSceneObjectsParams) override;
+
+	virtual bool CanExecuteOnAnyThread() const override
+	{
+		return false;
+	}
+
+	// Interchange factory base interface end
+	//////////////////////////////////////////////////////////////////////////
+
+private:
+
+	void SetupSkeletalMeshActor(const UInterchangeBaseNodeContainer* NodeContainer, const UInterchangeBaseNode* ActorFactoryNode, ASkeletalMeshActor* SkeletalMeshActor);
+};
+
+
diff --git a/Engine/Plugins/Experimental/Interchange/Runtime/Source/Import/Public/Scene/InterchangeStaticMeshActorFactory.h b/Engine/Plugins/Experimental/Interchange/Runtime/Source/Import/Public/Scene/InterchangeStaticMeshActorFactory.h
new file mode 100644
--- /dev/null
+++ b/Engine/Plugins/Experimental/Interchange/Runtime/Source/Import/Public/Scene/InterchangeStaticMeshActorFactory.h
@@ -0,0 +1,42 @@
+// Copyright Epic Games, Inc. All Rights Reserved.
+
+#pragma once
+
+#include "CoreMinimal.h"
+#include "InterchangeActorFactory.h"
+#include "UObject/Object.h"
+#include "UObject/ObjectMacros.h"
+
+#include "InterchangeStaticMeshActorFactory.generated.h"
+
+class AActor;
+class AStaticMeshActor;
+class UInterchangeActorFactoryNode;
+
+UCLASS(BlueprintType, Experimental)
+class INTERCHANGEIMPORT_API UInterchangeStaticMeshActorFactory : public UInterchangeFactoryBase
+{
+	GENERATED_BODY()
+public:
+
+	//////////////////////////////////////////////////////////////////////////
+	// Interchange factory base interface begin
+
+	virtual UClass* GetFactoryClass() const override;
+
+	virtual UObject* CreateSceneObject(const UInterchangeFactoryBase::FCreateSceneObjectsParams& CreateSceneObjectsParams) override;
+
+	virtual bool CanExecuteOnAnyThread() const override
+	{
+		return false; 
+	}
+
+	// Interchange factory base interface end
+	//////////////////////////////////////////////////////////////////////////
+
+private:
+
+	void SetupStaticMeshActor(const UInterchangeBaseNodeContainer* NodeContainer, const UInterchangeBaseNode* ActorFactoryNode, AStaticMeshActor* StaticMeshActor);
+};
+
+
diff --git a/Engine/Plugins/Experimental/Interchange/Runtime/Source/Pipelines/Private/InterchangeGenericScenesPipeline.cpp b/Engine/Plugins/Experimental/Interchange/Runtime/Source/Pipelines/Private/InterchangeGenericScenesPipeline.cpp
--- a/Engine/Plugins/Experimental/Interchange/Runtime/Source/Pipelines/Private/InterchangeGenericScenesPipeline.cpp
+++ b/Engine/Plugins/Experimental/Interchange/Runtime/Source/Pipelines/Private/InterchangeGenericScenesPipeline.cpp
@@ -98,9 +98,11 @@ void UInterchangeGenericLevelPipeline::CreateActorFactoryNode(UInterchangeBaseNo
 	if (!SceneNode->GetParentUid().IsEmpty())
 	{
 		ActorFactoryNode->SetParentUid(TEXT("Factory_") + SceneNode->GetParentUid());
+		ActorFactoryNode->AddFactoryDependencyUid(ActorFactoryNode->GetParentUid());
 	}
 
 	ActorFactoryNode->AddTargetNodeUid(SceneNode->GetUniqueID());
+	SceneNode->AddTargetNodeUid(ActorFactoryNode->GetUniqueID());
 
 	FTransform GlobalTransform;
 	if (SceneNode->GetCustomGlobalTransform(InBaseNodeContainer, GlobalTransform))
diff --git a/Engine/Source/Runtime/Experimental/Interchange/Core/Public/InterchangeFactoryBase.h b/Engine/Source/Runtime/Experimental/Interchange/Core/Public/InterchangeFactoryBase.h
--- a/Engine/Source/Runtime/Experimental/Interchange/Core/Public/InterchangeFactoryBase.h
+++ b/Engine/Source/Runtime/Experimental/Interchange/Core/Public/InterchangeFactoryBase.h
@@ -102,25 +102,21 @@ public:
 		FString ObjectName;
 
 		/** The base node that describe how to create the asset */
-		UInterchangeBaseNode* ObjectNode = nullptr;
+		UInterchangeBaseNode* FactoryNode = nullptr;
 
 		/** The node container associated with the current source index */
 		const UInterchangeBaseNodeContainer* NodeContainer = nullptr;
-
-		/** Whether to create the scene objects for the child nodes or not */
-		bool bCreateSceneObjectsForChildren = false;
 	};
 
 	/**
 	 * Creates the scene object from a Scene Node data.
-	 * If FCreateSceneObjectsParams::bCreateSceneObjectsForChildren is true, will also create the scene objects for our children.
 	 *
 	 * @param Arguments - The structure containing all necessary arguments, see the structure definition for the documentation.
-	 * @return The node uids and the scene objects that we created from them.
+	 * @return The scene object or nullptr if the operation was unsuccessful.
 	 */
-	virtual TMap<FString, UObject*> CreateSceneObjects(const FCreateSceneObjectsParams& Arguments)
+	virtual UObject* CreateSceneObject(const FCreateSceneObjectsParams& Arguments)
 	{
-		return {};
+		return nullptr;
 	}
 
 	/** Return true if the factory can create the asset asynchronously on any thread, false if it need to be on the main thread */
diff --git a/Engine/Source/Runtime/Experimental/Interchange/Engine/Private/Tasks/InterchangeTaskCreateSceneObjects.cpp b/Engine/Source/Runtime/Experimental/Interchange/Engine/Private/Tasks/InterchangeTaskCreateSceneObjects.cpp
--- a/Engine/Source/Runtime/Experimental/Interchange/Engine/Private/Tasks/InterchangeTaskCreateSceneObjects.cpp
+++ b/Engine/Source/Runtime/Experimental/Interchange/Engine/Private/Tasks/InterchangeTaskCreateSceneObjects.cpp
@@ -21,21 +21,14 @@
 #include "UObject/UObjectGlobals.h"
 #include "UObject/WeakObjectPtrTemplates.h"
 
-UE::Interchange::FTaskCreateSceneObjects::FTaskCreateSceneObjects(const FString& InPackageBasePath, const int32 InSourceIndex, TWeakPtr<FImportAsyncHelper, ESPMode::ThreadSafe> InAsyncHelper,
-	TArrayView<UInterchangeBaseNode*> InNodes, UInterchangeFactoryBase* InFactory, bool bInCreateSceneObjectsForChildren)
+UE::Interchange::FTaskCreateSceneObjects::FTaskCreateSceneObjects(const FString& InPackageBasePath, const int32 InSourceIndex, TWeakPtr<FImportAsyncHelper> InAsyncHelper, TArrayView<UInterchangeBaseNode*> InNodes, const UClass* InFactoryClass)
 	: PackageBasePath(InPackageBasePath)
 	, SourceIndex(InSourceIndex)
 	, WeakAsyncHelper(InAsyncHelper)
 	, Nodes(InNodes)
-	, Factory(InFactory)
-	, bCreateSceneObjectsForChildren(bInCreateSceneObjectsForChildren)
+	, FactoryClass(InFactoryClass)
 {
-	check(Factory);
-}
-
-ENamedThreads::Type UE::Interchange::FTaskCreateSceneObjects::GetDesiredThread()
-{
-	return Factory->CanExecuteOnAnyThread() ? ENamedThreads::AnyBackgroundThreadNormalTask : ENamedThreads::GameThread;
+	check(FactoryClass);
 }
 
 void UE::Interchange::FTaskCreateSceneObjects::DoTask(ENamedThreads::Type CurrentThread, const FGraphEventRef& MyCompletionGraphEvent)
@@ -50,7 +43,7 @@ void UE::Interchange::FTaskCreateSceneObjects::DoTask(ENamedThreads::Type Curren
 		GCScopeGuard.Emplace();
 	}
 
-	TSharedPtr<UE::Interchange::FImportAsyncHelper, ESPMode::ThreadSafe> AsyncHelper = WeakAsyncHelper.Pin();
+	TSharedPtr<UE::Interchange::FImportAsyncHelper> AsyncHelper = WeakAsyncHelper.Pin();
 	check(WeakAsyncHelper.IsValid());
 
 	//Verify if the task was canceled
@@ -61,49 +54,45 @@ void UE::Interchange::FTaskCreateSceneObjects::DoTask(ENamedThreads::Type Curren
 
 	for (UInterchangeBaseNode* Node : Nodes)
 	{
+		UInterchangeFactoryBase* Factory = NewObject<UInterchangeFactoryBase>(GetTransientPackage(), FactoryClass);
+		Factory->SetResultsContainer(AsyncHelper->AssetImportResult->GetResults());
+		{
+			FScopeLock Lock(&AsyncHelper->CreatedFactoriesLock);
+			AsyncHelper->CreatedFactories.Add(Node->GetUniqueID(), Factory);
+		}
+
 		FString NodeDisplayName = Node->GetDisplayLabel();
 		SanitizeObjectName(NodeDisplayName);
 
 		UInterchangeFactoryBase::FCreateSceneObjectsParams CreateSceneObjectsParams;
 		CreateSceneObjectsParams.ObjectName = NodeDisplayName;
-		CreateSceneObjectsParams.ObjectNode = Node;
+		CreateSceneObjectsParams.FactoryNode = Node;
 		CreateSceneObjectsParams.Level = GWorld->GetCurrentLevel();
-		CreateSceneObjectsParams.bCreateSceneObjectsForChildren = bCreateSceneObjectsForChildren;
 
 		if (AsyncHelper->BaseNodeContainers.IsValidIndex(SourceIndex))
 		{
 			CreateSceneObjectsParams.NodeContainer = AsyncHelper->BaseNodeContainers[SourceIndex].Get();
 		}
 
-		TMap<FString, UObject*> SceneObjects = Factory->CreateSceneObjects(CreateSceneObjectsParams);
-
-		for (const TPair<FString, UObject*>& SceneObject : SceneObjects)
+		UObject* SceneObject = Factory->CreateSceneObject(CreateSceneObjectsParams);
+		if (SceneObject)
 		{
-			if (UObject* NodeObject = SceneObject.Value)
-			{
-				FScopeLock Lock(&AsyncHelper->ImportedSceneObjectsPerSourceIndexLock);
-				TArray<UE::Interchange::FImportAsyncHelper::FImportedObjectInfo>& ImportedInfos = AsyncHelper->ImportedSceneObjectsPerSourceIndex.FindOrAdd(SourceIndex);
-				UE::Interchange::FImportAsyncHelper::FImportedObjectInfo* ImportedInfoPtr = ImportedInfos.FindByPredicate([NodeObject](const UE::Interchange::FImportAsyncHelper::FImportedObjectInfo& CurInfo)
-					{
-						return CurInfo.ImportedObject == NodeObject;
-					});
-
-				if (!ImportedInfoPtr)
+			FScopeLock Lock(&AsyncHelper->ImportedSceneObjectsPerSourceIndexLock);
+			TArray<UE::Interchange::FImportAsyncHelper::FImportedObjectInfo>& ImportedInfos = AsyncHelper->ImportedSceneObjectsPerSourceIndex.FindOrAdd(SourceIndex);
+			UE::Interchange::FImportAsyncHelper::FImportedObjectInfo* ImportedInfoPtr = ImportedInfos.FindByPredicate([SceneObject](const UE::Interchange::FImportAsyncHelper::FImportedObjectInfo& CurInfo)
 				{
-					UE::Interchange::FImportAsyncHelper::FImportedObjectInfo& ObjectInfo = ImportedInfos.AddDefaulted_GetRef();
-					ObjectInfo.ImportedObject = NodeObject;
-					ObjectInfo.Factory = Factory;
-					ObjectInfo.FactoryNode = Node;
-				}
+					return CurInfo.ImportedObject == SceneObject;
+				});
 
-				if (CreateSceneObjectsParams.NodeContainer)
-				{
-					if (const UInterchangeBaseNode* ActorNode = CreateSceneObjectsParams.NodeContainer->GetNode(SceneObject.Key))
-					{
-						ActorNode->ReferenceObject = FSoftObjectPath(SceneObject.Value);
-					}
-				}
+			if (!ImportedInfoPtr)
+			{
+				UE::Interchange::FImportAsyncHelper::FImportedObjectInfo& ObjectInfo = ImportedInfos.AddDefaulted_GetRef();
+				ObjectInfo.ImportedObject = SceneObject;
+				ObjectInfo.Factory = Factory;
+				ObjectInfo.FactoryNode = Node;
 			}
+
+			Node->ReferenceObject = FSoftObjectPath(SceneObject);
 		}
 	}
 }
\ No newline at end of file
diff --git a/Engine/Source/Runtime/Experimental/Interchange/Engine/Private/Tasks/InterchangeTaskCreateSceneObjects.h b/Engine/Source/Runtime/Experimental/Interchange/Engine/Private/Tasks/InterchangeTaskCreateSceneObjects.h
--- a/Engine/Source/Runtime/Experimental/Interchange/Engine/Private/Tasks/InterchangeTaskCreateSceneObjects.h
+++ b/Engine/Source/Runtime/Experimental/Interchange/Engine/Private/Tasks/InterchangeTaskCreateSceneObjects.h
@@ -9,7 +9,6 @@
 #include "Nodes/InterchangeBaseNode.h"
 
 class UInterchangeBaseNode;
-class UInterchangeFactoryBase;
 
 namespace UE
 {
@@ -22,16 +21,19 @@ namespace UE
 		private:
 			FString PackageBasePath;
 			int32 SourceIndex;
-			TWeakPtr<FImportAsyncHelper, ESPMode::ThreadSafe> WeakAsyncHelper;
+			TWeakPtr<FImportAsyncHelper> WeakAsyncHelper;
 			TArray<UInterchangeBaseNode*> Nodes;
-			UInterchangeFactoryBase* Factory;
-			bool bCreateSceneObjectsForChildren = false;
+			const UClass* FactoryClass;
 
 		public:
-			explicit FTaskCreateSceneObjects(const FString& InPackageBasePath, const int32 InSourceIndex, TWeakPtr<FImportAsyncHelper, ESPMode::ThreadSafe> InAsyncHelper, TArrayView<UInterchangeBaseNode*> InNodes,
-				UInterchangeFactoryBase* InFactory, bool bInCreateSceneObjectsForChildren);
+			explicit FTaskCreateSceneObjects(const FString& InPackageBasePath, const int32 InSourceIndex, TWeakPtr<FImportAsyncHelper> InAsyncHelper, TArrayView<UInterchangeBaseNode*> InNodes, const UClass* InFactoryClass);
 
-			ENamedThreads::Type GetDesiredThread();
+			ENamedThreads::Type GetDesiredThread()
+			{
+				// We are creating the factories in this task so it must execute on the GameThread.
+				// Also, there are no "CreatePackage Task" equivalent for scene objects right now, so the factories must create those on the game thread.
+				return ENamedThreads::GameThread;
+			}
 
 			static ESubsequentsMode::Type GetSubsequentsMode()
 			{
diff --git a/Engine/Source/Runtime/Experimental/Interchange/Engine/Private/Tasks/InterchangeTaskParsing.cpp b/Engine/Source/Runtime/Experimental/Interchange/Engine/Private/Tasks/InterchangeTaskParsing.cpp
--- a/Engine/Source/Runtime/Experimental/Interchange/Engine/Private/Tasks/InterchangeTaskParsing.cpp
+++ b/Engine/Source/Runtime/Experimental/Interchange/Engine/Private/Tasks/InterchangeTaskParsing.cpp
@@ -116,6 +116,7 @@ void UE::Interchange::FTaskParsing::DoTask(ENamedThreads::Type CurrentThread, co
 			{
 				continue;
 			}
+			const bool bCanImportSceneNode = AsyncHelper->TaskData.ImportType == EImportType::ImportType_Scene;
 			BaseNodeContainer->IterateNodes([&](const FString& NodeUID, UInterchangeBaseNode* Node)
 			{
 				if (!Node->IsEnabled())
@@ -129,13 +130,6 @@ void UE::Interchange::FTaskParsing::DoTask(ENamedThreads::Type CurrentThread, co
 					const UClass* RegisteredFactoryClass = InterchangeManager->GetRegisteredFactoryClass(Node->GetObjectClass());
 
 					const bool bIsAsset = !(Node->GetObjectClass()->IsChildOf<AActor>() || Node->GetObjectClass()->IsChildOf<UActorComponent>());
-					const bool bCanImportSceneNode = [&AsyncHelper, &Node]()
-					{
-						bool bCanImport = AsyncHelper->TaskData.ImportType == EImportType::ImportType_Scene;
-						bCanImport = bCanImport && Node->GetParentUid().IsEmpty(); // We only import root scene nodes since we ask the factory to spawn the children nodes
-
-						return bCanImport;
-					}();
 
 					if (!RegisteredFactoryClass || (!bIsAsset && !bCanImportSceneNode))
 					{
@@ -248,14 +242,12 @@ void UE::Interchange::FTaskParsing::DoTask(ENamedThreads::Type CurrentThread, co
 	{
 		TFunction<FGraphEventRef(FTaskData&)> CreateTasksForSceneObject = [this, &AsyncHelper, &AssetsCompletionPrerequistes](FTaskData& TaskData)
 		{
-			const bool bSpawnChildren = true;
 			const int32 SourceIndex = TaskData.SourceIndex;
 			const UClass* const FactoryClass = TaskData.FactoryClass;
 
 			return AsyncHelper->SceneTasks.Add_GetRef(
 				TGraphTask<FTaskCreateSceneObjects>::CreateTask(&AssetsCompletionPrerequistes)
-				.ConstructAndDispatchWhenReady(PackageBasePath, SourceIndex, WeakAsyncHelper, TaskData.Nodes, FactoryClass->GetDefaultObject<UInterchangeFactoryBase>(), bSpawnChildren));
-
+				.ConstructAndDispatchWhenReady(PackageBasePath, SourceIndex, WeakAsyncHelper, TaskData.Nodes, FactoryClass));
 		};
 
 		ScenesCompletionPrerequistes = CreateTasksForEachTaskData(SceneTaskDatas, CreateTasksForSceneObject);
