Fix build warning - adding missing header guard

diff --git a/Engine/Plugins/Experimental/Interchange/Runtime/Source/Import/Public/Scene/InterchangeActorHelper.h b/Engine/Plugins/Experimental/Interchange/Runtime/Source/Import/Public/Scene/InterchangeActorHelper.h
--- a/Engine/Plugins/Experimental/Interchange/Runtime/Source/Import/Public/Scene/InterchangeActorHelper.h
+++ b/Engine/Plugins/Experimental/Interchange/Runtime/Source/Import/Public/Scene/InterchangeActorHelper.h
@@ -1,5 +1,7 @@
 // Copyright Epic Games, Inc. All Rights Reserved. 
 
+#pragma once
+
 #include "InterchangeFactoryBase.h"
 
 class UInterchangeBaseNode;
