Interchange Materials - Added support for opacity masks

diff --git a/Engine/Plugins/Experimental/Interchange/Runtime/Source/FactoryNodes/Private/InterchangeMaterialFactoryNode.cpp b/Engine/Plugins/Experimental/Interchange/Runtime/Source/FactoryNodes/Private/InterchangeMaterialFactoryNode.cpp
--- a/Engine/Plugins/Experimental/Interchange/Runtime/Source/FactoryNodes/Private/InterchangeMaterialFactoryNode.cpp
+++ b/Engine/Plugins/Experimental/Interchange/Runtime/Source/FactoryNodes/Private/InterchangeMaterialFactoryNode.cpp
@@ -297,6 +297,16 @@ bool UInterchangeMaterialFactoryNode::SetCustomTwoSided(const bool& AttributeVal
 	IMPLEMENT_NODE_ATTRIBUTE_SETTER(UInterchangeMaterialFactoryNode, TwoSided, bool, UMaterial);
 }
 
+bool UInterchangeMaterialFactoryNode::GetCustomOpacityMaskClipValue(float& AttributeValue) const
+{
+	IMPLEMENT_NODE_ATTRIBUTE_GETTER(OpacityMaskClipValue, float);
+}
+
+bool UInterchangeMaterialFactoryNode::SetCustomOpacityMaskClipValue(const float& AttributeValue, bool bAddApplyDelegate)
+{
+	IMPLEMENT_NODE_ATTRIBUTE_SETTER(UInterchangeMaterialFactoryNode, OpacityMaskClipValue, float, UMaterial);
+}
+
 FString UInterchangeMaterialExpressionFactoryNode::GetTypeName() const
 {
 	const FString TypeName = TEXT("MaterialExpressionFactoryNode");
diff --git a/Engine/Plugins/Experimental/Interchange/Runtime/Source/FactoryNodes/Public/InterchangeMaterialFactoryNode.h b/Engine/Plugins/Experimental/Interchange/Runtime/Source/FactoryNodes/Public/InterchangeMaterialFactoryNode.h
--- a/Engine/Plugins/Experimental/Interchange/Runtime/Source/FactoryNodes/Public/InterchangeMaterialFactoryNode.h
+++ b/Engine/Plugins/Experimental/Interchange/Runtime/Source/FactoryNodes/Public/InterchangeMaterialFactoryNode.h
@@ -198,11 +198,18 @@ public:
 	UFUNCTION(BlueprintCallable, Category = "Interchange | Node | Material")
 	bool SetCustomTwoSided(const bool& AttributeValue, bool bAddApplyDelegate = true);
 
+	UFUNCTION(BlueprintCallable, Category = "Interchange | Node | Material")
+	bool GetCustomOpacityMaskClipValue(float& AttributeValue) const;
+
+	UFUNCTION(BlueprintCallable, Category = "Interchange | Node | Material")
+	bool SetCustomOpacityMaskClipValue(const float& AttributeValue, bool bAddApplyDelegate = true);
+
 private:
 	const UE::Interchange::FAttributeKey Macro_CustomShadingModelKey = UE::Interchange::FAttributeKey(TEXT("ShadingModel"));
 	const UE::Interchange::FAttributeKey Macro_CustomTranslucencyLightingModeKey = UE::Interchange::FAttributeKey(TEXT("TranslucencyLightingMode"));
 	const UE::Interchange::FAttributeKey Macro_CustomBlendModeKey = UE::Interchange::FAttributeKey(TEXT("BlendMode"));
 	const UE::Interchange::FAttributeKey Macro_CustomTwoSidedKey = UE::Interchange::FAttributeKey(TEXT("TwoSided"));
+	const UE::Interchange::FAttributeKey Macro_CustomOpacityMaskClipValueKey = UE::Interchange::FAttributeKey(TEXT("OpacityMaskClipValue"));
 };
 
 UCLASS(BlueprintType, Experimental)
diff --git a/Engine/Plugins/Experimental/Interchange/Runtime/Source/Import/Private/Gltf/InterchangeGltfTranslator.cpp b/Engine/Plugins/Experimental/Interchange/Runtime/Source/Import/Private/Gltf/InterchangeGltfTranslator.cpp
--- a/Engine/Plugins/Experimental/Interchange/Runtime/Source/Import/Private/Gltf/InterchangeGltfTranslator.cpp
+++ b/Engine/Plugins/Experimental/Interchange/Runtime/Source/Import/Private/Gltf/InterchangeGltfTranslator.cpp
@@ -358,6 +358,12 @@ void UInterchangeGltfTranslator::HandleGltfMaterial( UInterchangeBaseNodeContain
 				OpacityFactor, Standard::Nodes::TextureSample::Outputs::A.ToString() );
 		}
 
+		// Alpha cutoff
+		if ( GltfMaterial.AlphaMode == GLTF::FMaterial::EAlphaMode::Mask )
+		{
+			ShaderGraphNode.SetCustomOpacityMaskClipValue( GltfMaterial.AlphaCutoff );
+		}
+
 		// IOR
 		if ( GltfMaterial.bHasIOR )
 		{
diff --git a/Engine/Plugins/Experimental/Interchange/Runtime/Source/Import/Private/Material/InterchangeMaterialFactory.cpp b/Engine/Plugins/Experimental/Interchange/Runtime/Source/Import/Private/Material/InterchangeMaterialFactory.cpp
--- a/Engine/Plugins/Experimental/Interchange/Runtime/Source/Import/Private/Material/InterchangeMaterialFactory.cpp
+++ b/Engine/Plugins/Experimental/Interchange/Runtime/Source/Import/Private/Material/InterchangeMaterialFactory.cpp
@@ -493,7 +493,18 @@ void UInterchangeMaterialFactory::SetupMaterial(UMaterial* Material, const FCrea
 			{
 				if (UMaterialExpression* OpacityExpression = CreateExpressionsForNode(Material, Arguments, OpacityNode, Expressions))
 				{
-					if (FExpressionInput* OpacityInput = Material->GetExpressionInputForProperty(MP_Opacity))
+					FExpressionInput* OpacityInput = Material->GetExpressionInputForProperty(MP_Opacity);
+
+					TEnumAsByte<EBlendMode> BlendMode;
+					if (MaterialFactoryNode->GetCustomBlendMode(BlendMode))
+					{
+						if (BlendMode == EBlendMode::BLEND_Masked)
+						{
+							OpacityInput = Material->GetExpressionInputForProperty(MP_OpacityMask);
+						}
+					}
+					
+					if (OpacityInput)
 					{
 						OpacityExpression->ConnectExpression(OpacityInput, GetOutputIndex(*OpacityExpression, OutputName));
 					}
diff --git a/Engine/Plugins/Experimental/Interchange/Runtime/Source/Nodes/Private/InterchangeShaderGraphNode.cpp b/Engine/Plugins/Experimental/Interchange/Runtime/Source/Nodes/Private/InterchangeShaderGraphNode.cpp
--- a/Engine/Plugins/Experimental/Interchange/Runtime/Source/Nodes/Private/InterchangeShaderGraphNode.cpp
+++ b/Engine/Plugins/Experimental/Interchange/Runtime/Source/Nodes/Private/InterchangeShaderGraphNode.cpp
@@ -143,3 +143,13 @@ bool UInterchangeShaderGraphNode::SetCustomTwoSided(const bool& AttributeValue)
 {
 	IMPLEMENT_NODE_ATTRIBUTE_SETTER_NODELEGATE(TwoSided, bool);
 }
+
+bool UInterchangeShaderGraphNode::GetCustomOpacityMaskClipValue(float& AttributeValue) const
+{
+	IMPLEMENT_NODE_ATTRIBUTE_GETTER(OpacityMaskClipValue, float);
+}
+
+bool UInterchangeShaderGraphNode::SetCustomOpacityMaskClipValue(const float& AttributeValue, bool bAddApplyDelegate)
+{
+	IMPLEMENT_NODE_ATTRIBUTE_SETTER_NODELEGATE(OpacityMaskClipValue, float);
+}
diff --git a/Engine/Plugins/Experimental/Interchange/Runtime/Source/Nodes/Public/InterchangeShaderGraphNode.h b/Engine/Plugins/Experimental/Interchange/Runtime/Source/Nodes/Public/InterchangeShaderGraphNode.h
--- a/Engine/Plugins/Experimental/Interchange/Runtime/Source/Nodes/Public/InterchangeShaderGraphNode.h
+++ b/Engine/Plugins/Experimental/Interchange/Runtime/Source/Nodes/Public/InterchangeShaderGraphNode.h
@@ -140,6 +140,16 @@ public:
 	UFUNCTION(BlueprintCallable, Category = "Interchange | Node | Material")
 	bool SetCustomTwoSided(const bool& AttributeValue);
 
+	UFUNCTION(BlueprintCallable, Category = "Interchange | Node | Material")
+	bool GetCustomOpacityMaskClipValue(float& AttributeValue) const;
+
+	/**
+	 * Shader is transparent or opaque if it's alpha is lower or higher than the clip value.
+	 */
+	UFUNCTION(BlueprintCallable, Category = "Interchange | Node | Material")
+	bool SetCustomOpacityMaskClipValue(const float& AttributeValue, bool bAddApplyDelegate = true);
+
 private:
 	const UE::Interchange::FAttributeKey Macro_CustomTwoSidedKey = UE::Interchange::FAttributeKey(TEXT("TwoSided"));
+	const UE::Interchange::FAttributeKey Macro_CustomOpacityMaskClipValueKey = UE::Interchange::FAttributeKey(TEXT("OpacityMaskClipValue"));
 };
\ No newline at end of file
diff --git a/Engine/Plugins/Experimental/Interchange/Runtime/Source/Pipelines/Private/InterchangeGenericMaterialPipeline.cpp b/Engine/Plugins/Experimental/Interchange/Runtime/Source/Pipelines/Private/InterchangeGenericMaterialPipeline.cpp
--- a/Engine/Plugins/Experimental/Interchange/Runtime/Source/Pipelines/Private/InterchangeGenericMaterialPipeline.cpp
+++ b/Engine/Plugins/Experimental/Interchange/Runtime/Source/Pipelines/Private/InterchangeGenericMaterialPipeline.cpp
@@ -615,18 +615,29 @@ void UInterchangeGenericMaterialPipeline::HandleCommonParameters(const UIntercha
 				if (OpacityExpressionFactoryNode.Get<0>())
 				{
 					MaterialFactoryNode->ConnectOutputToOpacity(OpacityExpressionFactoryNode.Get<0>()->GetUniqueID(), OpacityExpressionFactoryNode.Get<1>());
+				}
 
-					// Don't change the blend mode or the lighting mode if they were already set
-					TEnumAsByte<EBlendMode> BlendMode = EBlendMode::BLEND_Translucent;
-					if (!MaterialFactoryNode->GetCustomBlendMode(BlendMode))
+				// Opacity Clip Value
+				bool bIsMasked = false;
+				{
+					float OpacityClipValue;
+					if (ShaderGraphNode->GetCustomOpacityMaskClipValue(OpacityClipValue))
 					{
-						MaterialFactoryNode->SetCustomBlendMode(BlendMode);
+						MaterialFactoryNode->SetCustomOpacityMaskClipValue(OpacityClipValue);
+						bIsMasked = true;
+					}
+				}
 
-						TEnumAsByte<ETranslucencyLightingMode> LightingMode = ETranslucencyLightingMode::TLM_Surface;
-						if (!MaterialFactoryNode->GetCustomTranslucencyLightingMode(LightingMode))
-						{
-							MaterialFactoryNode->SetCustomTranslucencyLightingMode(LightingMode);
-						}
+				// Don't change the blend mode or the lighting mode if they were already set
+				TEnumAsByte<EBlendMode> BlendMode = bIsMasked ? EBlendMode::BLEND_Masked : EBlendMode::BLEND_Translucent;
+				if (!MaterialFactoryNode->GetCustomBlendMode(BlendMode))
+				{
+					MaterialFactoryNode->SetCustomBlendMode(BlendMode);
+
+					TEnumAsByte<ETranslucencyLightingMode> LightingMode = ETranslucencyLightingMode::TLM_Surface;
+					if (!MaterialFactoryNode->GetCustomTranslucencyLightingMode(LightingMode))
+					{
+						MaterialFactoryNode->SetCustomTranslucencyLightingMode(LightingMode);
 					}
 				}
 			}
