- Set BaseNodeContainer and SourceDatas to protected instead of private in UInterchangeGenericMaterialPipeline to allow sub-classes to access those two members when overriding ExecutePreImportPipeline
- Added support for override materails in UInterchangeGenericLevelPipeline::CreateActorFactoryNode when handling a static mesh actor
- In UInterchangeGenericMaterialPipeline::HandleTextureSampleNode, set the ExpressionClassName to UMaterialExpressionTextureSampleParameter2D::StaticClass()->GetName() if the texture node is null. Note that this is the default if the type of texture node is not recognized.
- Added application of material function in SetupFunctionCallExpression when the  material function is created during the import
- A specialized factory node, UInterchangeMaterialFunctionCallExpressionFactoryNode, hasbeen added to hold on the potential depedency on a material function created during the import

#prefilight 6239f2c489625f0612af8c45

diff --git a/Engine/Plugins/Experimental/Interchange/Runtime/Source/FactoryNodes/Private/InterchangeMaterialFactoryNode.cpp b/Engine/Plugins/Experimental/Interchange/Runtime/Source/FactoryNodes/Private/InterchangeMaterialFactoryNode.cpp
--- a/Engine/Plugins/Experimental/Interchange/Runtime/Source/FactoryNodes/Private/InterchangeMaterialFactoryNode.cpp
+++ b/Engine/Plugins/Experimental/Interchange/Runtime/Source/FactoryNodes/Private/InterchangeMaterialFactoryNode.cpp
@@ -367,3 +367,20 @@ bool UInterchangeMaterialInstanceFactoryNode::SetCustomParent(const FString& Att
 {
 	IMPLEMENT_NODE_ATTRIBUTE_SETTER_NODELEGATE(Parent, FString);
 }
+
+FString UInterchangeMaterialFunctionCallExpressionFactoryNode::GetTypeName() const
+{
+	const FString TypeName = TEXT("MaterialFunctionCallExpressionFactoryNode");
+	return TypeName;
+}
+
+bool UInterchangeMaterialFunctionCallExpressionFactoryNode::GetCustomMaterialFunctionDependency(FString& AttributeValue) const
+{
+	IMPLEMENT_NODE_ATTRIBUTE_GETTER(MaterialFunctionDependency, FString);
+}
+
+bool UInterchangeMaterialFunctionCallExpressionFactoryNode::SetCustomMaterialFunctionDependency(const FString& AttributeValue)
+{
+	IMPLEMENT_NODE_ATTRIBUTE_SETTER_NODELEGATE(MaterialFunctionDependency, FString);
+	AddFactoryDependencyUid(AttributeValue);
+}
diff --git a/Engine/Plugins/Experimental/Interchange/Runtime/Source/FactoryNodes/Public/InterchangeMaterialFactoryNode.h b/Engine/Plugins/Experimental/Interchange/Runtime/Source/FactoryNodes/Public/InterchangeMaterialFactoryNode.h
--- a/Engine/Plugins/Experimental/Interchange/Runtime/Source/FactoryNodes/Public/InterchangeMaterialFactoryNode.h
+++ b/Engine/Plugins/Experimental/Interchange/Runtime/Source/FactoryNodes/Public/InterchangeMaterialFactoryNode.h
@@ -213,7 +213,7 @@ private:
 };
 
 UCLASS(BlueprintType, Experimental)
-class INTERCHANGEFACTORYNODES_API UInterchangeMaterialExpressionFactoryNode : public UInterchangeBaseNode
+class INTERCHANGEFACTORYNODES_API UInterchangeMaterialExpressionFactoryNode : public UInterchangeFactoryBaseNode
 {
 	GENERATED_BODY()
 
@@ -259,3 +259,29 @@ private:
 	const UE::Interchange::FAttributeKey Macro_CustomParentKey = UE::Interchange::FAttributeKey(TEXT("Parent"));
 	
 };
+
+UCLASS(BlueprintType, Experimental)
+class INTERCHANGEFACTORYNODES_API UInterchangeMaterialFunctionCallExpressionFactoryNode : public UInterchangeMaterialExpressionFactoryNode
+{
+	GENERATED_BODY()
+
+public:
+	virtual FString GetTypeName() const override;
+
+public:
+	UFUNCTION(BlueprintCallable, Category = "Interchange | Node | Material")
+	bool GetCustomMaterialFunctionDependency(FString& AttributeValue) const;
+
+	/**
+	 * Set the unique ID of the material function that the function call expression
+	 * is referring to.
+	 * Note that a call to AddFactoryDependencyUid is made to guarantee that
+	 * the material function is created before the function call expression
+	 */
+	UFUNCTION(BlueprintCallable, Category = "Interchange | Node | Material")
+	bool SetCustomMaterialFunctionDependency(const FString& AttributeValue);
+
+private:
+	const UE::Interchange::FAttributeKey Macro_CustomMaterialFunctionDependencyKey = UE::Interchange::FAttributeKey(TEXT("MaterialFunctionDependency"));
+
+};
diff --git a/Engine/Plugins/Experimental/Interchange/Runtime/Source/Import/Private/Material/InterchangeMaterialFactory.cpp b/Engine/Plugins/Experimental/Interchange/Runtime/Source/Import/Private/Material/InterchangeMaterialFactory.cpp
--- a/Engine/Plugins/Experimental/Interchange/Runtime/Source/Import/Private/Material/InterchangeMaterialFactory.cpp
+++ b/Engine/Plugins/Experimental/Interchange/Runtime/Source/Import/Private/Material/InterchangeMaterialFactory.cpp
@@ -109,8 +109,23 @@ namespace UE
 					return 0; // Consider 0 as the default output to connect to since most expressions have a single output
 				}
 
-				void SetupFunctionCallExpression(UMaterial* Material, const UInterchangeFactoryBase::FCreateAssetParams& Arguments, UMaterialExpressionMaterialFunctionCall* FunctionCallExpression)
+				void SetupFunctionCallExpression(const UInterchangeFactoryBase::FCreateAssetParams& Arguments, const UInterchangeMaterialExpressionFactoryNode& ExpressionNode, UMaterialExpressionMaterialFunctionCall* FunctionCallExpression)
 				{
+					if (const UInterchangeMaterialFunctionCallExpressionFactoryNode* FunctionCallFactoryNode = Cast<UInterchangeMaterialFunctionCallExpressionFactoryNode>(&ExpressionNode))
+					{
+						FString MaterialFactoryNodeUid;
+						FunctionCallFactoryNode->GetCustomMaterialFunctionDependency(MaterialFactoryNodeUid);
+
+						if (const UInterchangeMaterialFactoryNode* MaterialFactoryNode = Cast<UInterchangeMaterialFactoryNode>(Arguments.NodeContainer->GetNode(MaterialFactoryNodeUid)))
+						{
+							if (UMaterialFunctionInterface* MaterialFunction = Cast<UMaterialFunctionInterface>(MaterialFactoryNode->ReferenceObject.TryLoad()))
+							{
+								MaterialFunction->UpdateFromFunctionResource();
+								FunctionCallExpression->SetMaterialFunction(MaterialFunction);
+							}
+						}
+					}
+
 					FunctionCallExpression->UpdateFromFunctionResource();
 				}
 
@@ -730,7 +745,7 @@ UMaterialExpression* UInterchangeMaterialFactory::CreateExpression(UMaterial* Ma
 	}
 	else if (UMaterialExpressionMaterialFunctionCall* FunctionCallExpression = Cast<UMaterialExpressionMaterialFunctionCall>(MaterialExpression))
 	{
-		SetupFunctionCallExpression(Material, Arguments, FunctionCallExpression);
+		SetupFunctionCallExpression(Arguments, ExpressionNode, FunctionCallExpression);
 	}
 
 	return MaterialExpression;
diff --git a/Engine/Plugins/Experimental/Interchange/Runtime/Source/Pipelines/Private/InterchangeGenericMaterialPipeline.cpp b/Engine/Plugins/Experimental/Interchange/Runtime/Source/Pipelines/Private/InterchangeGenericMaterialPipeline.cpp
--- a/Engine/Plugins/Experimental/Interchange/Runtime/Source/Pipelines/Private/InterchangeGenericMaterialPipeline.cpp
+++ b/Engine/Plugins/Experimental/Interchange/Runtime/Source/Pipelines/Private/InterchangeGenericMaterialPipeline.cpp
@@ -768,25 +768,30 @@ void UInterchangeGenericMaterialPipeline::HandleTextureSampleNode(const UInterch
 		{
 			TextureFactoryUid = TextureTargetNodes[0];
 		}
-	}
 
-	TextureSampleFactoryNode->SetCustomExpressionClassName(ExpressionClassName);
-	TextureSampleFactoryNode->AddStringAttribute(UInterchangeShaderPortsAPI::MakeInputValueKey(Inputs::Texture.ToString()), TextureFactoryUid);
+		TextureSampleFactoryNode->SetCustomExpressionClassName(ExpressionClassName);
+		TextureSampleFactoryNode->AddStringAttribute(UInterchangeShaderPortsAPI::MakeInputValueKey(Inputs::Texture.ToString()), TextureFactoryUid);
 
-	if (bParsingForNormalInput)
-	{
-		if (UInterchangeTextureFactoryNode* TextureFactoryNode = Cast<UInterchangeTextureFactoryNode>(BaseNodeContainer->GetNode(TextureFactoryUid)))
+		if (bParsingForNormalInput)
 		{
-			TextureFactoryNode->SetCustomCompressionSettings(TextureCompressionSettings::TC_Normalmap);
+			if (UInterchangeTextureFactoryNode* TextureFactoryNode = Cast<UInterchangeTextureFactoryNode>(BaseNodeContainer->GetNode(TextureFactoryUid)))
+			{
+				TextureFactoryNode->SetCustomCompressionSettings(TextureCompressionSettings::TC_Normalmap);
+			}
 		}
-	}
-	else if (bParsingForLinearInput)
-	{
-		if (UInterchangeTextureFactoryNode* TextureFactoryNode = Cast<UInterchangeTextureFactoryNode>(BaseNodeContainer->GetNode(TextureFactoryUid)))
+		else if (bParsingForLinearInput)
 		{
-			TextureFactoryNode->SetCustomSRGB(false);
+			if (UInterchangeTextureFactoryNode* TextureFactoryNode = Cast<UInterchangeTextureFactoryNode>(BaseNodeContainer->GetNode(TextureFactoryUid)))
+			{
+				TextureFactoryNode->SetCustomSRGB(false);
+			}
 		}
 	}
+	else
+	{
+		TextureSampleFactoryNode->SetCustomExpressionClassName(UMaterialExpressionTextureSampleParameter2D::StaticClass()->GetName());
+		TextureSampleFactoryNode->AddStringAttribute(UInterchangeShaderPortsAPI::MakeInputValueKey(Inputs::Texture.ToString()), TextureFactoryUid);
+	}
 
 	// Coordinates
 	{
diff --git a/Engine/Plugins/Experimental/Interchange/Runtime/Source/Pipelines/Public/InterchangeGenericMaterialPipeline.h b/Engine/Plugins/Experimental/Interchange/Runtime/Source/Pipelines/Public/InterchangeGenericMaterialPipeline.h
--- a/Engine/Plugins/Experimental/Interchange/Runtime/Source/Pipelines/Public/InterchangeGenericMaterialPipeline.h
+++ b/Engine/Plugins/Experimental/Interchange/Runtime/Source/Pipelines/Public/InterchangeGenericMaterialPipeline.h
@@ -43,7 +43,6 @@ public:
 protected:
 	virtual void ExecutePreImportPipeline(UInterchangeBaseNodeContainer* InBaseNodeContainer, const TArray<UInterchangeSourceData*>& InSourceDatas) override;
 
-private:
 	UPROPERTY()
 	TObjectPtr<UInterchangeBaseNodeContainer> BaseNodeContainer;
 
