Adding support for the compressed source image data via interchange.
Added support for the import of compressed jpeg directly into source data payload.
Added a function to initialize a unique buffer via an array and added unit test for it.
Updated the test to compare the interchange image import to the legacy workflow to test properly the compressed source data and fixed an issue with the flip green channel setting not being test in a deterministic way.

diff --git a/Engine/Plugins/Experimental/Interchange/Runtime/Source/FactoryNodes/Public/InterchangeTextureFactoryNode.h b/Engine/Plugins/Experimental/Interchange/Runtime/Source/FactoryNodes/Public/InterchangeTextureFactoryNode.h
--- a/Engine/Plugins/Experimental/Interchange/Runtime/Source/FactoryNodes/Public/InterchangeTextureFactoryNode.h
+++ b/Engine/Plugins/Experimental/Interchange/Runtime/Source/FactoryNodes/Public/InterchangeTextureFactoryNode.h
@@ -655,6 +655,26 @@ public:
 #endif
 	}
 
+	//////////////////////////////////////////////////////////////////////////
+	//Factory Import settings
+
+	/** Return false if the Attribute was not set previously.*/
+	UFUNCTION(BlueprintCallable, Category = "Interchange | Node | Texture | Import Settings")
+	bool GetCustomPreferCompressedSourceData(bool& AttributeValue) const
+	{
+		IMPLEMENT_NODE_ATTRIBUTE_GETTER(PreferCompressedSourceData, bool);
+	}
+
+	/*
+	 * Should the factory tell the translator to provide a compressed source data payload when available.
+	 * This will generally result in smaller assets, but some operations like the texture build might be slower because the source data will need to be uncompressed.
+	 */
+	UFUNCTION(BlueprintCallable, Category = "Interchange | Node | Texture | Import Settings")
+	bool SetCustomPreferCompressedSourceData(const bool& AttributeValue)
+	{
+		IMPLEMENT_NODE_ATTRIBUTE_SETTER_NODELEGATE(PreferCompressedSourceData, bool)
+	}
+
 private:
 
 	const UE::Interchange::FAttributeKey ClassNameAttributeKey = UE::Interchange::FBaseNodeStaticData::ClassTypeAttributeKey();
@@ -710,6 +730,9 @@ private:
 	const UE::Interchange::FAttributeKey Macro_CustomCompositeTextureModeKey = UE::Interchange::FAttributeKey(TEXT("CompositeTextureMode"));
 	const UE::Interchange::FAttributeKey Macro_CustomCompositePowerKey = UE::Interchange::FAttributeKey(TEXT("CompositePower"));
 
+	// Texture Factory Settings
+	IMPLEMENT_NODE_ATTRIBUTE_KEY(PreferCompressedSourceData)
+
 	// no IMPLEMENT_NODE_ATTRIBUTE_APPLY_UOBJECT here ?
 
 	//TODO support per platform data in the FAttributeStorage, so we can set different value per platform at the pipeline stage, We set only the default value for now
diff --git a/Engine/Plugins/Experimental/Interchange/Runtime/Source/Import/Private/Texture/InterchangeJPGTranslator.cpp b/Engine/Plugins/Experimental/Interchange/Runtime/Source/Import/Private/Texture/InterchangeJPGTranslator.cpp
--- a/Engine/Plugins/Experimental/Interchange/Runtime/Source/Import/Private/Texture/InterchangeJPGTranslator.cpp
+++ b/Engine/Plugins/Experimental/Interchange/Runtime/Source/Import/Private/Texture/InterchangeJPGTranslator.cpp
@@ -9,6 +9,7 @@
 #include "IImageWrapperModule.h"
 #include "InterchangeImportLog.h"
 #include "InterchangeTextureNode.h"
+#include "Memory/SharedBuffer.h"
 #include "Misc/ConfigCacheIni.h"
 #include "Misc/FileHelper.h"
 #include "Misc/Paths.h"
@@ -34,37 +35,78 @@ bool UInterchangeJPGTranslator::Translate(UInterchangeBaseNodeContainer& BaseNod
 }
 
 TOptional<UE::Interchange::FImportImage> UInterchangeJPGTranslator::GetTexturePayloadData(const UInterchangeSourceData* PayloadSourceData, const FString& PayLoadKey) const
+{
+	TArray64<uint8> SourceDataBuffer = LoadSourceFile(PayloadSourceData, PayLoadKey);
+
+	if (SourceDataBuffer.IsEmpty())
+	{
+		return {};
+	}
+
+	const bool bImportRaw = false;
+	return GetTexturePayloadImplementation(MoveTemp(SourceDataBuffer), PayLoadKey, bImportRaw);
+
+}
+
+bool UInterchangeJPGTranslator::SupportCompressedTexturePayloadData() const
+{
+	// For now this option is opt in via the config files once there is no technical risk this will become the default path.
+	bool bRetainJpegFormat = false;
+	GConfig->GetBool(TEXT("TextureImporter"), TEXT("RetainJpegFormat"), bRetainJpegFormat, GEditorIni);
+	return bRetainJpegFormat;
+}
+
+TOptional<UE::Interchange::FImportImage> UInterchangeJPGTranslator::GetCompressedTexturePayloadData(const UInterchangeSourceData* PayloadSourceData, const FString& PayLoadKey) const
+{
+	TArray64<uint8> SourceDataBuffer = LoadSourceFile(PayloadSourceData, PayLoadKey);
+
+	if (SourceDataBuffer.IsEmpty())
+	{
+		return {};
+	}
+
+	const bool bImportRaw = true;
+	return GetTexturePayloadImplementation(MoveTemp(SourceDataBuffer), PayLoadKey, bImportRaw);
+}
+
+TArray64<uint8> UInterchangeJPGTranslator::LoadSourceFile(const UInterchangeSourceData* PayloadSourceData, const FString& PayLoadKey) const
 {
 	check(PayloadSourceData == GetSourceData());
 
+	TArray64<uint8> SourceDataBuffer;
+
 	if (!GetSourceData())
 	{
 		UE_LOG(LogInterchangeImport, Error, TEXT("Failed to import JPEG, bad source data."));
-		return TOptional<UE::Interchange::FImportImage>();
+		return SourceDataBuffer;
 	}
 
-	TArray64<uint8> SourceDataBuffer;
 	FString Filename = GetSourceData()->GetFilename();
-	
+
 	//Make sure the key fit the filename, The key should always be valid
 	if (!Filename.Equals(PayLoadKey))
 	{
 		UE_LOG(LogInterchangeImport, Error, TEXT("Failed to import JPEG, wrong payload key. [%s]"), *Filename);
-		return TOptional<UE::Interchange::FImportImage>();
+		return SourceDataBuffer;
 	}
 
 	if (!FPaths::FileExists(Filename))
 	{
 		UE_LOG(LogInterchangeImport, Error, TEXT("Failed to import JPEG, cannot open file. [%s]"), *Filename);
-		return TOptional<UE::Interchange::FImportImage>();
+		return SourceDataBuffer;
 	}
 
 	if (!FFileHelper::LoadFileToArray(SourceDataBuffer, *Filename))
 	{
 		UE_LOG(LogInterchangeImport, Error, TEXT("Failed to import JPEG, cannot load file content into an array. [%s]"), *Filename);
-		return TOptional<UE::Interchange::FImportImage>();
 	}
 
+	return SourceDataBuffer;
+}
+
+TOptional<UE::Interchange::FImportImage> UInterchangeJPGTranslator::GetTexturePayloadImplementation(TArray64<uint8>&& SourceDataBuffer, const FString& Filename, bool bShouldImportRaw) const
+{
+
 	const uint8* Buffer = SourceDataBuffer.GetData();
 	const uint8* BufferEnd = Buffer + SourceDataBuffer.Num();
 
@@ -123,17 +165,27 @@ TOptional<UE::Interchange::FImportImage> UInterchangeJPGTranslator::GetTexturePa
 
 	UE::Interchange::FImportImage PayloadData;
 
+
+	const bool bShouldAllocateRawDataBuffer = !bShouldImportRaw;
+
 	PayloadData.Init2DWithParams(
 		JpegImageWrapper->GetWidth(),
 		JpegImageWrapper->GetHeight(),
 		TextureFormat,
-		BitDepth < 16
+		BitDepth < 16,
+		bShouldAllocateRawDataBuffer
 	);
 
-	if (!JpegImageWrapper->GetRaw(Format, BitDepth, PayloadData.GetArrayViewOfRawData()))
+	if (bShouldImportRaw)
+	{
+		PayloadData.RawData = MakeUniqueBufferFromArray(MoveTemp(SourceDataBuffer));
+		PayloadData.RawDataCompressionFormat = ETextureSourceCompressionFormat::TSCF_JPEG;
+	}
+	else if (!JpegImageWrapper->GetRaw(Format, BitDepth, PayloadData.GetArrayViewOfRawData()))
 	{
 		UE_LOG(LogInterchangeImport, Error, TEXT("Failed to decode JPEG. [%s]"), *Filename);
 		return TOptional<UE::Interchange::FImportImage>();
 	}
+
 	return PayloadData;
-}
\ No newline at end of file
+}
diff --git a/Engine/Plugins/Experimental/Interchange/Runtime/Source/Import/Private/Texture/InterchangeTextureFactory.cpp b/Engine/Plugins/Experimental/Interchange/Runtime/Source/Import/Private/Texture/InterchangeTextureFactory.cpp
--- a/Engine/Plugins/Experimental/Interchange/Runtime/Source/Import/Private/Texture/InterchangeTextureFactory.cpp
+++ b/Engine/Plugins/Experimental/Interchange/Runtime/Source/Import/Private/Texture/InterchangeTextureFactory.cpp
@@ -222,6 +222,33 @@ namespace UE::Interchange::Private::InterchangeTextureFactory
 		return {};
 	}
 
+	UInterchangeTextureFactoryNode* GetTextureFactoryNodeFromVariant(const FTextureFactoryNodeVariant& FactoryNodeVariant)
+	{
+		static_assert(TVariantSize<FTextureFactoryNodeVariant>::Value == 5, "Please update the code below and this assert to reflect the change to the variant type.");
+
+		if (UInterchangeTexture2DFactoryNode* const* TextureFactoryNode = FactoryNodeVariant.TryGet<UInterchangeTexture2DFactoryNode*>())
+		{
+			return *TextureFactoryNode;
+		}
+
+		if (UInterchangeTextureCubeFactoryNode* const* TextureCubeFactoryNode = FactoryNodeVariant.TryGet<UInterchangeTextureCubeFactoryNode*>())
+		{
+			return *TextureCubeFactoryNode;
+		}
+
+		if (UInterchangeTexture2DArrayFactoryNode* const* Texture2DArrayFactoryNode = FactoryNodeVariant.TryGet<UInterchangeTexture2DArrayFactoryNode*>())
+		{
+			return *Texture2DArrayFactoryNode;
+		}
+
+		if (UInterchangeTextureLightProfileFactoryNode* const* TextureLightProfileFactoryNode = FactoryNodeVariant.TryGet<UInterchangeTextureLightProfileFactoryNode*>())
+		{
+			return *TextureLightProfileFactoryNode;
+		}
+
+		return nullptr;
+	}
+
 	using FTextureNodeVariant = TVariant<FEmptyVariantState
 		, const UInterchangeTexture2DNode*
 		, const UInterchangeTextureCubeNode*
@@ -419,17 +446,30 @@ namespace UE::Interchange::Private::InterchangeTextureFactory
 		if (const UInterchangeTexture2DNode* const* TextureNode =  TextureNodeVariant.TryGet<const UInterchangeTexture2DNode*>())
 		{
 			TMap<int32, FString> BlockAndSourceDataFiles;
+			bool bShoudImportCompressedImage = false;
 			if (const UInterchangeTexture2DFactoryNode* const* Texture2DFactoryNode = FactoryNodeVariant.TryGet<UInterchangeTexture2DFactoryNode*>())
 			{
 				BlockAndSourceDataFiles = (*Texture2DFactoryNode)->GetSourceBlocks();
 			}
 
+			if (const UInterchangeTextureFactoryNode* TextureFactoryNode = GetTextureFactoryNodeFromVariant(FactoryNodeVariant))
+			{
+				TextureFactoryNode->GetCustomPreferCompressedSourceData(bShoudImportCompressedImage);
+			}
+
 			// Is there a case were a translator can be both interface and how should the factory chose which to invoke?
 			if (const IInterchangeTexturePayloadInterface* TextureTranslator = Cast<IInterchangeTexturePayloadInterface>(Translator))
 			{
 				if (BlockAndSourceDataFiles.IsEmpty())
 				{
-					return FTexturePayloadVariant(TInPlaceType<TOptional<FImportImage>>(), TextureTranslator->GetTexturePayloadData(SourceData, PayloadKey));
+					if (bShoudImportCompressedImage && TextureTranslator->SupportCompressedTexturePayloadData())
+					{
+						return FTexturePayloadVariant(TInPlaceType<TOptional<FImportImage>>(), TextureTranslator->GetCompressedTexturePayloadData(SourceData, PayloadKey));
+					}
+					else
+					{
+						return FTexturePayloadVariant(TInPlaceType<TOptional<FImportImage>>(), TextureTranslator->GetTexturePayloadData(SourceData, PayloadKey));
+					}
 				}
 				else
 				{
@@ -466,13 +506,13 @@ namespace UE::Interchange::Private::InterchangeTextureFactory
 #if WITH_EDITORONLY_DATA
 	void SetupTextureSourceData(UTexture* Texture, const FImportImage& Image, UE::Serialization::FEditorBulkData::FSharedBufferWithID&& BufferAndId)
 	{
-		Texture->Source.Init(
+		Texture->Source.InitWithCompressedSourceData(
 			Image.SizeX,
 			Image.SizeY,
-			/*NumSlices=*/ 1,
 			Image.NumMips,
 			Image.Format,
-			MoveTemp(BufferAndId)
+			MoveTemp(BufferAndId),
+			Image.RawDataCompressionFormat
 		);
 
 		Texture->CompressionSettings = Image.CompressionSettings;
diff --git a/Engine/Plugins/Experimental/Interchange/Runtime/Source/Import/Private/Texture/InterchangeTexturePayloadData.cpp b/Engine/Plugins/Experimental/Interchange/Runtime/Source/Import/Private/Texture/InterchangeTexturePayloadData.cpp
--- a/Engine/Plugins/Experimental/Interchange/Runtime/Source/Import/Private/Texture/InterchangeTexturePayloadData.cpp
+++ b/Engine/Plugins/Experimental/Interchange/Runtime/Source/Import/Private/Texture/InterchangeTexturePayloadData.cpp
@@ -49,19 +49,22 @@ bool UE::Interchange::FImportImageHelper::IsImportResolutionValid(int32 Width, i
 	return bValid;
 }
 
-void UE::Interchange::FImportImage::Init2DWithParams(int32 InSizeX, int32 InSizeY, ETextureSourceFormat InFormat, bool InSRGB)
+void UE::Interchange::FImportImage::Init2DWithParams(int32 InSizeX, int32 InSizeY, ETextureSourceFormat InFormat, bool bInSRGB, bool bShouldAllocateRawData)
 {
-	Init2DWithParams(InSizeX, InSizeY, 1, InFormat, InSRGB);
+	Init2DWithParams(InSizeX, InSizeY, 1, InFormat, bInSRGB, bShouldAllocateRawData);
 }
 
-void UE::Interchange::FImportImage::Init2DWithParams(int32 InSizeX, int32 InSizeY, int32 InNumMips, ETextureSourceFormat InFormat, bool InSRGB)
+void UE::Interchange::FImportImage::Init2DWithParams(int32 InSizeX, int32 InSizeY, int32 InNumMips, ETextureSourceFormat InFormat, bool bInSRGB, bool bShouldAllocateRawData)
 {
 	SizeX = InSizeX;
 	SizeY = InSizeY;
 	NumMips = InNumMips;
 	Format = InFormat;
-	bSRGB = InSRGB;
-	RawData = FUniqueBuffer::Alloc(ComputeBufferSize());
+	bSRGB = bInSRGB;
+	if (bShouldAllocateRawData)
+	{
+		RawData = FUniqueBuffer::Alloc(ComputeBufferSize());
+	}
 }
 
 void UE::Interchange::FImportImage::Init2DWithOneMip(int32 InSizeX, int32 InSizeY, ETextureSourceFormat InFormat, const void* InData)
@@ -115,9 +118,20 @@ TArrayView64<uint8> UE::Interchange::FImportImage::GetArrayViewOfRawData()
 
 bool UE::Interchange::FImportImage::IsValid() const
 {
+	bool bIsRawDataBufferValid = false;
+
+	if (RawDataCompressionFormat == TSCF_None)
+	{
+		bIsRawDataBufferValid = ComputeBufferSize() == RawData.GetSize();
+	}
+	else
+	{
+		bIsRawDataBufferValid = !RawData.IsNull();
+	}
+
 	return SizeX > 0
 		&& SizeY > 0
 		&& NumMips > 0
 		&& Format != TSF_Invalid
-		&& ComputeBufferSize() == RawData.GetSize();
+		&& bIsRawDataBufferValid;
 }
diff --git a/Engine/Plugins/Experimental/Interchange/Runtime/Source/Import/Public/Texture/InterchangeJPGTranslator.h b/Engine/Plugins/Experimental/Interchange/Runtime/Source/Import/Public/Texture/InterchangeJPGTranslator.h
--- a/Engine/Plugins/Experimental/Interchange/Runtime/Source/Import/Public/Texture/InterchangeJPGTranslator.h
+++ b/Engine/Plugins/Experimental/Interchange/Runtime/Source/Import/Public/Texture/InterchangeJPGTranslator.h
@@ -41,7 +41,17 @@ public:
 	 */
 	virtual TOptional<UE::Interchange::FImportImage> GetTexturePayloadData(const UInterchangeSourceData* PayloadSourceData, const FString& PayLoadKey) const override;
 
+	virtual bool SupportCompressedTexturePayloadData() const override;
+
+	virtual TOptional<UE::Interchange::FImportImage> GetCompressedTexturePayloadData(const UInterchangeSourceData* PayloadSourceData, const FString& PayLoadKey) const override;
+	
 	/* IInterchangeTexturePayloadInterface End */
+
+private:
+
+	TArray64<uint8> LoadSourceFile(const UInterchangeSourceData* PayloadSourceData, const FString& PayLoadKey) const;
+
+	TOptional<UE::Interchange::FImportImage> GetTexturePayloadImplementation(TArray64<uint8>&& SourceDataBuffer, const FString& Filename, bool bShouldImportRaw) const;
 };
 
 
diff --git a/Engine/Plugins/Experimental/Interchange/Runtime/Source/Import/Public/Texture/InterchangeTexturePayloadData.h b/Engine/Plugins/Experimental/Interchange/Runtime/Source/Import/Public/Texture/InterchangeTexturePayloadData.h
--- a/Engine/Plugins/Experimental/Interchange/Runtime/Source/Import/Public/Texture/InterchangeTexturePayloadData.h
+++ b/Engine/Plugins/Experimental/Interchange/Runtime/Source/Import/Public/Texture/InterchangeTexturePayloadData.h
@@ -22,6 +22,10 @@ namespace UE
 			FImportImage& operator=(const FImportImage&) = delete;
 
 			FUniqueBuffer RawData;
+
+			/** Which compression format (if any) that is applied to RawData */
+			ETextureSourceCompressionFormat RawDataCompressionFormat = TSCF_None;
+
 			ETextureSourceFormat Format = TSF_Invalid;
 			TextureCompressionSettings CompressionSettings = TC_Default;
 			int32 NumMips = 0;
@@ -30,8 +34,8 @@ namespace UE
 			bool bSRGB = true;
 			TOptional<TextureMipGenSettings> MipGenSettings;
 
-			void Init2DWithParams(int32 InSizeX, int32 InSizeY, ETextureSourceFormat InFormat, bool InSRGB);
-			void Init2DWithParams(int32 InSizeX, int32 InSizeY, int32 InNumMips, ETextureSourceFormat InFormat, bool InSRGB);
+			void Init2DWithParams(int32 InSizeX, int32 InSizeY, ETextureSourceFormat InFormat, bool bInSRGB, bool bShouldAllocateRawData = true);
+			void Init2DWithParams(int32 InSizeX, int32 InSizeY, int32 InNumMips, ETextureSourceFormat InFormat, bool bInSRGB, bool bShouldAllocateRawData = true);
 			void Init2DWithOneMip(int32 InSizeX, int32 InSizeY, ETextureSourceFormat InFormat, const void* InData = nullptr);
 			void Init2DWithMips(int32 InSizeX, int32 InSizeY, int32 InNumMips, ETextureSourceFormat InFormat, const void* InData = nullptr);
 
diff --git a/Engine/Plugins/Experimental/Interchange/Runtime/Source/Import/Public/Texture/InterchangeTexturePayloadInterface.h b/Engine/Plugins/Experimental/Interchange/Runtime/Source/Import/Public/Texture/InterchangeTexturePayloadInterface.h
--- a/Engine/Plugins/Experimental/Interchange/Runtime/Source/Import/Public/Texture/InterchangeTexturePayloadInterface.h
+++ b/Engine/Plugins/Experimental/Interchange/Runtime/Source/Import/Public/Texture/InterchangeTexturePayloadInterface.h
@@ -32,6 +32,21 @@ public:
 	 * @return a PayloadData containing the import image data. The TOptional will not be set if there is an error.
 	 */
 	virtual TOptional<UE::Interchange::FImportImage> GetTexturePayloadData(const UInterchangeSourceData* PayloadSourceData, const FString& PayloadKey) const = 0;
+
+	/**
+	 * Does this translator support the request for a compressed source data?
+	 */
+	virtual bool SupportCompressedTexturePayloadData() const { return false; }
+
+	/**
+	 * Once the translation is done, the import process need a way to retrieve payload data.
+	 * This payload will be use by the factories to create the asset.
+	 *
+	 * @param PayloadSourceData - The source data containing the data to translate
+	 * @param PayloadKey - The key to retrieve the a particular payload contain into the specified source data.
+	 * @return a PayloadData containing the import image data with image data being the original file format generally. The TOptional will not be set if there is an error.
+	 */
+	virtual TOptional<UE::Interchange::FImportImage> GetCompressedTexturePayloadData(const UInterchangeSourceData* PayloadSourceData, const FString& PayloadKey) const { return {}; }
 };
 
 
diff --git a/Engine/Plugins/Experimental/Interchange/Runtime/Source/Pipelines/Private/InterchangeGenericTexturePipeline.cpp b/Engine/Plugins/Experimental/Interchange/Runtime/Source/Pipelines/Private/InterchangeGenericTexturePipeline.cpp
--- a/Engine/Plugins/Experimental/Interchange/Runtime/Source/Pipelines/Private/InterchangeGenericTexturePipeline.cpp
+++ b/Engine/Plugins/Experimental/Interchange/Runtime/Source/Pipelines/Private/InterchangeGenericTexturePipeline.cpp
@@ -223,6 +223,17 @@ UInterchangeTextureFactoryNode* UInterchangeGenericTexturePipeline::CreateTextur
 		TextureFactoryNode->AddTargetNodeUid(TextureNode->GetUniqueID());
 		TextureNode->AddTargetNodeUid(TextureFactoryNode->GetUniqueID());
 	}
+
+#if WITH_EDITORONLY_DATA
+	if (TextureFactoryNode)
+	{
+		if (bPreferCompressedSourceData)
+		{
+			TextureFactoryNode->SetCustomPreferCompressedSourceData(true);
+		}
+	}
+#endif // WITH_EDITORONLY_DATA
+
 	return TextureFactoryNode;
 }
 
diff --git a/Engine/Plugins/Experimental/Interchange/Runtime/Source/Pipelines/Public/InterchangeGenericTexturePipeline.h b/Engine/Plugins/Experimental/Interchange/Runtime/Source/Pipelines/Public/InterchangeGenericTexturePipeline.h
--- a/Engine/Plugins/Experimental/Interchange/Runtime/Source/Pipelines/Public/InterchangeGenericTexturePipeline.h
+++ b/Engine/Plugins/Experimental/Interchange/Runtime/Source/Pipelines/Public/InterchangeGenericTexturePipeline.h
@@ -42,6 +42,17 @@ public:
 	/** Specify the files type that should be imported as long/lat cubemap */
 	UPROPERTY(EditAnywhere, Category = "Textures")
 	TSet<FString> FileExtensionsToImportAsLongLatCubemap = {"hdr"};
+
+	/** 
+	 * If true, tell the translator to provide a compressed source data payload when available.
+	 * This will generally result in smaller assets, but some operations like the texture build might be slower because the source data will need to be uncompressed.
+	 * If false, it will let the factory or another step in the pipeline decide what to do.
+	 * 
+	 * Compressed source source data is generally store the data as it is in the source file.
+	 */
+	UPROPERTY(EditAnywhere, BlueprintReadWrite, AdvancedDisplay, Category = "Textures")
+	bool bPreferCompressedSourceData = false;
+
 #endif
 
 public:
diff --git a/Engine/Source/Runtime/Core/Public/Memory/SharedBuffer.h b/Engine/Source/Runtime/Core/Public/Memory/SharedBuffer.h
--- a/Engine/Source/Runtime/Core/Public/Memory/SharedBuffer.h
+++ b/Engine/Source/Runtime/Core/Public/Memory/SharedBuffer.h
@@ -880,6 +880,13 @@ private:
 
 } // UE::SharedBuffer::Private
 
+/** Construct a unique buffer by taking ownership of an array. */
+template <typename T, typename Allocator>
+[[nodiscard]] inline FUniqueBuffer MakeUniqueBufferFromArray(TArray<T, Allocator>&& Array)
+{
+	return FUniqueBuffer(new UE::SharedBuffer::Private::TBufferOwnerTArray<T, Allocator>(MoveTemp(Array)));
+}
+
 /** Construct a shared buffer by taking ownership of an array. */
 template <typename T, typename Allocator>
 [[nodiscard]] inline FSharedBuffer MakeSharedBufferFromArray(TArray<T, Allocator>&& Array)
diff --git a/Engine/Source/Runtime/Engine/Classes/Engine/Texture.h b/Engine/Source/Runtime/Engine/Classes/Engine/Texture.h
--- a/Engine/Source/Runtime/Engine/Classes/Engine/Texture.h
+++ b/Engine/Source/Runtime/Engine/Classes/Engine/Texture.h
@@ -267,6 +267,24 @@ struct FTextureSource
 		ETextureSourceCompressionFormat NewSourceFormat
 	);
 
+	/**
+	 * Initialize the source data with the given size, number of mips, and format.
+	 * @param NewSizeX - Width of the texture source data.
+	 * @param NewSizeY - Height of the texture source data.
+	 * @param NewNumMips - The number of mips in the texture source data.
+	 * @param NewFormat - The format in which source data is stored.
+	 * @param NewSourceData -The new source data.
+	 * @param NewSourceFormat -The compression format of the new source data.
+	 */
+	ENGINE_API void InitWithCompressedSourceData(
+		int32 NewSizeX,
+		int32 NewSizeY,
+		int32 NewNumMips,
+		ETextureSourceFormat NewFormat,
+		UE::Serialization::FEditorBulkData::FSharedBufferWithID NewSourceData,
+		ETextureSourceCompressionFormat NewSourceFormat
+	);
+
 	/** Make a copy with a torn-off BulkData that has the same Guid used for DDC as this->BulkData */
 	FTextureSource CopyTornOff() const;
 
diff --git a/Engine/Source/Runtime/Engine/Private/Texture.cpp b/Engine/Source/Runtime/Engine/Private/Texture.cpp
--- a/Engine/Source/Runtime/Engine/Private/Texture.cpp
+++ b/Engine/Source/Runtime/Engine/Private/Texture.cpp
@@ -1585,29 +1585,54 @@ void FTextureSource::InitWithCompressedSourceData(
 	ETextureSourceCompressionFormat NewSourceFormat
 )
 {
-	RemoveSourceData();
-
-	SizeX = NewSizeX;
-	SizeY = NewSizeY;
-
-	NumLayers = 1;
-	NumSlices = 1;
-	NumMips = NewNumMips;
-
-	check( NewNumMips <= GetFullMipCount(SizeX,SizeY) );
-
-	Format = NewFormat;
-	LayerFormat.SetNum(1, true);
-	LayerFormat[0] = NewFormat;
+	const int32 NewNumSlice = 1;
+	const int32 NewNumLayer = 1;
+	InitLayeredImpl(
+		NewSizeX,
+		NewSizeY,
+		NewNumSlice,
+		NewNumLayer,
+		NewNumMips,
+		&NewFormat
+		);
 
 	CompressionFormat = NewSourceFormat;
 
-	BlockDataOffsets.Add(0);
+	BulkData.UpdatePayload(FSharedBuffer::Clone(NewData.GetData(), NewData.Num()), Owner);
+	// Disable the internal bulkdata compression if the source data is already compressed
+	if (CompressionFormat == TSCF_None)
+	{
+		BulkData.SetCompressionOptions(UE::Serialization::ECompressionOptions::Default);
+	}
+	else
+	{
+		BulkData.SetCompressionOptions(UE::Serialization::ECompressionOptions::Disabled);
+	}
+}
 
-	checkf(LockState == ELockState::None, TEXT("InitWithCompressedSourceData shouldn't be called in-between LockMip/UnlockMip"));
+void FTextureSource::InitWithCompressedSourceData(
+	int32 NewSizeX,
+	int32 NewSizeY,
+	int32 NewNumMips,
+	ETextureSourceFormat NewFormat,
+	UE::Serialization::FEditorBulkData::FSharedBufferWithID NewSourceData,
+	ETextureSourceCompressionFormat NewSourceFormat
+)
+{
+	const int32 NewNumSlice = 1;
+	const int32 NewNumLayer = 1;
+	InitLayeredImpl(
+		NewSizeX,
+		NewSizeY,
+		NewNumSlice,
+		NewNumLayer,
+		NewNumMips,
+		&NewFormat
+	);
 
-	BulkData.UpdatePayload(FSharedBuffer::Clone(NewData.GetData(), NewData.Num()), Owner);
+	CompressionFormat = NewSourceFormat;
 
+	BulkData.UpdatePayload(MoveTemp(NewSourceData), Owner);
 	// Disable the internal bulkdata compression if the source data is already compressed
 	if (CompressionFormat == TSCF_None)
 	{
diff --git a/Engine/Source/Runtime/Experimental/Interchange/Engine/Private/Tasks/InterchangeTaskCreateAsset.cpp b/Engine/Source/Runtime/Experimental/Interchange/Engine/Private/Tasks/InterchangeTaskCreateAsset.cpp
--- a/Engine/Source/Runtime/Experimental/Interchange/Engine/Private/Tasks/InterchangeTaskCreateAsset.cpp
+++ b/Engine/Source/Runtime/Experimental/Interchange/Engine/Private/Tasks/InterchangeTaskCreateAsset.cpp
@@ -193,7 +193,7 @@ void UE::Interchange::FTaskCreatePackage::DoTask(ENamedThreads::Type CurrentThre
 		CreateAssetParams.AssetNode = Node;
 		CreateAssetParams.Parent = Pkg;
 		CreateAssetParams.SourceData = AsyncHelper->SourceDatas[SourceIndex];
-		CreateAssetParams.Translator = nullptr;
+		CreateAssetParams.Translator = AsyncHelper->Translators[SourceIndex];
 		if (AsyncHelper->BaseNodeContainers.IsValidIndex(SourceIndex))
 		{
 			CreateAssetParams.NodeContainer = AsyncHelper->BaseNodeContainers[SourceIndex].Get();
