Mesh pipeline support for socket import.
Create static mesh sockets in static mesh factory.

diff --git a/Engine/Plugins/Experimental/Interchange/Runtime/Source/FactoryNodes/Private/InterchangeStaticMeshFactoryNode.cpp b/Engine/Plugins/Experimental/Interchange/Runtime/Source/FactoryNodes/Private/InterchangeStaticMeshFactoryNode.cpp
--- a/Engine/Plugins/Experimental/Interchange/Runtime/Source/FactoryNodes/Private/InterchangeStaticMeshFactoryNode.cpp
+++ b/Engine/Plugins/Experimental/Interchange/Runtime/Source/FactoryNodes/Private/InterchangeStaticMeshFactoryNode.cpp
@@ -15,6 +15,12 @@ namespace UE
 			static FString LodDependencies_BaseKey = TEXT("Lod_Dependencies");
 			return LodDependencies_BaseKey;
 		}
+
+		const FString& FStaticMeshNodeStaticData::GetSocketUidsBaseKey()
+		{
+			static FString SocketUids_BaseKey = TEXT("SocketUids");
+			return SocketUids_BaseKey;
+		}
 	} // namespace Interchange
 } // namespace UE
 
@@ -25,6 +31,7 @@ UInterchangeStaticMeshFactoryNode::UInterchangeStaticMeshFactoryNode()
 	AssetClass = nullptr;
 #endif
 	LodDependencies.Initialize(Attributes, UE::Interchange::FStaticMeshNodeStaticData::GetLodDependenciesBaseKey());
+	SocketUids.Initialize(Attributes, UE::Interchange::FStaticMeshNodeStaticData::GetSocketUidsBaseKey());
 }
 
 void UInterchangeStaticMeshFactoryNode::InitializeStaticMeshNode(const FString& UniqueID, const FString& DisplayLabel, const FString& InAssetClass)
@@ -107,6 +114,39 @@ bool UInterchangeStaticMeshFactoryNode::RemoveLodDataUniqueId(const FString& Lod
 	return LodDependencies.RemoveItem(LodDataUniqueId);
 }
 
+int32 UInterchangeStaticMeshFactoryNode::GetSocketUidCount() const
+{
+	return SocketUids.GetCount();
+}
+
+void UInterchangeStaticMeshFactoryNode::GetSocketUids(TArray<FString>& OutSocketUids) const
+{
+	SocketUids.GetItems(OutSocketUids);
+}
+
+bool UInterchangeStaticMeshFactoryNode::AddSocketUid(const FString& SocketUid)
+{
+	return SocketUids.AddItem(SocketUid);
+}
+
+bool UInterchangeStaticMeshFactoryNode::AddSocketUids(const TArray<FString>& InSocketUids)
+{
+	for (const FString& SocketUid : InSocketUids)
+	{
+		if (!SocketUids.AddItem(SocketUid))
+		{
+			return false;
+		}
+	}
+
+	return true;
+}
+
+bool UInterchangeStaticMeshFactoryNode::RemoveSocketUd(const FString& SocketUid)
+{
+	return SocketUids.RemoveItem(SocketUid);
+}
+
 bool UInterchangeStaticMeshFactoryNode::GetCustomVertexColorReplace(bool& AttributeValue) const
 {
 	IMPLEMENT_NODE_ATTRIBUTE_GETTER(VertexColorReplace, bool)
diff --git a/Engine/Plugins/Experimental/Interchange/Runtime/Source/FactoryNodes/Public/InterchangeStaticMeshFactoryNode.h b/Engine/Plugins/Experimental/Interchange/Runtime/Source/FactoryNodes/Public/InterchangeStaticMeshFactoryNode.h
--- a/Engine/Plugins/Experimental/Interchange/Runtime/Source/FactoryNodes/Public/InterchangeStaticMeshFactoryNode.h
+++ b/Engine/Plugins/Experimental/Interchange/Runtime/Source/FactoryNodes/Public/InterchangeStaticMeshFactoryNode.h
@@ -19,6 +19,7 @@ namespace UE
 		struct INTERCHANGEFACTORYNODES_API FStaticMeshNodeStaticData : public FBaseNodeStaticData
 		{
 			static const FString& GetLodDependenciesBaseKey();
+			static const FString& GetSocketUidsBaseKey();
 		};
 	} // namespace Interchange
 } // namespace UE
@@ -68,6 +69,22 @@ public:
 	UFUNCTION(BlueprintCallable, Category = "Interchange | Node | StaticMesh")
 	bool RemoveLodDataUniqueId(const FString& LodDataUniqueId);
 
+	/** Return The number of socket UIDs this static mesh has.*/
+	UFUNCTION(BlueprintCallable, Category = "Interchange | Node | StaticMesh")
+	int32 GetSocketUidCount() const;
+
+	UFUNCTION(BlueprintCallable, Category = "Interchange | Node | StaticMesh")
+	void GetSocketUids(TArray<FString>& OutSocketUids) const;
+
+	UFUNCTION(BlueprintCallable, Category = "Interchange | Node | StaticMesh")
+	bool AddSocketUid(const FString& SocketUid);
+
+	UFUNCTION(BlueprintCallable, Category = "Interchange | Node | StaticMesh")
+	bool AddSocketUids(const TArray<FString>& InSocketUids);
+
+	UFUNCTION(BlueprintCallable, Category = "Interchange | Node | StaticMesh")
+	bool RemoveSocketUd(const FString& SocketUid);
+
 	/** Query weather the static mesh factory should replace the vertex color. Return false if the attribute was not set.*/
 	UFUNCTION(BlueprintCallable, Category = "Interchange | Node | StaticMesh")
 	bool GetCustomVertexColorReplace(bool& AttributeValue) const;
@@ -104,6 +121,7 @@ private:
 	const UE::Interchange::FAttributeKey Macro_CustomVertexColorOverrideKey = UE::Interchange::FAttributeKey(TEXT("VertexColorOverride"));
 
 	UE::Interchange::TArrayAttributeHelper<FString> LodDependencies;
+	UE::Interchange::TArrayAttributeHelper<FString> SocketUids;
 protected:
 #if WITH_ENGINE
 	TSubclassOf<UStaticMesh> AssetClass = nullptr;
diff --git a/Engine/Plugins/Experimental/Interchange/Runtime/Source/Import/Private/Mesh/InterchangeStaticMeshFactory.cpp b/Engine/Plugins/Experimental/Interchange/Runtime/Source/Import/Private/Mesh/InterchangeStaticMeshFactory.cpp
--- a/Engine/Plugins/Experimental/Interchange/Runtime/Source/Import/Private/Mesh/InterchangeStaticMeshFactory.cpp
+++ b/Engine/Plugins/Experimental/Interchange/Runtime/Source/Import/Private/Mesh/InterchangeStaticMeshFactory.cpp
@@ -8,6 +8,7 @@
 #include "Components.h"
 #include "Engine/Polys.h"
 #include "Engine/StaticMesh.h"
+#include "Engine/StaticMeshSocket.h"
 #include "InterchangeAssetImportData.h"
 #include "InterchangeImportCommon.h"
 #include "InterchangeImportLog.h"
@@ -358,6 +359,8 @@ UObject* UInterchangeStaticMeshFactory::CreateAsset(const FCreateAssetParams& Ar
 		CurrentLodIndex++;
 	}
 
+	ImportSockets(Arguments, StaticMesh, StaticMeshFactoryNode);
+
 	if (!Arguments.ReimportObject)
 	{
 		// Apply all StaticMeshFactoryNode custom attributes to the static mesh asset
@@ -1276,3 +1279,62 @@ bool UInterchangeStaticMeshFactory::GenerateKDopCollision(const FCreateAssetPara
 
 #endif
 }
+
+
+bool UInterchangeStaticMeshFactory::ImportSockets(const FCreateAssetParams& Arguments, UStaticMesh* StaticMesh, const UInterchangeStaticMeshFactoryNode* FactoryNode)
+{
+	TArray<FString> SocketUids;
+	FactoryNode->GetSocketUids(SocketUids);
+
+	TSet<FName> ImportedSocketNames;
+
+	for (const FString& SocketUid : SocketUids)
+	{
+		if (const UInterchangeSceneNode* SceneNode = Cast<UInterchangeSceneNode>(Arguments.NodeContainer->GetNode(SocketUid)))
+		{
+			FString NodeDisplayName = SceneNode->GetDisplayLabel();
+			if (NodeDisplayName.StartsWith(TEXT("SOCKET_")))
+			{
+				constexpr bool bAllowShrinking = false;
+				NodeDisplayName.RightChopInline(sizeof("SOCKET_") - 1, bAllowShrinking);
+			}
+			FName SocketName = FName(NodeDisplayName);
+			ImportedSocketNames.Add(SocketName);
+
+			FTransform GlobalTransform;
+			SceneNode->GetCustomGlobalTransform(Arguments.NodeContainer, GlobalTransform);
+
+			UStaticMeshSocket* Socket = StaticMesh->FindSocket(SocketName);
+			if (!Socket)
+			{
+				// If the socket didn't exist create a new one now
+				Socket = NewObject<UStaticMeshSocket>(StaticMesh);
+#if WITH_EDITORONLY_DATA
+				Socket->bSocketCreatedAtImport = true;
+#endif
+				Socket->SocketName = SocketName;
+				StaticMesh->AddSocket(Socket);
+			}
+
+			Socket->RelativeLocation = GlobalTransform.GetLocation();
+			Socket->RelativeRotation = GlobalTransform.GetRotation().Rotator();
+			Socket->RelativeScale = GlobalTransform.GetScale3D();
+		}
+	}
+
+	// Delete any sockets which were previously imported but which no longer exist in the imported scene
+	for (TArray<TObjectPtr<UStaticMeshSocket>>::TIterator It = StaticMesh->Sockets.CreateIterator(); It; ++It)
+	{
+		UStaticMeshSocket* Socket = *It;
+		if (
+#if WITH_EDITORONLY_DATA
+			Socket->bSocketCreatedAtImport &&
+#endif
+			!ImportedSocketNames.Contains(Socket->SocketName))
+		{
+			It.RemoveCurrent();
+		}
+	}
+
+	return true;
+}
diff --git a/Engine/Plugins/Experimental/Interchange/Runtime/Source/Import/Public/Mesh/InterchangeStaticMeshFactory.h b/Engine/Plugins/Experimental/Interchange/Runtime/Source/Import/Public/Mesh/InterchangeStaticMeshFactory.h
--- a/Engine/Plugins/Experimental/Interchange/Runtime/Source/Import/Public/Mesh/InterchangeStaticMeshFactory.h
+++ b/Engine/Plugins/Experimental/Interchange/Runtime/Source/Import/Public/Mesh/InterchangeStaticMeshFactory.h
@@ -13,6 +13,7 @@
 
 class UBodySetup;
 class UStaticMesh;
+class UInterchangeStaticMeshFactoryNode;
 class UInterchangeStaticMeshLodDataNode;
 struct FMeshDescription;
 struct FKAggregateGeom;
@@ -57,4 +58,5 @@ private:
 	bool ImportSphereCollision(const FCreateAssetParams& Arguments, UStaticMesh* StaticMesh, const UInterchangeStaticMeshLodDataNode* LodDataNode);
 	bool ImportConvexCollision(const FCreateAssetParams& Arguments, UStaticMesh* StaticMesh, const UInterchangeStaticMeshLodDataNode* LodDataNode);
 	bool GenerateKDopCollision(const FCreateAssetParams& Arguments, UStaticMesh* StaticMesh);
+	bool ImportSockets(const FCreateAssetParams& Arguments, UStaticMesh* StaticMesh, const UInterchangeStaticMeshFactoryNode* FactoryNode);
 };
diff --git a/Engine/Plugins/Experimental/Interchange/Runtime/Source/Pipelines/Private/InterchangeGenericStaticMeshPipeline.cpp b/Engine/Plugins/Experimental/Interchange/Runtime/Source/Pipelines/Private/InterchangeGenericStaticMeshPipeline.cpp
--- a/Engine/Plugins/Experimental/Interchange/Runtime/Source/Pipelines/Private/InterchangeGenericStaticMeshPipeline.cpp
+++ b/Engine/Plugins/Experimental/Interchange/Runtime/Source/Pipelines/Private/InterchangeGenericStaticMeshPipeline.cpp
@@ -480,6 +480,7 @@ void UInterchangeGenericMeshPipeline::AddLodDataToStaticMesh(UInterchangeStaticM
 					const UInterchangeMeshNode* MeshDependencyNode = Cast<UInterchangeMeshNode>(BaseNodeContainer->GetNode(MeshDependency));
 					UInterchangeUserDefinedAttributesAPI::DuplicateAllUserDefinedAttribute(MeshDependencyNode, StaticMeshFactoryNode, bAddSourceNodeName);
 					StaticMeshFactoryNode->AddTargetNodeUid(MeshDependency);
+					StaticMeshFactoryNode->AddSocketUids(PipelineMeshesUtilities->GetMeshGeometryByUid(MeshDependency).AttachedSocketUids);
 					MeshDependencyNode->AddTargetNodeUid(StaticMeshFactoryNode->GetUniqueID());
 					MeshDependencyNode->GetMaterialDependencies(MaterialDependencies);
 				}
@@ -487,12 +488,14 @@ void UInterchangeGenericMeshPipeline::AddLodDataToStaticMesh(UInterchangeStaticM
 				{
 					SceneNode->GetMaterialDependencyUids(MaterialDependencies);
 				}
+
 				UInterchangeUserDefinedAttributesAPI::DuplicateAllUserDefinedAttribute(SceneNode, StaticMeshFactoryNode, bAddSourceNodeName);
 			}
 			else if (const UInterchangeMeshNode* MeshNode = Cast<UInterchangeMeshNode>(BaseNodeContainer->GetNode(NodeUid)))
 			{
 				UInterchangeUserDefinedAttributesAPI::DuplicateAllUserDefinedAttribute(MeshNode, StaticMeshFactoryNode, bAddSourceNodeName);
 				StaticMeshFactoryNode->AddTargetNodeUid(NodeUid);
+				StaticMeshFactoryNode->AddSocketUids(PipelineMeshesUtilities->GetMeshGeometryByUid(NodeUid).AttachedSocketUids);
 				MeshNode->AddTargetNodeUid(StaticMeshFactoryNode->GetUniqueID());
 				MeshNode->GetMaterialDependencies(MaterialDependencies);
 			}
diff --git a/Engine/Plugins/Experimental/Interchange/Runtime/Source/Pipelines/Private/InterchangePipelineMeshesUtilities.cpp b/Engine/Plugins/Experimental/Interchange/Runtime/Source/Pipelines/Private/InterchangePipelineMeshesUtilities.cpp
--- a/Engine/Plugins/Experimental/Interchange/Runtime/Source/Pipelines/Private/InterchangePipelineMeshesUtilities.cpp
+++ b/Engine/Plugins/Experimental/Interchange/Runtime/Source/Pipelines/Private/InterchangePipelineMeshesUtilities.cpp
@@ -11,6 +11,15 @@
 #include "UObject/ObjectMacros.h"
 
 
+static bool IsSceneNodeASocket(const UInterchangeSceneNode* SceneNode)
+{
+	// Generic pipeline determines where a scene node should be considered as a socket from its naming convention.
+	// This is no longer decided by the translator.
+	FString NodeDisplayName = SceneNode->GetDisplayLabel();
+	return NodeDisplayName.StartsWith(TEXT("SOCKET_"));
+}
+
+
 UInterchangePipelineMeshesUtilities* UInterchangePipelineMeshesUtilities::CreateInterchangePipelineMeshesUtilities(UInterchangeBaseNodeContainer* BaseNodeContainer)
 {
 	check(BaseNodeContainer);
@@ -19,108 +28,152 @@ UInterchangePipelineMeshesUtilities* UInterchangePipelineMeshesUtilities::Create
 	TArray<FString> SkeletonRootNodeUids;
 	
 	//Find all translated node we need for this pipeline
-	BaseNodeContainer->IterateNodes([&PipelineMeshesUtilities, &BaseNodeContainer, &SkeletonRootNodeUids](const FString& NodeUid, UInterchangeBaseNode* Node)
-	{
-		if (Node->GetNodeContainerType() == EInterchangeNodeContainerType::TranslatedAsset)
+	BaseNodeContainer->IterateNodes(
+		[&PipelineMeshesUtilities, &BaseNodeContainer, &SkeletonRootNodeUids](const FString& NodeUid, UInterchangeBaseNode* Node)
 		{
-			if (UInterchangeMeshNode* MeshNode = Cast<UInterchangeMeshNode>(Node))
+			if (Node->GetNodeContainerType() == EInterchangeNodeContainerType::TranslatedAsset)
 			{
-				FInterchangeMeshGeometry& MeshGeometry = PipelineMeshesUtilities->MeshGeometriesPerMeshUid.FindOrAdd(NodeUid);
-				MeshGeometry.MeshUid = NodeUid;
-				MeshGeometry.MeshNode = MeshNode;
+				if (UInterchangeMeshNode* MeshNode = Cast<UInterchangeMeshNode>(Node))
+				{
+					FInterchangeMeshGeometry& MeshGeometry = PipelineMeshesUtilities->MeshGeometriesPerMeshUid.FindOrAdd(NodeUid);
+					MeshGeometry.MeshUid = NodeUid;
+					MeshGeometry.MeshNode = MeshNode;
+				}
 			}
 		}
-	});
+	);
 
 	//Find all translated scene node we need for this pipeline
-	BaseNodeContainer->IterateNodes([&PipelineMeshesUtilities, &BaseNodeContainer, &SkeletonRootNodeUids](const FString& NodeUid, UInterchangeBaseNode* Node)
-	{
-		if (Node->GetNodeContainerType() == EInterchangeNodeContainerType::TranslatedScene)
+	bool bHasSockets = false;
+	BaseNodeContainer->IterateNodes(
+		[&PipelineMeshesUtilities, &BaseNodeContainer, &SkeletonRootNodeUids, &bHasSockets](const FString& NodeUid, const UInterchangeBaseNode* Node)
 		{
-			if (UInterchangeSceneNode* SceneNode = Cast<UInterchangeSceneNode>(Node))
+			if (Node->GetNodeContainerType() == EInterchangeNodeContainerType::TranslatedScene)
 			{
-				if (SceneNode->IsSpecializedTypeContains(UE::Interchange::FSceneNodeStaticData::GetJointSpecializeTypeString()))
+				if (const UInterchangeSceneNode* SceneNode = Cast<UInterchangeSceneNode>(Node))
 				{
-					const UInterchangeSceneNode* ParentJointNode = Cast<UInterchangeSceneNode>(BaseNodeContainer->GetNode(SceneNode->GetParentUid()));
-					if (!ParentJointNode || !ParentJointNode->IsSpecializedTypeContains(UE::Interchange::FSceneNodeStaticData::GetJointSpecializeTypeString()))
+					if (SceneNode->IsSpecializedTypeContains(UE::Interchange::FSceneNodeStaticData::GetJointSpecializeTypeString()))
 					{
-						SkeletonRootNodeUids.Add(SceneNode->GetUniqueID());
+						const UInterchangeSceneNode* ParentJointNode = Cast<UInterchangeSceneNode>(BaseNodeContainer->GetNode(SceneNode->GetParentUid()));
+						if (!ParentJointNode || !ParentJointNode->IsSpecializedTypeContains(UE::Interchange::FSceneNodeStaticData::GetJointSpecializeTypeString()))
+						{
+							SkeletonRootNodeUids.Add(SceneNode->GetUniqueID());
+						}
 					}
-				}
 
-				FString MeshUid;
-				if (SceneNode->GetCustomAssetInstanceUid(MeshUid))
-				{
-					if (BaseNodeContainer->GetNode(MeshUid)->IsA<UInterchangeMeshNode>())
+					if (IsSceneNodeASocket(SceneNode))
 					{
-						UInterchangeSceneNode* ParentMeshSceneNode = Cast<UInterchangeSceneNode>(BaseNodeContainer->GetNode(SceneNode->GetParentUid()));
-						if (ParentMeshSceneNode)
+						bHasSockets = true;
+					}
+
+					FString MeshUid;
+					if (SceneNode->GetCustomAssetInstanceUid(MeshUid))
+					{
+						if (BaseNodeContainer->GetNode(MeshUid)->IsA<UInterchangeMeshNode>())
 						{
-							UInterchangeSceneNode* LodGroupNode = nullptr;
-							int32 LodIndex = 0;
-							FString LastChildUid = SceneNode->GetUniqueID();
-							do
+							const UInterchangeSceneNode* ParentMeshSceneNode = Cast<UInterchangeSceneNode>(BaseNodeContainer->GetNode(SceneNode->GetParentUid()));
+							if (ParentMeshSceneNode)
 							{
-								if (ParentMeshSceneNode->IsSpecializedTypeContains(UE::Interchange::FSceneNodeStaticData::GetLodGroupSpecializeTypeString()))
+								const UInterchangeSceneNode* LodGroupNode = nullptr;
+								int32 LodIndex = 0;
+								FString LastChildUid = SceneNode->GetUniqueID();
+								do
 								{
-									LodGroupNode = ParentMeshSceneNode;
-									TArray<FString> LodGroupChildrens = BaseNodeContainer->GetNodeChildrenUids(ParentMeshSceneNode->GetUniqueID());
-									for (int32 ChildLodIndex = 0; ChildLodIndex < LodGroupChildrens.Num(); ++ChildLodIndex)
+									if (ParentMeshSceneNode->IsSpecializedTypeContains(UE::Interchange::FSceneNodeStaticData::GetLodGroupSpecializeTypeString()))
 									{
-										const FString& ChildrenUid = LodGroupChildrens[ChildLodIndex];
-										if (ChildrenUid.Equals(LastChildUid))
+										LodGroupNode = ParentMeshSceneNode;
+										TArray<FString> LodGroupChildrens = BaseNodeContainer->GetNodeChildrenUids(ParentMeshSceneNode->GetUniqueID());
+										for (int32 ChildLodIndex = 0; ChildLodIndex < LodGroupChildrens.Num(); ++ChildLodIndex)
 										{
-											LodIndex = ChildLodIndex;
-											break;
+											const FString& ChildrenUid = LodGroupChildrens[ChildLodIndex];
+											if (ChildrenUid.Equals(LastChildUid))
+											{
+												LodIndex = ChildLodIndex;
+												break;
+											}
 										}
+										break;
 									}
-									break;
-								}
-								LastChildUid = ParentMeshSceneNode->GetUniqueID();
-								ParentMeshSceneNode = Cast<UInterchangeSceneNode>(BaseNodeContainer->GetNode(ParentMeshSceneNode->GetParentUid()));
-							} while (ParentMeshSceneNode);
+									LastChildUid = ParentMeshSceneNode->GetUniqueID();
+									ParentMeshSceneNode = Cast<UInterchangeSceneNode>(BaseNodeContainer->GetNode(ParentMeshSceneNode->GetParentUid()));
+								} while (ParentMeshSceneNode);
 
-							FInterchangeMeshGeometry& MeshGeometry = PipelineMeshesUtilities->MeshGeometriesPerMeshUid.FindChecked(MeshUid);
-							if (LodGroupNode)
-							{
-								//We have a LOD
-								FInterchangeMeshInstance& MeshInstance = PipelineMeshesUtilities->MeshInstancesPerMeshInstanceUid.FindOrAdd(LodGroupNode->GetUniqueID());
-								if (MeshInstance.LodGroupNode != nullptr)
+								FInterchangeMeshGeometry& MeshGeometry = PipelineMeshesUtilities->MeshGeometriesPerMeshUid.FindChecked(MeshUid);
+								if (LodGroupNode)
 								{
-									//This LodGroup was already created, verify everything is ok
-									checkSlow(MeshInstance.LodGroupNode == LodGroupNode);
-									checkSlow(MeshInstance.MeshInstanceUid.Equals(LodGroupNode->GetUniqueID()));
+									//We have a LOD
+									FInterchangeMeshInstance& MeshInstance = PipelineMeshesUtilities->MeshInstancesPerMeshInstanceUid.FindOrAdd(LodGroupNode->GetUniqueID());
+									if (MeshInstance.LodGroupNode != nullptr)
+									{
+										//This LodGroup was already created, verify everything is ok
+										checkSlow(MeshInstance.LodGroupNode == LodGroupNode);
+										checkSlow(MeshInstance.MeshInstanceUid.Equals(LodGroupNode->GetUniqueID()));
+									}
+									else
+									{
+										MeshInstance.LodGroupNode = LodGroupNode;
+										MeshInstance.MeshInstanceUid = LodGroupNode->GetUniqueID();
+									}
+									FInterchangeLodSceneNodeContainer& InstancedSceneNodes = MeshInstance.SceneNodePerLodIndex.FindOrAdd(LodIndex);
+									InstancedSceneNodes.SceneNodes.AddUnique(SceneNode);
+									MeshGeometry.ReferencingMeshInstanceUids.Add(MeshInstance.MeshInstanceUid);
+									MeshInstance.ReferencingMeshGeometryUids.Add(MeshUid);
+									MeshInstance.bReferenceSkinnedMesh |= MeshGeometry.MeshNode->IsSkinnedMesh();
+									MeshInstance.bReferenceBlendShape |= MeshGeometry.MeshNode->IsBlendShape();
 								}
 								else
 								{
-									MeshInstance.LodGroupNode = LodGroupNode;
-									MeshInstance.MeshInstanceUid = LodGroupNode->GetUniqueID();
+									FInterchangeMeshInstance& MeshInstance = PipelineMeshesUtilities->MeshInstancesPerMeshInstanceUid.FindOrAdd(NodeUid);
+									MeshInstance.LodGroupNode = nullptr;
+									MeshInstance.MeshInstanceUid = NodeUid;
+									FInterchangeLodSceneNodeContainer& InstancedSceneNodes = MeshInstance.SceneNodePerLodIndex.FindOrAdd(LodIndex);
+									InstancedSceneNodes.SceneNodes.AddUnique(SceneNode);
+									MeshGeometry.ReferencingMeshInstanceUids.Add(MeshInstance.MeshInstanceUid);
+									MeshInstance.ReferencingMeshGeometryUids.Add(MeshUid);
+									MeshInstance.bReferenceSkinnedMesh |= MeshGeometry.MeshNode->IsSkinnedMesh();
+									MeshInstance.bReferenceBlendShape |= MeshGeometry.MeshNode->IsBlendShape();
 								}
-								FInterchangeLodSceneNodeContainer& InstancedSceneNodes = MeshInstance.SceneNodePerLodIndex.FindOrAdd(LodIndex);
-								InstancedSceneNodes.SceneNodes.AddUnique(SceneNode);
-								MeshGeometry.ReferencingMeshInstanceUids.Add(MeshInstance.MeshInstanceUid);
-								MeshInstance.ReferencingMeshGeometryUids.Add(MeshUid);
-								MeshInstance.bReferenceSkinnedMesh |= MeshGeometry.MeshNode->IsSkinnedMesh();
-								MeshInstance.bReferenceBlendShape |= MeshGeometry.MeshNode->IsBlendShape();
 							}
-							else
+						}
+					}
+				}
+			}
+		}
+	);
+
+	// Do a second pass to discover sockets
+	if (bHasSockets)
+	{
+		BaseNodeContainer->IterateNodes(
+			[&PipelineMeshesUtilities, &BaseNodeContainer](const FString& NodeUid, const UInterchangeBaseNode* Node)
+			{
+				if (const UInterchangeSceneNode* SceneNode = Cast<UInterchangeSceneNode>(Node))
+				{
+					if (IsSceneNodeASocket(SceneNode))
+					{
+						FString MeshUid;
+						const UInterchangeSceneNode* ParentMeshSceneNode = Cast<UInterchangeSceneNode>(BaseNodeContainer->GetNode(SceneNode->GetParentUid()));
+						while (ParentMeshSceneNode)
+						{
+							if (ParentMeshSceneNode->GetCustomAssetInstanceUid(MeshUid))
 							{
-								FInterchangeMeshInstance& MeshInstance = PipelineMeshesUtilities->MeshInstancesPerMeshInstanceUid.FindOrAdd(NodeUid);
-								MeshInstance.LodGroupNode = nullptr;
-								MeshInstance.MeshInstanceUid = NodeUid;
-								FInterchangeLodSceneNodeContainer& InstancedSceneNodes = MeshInstance.SceneNodePerLodIndex.FindOrAdd(LodIndex);
-								InstancedSceneNodes.SceneNodes.AddUnique(SceneNode);
-								MeshGeometry.ReferencingMeshInstanceUids.Add(MeshInstance.MeshInstanceUid);
-								MeshInstance.ReferencingMeshGeometryUids.Add(MeshUid);
-								MeshInstance.bReferenceSkinnedMesh |= MeshGeometry.MeshNode->IsSkinnedMesh();
-								MeshInstance.bReferenceBlendShape |= MeshGeometry.MeshNode->IsBlendShape();
+								break;
 							}
+
+							ParentMeshSceneNode = Cast<UInterchangeSceneNode>(BaseNodeContainer->GetNode(SceneNode->GetParentUid()));
+						}
+
+						if (!MeshUid.IsEmpty())
+						{
+							FInterchangeMeshGeometry& MeshGeometry = PipelineMeshesUtilities->MeshGeometriesPerMeshUid.FindChecked(MeshUid);
+							MeshGeometry.AttachedSocketUids.Add(SceneNode->GetUniqueID());
 						}
 					}
 				}
 			}
-		}
-	});
+		);
+	}
+
 
 	//Fill the SkeletonRootUidPerMeshUid data
 	for (const TPair<FString, FInterchangeMeshGeometry>& MeshGeometryUidAndMeshGeometry : PipelineMeshesUtilities->MeshGeometriesPerMeshUid)
diff --git a/Engine/Plugins/Experimental/Interchange/Runtime/Source/Pipelines/Public/InterchangePipelineMeshesUtilities.h b/Engine/Plugins/Experimental/Interchange/Runtime/Source/Pipelines/Public/InterchangePipelineMeshesUtilities.h
--- a/Engine/Plugins/Experimental/Interchange/Runtime/Source/Pipelines/Public/InterchangePipelineMeshesUtilities.h
+++ b/Engine/Plugins/Experimental/Interchange/Runtime/Source/Pipelines/Public/InterchangePipelineMeshesUtilities.h
@@ -24,7 +24,7 @@ struct FInterchangeLodSceneNodeContainer
 	* Each scene node here represent a mesh scene node. Only if we represent a lod group we can have more then 1 mesh scene node for a specific lod index.
 	*/
 	UPROPERTY(EditAnywhere, Category = "Interchange | Pipeline | MeshesInstance")
-	TArray<TObjectPtr<UInterchangeSceneNode>> SceneNodes;
+	TArray<TObjectPtr<const UInterchangeSceneNode>> SceneNodes;
 };
 
 /*
@@ -53,7 +53,7 @@ struct FInterchangeMeshInstance
 	* If this mesh instance represent a LodGroup this member will not be null, but will be null if the mesh instance do not represent a lod group
 	*/
 	UPROPERTY(EditAnywhere, Category = "Interchange | Pipeline | MeshesInstance")
-	TObjectPtr<UInterchangeSceneNode> LodGroupNode;
+	TObjectPtr<const UInterchangeSceneNode> LodGroupNode;
 
 	UPROPERTY(EditAnywhere, Category = "Interchange | Pipeline | MeshesInstance")
 	bool bReferenceSkinnedMesh;
@@ -97,13 +97,19 @@ struct FInterchangeMeshGeometry
 	* The UInterchangeMeshNode pointer represent by this structure.
 	*/
 	UPROPERTY(EditAnywhere, Category = "Interchange | Pipeline | MeshesGeometry")
-	TObjectPtr<UInterchangeMeshNode> MeshNode = nullptr;
+	TObjectPtr<const UInterchangeMeshNode> MeshNode = nullptr;
 
 	/**
 	* All mesh instance referencing this UInterchangeMeshNode pointer.
 	*/
 	UPROPERTY(EditAnywhere, Category = "Interchange | Pipeline | MeshesGeometry")
 	TArray<FString> ReferencingMeshInstanceUids;
+
+	/**
+	 * A list of all scene nodes representing sockets attached to this mesh
+	 */
+	UPROPERTY(EditAnywhere, Category = "Interchange | Pipeline | MeshesGeometry")
+	TArray<FString> AttachedSocketUids;
 };
 
 
diff --git a/Engine/Plugins/Tests/InterchangeTests/Source/InterchangeTests/Private/ImportTestFunctions/StaticMeshImportTestFunctions.cpp b/Engine/Plugins/Tests/InterchangeTests/Source/InterchangeTests/Private/ImportTestFunctions/StaticMeshImportTestFunctions.cpp
--- a/Engine/Plugins/Tests/InterchangeTests/Source/InterchangeTests/Private/ImportTestFunctions/StaticMeshImportTestFunctions.cpp
+++ b/Engine/Plugins/Tests/InterchangeTests/Source/InterchangeTests/Private/ImportTestFunctions/StaticMeshImportTestFunctions.cpp
@@ -2,6 +2,7 @@
 
 #include "ImportTestFunctions/StaticMeshImportTestFunctions.h"
 #include "Engine/StaticMesh.h"
+#include "Engine/StaticMeshSocket.h"
 #include "Engine/SkeletalMesh.h"
 #include "PhysicsEngine/BodySetup.h"
 #include "StaticMeshAttributes.h"
@@ -454,6 +455,60 @@ FInterchangeTestFunctionResult UStaticMeshImportTestFunctions::CheckSimpleCollis
 }
 
 
+FInterchangeTestFunctionResult UStaticMeshImportTestFunctions::CheckSocketCount(UStaticMesh* Mesh, int32 ExpectedSocketCount)
+{
+	FInterchangeTestFunctionResult Result;
+
+	int32 SocketCount = Mesh->Sockets.Num();
+	if (SocketCount != ExpectedSocketCount)
+	{
+		Result.AddError(FString::Printf(TEXT("Expected %d sockets, imported %d."), ExpectedSocketCount, SocketCount));
+	}
+
+	return Result;
+}
+
+
+FInterchangeTestFunctionResult UStaticMeshImportTestFunctions::CheckSocketName(UStaticMesh* Mesh, int32 SocketIndex, const FString& ExpectedSocketName)
+{
+	FInterchangeTestFunctionResult Result;
+
+	int32 SocketCount = Mesh->Sockets.Num();
+	if (SocketIndex >= SocketCount)
+	{
+		Result.AddError(FString::Printf(TEXT("The imported mesh doesn't contain %s sockets."), SocketIndex));
+	}
+
+	FString SocketName = Mesh->Sockets[SocketIndex]->SocketName.ToString();
+	if (SocketName != ExpectedSocketName)
+	{
+		Result.AddError(FString::Printf(TEXT("Expected socket name '%s', imported '%s'"), *ExpectedSocketName, *SocketName));
+	}
+
+	return Result;
+}
+
+
+FInterchangeTestFunctionResult UStaticMeshImportTestFunctions::CheckSocketLocation(UStaticMesh* Mesh, int32 SocketIndex, const FVector& ExpectedSocketLocation)
+{
+	FInterchangeTestFunctionResult Result;
+
+	int32 SocketCount = Mesh->Sockets.Num();
+	if (SocketIndex >= SocketCount)
+	{
+		Result.AddError(FString::Printf(TEXT("The imported mesh doesn't contain %s sockets."), SocketIndex));
+	}
+
+	FVector SocketLocation = Mesh->Sockets[SocketIndex]->RelativeLocation;
+	if (SocketLocation != ExpectedSocketLocation)
+	{
+		Result.AddError(FString::Printf(TEXT("Expected socket location %s, imported %s"), *ExpectedSocketLocation.ToString(), *SocketLocation.ToString()));
+	}
+
+	return Result;
+}
+
+
 FInterchangeTestFunctionResult UStaticMeshImportTestFunctions::CheckAgainstGroundTruth(UStaticMesh* Mesh, TSoftObjectPtr<UStaticMesh> MeshToCompare,
 	bool bCheckVertexCountEqual,
 	bool bCheckTriangleCountEqual,
diff --git a/Engine/Plugins/Tests/InterchangeTests/Source/InterchangeTests/Public/ImportTestFunctions/StaticMeshImportTestFunctions.h b/Engine/Plugins/Tests/InterchangeTests/Source/InterchangeTests/Public/ImportTestFunctions/StaticMeshImportTestFunctions.h
--- a/Engine/Plugins/Tests/InterchangeTests/Source/InterchangeTests/Public/ImportTestFunctions/StaticMeshImportTestFunctions.h
+++ b/Engine/Plugins/Tests/InterchangeTests/Source/InterchangeTests/Public/ImportTestFunctions/StaticMeshImportTestFunctions.h
@@ -106,6 +106,18 @@ public:
 	UFUNCTION(Exec)
 	static FInterchangeTestFunctionResult CheckSimpleCollisionPrimitiveCount(UStaticMesh* Mesh, int32 ExpectedSphereElementCount, int32 ExpectedBoxElementCount, int32 ExpectedCapsuleElementCount, int32 ExpectedConvexElementCount, int32 ExpectedTaperedCapsuleElementCount);
 
+	/** Check whether the expected number of sockets were imported */
+	UFUNCTION(Exec)
+	static FInterchangeTestFunctionResult CheckSocketCount(UStaticMesh* Mesh, int32 ExpectedSocketCount);
+
+	/** Check whether the given socket index has the expected name */
+	UFUNCTION(Exec)
+	static FInterchangeTestFunctionResult CheckSocketName(UStaticMesh* Mesh, int32 SocketIndex, const FString& ExpectedSocketName);
+
+	/** Check whether the given socket index has the expected location */
+	UFUNCTION(Exec)
+	static FInterchangeTestFunctionResult CheckSocketLocation(UStaticMesh* Mesh, int32 SocketIndex, const FVector& ExpectedSocketLocation);
+
 	/** Check whether the mesh is equivalent to a ground truth asset */
 	UFUNCTION(Exec)
 	static FInterchangeTestFunctionResult CheckAgainstGroundTruth(UStaticMesh* Mesh, TSoftObjectPtr<UStaticMesh> MeshToCompare,
