Made interchange factory base nodes a more concrete and distinct base type. Amended various parts of Interchange Manager to require factory base nodes where they are expected.  Moved factory node specific properties, such as the ReferencedObject, out of the base node class.

diff --git a/Engine/Plugins/Enterprise/DatasmithImporter/Source/DatasmithImporter/Private/DatasmithTextureImporter.cpp b/Engine/Plugins/Enterprise/DatasmithImporter/Source/DatasmithImporter/Private/DatasmithTextureImporter.cpp
--- a/Engine/Plugins/Enterprise/DatasmithImporter/Source/DatasmithImporter/Private/DatasmithTextureImporter.cpp
+++ b/Engine/Plugins/Enterprise/DatasmithImporter/Source/DatasmithImporter/Private/DatasmithTextureImporter.cpp
@@ -423,7 +423,7 @@ void UDatasmithTexturePipeline::ExecutePreImportPipeline(UInterchangeBaseNodeCon
 		return;
 	}
 
-	UInterchangeTextureNode* TextureTranslatedNode = Cast<UInterchangeTextureNode >( BaseNodeContainer->GetNode( Nodes[0] ) );
+	const UInterchangeTextureNode* TextureTranslatedNode = Cast<const UInterchangeTextureNode >( BaseNodeContainer->GetNode( Nodes[0] ) );
 	if (!ensure(TextureTranslatedNode))
 	{
 		return;
diff --git a/Engine/Plugins/Experimental/Interchange/Editor/Source/Pipelines/Private/SInterchangeGraphInspectorWindow.cpp b/Engine/Plugins/Experimental/Interchange/Editor/Source/Pipelines/Private/SInterchangeGraphInspectorWindow.cpp
--- a/Engine/Plugins/Experimental/Interchange/Editor/Source/Pipelines/Private/SInterchangeGraphInspectorWindow.cpp
+++ b/Engine/Plugins/Experimental/Interchange/Editor/Source/Pipelines/Private/SInterchangeGraphInspectorWindow.cpp
@@ -9,6 +9,7 @@
 #include "Modules/ModuleManager.h"
 #include "Nodes/InterchangeBaseNode.h"
 #include "Nodes/InterchangeBaseNodeContainer.h"
+#include "Nodes/InterchangeFactoryBaseNode.h"
 #include "PropertyEditorModule.h"
 #include "Styling/SlateIconFinder.h"
 #include "Styling/StyleColors.h"
@@ -42,9 +43,8 @@ void SInterchangeGraphInspectorTreeView::Construct(const FArguments& InArgs)
 	InterchangeBaseNodeContainer->GetRoots(Roots);
 	for (FString RootID : Roots)
 	{
-		UInterchangeBaseNode* RootNode = InterchangeBaseNodeContainer->GetNode(RootID);
-
-		RootNodeArray.Add(RootNode);
+		const UInterchangeBaseNode* RootNode = InterchangeBaseNodeContainer->GetNode(RootID);
+		RootNodeArray.Add(const_cast<UInterchangeBaseNode*>(RootNode));
 	}
 
 	STreeView::Construct
@@ -83,33 +83,35 @@ public:
 	{
 		InterchangeNode = InArgs._InterchangeNode;
 		InterchangeBaseNodeContainer = InArgs._InterchangeBaseNodeContainer;
-		//This is suppose to always be valid
+
+		// This is supposed to always be valid
 		check(InterchangeNode);
 		check(InterchangeBaseNodeContainer);
 
-		UClass* IconClass = InterchangeNode->GetObjectClass();
+		FString Tooltip = InterchangeNode->GetDisplayLabel();
 		const FSlateBrush* TypeIcon = nullptr;
 		FName IconName = InterchangeNode->GetIconName();
-		if (IconClass)
-		{
-			TypeIcon = FSlateIconFinder::FindIconBrushForClass(IconClass);
-		}
-		else if (IconName != NAME_None)
+		if (IconName != NAME_None)
 		{
 			const FSlateIcon SlateIcon = FSlateIconFinder::FindIcon(IconName);
 			TypeIcon = SlateIcon.GetOptionalIcon();
 		}
-		
-		if (!TypeIcon)
+
+		// Factory nodes take their icons directly from the class they represent, regardless of any icon which may be defined
+		// @TODO: allow GetIconName() to override this for factory nodes?
+		if (UInterchangeFactoryBaseNode* FactoryNode = Cast<UInterchangeFactoryBaseNode>(InterchangeNode))
 		{
-			TypeIcon = FSlateIconFinder::FindIconBrushForClass(AActor::StaticClass());
+			if (UClass* IconClass = FactoryNode->GetObjectClass())
+			{
+				TypeIcon = FSlateIconFinder::FindIconBrushForClass(IconClass);
+			}
+
+			Tooltip += TEXT(" [") + FactoryNode->GetObjectClass()->GetName() + TEXT("]");
 		}
 
-		//Prepare the tooltip
-		FString Tooltip = InterchangeNode->GetDisplayLabel();
-		if (InterchangeNode->GetObjectClass())
+		if (!TypeIcon)
 		{
-			Tooltip += TEXT(" [") + InterchangeNode->GetObjectClass()->GetName() + TEXT("]");
+			TypeIcon = FSlateIconFinder::FindIconBrushForClass(AActor::StaticClass());
 		}
 
 		FSlateColor TypeIconColor;
@@ -182,7 +184,7 @@ private:
 		TArray<FString> Childrens = InterchangeBaseNodeContainer->GetNodeChildrenUids(Node->GetUniqueID());
 		for (const FString& ChildID : Childrens)
 		{
-			UInterchangeBaseNode* ChildNode = InterchangeBaseNodeContainer->GetNode(ChildID);
+			UInterchangeBaseNode* ChildNode = const_cast<UInterchangeBaseNode*>(InterchangeBaseNodeContainer->GetNode(ChildID));
 			if (!ChildNode)
 				continue;
 			ChildNode->SetEnabled(bState);
@@ -221,10 +223,10 @@ void SInterchangeGraphInspectorTreeView::OnGetChildrenGraphInspectorTreeView(UIn
 	TArray<FString> Childrens = InterchangeBaseNodeContainer->GetNodeChildrenUids(InParent->GetUniqueID());
 	for (const FString& ChildID : Childrens)
 	{
-		UInterchangeBaseNode* ChildNode = InterchangeBaseNodeContainer->GetNode(ChildID);
+		const UInterchangeBaseNode* ChildNode = InterchangeBaseNodeContainer->GetNode(ChildID);
 		if (!ChildNode)
 			continue;
-		OutChildren.Add(ChildNode);
+		OutChildren.Add(const_cast<UInterchangeBaseNode*>(ChildNode));
 	}
 }
 
@@ -233,7 +235,7 @@ void SInterchangeGraphInspectorTreeView::RecursiveSetImport(UInterchangeBaseNode
 	TArray<FString> Childrens = InterchangeBaseNodeContainer->GetNodeChildrenUids(Node->GetUniqueID());
 	for (const FString& ChildID : Childrens)
 	{
-		UInterchangeBaseNode* ChildNode = InterchangeBaseNodeContainer->GetNode(ChildID);
+		UInterchangeBaseNode* ChildNode = const_cast<UInterchangeBaseNode*>(InterchangeBaseNodeContainer->GetNode(ChildID));
 		if (!ChildNode)
 			continue;
 		ChildNode->SetEnabled(bState);
@@ -259,7 +261,7 @@ void SInterchangeGraphInspectorTreeView::RecursiveSetExpand(UInterchangeBaseNode
 	TArray<FString> Childrens = InterchangeBaseNodeContainer->GetNodeChildrenUids(Node->GetUniqueID());
 	for (const FString& ChildID : Childrens)
 	{
-		UInterchangeBaseNode* ChildNode = InterchangeBaseNodeContainer->GetNode(ChildID);
+		UInterchangeBaseNode* ChildNode = const_cast<UInterchangeBaseNode*>(InterchangeBaseNodeContainer->GetNode(ChildID));
 		if (!ChildNode)
 			continue;
 		RecursiveSetExpand(ChildNode, ExpandState);
diff --git a/Engine/Plugins/Experimental/Interchange/Runtime/Source/FactoryNodes/Public/InterchangeSkeletalMeshLodDataNode.h b/Engine/Plugins/Experimental/Interchange/Runtime/Source/FactoryNodes/Public/InterchangeSkeletalMeshLodDataNode.h
--- a/Engine/Plugins/Experimental/Interchange/Runtime/Source/FactoryNodes/Public/InterchangeSkeletalMeshLodDataNode.h
+++ b/Engine/Plugins/Experimental/Interchange/Runtime/Source/FactoryNodes/Public/InterchangeSkeletalMeshLodDataNode.h
@@ -3,7 +3,7 @@
 #pragma once
 
 #include "CoreMinimal.h"
-#include "Nodes/InterchangeBaseNode.h"
+#include "Nodes/InterchangeFactoryBaseNode.h"
 
 #include "InterchangeSkeletalMeshLodDataNode.generated.h"
 
@@ -22,7 +22,7 @@ namespace UE
 }//ns UE
 
 UCLASS(BlueprintType, Experimental)
-class INTERCHANGEFACTORYNODES_API UInterchangeSkeletalMeshLodDataNode : public UInterchangeBaseNode
+class INTERCHANGEFACTORYNODES_API UInterchangeSkeletalMeshLodDataNode : public UInterchangeFactoryBaseNode
 {
 	GENERATED_BODY()
 
diff --git a/Engine/Plugins/Experimental/Interchange/Runtime/Source/FactoryNodes/Public/InterchangeStaticMeshLodDataNode.h b/Engine/Plugins/Experimental/Interchange/Runtime/Source/FactoryNodes/Public/InterchangeStaticMeshLodDataNode.h
--- a/Engine/Plugins/Experimental/Interchange/Runtime/Source/FactoryNodes/Public/InterchangeStaticMeshLodDataNode.h
+++ b/Engine/Plugins/Experimental/Interchange/Runtime/Source/FactoryNodes/Public/InterchangeStaticMeshLodDataNode.h
@@ -4,7 +4,7 @@
 
 #include "CoreMinimal.h"
 #include "UObject/ObjectMacros.h"
-#include "Nodes/InterchangeBaseNode.h"
+#include "Nodes/InterchangeFactoryBaseNode.h"
 
 #include "InterchangeStaticMeshLodDataNode.generated.h"
 
@@ -26,7 +26,7 @@ namespace UE
 
 
 UCLASS(BlueprintType, Experimental)
-class INTERCHANGEFACTORYNODES_API UInterchangeStaticMeshLodDataNode : public UInterchangeBaseNode
+class INTERCHANGEFACTORYNODES_API UInterchangeStaticMeshLodDataNode : public UInterchangeFactoryBaseNode
 {
 	GENERATED_BODY()
 
diff --git a/Engine/Plugins/Experimental/Interchange/Runtime/Source/Import/Private/Animation/InterchangeAnimSequenceFactory.cpp b/Engine/Plugins/Experimental/Interchange/Runtime/Source/Import/Private/Animation/InterchangeAnimSequenceFactory.cpp
--- a/Engine/Plugins/Experimental/Interchange/Runtime/Source/Import/Private/Animation/InterchangeAnimSequenceFactory.cpp
+++ b/Engine/Plugins/Experimental/Interchange/Runtime/Source/Import/Private/Animation/InterchangeAnimSequenceFactory.cpp
@@ -412,12 +412,12 @@ UObject* UInterchangeAnimSequenceFactory::CreateAsset(const FCreateAssetParams&
 	{
 		//Apply the re import strategy 
 		UInterchangeAssetImportData* InterchangeAssetImportData = Cast<UInterchangeAssetImportData>(AnimSequence->AssetImportData);
-		UInterchangeBaseNode* PreviousNode = nullptr;
+		UInterchangeFactoryBaseNode* PreviousNode = nullptr;
 		if (InterchangeAssetImportData)
 		{
-			PreviousNode = InterchangeAssetImportData->NodeContainer->GetNode(InterchangeAssetImportData->NodeUniqueID);
+			PreviousNode = InterchangeAssetImportData->NodeContainer->GetFactoryNode(InterchangeAssetImportData->NodeUniqueID);
 		}
-		UInterchangeBaseNode* CurrentNode = NewObject<UInterchangeBaseNode>(GetTransientPackage(), UInterchangeAnimSequenceFactoryNode::StaticClass());
+		UInterchangeFactoryBaseNode* CurrentNode = NewObject<UInterchangeAnimSequenceFactoryNode>(GetTransientPackage());
 		UInterchangeBaseNode::CopyStorage(AnimSequenceFactoryNode, CurrentNode);
 		CurrentNode->FillAllCustomAttributeFromObject(AnimSequence);
 		UE::Interchange::FFactoryCommon::ApplyReimportStrategyToAsset(AnimSequence, PreviousNode, CurrentNode, AnimSequenceFactoryNode);
diff --git a/Engine/Plugins/Experimental/Interchange/Runtime/Source/Import/Private/Gltf/InterchangeGltfTranslator.cpp b/Engine/Plugins/Experimental/Interchange/Runtime/Source/Import/Private/Gltf/InterchangeGltfTranslator.cpp
--- a/Engine/Plugins/Experimental/Interchange/Runtime/Source/Import/Private/Gltf/InterchangeGltfTranslator.cpp
+++ b/Engine/Plugins/Experimental/Interchange/Runtime/Source/Import/Private/Gltf/InterchangeGltfTranslator.cpp
@@ -39,7 +39,7 @@ void UInterchangeGltfTranslator::HandleGltfNode( UInterchangeBaseNodeContainer&
 
 	const FString NodeUid = ParentNodeUid + TEXT("\\") + GenerateUniqueIdForGltfNode( GltfNode.Name, NodeIndex );
 
-	UInterchangeSceneNode* ParentSceneNode = Cast< UInterchangeSceneNode >( NodeContainer.GetNode( ParentNodeUid ) );
+	const UInterchangeSceneNode* ParentSceneNode = Cast< UInterchangeSceneNode >( NodeContainer.GetNode( ParentNodeUid ) );
 
 	UInterchangeSceneNode* InterchangeSceneNode = NewObject< UInterchangeSceneNode >( &NodeContainer );
 	InterchangeSceneNode->InitializeNode( NodeUid, GltfNode.Name, EInterchangeNodeContainerType::TranslatedScene );
@@ -556,7 +556,7 @@ void UInterchangeGltfTranslator::HandleGltfTransmission( UInterchangeBaseNodeCon
 	{
 		// There's no separation of reflected and transmitted color in this model. So the same color is used for the base color and the transmitted color.
 		// Since this extension is only supported with the metallic-roughness model, we can reuse its base color
-		UInterchangeBaseNode* CurrentNode = &ShaderGraphNode;
+		const UInterchangeBaseNode* CurrentNode = &ShaderGraphNode;
 		FString CurrentOuput = TEXT("");
 		FLinearColor CurrentColor = FLinearColor::White;
 
diff --git a/Engine/Plugins/Experimental/Interchange/Runtime/Source/Import/Private/Mesh/InterchangeOBJTranslator.cpp b/Engine/Plugins/Experimental/Interchange/Runtime/Source/Import/Private/Mesh/InterchangeOBJTranslator.cpp
--- a/Engine/Plugins/Experimental/Interchange/Runtime/Source/Import/Private/Mesh/InterchangeOBJTranslator.cpp
+++ b/Engine/Plugins/Experimental/Interchange/Runtime/Source/Import/Private/Mesh/InterchangeOBJTranslator.cpp
@@ -1002,16 +1002,16 @@ void UInterchangeOBJTranslator::AddMaterialNodes(UInterchangeBaseNodeContainer&
 		BaseNodeContainer.SetNodeParentUid(TextureSampleShaderUid, ShaderGraphNode->GetUniqueID());
 
 		FString TextureNodeUid = MakeTextureNodeUid(TextureName);
-		UInterchangeTexture2DNode* TextureNode = Cast<UInterchangeTexture2DNode>(BaseNodeContainer.GetNode(TextureNodeUid));
+		const UInterchangeTexture2DNode* TextureNode = Cast<const UInterchangeTexture2DNode>(BaseNodeContainer.GetNode(TextureNodeUid));
 		if (!TextureNode)
 		{
-			TextureNode = NewObject<UInterchangeTexture2DNode>(&BaseNodeContainer);
-			TextureNode->InitializeNode(TextureNodeUid, TextureName, EInterchangeNodeContainerType::TranslatedAsset);
-			BaseNodeContainer.AddNode(TextureNode);
+			UInterchangeTexture2DNode* NewTextureNode = NewObject<UInterchangeTexture2DNode>(&BaseNodeContainer);
+			NewTextureNode->InitializeNode(TextureNodeUid, TextureName, EInterchangeNodeContainerType::TranslatedAsset);
+			BaseNodeContainer.AddNode(NewTextureNode);
 
 			FString NormalizedTexturePath(TexturePath);
 			FPaths::NormalizeFilename(NormalizedTexturePath);
-			TextureNode->SetPayLoadKey(NormalizedTexturePath);
+			NewTextureNode->SetPayLoadKey(NormalizedTexturePath);
 		}
 
 		TextureSampleShader->AddStringAttribute(UInterchangeShaderPortsAPI::MakeInputValueKey(TextureSample::Inputs::Texture.ToString()), TextureNodeUid);
diff --git a/Engine/Plugins/Experimental/Interchange/Runtime/Source/Import/Private/Mesh/InterchangeSkeletalMeshFactory.cpp b/Engine/Plugins/Experimental/Interchange/Runtime/Source/Import/Private/Mesh/InterchangeSkeletalMeshFactory.cpp
--- a/Engine/Plugins/Experimental/Interchange/Runtime/Source/Import/Private/Mesh/InterchangeSkeletalMeshFactory.cpp
+++ b/Engine/Plugins/Experimental/Interchange/Runtime/Source/Import/Private/Mesh/InterchangeSkeletalMeshFactory.cpp
@@ -1406,12 +1406,12 @@ UObject* UInterchangeSkeletalMeshFactory::CreateAsset(const FCreateAssetParams&
 	{
 		//Apply the re import strategy 
 		UInterchangeAssetImportData* InterchangeAssetImportData = Cast<UInterchangeAssetImportData>(SkeletalMesh->GetAssetImportData());
-		UInterchangeBaseNode* PreviousNode = nullptr;
+		UInterchangeFactoryBaseNode* PreviousNode = nullptr;
 		if (InterchangeAssetImportData)
 		{
-			PreviousNode = InterchangeAssetImportData->NodeContainer->GetNode(InterchangeAssetImportData->NodeUniqueID);
+			PreviousNode = InterchangeAssetImportData->NodeContainer->GetFactoryNode(InterchangeAssetImportData->NodeUniqueID);
 		}
-		UInterchangeBaseNode* CurrentNode = NewObject<UInterchangeBaseNode>(GetTransientPackage(), UInterchangeSkeletalMeshFactoryNode::StaticClass());
+		UInterchangeFactoryBaseNode* CurrentNode = NewObject<UInterchangeSkeletalMeshFactoryNode>(GetTransientPackage());
 		UInterchangeBaseNode::CopyStorage(SkeletalMeshFactoryNode, CurrentNode);
 		CurrentNode->FillAllCustomAttributeFromObject(SkeletalMesh);
 		UE::Interchange::FFactoryCommon::ApplyReimportStrategyToAsset(SkeletalMesh, PreviousNode, CurrentNode, SkeletalMeshFactoryNode);
@@ -1467,7 +1467,7 @@ void UInterchangeSkeletalMeshFactory::PreImportPreCompletedCallback(const FImpor
 						: NSSkeletalMeshSourceFileLabels::GeoAndSkinningText().ToString();
 				};
 
-				if (UInterchangeSkeletalMeshFactoryNode* SkeletalMeshFactoryNode = Cast<UInterchangeSkeletalMeshFactoryNode>(Arguments.NodeContainer->GetNode(Arguments.NodeUniqueID)))
+				if (const UInterchangeSkeletalMeshFactoryNode* SkeletalMeshFactoryNode = Cast<const UInterchangeSkeletalMeshFactoryNode>(Arguments.NodeContainer->GetFactoryNode(Arguments.NodeUniqueID)))
 				{
 					EInterchangeSkeletalMeshContentType ImportContentType = EInterchangeSkeletalMeshContentType::All;
 					SkeletalMeshFactoryNode->GetCustomImportContentType(ImportContentType);
diff --git a/Engine/Plugins/Experimental/Interchange/Runtime/Source/Import/Private/Mesh/InterchangeStaticMeshFactory.cpp b/Engine/Plugins/Experimental/Interchange/Runtime/Source/Import/Private/Mesh/InterchangeStaticMeshFactory.cpp
--- a/Engine/Plugins/Experimental/Interchange/Runtime/Source/Import/Private/Mesh/InterchangeStaticMeshFactory.cpp
+++ b/Engine/Plugins/Experimental/Interchange/Runtime/Source/Import/Private/Mesh/InterchangeStaticMeshFactory.cpp
@@ -370,12 +370,12 @@ UObject* UInterchangeStaticMeshFactory::CreateAsset(const FCreateAssetParams& Ar
 	{
 		//Apply the re import strategy 
 		UInterchangeAssetImportData* InterchangeAssetImportData = Cast<UInterchangeAssetImportData>(StaticMesh->GetAssetImportData());
-		UInterchangeBaseNode* PreviousNode = nullptr;
+		UInterchangeFactoryBaseNode* PreviousNode = nullptr;
 		if (InterchangeAssetImportData)
 		{
-			PreviousNode = InterchangeAssetImportData->NodeContainer->GetNode(InterchangeAssetImportData->NodeUniqueID);
+			PreviousNode = InterchangeAssetImportData->NodeContainer->GetFactoryNode(InterchangeAssetImportData->NodeUniqueID);
 		}
-		UInterchangeBaseNode* CurrentNode = NewObject<UInterchangeBaseNode>(GetTransientPackage(), UInterchangeStaticMeshFactoryNode::StaticClass());
+		UInterchangeFactoryBaseNode* CurrentNode = NewObject<UInterchangeStaticMeshFactoryNode>(GetTransientPackage());
 		UInterchangeBaseNode::CopyStorage(StaticMeshFactoryNode, CurrentNode);
 		CurrentNode->FillAllCustomAttributeFromObject(StaticMesh);
 		UE::Interchange::FFactoryCommon::ApplyReimportStrategyToAsset(StaticMesh, PreviousNode, CurrentNode, StaticMeshFactoryNode);
diff --git a/Engine/Plugins/Experimental/Interchange/Runtime/Source/Import/Private/Scene/InterchangeActorFactory.cpp b/Engine/Plugins/Experimental/Interchange/Runtime/Source/Import/Private/Scene/InterchangeActorFactory.cpp
--- a/Engine/Plugins/Experimental/Interchange/Runtime/Source/Import/Private/Scene/InterchangeActorFactory.cpp
+++ b/Engine/Plugins/Experimental/Interchange/Runtime/Source/Import/Private/Scene/InterchangeActorFactory.cpp
@@ -1,7 +1,7 @@
 // Copyright Epic Games, Inc. All Rights Reserved. 
 #include "Scene/InterchangeActorFactory.h"
 
-#include "Nodes/InterchangeBaseNode.h"
+#include "Nodes/InterchangeFactoryBaseNode.h"
 #include "Scene/InterchangeActorHelper.h"
 
 #include "GameFramework/Actor.h"
diff --git a/Engine/Plugins/Experimental/Interchange/Runtime/Source/Import/Private/Scene/InterchangeActorHelper.cpp b/Engine/Plugins/Experimental/Interchange/Runtime/Source/Import/Private/Scene/InterchangeActorHelper.cpp
--- a/Engine/Plugins/Experimental/Interchange/Runtime/Source/Import/Private/Scene/InterchangeActorHelper.cpp
+++ b/Engine/Plugins/Experimental/Interchange/Runtime/Source/Import/Private/Scene/InterchangeActorHelper.cpp
@@ -18,7 +18,7 @@ AActor* UE::Interchange::ActorHelper::GetSpawnedParentActor(const UInterchangeBa
 {
 	AActor* ParentActor = nullptr;
 
-	if (const UInterchangeFactoryBaseNode* ParentFactoryNode = Cast<UInterchangeFactoryBaseNode>(NodeContainer->GetNode(FactoryNode->GetParentUid())))
+	if (const UInterchangeFactoryBaseNode* ParentFactoryNode = NodeContainer->GetFactoryNode(FactoryNode->GetParentUid()))
 	{
 		ParentActor = Cast<AActor>(ParentFactoryNode->ReferenceObject.TryLoad());
 	}
@@ -123,7 +123,7 @@ AActor* UE::Interchange::ActorHelper::SpawnFactoryActor(const UInterchangeFactor
 	return SpawnedActor;
 }
 
-const UInterchangeBaseNode* UE::Interchange::ActorHelper::FindAssetInstanceFactoryNode(const UInterchangeBaseNodeContainer* NodeContainer, const UInterchangeBaseNode* ActorFactoryNode)
+const UInterchangeFactoryBaseNode* UE::Interchange::ActorHelper::FindAssetInstanceFactoryNode(const UInterchangeBaseNodeContainer* NodeContainer, const UInterchangeFactoryBaseNode* ActorFactoryNode)
 {
 	TArray<FString> ActorTargetNodes;
 	ActorFactoryNode->GetTargetNodeUids(ActorTargetNodes);
@@ -143,5 +143,5 @@ const UInterchangeBaseNode* UE::Interchange::ActorHelper::FindAssetInstanceFacto
 
 	TArray<FString> AssetTargetNodeIds;
 	AssetNode->GetTargetNodeUids(AssetTargetNodeIds);
-	return AssetTargetNodeIds.IsEmpty() ? nullptr : NodeContainer->GetNode(AssetTargetNodeIds[0]);
+	return AssetTargetNodeIds.IsEmpty() ? nullptr : NodeContainer->GetFactoryNode(AssetTargetNodeIds[0]);
 }
\ No newline at end of file
diff --git a/Engine/Plugins/Experimental/Interchange/Runtime/Source/Import/Private/Scene/InterchangeSkeletalMeshActorFactory.cpp b/Engine/Plugins/Experimental/Interchange/Runtime/Source/Import/Private/Scene/InterchangeSkeletalMeshActorFactory.cpp
--- a/Engine/Plugins/Experimental/Interchange/Runtime/Source/Import/Private/Scene/InterchangeSkeletalMeshActorFactory.cpp
+++ b/Engine/Plugins/Experimental/Interchange/Runtime/Source/Import/Private/Scene/InterchangeSkeletalMeshActorFactory.cpp
@@ -19,7 +19,7 @@ UObject* UInterchangeSkeletalMeshActorFactory::CreateSceneObject(const UIntercha
 		return nullptr;
 	}
 
-	UInterchangeBaseNode* FactoryNode = CreateSceneObjectsParams.FactoryNode;
+	UInterchangeFactoryBaseNode* FactoryNode = CreateSceneObjectsParams.FactoryNode;
 	SetupSkeletalMeshActor(CreateSceneObjectsParams.NodeContainer, FactoryNode, SpawnedActor);
 
 	if (USkeletalMeshComponent* SkeletalMeshComponent = SpawnedActor->GetSkeletalMeshComponent())
@@ -35,16 +35,16 @@ UClass* UInterchangeSkeletalMeshActorFactory::GetFactoryClass() const
 	return ASkeletalMeshActor::StaticClass();
 }
 
-void UInterchangeSkeletalMeshActorFactory::SetupSkeletalMeshActor(const UInterchangeBaseNodeContainer* NodeContainer, const UInterchangeBaseNode* ActorFactoryNode, ASkeletalMeshActor* SkeletalMeshActor)
+void UInterchangeSkeletalMeshActorFactory::SetupSkeletalMeshActor(const UInterchangeBaseNodeContainer* NodeContainer, const UInterchangeFactoryBaseNode* ActorFactoryNode, ASkeletalMeshActor* SkeletalMeshActor)
 {
 	USkeletalMeshComponent* SkeletalMeshComponent = SkeletalMeshActor->GetSkeletalMeshComponent();
 	SkeletalMeshComponent->UnregisterComponent();
 
-	if (const UInterchangeBaseNode* MeshNode = UE::Interchange::ActorHelper::FindAssetInstanceFactoryNode(NodeContainer, ActorFactoryNode))
+	if (const UInterchangeFactoryBaseNode* MeshNode = UE::Interchange::ActorHelper::FindAssetInstanceFactoryNode(NodeContainer, ActorFactoryNode))
 	{
 		if (USkeletalMesh* SkeletalMesh = Cast<USkeletalMesh>(MeshNode->ReferenceObject.TryLoad()))
 		{
 			SkeletalMeshComponent->SetSkeletalMesh(SkeletalMesh);
 		}
 	}
-}
\ No newline at end of file
+}
diff --git a/Engine/Plugins/Experimental/Interchange/Runtime/Source/Import/Private/Scene/InterchangeStaticMeshActorFactory.cpp b/Engine/Plugins/Experimental/Interchange/Runtime/Source/Import/Private/Scene/InterchangeStaticMeshActorFactory.cpp
--- a/Engine/Plugins/Experimental/Interchange/Runtime/Source/Import/Private/Scene/InterchangeStaticMeshActorFactory.cpp
+++ b/Engine/Plugins/Experimental/Interchange/Runtime/Source/Import/Private/Scene/InterchangeStaticMeshActorFactory.cpp
@@ -23,7 +23,7 @@ UObject* UInterchangeStaticMeshActorFactory::CreateSceneObject(const UInterchang
 		return nullptr;
 	}
 
-	UInterchangeBaseNode* FactoryNode = CreateSceneObjectsParams.FactoryNode;
+	UInterchangeFactoryBaseNode* FactoryNode = CreateSceneObjectsParams.FactoryNode;
 	SetupStaticMeshActor(CreateSceneObjectsParams.NodeContainer, FactoryNode, SpawnedActor);
 
 	if (UStaticMeshComponent* StaticMeshComponent = SpawnedActor->GetStaticMeshComponent())
@@ -34,12 +34,12 @@ UObject* UInterchangeStaticMeshActorFactory::CreateSceneObject(const UInterchang
 	return SpawnedActor;
 };
 
-void UInterchangeStaticMeshActorFactory::SetupStaticMeshActor(const UInterchangeBaseNodeContainer* NodeContainer, const UInterchangeBaseNode* ActorFactoryNode, AStaticMeshActor* StaticMeshActor)
+void UInterchangeStaticMeshActorFactory::SetupStaticMeshActor(const UInterchangeBaseNodeContainer* NodeContainer, const UInterchangeFactoryBaseNode* ActorFactoryNode, AStaticMeshActor* StaticMeshActor)
 {
 	UStaticMeshComponent* StaticMeshComponent = StaticMeshActor->GetStaticMeshComponent();
 	StaticMeshComponent->UnregisterComponent();
 
-	if (const UInterchangeBaseNode* MeshNode = UE::Interchange::ActorHelper::FindAssetInstanceFactoryNode(NodeContainer, ActorFactoryNode))
+	if (const UInterchangeFactoryBaseNode* MeshNode = UE::Interchange::ActorHelper::FindAssetInstanceFactoryNode(NodeContainer, ActorFactoryNode))
 	{
 		if (UStaticMesh* StaticMesh = Cast<UStaticMesh>(MeshNode->ReferenceObject.TryLoad()))
 		{
diff --git a/Engine/Plugins/Experimental/Interchange/Runtime/Source/Import/Private/Texture/InterchangeTextureFactory.cpp b/Engine/Plugins/Experimental/Interchange/Runtime/Source/Import/Private/Texture/InterchangeTextureFactory.cpp
--- a/Engine/Plugins/Experimental/Interchange/Runtime/Source/Import/Private/Texture/InterchangeTextureFactory.cpp
+++ b/Engine/Plugins/Experimental/Interchange/Runtime/Source/Import/Private/Texture/InterchangeTextureFactory.cpp
@@ -122,7 +122,7 @@ namespace UE::Interchange::Private::InterchangeTextureFactory
 	/**
 	 * Return the supported class if the node is one otherwise return nullptr
 	 */
-	UClass* GetSupportedFactoryNodeClass(const UInterchangeBaseNode* AssetNode)
+	UClass* GetSupportedFactoryNodeClass(const UInterchangeFactoryBaseNode* AssetNode)
 	{
 		UClass* TextureCubeFactoryClass = UInterchangeTextureCubeFactoryNode::StaticClass();
 		UClass* Texture2DFactoryClass = UInterchangeTexture2DFactoryNode::StaticClass();
@@ -170,7 +170,7 @@ namespace UE::Interchange::Private::InterchangeTextureFactory
 		, UInterchangeTexture2DArrayFactoryNode*
 		,UInterchangeTextureLightProfileFactoryNode* >;
 
-	FTextureFactoryNodeVariant GetAsTextureFactoryNodeVariant(UInterchangeBaseNode* AssetNode, UClass* SupportedFactoryNodeClass)
+	FTextureFactoryNodeVariant GetAsTextureFactoryNodeVariant(UInterchangeFactoryBaseNode* AssetNode, UClass* SupportedFactoryNodeClass)
 	{
 		if (AssetNode)
 		{
@@ -1240,7 +1240,7 @@ void UInterchangeTextureFactory::PreImportPreCompletedCallback(const FImportPreC
 		}
 
 
-		UInterchangeBaseNode* TextureFactoryNode = Arguments.FactoryNode;
+		UInterchangeFactoryBaseNode* TextureFactoryNode = Arguments.FactoryNode;
 		if (!Arguments.bIsReimport)
 		{
 			/** Apply all TextureNode custom attributes to the texture asset */
@@ -1249,13 +1249,13 @@ void UInterchangeTextureFactory::PreImportPreCompletedCallback(const FImportPreC
 		else
 		{
 			UInterchangeAssetImportData* InterchangeAssetImportData = Cast<UInterchangeAssetImportData>(Texture->AssetImportData);
-			UInterchangeBaseNode* PreviousNode = nullptr;
+			UInterchangeFactoryBaseNode* PreviousNode = nullptr;
 			if (InterchangeAssetImportData)
 			{
-				PreviousNode = InterchangeAssetImportData->NodeContainer->GetNode(InterchangeAssetImportData->NodeUniqueID);
+				PreviousNode = InterchangeAssetImportData->NodeContainer->GetFactoryNode(InterchangeAssetImportData->NodeUniqueID);
 			}
 
-			UInterchangeBaseNode* CurrentNode = NewObject<UInterchangeBaseNode>(GetTransientPackage(), GetSupportedFactoryNodeClass(TextureFactoryNode));
+			UInterchangeFactoryBaseNode* CurrentNode = NewObject<UInterchangeFactoryBaseNode>(GetTransientPackage(), GetSupportedFactoryNodeClass(TextureFactoryNode));
 			UInterchangeBaseNode::CopyStorage(TextureFactoryNode, CurrentNode);
 			CurrentNode->FillAllCustomAttributeFromObject(Texture);
 			//Apply reimport strategy
diff --git a/Engine/Plugins/Experimental/Interchange/Runtime/Source/Import/Public/Scene/InterchangeActorHelper.h b/Engine/Plugins/Experimental/Interchange/Runtime/Source/Import/Public/Scene/InterchangeActorHelper.h
--- a/Engine/Plugins/Experimental/Interchange/Runtime/Source/Import/Public/Scene/InterchangeActorHelper.h
+++ b/Engine/Plugins/Experimental/Interchange/Runtime/Source/Import/Public/Scene/InterchangeActorHelper.h
@@ -4,7 +4,7 @@
 
 #include "InterchangeFactoryBase.h"
 
-class UInterchangeBaseNode;
+class UInterchangeFactoryBaseNode;
 class UInterchangeBaseNodeContainer;
 class UInterchangeActorFactoryNode;
 
@@ -24,5 +24,5 @@ namespace UE::Interchange::ActorHelper
 	/**
 	 * Returns the factory node of the asset instanced by ActorFactoryNode.
 	 */
-	const UInterchangeBaseNode* FindAssetInstanceFactoryNode(const UInterchangeBaseNodeContainer* NodeContainer, const UInterchangeBaseNode* ActorFactoryNode);
+	const UInterchangeFactoryBaseNode* FindAssetInstanceFactoryNode(const UInterchangeBaseNodeContainer* NodeContainer, const UInterchangeFactoryBaseNode* ActorFactoryNode);
 }
diff --git a/Engine/Plugins/Experimental/Interchange/Runtime/Source/Import/Public/Scene/InterchangeSkeletalMeshActorFactory.h b/Engine/Plugins/Experimental/Interchange/Runtime/Source/Import/Public/Scene/InterchangeSkeletalMeshActorFactory.h
--- a/Engine/Plugins/Experimental/Interchange/Runtime/Source/Import/Public/Scene/InterchangeSkeletalMeshActorFactory.h
+++ b/Engine/Plugins/Experimental/Interchange/Runtime/Source/Import/Public/Scene/InterchangeSkeletalMeshActorFactory.h
@@ -36,7 +36,7 @@ public:
 
 private:
 
-	void SetupSkeletalMeshActor(const UInterchangeBaseNodeContainer* NodeContainer, const UInterchangeBaseNode* ActorFactoryNode, ASkeletalMeshActor* SkeletalMeshActor);
+	void SetupSkeletalMeshActor(const UInterchangeBaseNodeContainer* NodeContainer, const UInterchangeFactoryBaseNode* ActorFactoryNode, ASkeletalMeshActor* SkeletalMeshActor);
 };
 
 
diff --git a/Engine/Plugins/Experimental/Interchange/Runtime/Source/Import/Public/Scene/InterchangeStaticMeshActorFactory.h b/Engine/Plugins/Experimental/Interchange/Runtime/Source/Import/Public/Scene/InterchangeStaticMeshActorFactory.h
--- a/Engine/Plugins/Experimental/Interchange/Runtime/Source/Import/Public/Scene/InterchangeStaticMeshActorFactory.h
+++ b/Engine/Plugins/Experimental/Interchange/Runtime/Source/Import/Public/Scene/InterchangeStaticMeshActorFactory.h
@@ -36,7 +36,7 @@ public:
 
 private:
 
-	void SetupStaticMeshActor(const UInterchangeBaseNodeContainer* NodeContainer, const UInterchangeBaseNode* ActorFactoryNode, AStaticMeshActor* StaticMeshActor);
+	void SetupStaticMeshActor(const UInterchangeBaseNodeContainer* NodeContainer, const UInterchangeFactoryBaseNode* ActorFactoryNode, AStaticMeshActor* StaticMeshActor);
 };
 
 
diff --git a/Engine/Plugins/Experimental/Interchange/Runtime/Source/Parsers/Fbx/Private/FbxCamera.cpp b/Engine/Plugins/Experimental/Interchange/Runtime/Source/Parsers/Fbx/Private/FbxCamera.cpp
--- a/Engine/Plugins/Experimental/Interchange/Runtime/Source/Parsers/Fbx/Private/FbxCamera.cpp
+++ b/Engine/Plugins/Experimental/Interchange/Runtime/Source/Parsers/Fbx/Private/FbxCamera.cpp
@@ -69,12 +69,12 @@ namespace UE
 						FString NodeName = FFbxHelper::GetNodeAttributeName(NodeAttribute, UInterchangeCameraNode::StaticAssetTypeName());
 						FString NodeUid = FFbxHelper::GetNodeAttributeUniqueID(NodeAttribute, UInterchangeCameraNode::StaticAssetTypeName());
 
-						UInterchangeCameraNode* CameraNode = Cast<UInterchangeCameraNode>(NodeContainer.GetNode(NodeUid));
+						const UInterchangeCameraNode* CameraNode = Cast<const UInterchangeCameraNode>(NodeContainer.GetNode(NodeUid));
 
 						if (!CameraNode)
 						{
-							CameraNode = CreateCameraNode(NodeContainer, NodeUid, NodeName);
-							FillCameraNode(CameraNode, static_cast<FbxCamera&>(*NodeAttribute));
+							UInterchangeCameraNode* NewCameraNode = CreateCameraNode(NodeContainer, NodeUid, NodeName);
+							FillCameraNode(NewCameraNode, static_cast<FbxCamera&>(*NodeAttribute));
 						}
 					}
 				}
diff --git a/Engine/Plugins/Experimental/Interchange/Runtime/Source/Parsers/Fbx/Private/FbxLight.cpp b/Engine/Plugins/Experimental/Interchange/Runtime/Source/Parsers/Fbx/Private/FbxLight.cpp
--- a/Engine/Plugins/Experimental/Interchange/Runtime/Source/Parsers/Fbx/Private/FbxLight.cpp
+++ b/Engine/Plugins/Experimental/Interchange/Runtime/Source/Parsers/Fbx/Private/FbxLight.cpp
@@ -68,7 +68,7 @@ namespace UE
 						FString NodeName = FFbxHelper::GetNodeAttributeName(NodeAttribute, UInterchangeLightNode::StaticAssetTypeName());
 						FString NodeUid = FFbxHelper::GetNodeAttributeUniqueID(NodeAttribute, UInterchangeLightNode::StaticAssetTypeName());
 
-						UInterchangeLightNode* LightNode = Cast<UInterchangeLightNode>(NodeContainer.GetNode(NodeUid));
+						const UInterchangeLightNode* LightNode = Cast<const UInterchangeLightNode>(NodeContainer.GetNode(NodeUid));
 
 						if (!LightNode)
 						{
diff --git a/Engine/Plugins/Experimental/Interchange/Runtime/Source/Parsers/Fbx/Private/FbxMaterial.cpp b/Engine/Plugins/Experimental/Interchange/Runtime/Source/Parsers/Fbx/Private/FbxMaterial.cpp
--- a/Engine/Plugins/Experimental/Interchange/Runtime/Source/Parsers/Fbx/Private/FbxMaterial.cpp
+++ b/Engine/Plugins/Experimental/Interchange/Runtime/Source/Parsers/Fbx/Private/FbxMaterial.cpp
@@ -124,7 +124,8 @@ namespace UE
 					NodeContainer.SetNodeParentUid(TextureSampleShaderUid, ShaderGraphNode->GetUniqueID());
 
 					FString TextureNodeUid = TEXT("\\Texture\\") + TextureFilename;
-					UInterchangeTexture2DNode* TextureNode = Cast<UInterchangeTexture2DNode>(NodeContainer.GetNode(TextureNodeUid));
+					const UInterchangeTexture2DNode* TextureNode = Cast<const UInterchangeTexture2DNode>(NodeContainer.GetNode(TextureNodeUid));
+
 					if (!TextureNode)
 					{
 						TextureNode = CreateTexture2DNode(NodeContainer, TextureNodeUid, TextureFilename);
@@ -172,101 +173,104 @@ namespace UE
 				}
 			}
 
-			UInterchangeShaderGraphNode* FFbxMaterial::AddShaderGraphNode(FbxSurfaceMaterial* SurfaceMaterial, UInterchangeBaseNodeContainer& NodeContainer)
+			const UInterchangeShaderGraphNode* FFbxMaterial::AddShaderGraphNode(FbxSurfaceMaterial* SurfaceMaterial, UInterchangeBaseNodeContainer& NodeContainer)
 			{
 				using namespace UE::Interchange::Materials;
 
 				//Create a material node
 				FString MaterialName = FFbxHelper::GetFbxObjectName(SurfaceMaterial);
 				FString NodeUid = TEXT("\\Material\\") + MaterialName;
-				UInterchangeShaderGraphNode* ShaderGraphNode = Cast<UInterchangeShaderGraphNode>(NodeContainer.GetNode(NodeUid));
-				if (!ShaderGraphNode)
+				const UInterchangeShaderGraphNode* ExistingShaderGraphNode = Cast<const UInterchangeShaderGraphNode>(NodeContainer.GetNode(NodeUid));
+				if (ExistingShaderGraphNode)
+				{
+					return ExistingShaderGraphNode;
+				}
+
+
+				UInterchangeShaderGraphNode* ShaderGraphNode = CreateShaderGraphNode(NodeContainer, NodeUid, MaterialName);
+				if (ShaderGraphNode == nullptr)
 				{
-					ShaderGraphNode = CreateShaderGraphNode(NodeContainer, NodeUid, MaterialName);
-					if (ShaderGraphNode == nullptr)
+					FFormatNamedArguments Args
 					{
-						FFormatNamedArguments Args
-						{
-							{ TEXT("MaterialName"), FText::FromString(MaterialName) }
-						};
-						UInterchangeResultError_Generic* Message = Parser.AddMessage<UInterchangeResultError_Generic>();
-						Message->Text = FText::Format(LOCTEXT("CannotCreateFBXMaterial", "Cannot create FBX material '{MaterialName}'."), Args);
-						return nullptr;
-					}
+						{ TEXT("MaterialName"), FText::FromString(MaterialName) }
+					};
+					UInterchangeResultError_Generic* Message = Parser.AddMessage<UInterchangeResultError_Generic>();
+					Message->Text = FText::Format(LOCTEXT("CannotCreateFBXMaterial", "Cannot create FBX material '{MaterialName}'."), Args);
+					return nullptr;
+				}
 
-					// for each material property, add an input value or a shader node
-					FbxProperty CurrentProperty = SurfaceMaterial->GetFirstProperty();
-					while (CurrentProperty.IsValid())
+				// for each material property, add an input value or a shader node
+				FbxProperty CurrentProperty = SurfaceMaterial->GetFirstProperty();
+				while (CurrentProperty.IsValid())
+				{
+					if (CurrentProperty.Modified())
 					{
-						if (CurrentProperty.Modified())
-						{
-							const FName PropertyName(CurrentProperty.GetNameAsCStr());
+						const FName PropertyName(CurrentProperty.GetNameAsCStr());
 
-							// Convert FBX property names to standard Interchange Phong parameter names when possible
+						// Convert FBX property names to standard Interchange Phong parameter names when possible
+						{
+							if (PropertyName == FbxSurfaceMaterial::sDiffuse)
 							{
-								if (PropertyName == FbxSurfaceMaterial::sDiffuse)
-								{
-									const float Factor = (float)((FbxSurfaceLambert*)SurfaceMaterial)->DiffuseFactor.Get();
-									TVariant<FLinearColor, float> DefaultValue;
-									DefaultValue.Set<FLinearColor>(FLinearColor::Black);
-									ConvertPropertyToShaderNode(NodeContainer, ShaderGraphNode, CurrentProperty, Factor, Phong::Parameters::DiffuseColor, DefaultValue);
-								}
-								else if (PropertyName == FbxSurfaceMaterial::sSpecular)
-								{
-									const float Factor = (float)((FbxSurfacePhong*)SurfaceMaterial)->SpecularFactor.Get();
-									TVariant<FLinearColor, float> DefaultValue;
-									DefaultValue.Set<FLinearColor>(FLinearColor::Black);
-									ConvertPropertyToShaderNode(NodeContainer, ShaderGraphNode, CurrentProperty, Factor, Phong::Parameters::SpecularColor, DefaultValue);
-								}
-								else if (PropertyName == FbxSurfaceMaterial::sShininess)
-								{
-									const float Factor = 1.f;
-									TVariant<FLinearColor, float> DefaultValue;
-									DefaultValue.Set<float>(0.2f);
-									ConvertPropertyToShaderNode(NodeContainer, ShaderGraphNode, CurrentProperty, Factor, Phong::Parameters::Shininess, DefaultValue);
-								}
-								else if (PropertyName == FbxSurfaceMaterial::sEmissive)
-								{
-									const float Factor = (float)((FbxSurfaceLambert*)SurfaceMaterial)->EmissiveFactor.Get();
-									TVariant<FLinearColor, float> DefaultValue;
-									DefaultValue.Set<FLinearColor>(FLinearColor::Black);
-									ConvertPropertyToShaderNode(NodeContainer, ShaderGraphNode, CurrentProperty, Factor, Phong::Parameters::EmissiveColor, DefaultValue);
-								}
-								else if (PropertyName == FbxSurfaceMaterial::sNormalMap)
-								{
-									const float Factor = 1.f;
-									TVariant<FLinearColor, float> DefaultValue;
-									DefaultValue.Set<FLinearColor>(FLinearColor(FVector::UpVector));
-									ConvertPropertyToShaderNode(NodeContainer, ShaderGraphNode, CurrentProperty, Factor, Phong::Parameters::Normal, DefaultValue);
-								}
-								else if (PropertyName == FbxSurfaceMaterial::sTransparentColor)
-								{
-									const float TransparencyFactor = (float)((FbxSurfaceLambert*)SurfaceMaterial)->TransparencyFactor.Get();
-									TVariant<FLinearColor, float> DefaultValue;
-									DefaultValue.Set<float>(0.f); // Opaque
-									const bool bInverse = true;
-									ConvertPropertyToShaderNode(NodeContainer, ShaderGraphNode, CurrentProperty, TransparencyFactor, Phong::Parameters::Opacity, DefaultValue, bInverse);
-								}
-								else if (PropertyName == FbxSurfaceMaterial::sDiffuseFactor || PropertyName == FbxSurfaceMaterial::sSpecularFactor ||
-										 PropertyName == FbxSurfaceMaterial::sEmissiveFactor || PropertyName == FbxSurfaceMaterial::sTransparencyFactor)
-								{
-									// Skip factors as they were processed with their corresponding inputs
-								}
-								else
+								const float Factor = (float)((FbxSurfaceLambert*)SurfaceMaterial)->DiffuseFactor.Get();
+								TVariant<FLinearColor, float> DefaultValue;
+								DefaultValue.Set<FLinearColor>(FLinearColor::Black);
+								ConvertPropertyToShaderNode(NodeContainer, ShaderGraphNode, CurrentProperty, Factor, Phong::Parameters::DiffuseColor, DefaultValue);
+							}
+							else if (PropertyName == FbxSurfaceMaterial::sSpecular)
+							{
+								const float Factor = (float)((FbxSurfacePhong*)SurfaceMaterial)->SpecularFactor.Get();
+								TVariant<FLinearColor, float> DefaultValue;
+								DefaultValue.Set<FLinearColor>(FLinearColor::Black);
+								ConvertPropertyToShaderNode(NodeContainer, ShaderGraphNode, CurrentProperty, Factor, Phong::Parameters::SpecularColor, DefaultValue);
+							}
+							else if (PropertyName == FbxSurfaceMaterial::sShininess)
+							{
+								const float Factor = 1.f;
+								TVariant<FLinearColor, float> DefaultValue;
+								DefaultValue.Set<float>(0.2f);
+								ConvertPropertyToShaderNode(NodeContainer, ShaderGraphNode, CurrentProperty, Factor, Phong::Parameters::Shininess, DefaultValue);
+							}
+							else if (PropertyName == FbxSurfaceMaterial::sEmissive)
+							{
+								const float Factor = (float)((FbxSurfaceLambert*)SurfaceMaterial)->EmissiveFactor.Get();
+								TVariant<FLinearColor, float> DefaultValue;
+								DefaultValue.Set<FLinearColor>(FLinearColor::Black);
+								ConvertPropertyToShaderNode(NodeContainer, ShaderGraphNode, CurrentProperty, Factor, Phong::Parameters::EmissiveColor, DefaultValue);
+							}
+							else if (PropertyName == FbxSurfaceMaterial::sNormalMap)
+							{
+								const float Factor = 1.f;
+								TVariant<FLinearColor, float> DefaultValue;
+								DefaultValue.Set<FLinearColor>(FLinearColor(FVector::UpVector));
+								ConvertPropertyToShaderNode(NodeContainer, ShaderGraphNode, CurrentProperty, Factor, Phong::Parameters::Normal, DefaultValue);
+							}
+							else if (PropertyName == FbxSurfaceMaterial::sTransparentColor)
+							{
+								const float TransparencyFactor = (float)((FbxSurfaceLambert*)SurfaceMaterial)->TransparencyFactor.Get();
+								TVariant<FLinearColor, float> DefaultValue;
+								DefaultValue.Set<float>(0.f); // Opaque
+								const bool bInverse = true;
+								ConvertPropertyToShaderNode(NodeContainer, ShaderGraphNode, CurrentProperty, TransparencyFactor, Phong::Parameters::Opacity, DefaultValue, bInverse);
+							}
+							else if (PropertyName == FbxSurfaceMaterial::sDiffuseFactor || PropertyName == FbxSurfaceMaterial::sSpecularFactor ||
+										PropertyName == FbxSurfaceMaterial::sEmissiveFactor || PropertyName == FbxSurfaceMaterial::sTransparencyFactor)
+							{
+								// Skip factors as they were processed with their corresponding inputs
+							}
+							else
+							{
+								if (!UInterchangeShaderPortsAPI::HasInput(ShaderGraphNode, PropertyName))
 								{
-									if (!UInterchangeShaderPortsAPI::HasInput(ShaderGraphNode, PropertyName))
-									{
-										const float DummyFactor = 1.f;
-										TVariant<FLinearColor, float> DummyDefaultValue;
-										DummyDefaultValue.Set<float>(1.f); // Will be ignored because we have a factor of 1
-										ConvertPropertyToShaderNode(NodeContainer, ShaderGraphNode, CurrentProperty, DummyFactor, PropertyName, DummyDefaultValue);
-									}
+									const float DummyFactor = 1.f;
+									TVariant<FLinearColor, float> DummyDefaultValue;
+									DummyDefaultValue.Set<float>(1.f); // Will be ignored because we have a factor of 1
+									ConvertPropertyToShaderNode(NodeContainer, ShaderGraphNode, CurrentProperty, DummyFactor, PropertyName, DummyDefaultValue);
 								}
 							}
 						}
-
-						CurrentProperty = SurfaceMaterial->GetNextProperty(CurrentProperty);
 					}
+
+					CurrentProperty = SurfaceMaterial->GetNextProperty(CurrentProperty);
 				}
 
 				return ShaderGraphNode;
@@ -287,7 +291,7 @@ namespace UE
 					//Create a texture node and make it child of the material node
 					TArray<FString> JsonErrorMessage;
 					FString NodeUid = TEXT("\\Texture\\") + TextureFilename;
-					UInterchangeTexture2DNode* TextureNode = Cast<UInterchangeTexture2DNode>(NodeContainer.GetNode(NodeUid));
+					const UInterchangeTexture2DNode* TextureNode = Cast<UInterchangeTexture2DNode>(NodeContainer.GetNode(NodeUid));
 					if (!TextureNode)
 					{
 						TextureNode = CreateTexture2DNode(NodeContainer, NodeUid, TextureFilename);
@@ -301,7 +305,7 @@ namespace UE
 				for (int32 MaterialIndex = 0; MaterialIndex < MaterialCount; ++MaterialIndex)
 				{
 					FbxSurfaceMaterial* SurfaceMaterial = ParentFbxNode->GetMaterial(MaterialIndex);
-					UInterchangeShaderGraphNode* ShaderGraphNode = AddShaderGraphNode(SurfaceMaterial, NodeContainer);
+					const UInterchangeShaderGraphNode* ShaderGraphNode = AddShaderGraphNode(SurfaceMaterial, NodeContainer);
 					//The dependencies order is important because mesh will use index in that order to determine material use by a face
 					SceneNode->AddMaterialDependencyUid(ShaderGraphNode->GetUniqueID());
 				}
diff --git a/Engine/Plugins/Experimental/Interchange/Runtime/Source/Parsers/Fbx/Private/FbxMaterial.h b/Engine/Plugins/Experimental/Interchange/Runtime/Source/Parsers/Fbx/Private/FbxMaterial.h
--- a/Engine/Plugins/Experimental/Interchange/Runtime/Source/Parsers/Fbx/Private/FbxMaterial.h
+++ b/Engine/Plugins/Experimental/Interchange/Runtime/Source/Parsers/Fbx/Private/FbxMaterial.h
@@ -59,7 +59,7 @@ namespace UE
 					FbxProperty& Property, float Factor, FName PropertyName, TVariant<FLinearColor, float> DefaultValue, bool bInverse = false);
 
 			private:
-				UInterchangeShaderGraphNode* AddShaderGraphNode(FbxSurfaceMaterial* SurfaceMaterial, UInterchangeBaseNodeContainer& NodeContainer);
+				const UInterchangeShaderGraphNode* AddShaderGraphNode(FbxSurfaceMaterial* SurfaceMaterial, UInterchangeBaseNodeContainer& NodeContainer);
 
 				FFbxParser& Parser;
 			};
diff --git a/Engine/Plugins/Experimental/Interchange/Runtime/Source/Parsers/Fbx/Private/FbxMesh.cpp b/Engine/Plugins/Experimental/Interchange/Runtime/Source/Parsers/Fbx/Private/FbxMesh.cpp
--- a/Engine/Plugins/Experimental/Interchange/Runtime/Source/Parsers/Fbx/Private/FbxMesh.cpp
+++ b/Engine/Plugins/Experimental/Interchange/Runtime/Source/Parsers/Fbx/Private/FbxMesh.cpp
@@ -19,1227 +19,1302 @@
 
 #define LOCTEXT_NAMESPACE "InterchangeFbxMesh"
 
-namespace UE
+namespace UE::Interchange::Private {
+
+// Wraps some common code useful for multiple fbx import code path
+struct FFBXUVs
 {
-	namespace Interchange
+	// constructor
+	FFBXUVs(FMeshDescriptionImporter& InMeshDescriptionImporter, FbxMesh* Mesh)
+		: MeshDescriptionImporter(InMeshDescriptionImporter)
+		, UniqueUVCount(0)
 	{
-		namespace Private
+		check(Mesh);
+
+		//
+		//	store the UVs in arrays for fast access in the later looping of triangles 
+		//
+		// mapping from UVSets to Fbx LayerElementUV
+		// Fbx UVSets may be duplicated, remove the duplicated UVSets in the mapping 
+		int32 LayerCount = Mesh->GetLayerCount();
+		if (LayerCount > 0)
 		{
-			// Wraps some common code useful for multiple fbx import code path
-			struct FFBXUVs
+			int32 UVLayerIndex;
+			for (UVLayerIndex = 0; UVLayerIndex < LayerCount; UVLayerIndex++)
 			{
-				// constructor
-				FFBXUVs(FMeshDescriptionImporter& InMeshDescriptionImporter, FbxMesh* Mesh)
-					: MeshDescriptionImporter(InMeshDescriptionImporter)
-					, UniqueUVCount(0)
+				FbxLayer* lLayer = Mesh->GetLayer(UVLayerIndex);
+				int UVSetCount = lLayer->GetUVSetCount();
+				if (UVSetCount)
 				{
-					check(Mesh);
-
-					//
-					//	store the UVs in arrays for fast access in the later looping of triangles 
-					//
-					// mapping from UVSets to Fbx LayerElementUV
-					// Fbx UVSets may be duplicated, remove the duplicated UVSets in the mapping 
-					int32 LayerCount = Mesh->GetLayerCount();
-					if (LayerCount > 0)
+					FbxArray<FbxLayerElementUV const*> EleUVs = lLayer->GetUVSets();
+					for (int UVIndex = 0; UVIndex < UVSetCount; UVIndex++)
 					{
-						int32 UVLayerIndex;
-						for (UVLayerIndex = 0; UVLayerIndex < LayerCount; UVLayerIndex++)
+						FbxLayerElementUV const* ElementUV = EleUVs[UVIndex];
+						if (ElementUV)
 						{
-							FbxLayer* lLayer = Mesh->GetLayer(UVLayerIndex);
-							int UVSetCount = lLayer->GetUVSetCount();
-							if (UVSetCount)
+							const char* UVSetName = ElementUV->GetName();
+							FString LocalUVSetName = UTF8_TO_TCHAR(UVSetName);
+							if (LocalUVSetName.IsEmpty())
 							{
-								FbxArray<FbxLayerElementUV const*> EleUVs = lLayer->GetUVSets();
-								for (int UVIndex = 0; UVIndex < UVSetCount; UVIndex++)
-								{
-									FbxLayerElementUV const* ElementUV = EleUVs[UVIndex];
-									if (ElementUV)
-									{
-										const char* UVSetName = ElementUV->GetName();
-										FString LocalUVSetName = UTF8_TO_TCHAR(UVSetName);
-										if (LocalUVSetName.IsEmpty())
-										{
-											LocalUVSetName = TEXT("UVmap_") + FString::FromInt(UVLayerIndex);
-										}
-
-										UVSets.AddUnique(LocalUVSetName);
-									}
-								}
+								LocalUVSetName = TEXT("UVmap_") + FString::FromInt(UVLayerIndex);
 							}
+
+							UVSets.AddUnique(LocalUVSetName);
 						}
 					}
+				}
+			}
+		}
 
 
-					// If the the UV sets are named using the following format (UVChannel_X; where X ranges from 1 to 4)
-					// we will re-order them based on these names.  Any UV sets that do not follow this naming convention
-					// will be slotted into available spaces.
-					if (UVSets.Num())
-					{
-						for (int32 ChannelNumIdx = 0; ChannelNumIdx < 4; ChannelNumIdx++)
-						{
-							FString ChannelName = FString::Printf(TEXT("UVChannel_%d"), ChannelNumIdx + 1);
-							int32 SetIdx = UVSets.Find(ChannelName);
+		// If the the UV sets are named using the following format (UVChannel_X; where X ranges from 1 to 4)
+		// we will re-order them based on these names.  Any UV sets that do not follow this naming convention
+		// will be slotted into available spaces.
+		if (UVSets.Num())
+		{
+			for (int32 ChannelNumIdx = 0; ChannelNumIdx < 4; ChannelNumIdx++)
+			{
+				FString ChannelName = FString::Printf(TEXT("UVChannel_%d"), ChannelNumIdx + 1);
+				int32 SetIdx = UVSets.Find(ChannelName);
 
-							// If the specially formatted UVSet name appears in the list and it is in the wrong spot,
-							// we will swap it into the correct spot.
-							if (SetIdx != INDEX_NONE && SetIdx != ChannelNumIdx)
-							{
-								// If we are going to swap to a position that is outside the bounds of the
-								// array, then we pad out to that spot with empty data.
-								for (int32 ArrSize = UVSets.Num(); ArrSize < ChannelNumIdx + 1; ArrSize++)
-								{
-									UVSets.Add(FString(TEXT("")));
-								}
-								//Swap the entry into the appropriate spot.
-								UVSets.Swap(SetIdx, ChannelNumIdx);
-							}
-						}
+				// If the specially formatted UVSet name appears in the list and it is in the wrong spot,
+				// we will swap it into the correct spot.
+				if (SetIdx != INDEX_NONE && SetIdx != ChannelNumIdx)
+				{
+					// If we are going to swap to a position that is outside the bounds of the
+					// array, then we pad out to that spot with empty data.
+					for (int32 ArrSize = UVSets.Num(); ArrSize < ChannelNumIdx + 1; ArrSize++)
+					{
+						UVSets.Add(FString(TEXT("")));
 					}
+					//Swap the entry into the appropriate spot.
+					UVSets.Swap(SetIdx, ChannelNumIdx);
 				}
+			}
+		}
+	}
 
-				void Phase2(FbxMesh* Mesh)
+	void Phase2(FbxMesh* Mesh)
+	{
+		//
+		//	store the UVs in arrays for fast access in the later looping of triangles 
+		//
+		UniqueUVCount = UVSets.Num();
+		if (UniqueUVCount > 0)
+		{
+			LayerElementUV.AddZeroed(UniqueUVCount);
+			UVReferenceMode.AddZeroed(UniqueUVCount);
+			UVMappingMode.AddZeroed(UniqueUVCount);
+		}
+		for (int32 UVIndex = 0; UVIndex < UniqueUVCount; UVIndex++)
+		{
+			LayerElementUV[UVIndex] = NULL;
+			for (int32 UVLayerIndex = 0, LayerCount = Mesh->GetLayerCount(); UVLayerIndex < LayerCount; UVLayerIndex++)
+			{
+				FbxLayer* lLayer = Mesh->GetLayer(UVLayerIndex);
+				int UVSetCount = lLayer->GetUVSetCount();
+				if (UVSetCount)
 				{
-					//
-					//	store the UVs in arrays for fast access in the later looping of triangles 
-					//
-					UniqueUVCount = UVSets.Num();
-					if (UniqueUVCount > 0)
+					FbxArray<FbxLayerElementUV const*> EleUVs = lLayer->GetUVSets();
+					for (int32 FbxUVIndex = 0; FbxUVIndex < UVSetCount; FbxUVIndex++)
 					{
-						LayerElementUV.AddZeroed(UniqueUVCount);
-						UVReferenceMode.AddZeroed(UniqueUVCount);
-						UVMappingMode.AddZeroed(UniqueUVCount);
-					}
-					for (int32 UVIndex = 0; UVIndex < UniqueUVCount; UVIndex++)
-					{
-						LayerElementUV[UVIndex] = NULL;
-						for (int32 UVLayerIndex = 0, LayerCount = Mesh->GetLayerCount(); UVLayerIndex < LayerCount; UVLayerIndex++)
+						FbxLayerElementUV const* ElementUV = EleUVs[FbxUVIndex];
+						if (ElementUV)
 						{
-							FbxLayer* lLayer = Mesh->GetLayer(UVLayerIndex);
-							int UVSetCount = lLayer->GetUVSetCount();
-							if (UVSetCount)
+							const char* UVSetName = ElementUV->GetName();
+							FString LocalUVSetName = UTF8_TO_TCHAR(UVSetName);
+							if (LocalUVSetName.IsEmpty())
 							{
-								FbxArray<FbxLayerElementUV const*> EleUVs = lLayer->GetUVSets();
-								for (int32 FbxUVIndex = 0; FbxUVIndex < UVSetCount; FbxUVIndex++)
-								{
-									FbxLayerElementUV const* ElementUV = EleUVs[FbxUVIndex];
-									if (ElementUV)
-									{
-										const char* UVSetName = ElementUV->GetName();
-										FString LocalUVSetName = UTF8_TO_TCHAR(UVSetName);
-										if (LocalUVSetName.IsEmpty())
-										{
-											LocalUVSetName = TEXT("UVmap_") + FString::FromInt(UVLayerIndex);
-										}
-										if (LocalUVSetName == UVSets[UVIndex])
-										{
-											LayerElementUV[UVIndex] = ElementUV;
-											UVReferenceMode[UVIndex] = ElementUV->GetReferenceMode();
-											UVMappingMode[UVIndex] = ElementUV->GetMappingMode();
-											break;
-										}
-									}
-								}
+								LocalUVSetName = TEXT("UVmap_") + FString::FromInt(UVLayerIndex);
+							}
+							if (LocalUVSetName == UVSets[UVIndex])
+							{
+								LayerElementUV[UVIndex] = ElementUV;
+								UVReferenceMode[UVIndex] = ElementUV->GetReferenceMode();
+								UVMappingMode[UVIndex] = ElementUV->GetMappingMode();
+								break;
 							}
 						}
 					}
+				}
+			}
+		}
 
-					if (UniqueUVCount > MAX_MESH_TEXTURE_COORDS_MD)
-					{
-						UInterchangeResultMeshWarning_TooManyUVs* Message = MeshDescriptionImporter.AddMessage<UInterchangeResultMeshWarning_TooManyUVs>(Mesh);
-						Message->ExcessUVs = UniqueUVCount - MAX_MESH_TEXTURE_COORDS_MD;
-					}
+		if (UniqueUVCount > MAX_MESH_TEXTURE_COORDS_MD)
+		{
+			UInterchangeResultMeshWarning_TooManyUVs* Message = MeshDescriptionImporter.AddMessage<UInterchangeResultMeshWarning_TooManyUVs>(Mesh);
+			Message->ExcessUVs = UniqueUVCount - MAX_MESH_TEXTURE_COORDS_MD;
+		}
 
-					UniqueUVCount = FMath::Min<int32>(UniqueUVCount, MAX_MESH_TEXTURE_COORDS_MD);
-				}
+		UniqueUVCount = FMath::Min<int32>(UniqueUVCount, MAX_MESH_TEXTURE_COORDS_MD);
+	}
 
-				// @param FaceCornerIndex usually TriangleIndex * 3 + CornerIndex but more complicated for mixed n-gons
-				int32 ComputeUVIndex(int32 UVLayerIndex, int32 lControlPointIndex, int32 FaceCornerIndex) const
-				{
-					int32 UVMapIndex = (UVMappingMode[UVLayerIndex] == FbxLayerElement::eByControlPoint) ? lControlPointIndex : FaceCornerIndex;
+	// @param FaceCornerIndex usually TriangleIndex * 3 + CornerIndex but more complicated for mixed n-gons
+	int32 ComputeUVIndex(int32 UVLayerIndex, int32 lControlPointIndex, int32 FaceCornerIndex) const
+	{
+		int32 UVMapIndex = (UVMappingMode[UVLayerIndex] == FbxLayerElement::eByControlPoint) ? lControlPointIndex : FaceCornerIndex;
 
-					int32 Ret;
+		int32 Ret;
 
-					if (UVReferenceMode[UVLayerIndex] == FbxLayerElement::eDirect)
-					{
-						Ret = UVMapIndex;
-					}
-					else
-					{
-						FbxLayerElementArrayTemplate<int>& Array = LayerElementUV[UVLayerIndex]->GetIndexArray();
-						Ret = Array.GetAt(UVMapIndex);
-					}
+		if (UVReferenceMode[UVLayerIndex] == FbxLayerElement::eDirect)
+		{
+			Ret = UVMapIndex;
+		}
+		else
+		{
+			FbxLayerElementArrayTemplate<int>& Array = LayerElementUV[UVLayerIndex]->GetIndexArray();
+			Ret = Array.GetAt(UVMapIndex);
+		}
 
-					return Ret;
-				}
+		return Ret;
+	}
 
-				// todo: is that needed? could the dtor do it?
-				void Cleanup()
-				{
-					//
-					// clean up.  This needs to happen before the mesh is destroyed
-					//
-					LayerElementUV.Empty();
-					UVReferenceMode.Empty();
-					UVMappingMode.Empty();
-				}
+	// todo: is that needed? could the dtor do it?
+	void Cleanup()
+	{
+		//
+		// clean up.  This needs to happen before the mesh is destroyed
+		//
+		LayerElementUV.Empty();
+		UVReferenceMode.Empty();
+		UVMappingMode.Empty();
+	}
+
+	FMeshDescriptionImporter& MeshDescriptionImporter;
+	TArray<FString> UVSets;
+	TArray<FbxLayerElementUV const*> LayerElementUV;
+	TArray<FbxLayerElement::EReferenceMode> UVReferenceMode;
+	TArray<FbxLayerElement::EMappingMode> UVMappingMode;
+	int32 UniqueUVCount;
+};
+
+
+FMeshDescriptionImporter::FMeshDescriptionImporter(FFbxParser& InParser, FMeshDescription* InMeshDescription, FbxScene* InSDKScene, FbxGeometryConverter* InSDKGeometryConverter)
+	: Parser(InParser)
+	, MeshDescription(InMeshDescription)
+	, SDKScene(InSDKScene)
+	, SDKGeometryConverter(InSDKGeometryConverter)
+{
+	bInitialized = ensure(MeshDescription) && ensure(SDKScene) && ensure(SDKGeometryConverter);
+}
+			
+bool FMeshDescriptionImporter::FillStaticMeshDescriptionFromFbxMesh(FbxMesh* Mesh)
+{
+	if (!ensure(bInitialized) || !ensure(MeshDescription) || !ensure(Mesh))
+	{
+		return false;
+	}
+	TArray<FString> UnusedArray;
+	return FillMeshDescriptionFromFbxMesh(Mesh, UnusedArray, EMeshType::Static);
+}
+
+bool FMeshDescriptionImporter::FillSkinnedMeshDescriptionFromFbxMesh(FbxMesh* Mesh, TArray<FString>& OutJointUniqueNames)
+{
+	if (!ensure(bInitialized) || !ensure(MeshDescription) || !ensure(Mesh) || !ensure(Mesh->GetDeformerCount() > 0))
+	{
+		return false;
+	}
+
+	return FillMeshDescriptionFromFbxMesh(Mesh, OutJointUniqueNames, EMeshType::Skinned);
+}
+
+bool FMeshDescriptionImporter::FillMeshDescriptionFromFbxShape(FbxShape* Shape)
+{
+	//For the shape we just need the vertex positions
+	if (!ensure(bInitialized) || !ensure(MeshDescription))
+	{
+		return false;
+	}
+	FStaticMeshAttributes Attributes(*MeshDescription);
+
+	//The shape should be for a specified mesh
+	if (!ensure(Shape))
+	{
+		return false;
+	}
 
-				FMeshDescriptionImporter& MeshDescriptionImporter;
-				TArray<FString> UVSets;
-				TArray<FbxLayerElementUV const*> LayerElementUV;
-				TArray<FbxLayerElement::EReferenceMode> UVReferenceMode;
-				TArray<FbxLayerElement::EMappingMode> UVMappingMode;
-				int32 UniqueUVCount;
-			};
-
-			FMeshDescriptionImporter::FMeshDescriptionImporter(FFbxParser& InParser, FMeshDescription* InMeshDescription, FbxScene* InSDKScene, FbxGeometryConverter* InSDKGeometryConverter)
-				: Parser(InParser)
-				, MeshDescription(InMeshDescription)
-				, SDKScene(InSDKScene)
-				, SDKGeometryConverter(InSDKGeometryConverter)
+	//Extract the points into a simplified MeshDescription
+	{
+		TRACE_CPUPROFILER_EVENT_SCOPE(GatherMorphTargetPoints);
+		MeshDescription->SuspendVertexIndexing();
+
+		// Construct the matrices for the conversion from right handed to left handed system
+		FbxAMatrix TotalMatrix;
+		FbxAMatrix TotalMatrixForNormal;
+		TotalMatrix.SetIdentity(); //We use the identity since we want to bake in the interchange factory if the pipeline request it
+		TotalMatrixForNormal = TotalMatrix.Inverse();
+		TotalMatrixForNormal = TotalMatrixForNormal.Transpose();
+
+		FbxGeometryBase* GeoBase = (FbxGeometryBase*)Shape;
+
+		int32 VertexCount = GeoBase->GetControlPointsCount();
+
+		TVertexAttributesRef<FVector3f> VertexPositions = Attributes.GetVertexPositions();
+		int32 VertexOffset = MeshDescription->Vertices().Num();
+		// The below code expects Num() to be equivalent to GetArraySize(), i.e. that all added elements are appended, not inserted into existing gaps
+		check(VertexOffset == MeshDescription->Vertices().GetArraySize());
+		//Fill the vertex array
+		auto GetFinalPosition = [&TotalMatrix](FbxVector4& FbxPosition)
+		{
+			FbxPosition = TotalMatrix.MultT(FbxPosition);
+			return FFbxConvert::ConvertPos(FbxPosition);
+		};
+
+		MeshDescription->ReserveNewVertices(VertexCount);
+		for (int32 VertexIndex = 0; VertexIndex < VertexCount; ++VertexIndex)
+		{
+			int32 RealVertexIndex = VertexOffset + VertexIndex;
+			const FVector VertexPosition = GetFinalPosition(GeoBase->GetControlPoints()[VertexIndex]);
+			//Maybe we want to do some work here like computing the deltas
+			//const FVector MeshVertexPosition = GetFinalPosition(Mesh->GetControlPoints()[VertexIndex]);
+			FVertexID AddedVertexId = MeshDescription->CreateVertex();
+			if (AddedVertexId.GetValue() != RealVertexIndex)
 			{
-				bInitialized = ensure(MeshDescription) && ensure(SDKScene) && ensure(SDKGeometryConverter);
+				UInterchangeResultMeshError_Generic* Message = AddMessage<UInterchangeResultMeshError_Generic>(Shape);
+				Message->Text = LOCTEXT("CantCreateVertex", "Cannot create valid vertex for the mesh '{MeshName}'.");
+
+				return false;
 			}
-			
-			bool FMeshDescriptionImporter::FillStaticMeshDescriptionFromFbxMesh(FbxMesh* Mesh)
+			//Add the delta position, so we do not have to recompute it later
+			VertexPositions[AddedVertexId] = (FVector3f)VertexPosition;// -MeshVertexPosition;
+		}
+		MeshDescription->ResumeVertexIndexing();
+		FElementIDRemappings OutRemappings;
+		MeshDescription->Compact(OutRemappings);
+	}
+	return true;
+}
+
+bool FMeshDescriptionImporter::FillMeshDescriptionFromFbxMesh(FbxMesh* Mesh, TArray<FString>& OutJointUniqueNames, EMeshType MeshType)
+{
+	if (!ensure(bInitialized) || !ensure(MeshDescription) || !ensure(Mesh))
+	{
+		return false;
+	}
+
+	FStaticMeshAttributes Attributes(*MeshDescription);
+	Attributes.Register();
+
+	//Get the base layer of the mesh
+	FbxLayer* BaseLayer = Mesh->GetLayer(0);
+	if (BaseLayer == NULL)
+	{
+		UInterchangeResultMeshError_Generic* Message = AddMessage<UInterchangeResultMeshError_Generic>(Mesh);
+		Message->Text = LOCTEXT("CantGetLayer", "Unable to get FBX base layer for the mesh '{MeshName}'.");
+		return false;
+	}
+
+	FFBXUVs FBXUVs(*this, Mesh);
+
+	FbxNode* MeshNode = Mesh->GetNodeCount() > 0 ? Mesh->GetNode() : nullptr; //Get the first node using the mesh to setup default asset materials
+	//
+	// create materials
+	//
+	//Create a material name array in the node order, also fill the Meshdescription PolygonGroup
+	TArray<FName> MaterialNames;
+	int32 MaterialCount = (MeshNode != nullptr) ? MeshNode->GetMaterialCount() : Mesh->GetElementMaterialCount();
+	if (MeshNode)
+	{
+		MaterialCount = MeshNode->GetMaterialCount();
+		MaterialNames.Reserve(MaterialCount);
+		for (int32 MaterialIndex = 0; MaterialIndex < MaterialCount; ++MaterialIndex)
+		{
+			FbxSurfaceMaterial* FbxMaterial = MeshNode->GetMaterial(MaterialIndex);
+			MaterialNames.Add(*FFbxHelper::GetFbxObjectName(FbxMaterial));
+		}
+	}
+
+	// Must do this before triangulating the mesh due to an FBX bug in TriangulateMeshAdvance
+	int32 LayerSmoothingCount = Mesh->GetLayerCount(FbxLayerElement::eSmoothing);
+	for (int32 i = 0; i < LayerSmoothingCount; i++)
+	{
+		FbxLayerElementSmoothing const* SmoothingInfo = Mesh->GetLayer(0)->GetSmoothing();
+		if (SmoothingInfo && SmoothingInfo->GetMappingMode() != FbxLayerElement::eByPolygon)
+		{
+			SDKGeometryConverter->ComputePolygonSmoothingFromEdgeSmoothing(Mesh, i);
+		}
+	}
+
+	if (!Mesh->IsTriangleMesh())
+	{
+		const bool bReplace = true;
+		FbxNodeAttribute* ConvertedNode = SDKGeometryConverter->Triangulate(Mesh, bReplace);
+
+		if (ConvertedNode != NULL && ConvertedNode->GetAttributeType() == FbxNodeAttribute::eMesh)
+		{
+			Mesh = (fbxsdk::FbxMesh*)ConvertedNode;
+		}
+		else
+		{
+			UInterchangeResultMeshError_Generic* Message = AddMessage<UInterchangeResultMeshError_Generic>(Mesh);
+			Message->Text = LOCTEXT("TriangulationError", "Unable to triangulate the mesh '{MeshName}'.");
+
+			MeshDescription->Empty();
+			return false;
+		}
+	}
+
+	// renew the base layer
+	BaseLayer = Mesh->GetLayer(0);
+
+	//
+	// get the smoothing group layer
+	//
+	bool bSmoothingAvailable = false;
+
+	FbxLayerElementSmoothing* SmoothingInfo = BaseLayer->GetSmoothing();
+	FbxLayerElement::EReferenceMode SmoothingReferenceMode(FbxLayerElement::eDirect);
+	FbxLayerElement::EMappingMode SmoothingMappingMode(FbxLayerElement::eByEdge);
+	if (SmoothingInfo)
+	{
+		if (SmoothingInfo->GetMappingMode() == FbxLayerElement::eByPolygon)
+		{
+			//Convert the base layer to edge smoothing
+			SDKGeometryConverter->ComputeEdgeSmoothingFromPolygonSmoothing(Mesh, 0);
+			BaseLayer = Mesh->GetLayer(0);
+			SmoothingInfo = BaseLayer->GetSmoothing();
+		}
+
+		if (SmoothingInfo->GetMappingMode() == FbxLayerElement::eByEdge)
+		{
+			bSmoothingAvailable = true;
+		}
+
+		SmoothingReferenceMode = SmoothingInfo->GetReferenceMode();
+		SmoothingMappingMode = SmoothingInfo->GetMappingMode();
+
+		//If all smooth group are 0's, we can say that there is no valid smoothing data, but only for skeletal mesh
+		//Static mesh will be faceted in this case since we will not compute smooth group from normals like we do here
+		//This is legacy code to get the same result has before
+		if (MeshType == EMeshType::Skinned)
+		{
+			// Check and see if the smooothing data is valid.  If not generate it from the normals
+			BaseLayer = Mesh->GetLayer(0);
+			if (BaseLayer)
 			{
-				if (!ensure(bInitialized) || !ensure(MeshDescription) || !ensure(Mesh))
+				const FbxLayerElementSmoothing* SmoothingLayer = BaseLayer->GetSmoothing();
+
+				if (SmoothingLayer)
 				{
-					return false;
+					bool bValidSmoothingData = false;
+					FbxLayerElementArrayTemplate<int32>& Array = SmoothingLayer->GetDirectArray();
+					for (int32 SmoothingIndex = 0; SmoothingIndex < Array.GetCount(); ++SmoothingIndex)
+					{
+						if (Array[SmoothingIndex] != 0)
+						{
+							bValidSmoothingData = true;
+							break;
+						}
+					}
+
+					if (!bValidSmoothingData && Mesh->GetPolygonVertexCount() > 0)
+					{
+						SDKGeometryConverter->ComputeEdgeSmoothingFromNormals(Mesh);
+						BaseLayer = Mesh->GetLayer(0);
+						SmoothingInfo = BaseLayer->GetSmoothing();
+					}
 				}
-				TArray<FString> UnusedArray;
-				return FillMeshDescriptionFromFbxMesh(Mesh, UnusedArray, EMeshType::Static);
 			}
+		}
+	}
+
+	//
+	//	get the "material index" layer.  Do this AFTER the triangulation step as that may reorder material indices
+	//
+	FbxLayerElementMaterial* LayerElementMaterial = BaseLayer->GetMaterials();
+	FbxLayerElement::EMappingMode MaterialMappingMode = LayerElementMaterial ?
+		LayerElementMaterial->GetMappingMode() : FbxLayerElement::eByPolygon;
+
+	//	todo second phase UV, ok to put in first phase?
+	FBXUVs.Phase2(Mesh);
+
+	//
+	// get the first vertex color layer
+	//
+	FbxLayerElementVertexColor* LayerElementVertexColor = BaseLayer->GetVertexColors();
+	FbxLayerElement::EReferenceMode VertexColorReferenceMode(FbxLayerElement::eDirect);
+	FbxLayerElement::EMappingMode VertexColorMappingMode(FbxLayerElement::eByControlPoint);
+	if (LayerElementVertexColor)
+	{
+		VertexColorReferenceMode = LayerElementVertexColor->GetReferenceMode();
+		VertexColorMappingMode = LayerElementVertexColor->GetMappingMode();
+	}
+
+	//
+	// get the first normal layer
+	//
+	FbxLayerElementNormal* LayerElementNormal = BaseLayer->GetNormals();
+	FbxLayerElementTangent* LayerElementTangent = BaseLayer->GetTangents();
+	FbxLayerElementBinormal* LayerElementBinormal = BaseLayer->GetBinormals();
+
+	//whether there is normal, tangent and binormal data in this mesh
+	bool bHasNTBInformation = LayerElementNormal && LayerElementTangent && LayerElementBinormal;
+
+	FbxLayerElement::EReferenceMode NormalReferenceMode(FbxLayerElement::eDirect);
+	FbxLayerElement::EMappingMode NormalMappingMode(FbxLayerElement::eByControlPoint);
+	if (LayerElementNormal)
+	{
+		NormalReferenceMode = LayerElementNormal->GetReferenceMode();
+		NormalMappingMode = LayerElementNormal->GetMappingMode();
+	}
+
+	FbxLayerElement::EReferenceMode TangentReferenceMode(FbxLayerElement::eDirect);
+	FbxLayerElement::EMappingMode TangentMappingMode(FbxLayerElement::eByControlPoint);
+	if (LayerElementTangent)
+	{
+		TangentReferenceMode = LayerElementTangent->GetReferenceMode();
+		TangentMappingMode = LayerElementTangent->GetMappingMode();
+	}
+
+	FbxLayerElement::EReferenceMode BinormalReferenceMode(FbxLayerElement::eDirect);
+	FbxLayerElement::EMappingMode BinormalMappingMode(FbxLayerElement::eByControlPoint);
+	if (LayerElementBinormal)
+	{
+		BinormalReferenceMode = LayerElementBinormal->GetReferenceMode();
+		BinormalMappingMode = LayerElementBinormal->GetMappingMode();
+	}
+
+	bool bHasNonDegeneratePolygons = false;
+	{
+		MeshDescription->SuspendVertexInstanceIndexing();
+		MeshDescription->SuspendEdgeIndexing();
+		MeshDescription->SuspendPolygonIndexing();
+		MeshDescription->SuspendPolygonGroupIndexing();
+		MeshDescription->SuspendUVIndexing();
+
+		TRACE_CPUPROFILER_EVENT_SCOPE(BuildTriangles);
+
+		// Construct the matrices for the conversion from right handed to left handed system
+		FbxAMatrix TotalMatrix;
+		FbxAMatrix TotalMatrixForNormal;
+		TotalMatrix.SetIdentity();
+		TotalMatrixForNormal = TotalMatrix.Inverse();
+		TotalMatrixForNormal = TotalMatrixForNormal.Transpose();
+		int32 PolygonCount = Mesh->GetPolygonCount();
+
+		if (PolygonCount == 0)
+		{
+			UInterchangeResultMeshError_Generic* Message = AddMessage<UInterchangeResultMeshError_Generic>(Mesh);
+			Message->Text = LOCTEXT("NoPolygons", "No polygons were found in the mesh '{MeshName}'.");
+
+			return false;
+		}
+
+		int32 VertexCount = Mesh->GetControlPointsCount();
+		bool OddNegativeScale = IsOddNegativeScale(TotalMatrix);
+
+		TVertexAttributesRef<FVector3f> VertexPositions = Attributes.GetVertexPositions();
+		TVertexInstanceAttributesRef<FVector3f> VertexInstanceNormals = Attributes.GetVertexInstanceNormals();
+		TVertexInstanceAttributesRef<FVector3f> VertexInstanceTangents = Attributes.GetVertexInstanceTangents();
+		TVertexInstanceAttributesRef<float> VertexInstanceBinormalSigns = Attributes.GetVertexInstanceBinormalSigns();
+		TVertexInstanceAttributesRef<FVector4f> VertexInstanceColors = Attributes.GetVertexInstanceColors();
+		TVertexInstanceAttributesRef<FVector2f> VertexInstanceUVs = Attributes.GetVertexInstanceUVs();
+		TEdgeAttributesRef<bool> EdgeHardnesses = Attributes.GetEdgeHardnesses();
+		TPolygonGroupAttributesRef<FName> PolygonGroupImportedMaterialSlotNames = Attributes.GetPolygonGroupMaterialSlotNames();
+
+		int32 VertexOffset = MeshDescription->Vertices().Num();
+		int32 VertexInstanceOffset = MeshDescription->VertexInstances().Num();
+		int32 PolygonOffset = MeshDescription->Polygons().Num();
+
+		// The below code expects Num() to be equivalent to GetArraySize(), i.e. that all added elements are appended, not inserted into existing gaps
+		check(VertexOffset == MeshDescription->Vertices().GetArraySize());
+		check(VertexInstanceOffset == MeshDescription->VertexInstances().GetArraySize());
+		check(PolygonOffset == MeshDescription->Polygons().GetArraySize());
+
+		TMap<int32, FPolygonGroupID> PolygonGroupMapping;
+
+		// When importing multiple mesh pieces to the same static mesh.  Ensure each mesh piece has the same number of Uv's
+		int32 ExistingUVCount = VertexInstanceUVs.GetNumChannels();
+
+		int32 NumUVs = FMath::Max(FBXUVs.UniqueUVCount, ExistingUVCount);
+		NumUVs = FMath::Min<int32>(MAX_MESH_TEXTURE_COORDS_MD, NumUVs);
+		// At least one UV set must exist.  
+		// @todo: this needn't be true; we should be able to handle zero UV channels
+		NumUVs = FMath::Max(1, NumUVs);
+
+		//Make sure all Vertex instance have the correct number of UVs
+		VertexInstanceUVs.SetNumChannels(NumUVs);
+		MeshDescription->SetNumUVChannels(NumUVs);
+
+		TArray<int32> UVOffsets;
+		UVOffsets.SetNumUninitialized(NumUVs);
+		for (int UVChannel = 0; UVChannel < NumUVs; UVChannel++)
+		{
+			UVOffsets[UVChannel] = MeshDescription->UVs(UVChannel).GetArraySize();
+		}
 
-			bool FMeshDescriptionImporter::FillSkinnedMeshDescriptionFromFbxMesh(FbxMesh* Mesh, TArray<FString>& OutJointUniqueNames)
+		//Fill the vertex array
+		MeshDescription->ReserveNewVertices(VertexCount);
+		for (int32 VertexIndex = 0; VertexIndex < VertexCount; ++VertexIndex)
+		{
+			int32 RealVertexIndex = VertexOffset + VertexIndex;
+			FbxVector4 FbxPosition = Mesh->GetControlPoints()[VertexIndex];
+			FbxPosition = TotalMatrix.MultT(FbxPosition);
+			const FVector VertexPosition = FFbxConvert::ConvertPos(FbxPosition);
+
+			FVertexID AddedVertexId = MeshDescription->CreateVertex();
+			VertexPositions[AddedVertexId] = (FVector3f)VertexPosition;
+			if (AddedVertexId.GetValue() != RealVertexIndex)
 			{
-				if (!ensure(bInitialized) || !ensure(MeshDescription) || !ensure(Mesh) || !ensure(Mesh->GetDeformerCount() > 0))
-				{
-					return false;
-				}
+				UInterchangeResultMeshError_Generic* Message = AddMessage<UInterchangeResultMeshError_Generic>(Mesh);
+				Message->Text = LOCTEXT("CannotCreateVertex", "Cannot create valid vertex for mesh '{MeshName}'.");
 
-				return FillMeshDescriptionFromFbxMesh(Mesh, OutJointUniqueNames, EMeshType::Skinned);
+				return false;
 			}
+		}
 
-			bool FMeshDescriptionImporter::FillMeshDescriptionFromFbxShape(FbxShape* Shape)
+		// Fill the UV arrays
+		for (int32 UVLayerIndex = 0; UVLayerIndex < FBXUVs.UniqueUVCount; UVLayerIndex++)
+		{
+			check(FBXUVs.LayerElementUV[UVLayerIndex]);
+			if (FBXUVs.LayerElementUV[UVLayerIndex] != nullptr)
 			{
-				//For the shape we just need the vertex positions
-				if (!ensure(bInitialized) || !ensure(MeshDescription))
+				int32 UVCount = FBXUVs.LayerElementUV[UVLayerIndex]->GetDirectArray().GetCount();
+				TUVAttributesRef<FVector2f> UVCoordinates = MeshDescription->UVAttributes(UVLayerIndex).GetAttributesRef<FVector2f>(MeshAttribute::UV::UVCoordinate);
+				MeshDescription->ReserveNewUVs(UVCount, UVLayerIndex);
+				for (int32 UVIndex = 0; UVIndex < UVCount; UVIndex++)
 				{
-					return false;
+					FUVID UVID = MeshDescription->CreateUV(UVLayerIndex);
+					FbxVector2 UVVector = FBXUVs.LayerElementUV[UVLayerIndex]->GetDirectArray().GetAt(UVIndex);
+					UVCoordinates[UVID] = FVector2f(static_cast<float>(UVVector[0]), 1.0f - static_cast<float>(UVVector[1]));	// flip the Y of UVs for DirectX
 				}
-				FStaticMeshAttributes Attributes(*MeshDescription);
+			}
+		}
 
-				//The shape should be for a specified mesh
-				if (!ensure(Shape))
-				{
-					return false;
-				}
+		TMap<uint64, int32> RemapEdgeID;
+		{
+			TRACE_CPUPROFILER_EVENT_SCOPE(BuildMeshEdgeVertices);
+			Mesh->BeginGetMeshEdgeVertices();
 
-				//Extract the points into a simplified MeshDescription
+			//Fill the edge array
+			int32 FbxEdgeCount = Mesh->GetMeshEdgeCount();
+			RemapEdgeID.Reserve(FbxEdgeCount * 2);
+			for (int32 FbxEdgeIndex = 0; FbxEdgeIndex < FbxEdgeCount; ++FbxEdgeIndex)
+			{
+				int32 EdgeStartVertexIndex = -1;
+				int32 EdgeEndVertexIndex = -1;
+				Mesh->GetMeshEdgeVertices(FbxEdgeIndex, EdgeStartVertexIndex, EdgeEndVertexIndex);
+				// Skip invalid edges, i.e. one of the ends is invalid, or degenerated ones
+				if (EdgeStartVertexIndex == -1 || EdgeEndVertexIndex == -1 || EdgeStartVertexIndex == EdgeEndVertexIndex)
 				{
-					TRACE_CPUPROFILER_EVENT_SCOPE(GatherMorphTargetPoints);
-					MeshDescription->SuspendVertexIndexing();
-
-					// Construct the matrices for the conversion from right handed to left handed system
-					FbxAMatrix TotalMatrix;
-					FbxAMatrix TotalMatrixForNormal;
-					TotalMatrix.SetIdentity(); //We use the identity since we want to bake in the interchange factory if the pipeline request it
-					TotalMatrixForNormal = TotalMatrix.Inverse();
-					TotalMatrixForNormal = TotalMatrixForNormal.Transpose();
-
-					FbxGeometryBase* GeoBase = (FbxGeometryBase*)Shape;
-
-					int32 VertexCount = GeoBase->GetControlPointsCount();
-
-					TVertexAttributesRef<FVector3f> VertexPositions = Attributes.GetVertexPositions();
-					int32 VertexOffset = MeshDescription->Vertices().Num();
-					// The below code expects Num() to be equivalent to GetArraySize(), i.e. that all added elements are appended, not inserted into existing gaps
-					check(VertexOffset == MeshDescription->Vertices().GetArraySize());
-					//Fill the vertex array
-					auto GetFinalPosition = [&TotalMatrix](FbxVector4& FbxPosition)
-					{
-						FbxPosition = TotalMatrix.MultT(FbxPosition);
-						return FFbxConvert::ConvertPos(FbxPosition);
-					};
-
-					MeshDescription->ReserveNewVertices(VertexCount);
-					for (int32 VertexIndex = 0; VertexIndex < VertexCount; ++VertexIndex)
-					{
-						int32 RealVertexIndex = VertexOffset + VertexIndex;
-						const FVector VertexPosition = GetFinalPosition(GeoBase->GetControlPoints()[VertexIndex]);
-						//Maybe we want to do some work here like computing the deltas
-						//const FVector MeshVertexPosition = GetFinalPosition(Mesh->GetControlPoints()[VertexIndex]);
-						FVertexID AddedVertexId = MeshDescription->CreateVertex();
-						if (AddedVertexId.GetValue() != RealVertexIndex)
-						{
-							UInterchangeResultMeshError_Generic* Message = AddMessage<UInterchangeResultMeshError_Generic>(Shape);
-							Message->Text = LOCTEXT("CantCreateVertex", "Cannot create valid vertex for the mesh '{MeshName}'.");
-
-							return false;
-						}
-						//Add the delta position, so we do not have to recompute it later
-						VertexPositions[AddedVertexId] = (FVector3f)VertexPosition;// -MeshVertexPosition;
-					}
-					MeshDescription->ResumeVertexIndexing();
-					FElementIDRemappings OutRemappings;
-					MeshDescription->Compact(OutRemappings);
+					UInterchangeResultMeshWarning_Generic* Message = AddMessage<UInterchangeResultMeshWarning_Generic>(Mesh);
+					Message->Text = LOCTEXT("SkippingEdge", "Skipping invalid edge on mesh '{MeshName}'.");
+					continue;
 				}
-				return true;
+				FVertexID EdgeVertexStart(EdgeStartVertexIndex + VertexOffset);
+				check(MeshDescription->Vertices().IsValid(EdgeVertexStart));
+				FVertexID EdgeVertexEnd(EdgeEndVertexIndex + VertexOffset);
+				check(MeshDescription->Vertices().IsValid(EdgeVertexEnd));
+				uint64 CompactedKey = (((uint64)EdgeVertexStart.GetValue()) << 32) | ((uint64)EdgeVertexEnd.GetValue());
+				RemapEdgeID.Add(CompactedKey, FbxEdgeIndex);
+				//Add the other edge side
+				CompactedKey = (((uint64)EdgeVertexEnd.GetValue()) << 32) | ((uint64)EdgeVertexStart.GetValue());
+				RemapEdgeID.Add(CompactedKey, FbxEdgeIndex);
 			}
+			//Call this after all GetMeshEdgeIndexForPolygon call this is for optimization purpose.
+			Mesh->EndGetMeshEdgeVertices();
+		}
+
+		{
+			TRACE_CPUPROFILER_EVENT_SCOPE(BuildMeshEdgeIndexForPolygon);
 
-			bool FMeshDescriptionImporter::FillMeshDescriptionFromFbxMesh(FbxMesh* Mesh, TArray<FString>& OutJointUniqueNames, EMeshType MeshType)
+			// Compute and reserve memory to be used for vertex instances
 			{
-				if (!ensure(bInitialized) || !ensure(MeshDescription) || !ensure(Mesh))
+				int32 TotalVertexCount = 0;
+				for (int32 PolygonIndex = 0; PolygonIndex < PolygonCount; PolygonIndex++)
 				{
-					return false;
+					TotalVertexCount += Mesh->GetPolygonSize(PolygonIndex);
 				}
 
-				FStaticMeshAttributes Attributes(*MeshDescription);
-				Attributes.Register();
+				MeshDescription->ReserveNewPolygons(PolygonCount);
+				MeshDescription->ReserveNewVertexInstances(TotalVertexCount);
+				MeshDescription->ReserveNewEdges(TotalVertexCount);
+			}
 
-				//Get the base layer of the mesh
-				FbxLayer* BaseLayer = Mesh->GetLayer(0);
-				if (BaseLayer == NULL)
-				{
-					UInterchangeResultMeshError_Generic* Message = AddMessage<UInterchangeResultMeshError_Generic>(Mesh);
-					Message->Text = LOCTEXT("CantGetLayer", "Unable to get FBX base layer for the mesh '{MeshName}'.");
-					return false;
-				}
+			bool  bBeginGetMeshEdgeIndexForPolygonCalled = false;
+			bool  bBeginGetMeshEdgeIndexForPolygonRequired = true;
+			int32 CurrentVertexInstanceIndex = 0;
+			int32 SkippedVertexInstance = 0;
 
-				FFBXUVs FBXUVs(*this, Mesh);
+			// keep those for all iterations to avoid heap allocations
+			TArray<FVertexInstanceID> CornerInstanceIDs;
+			TArray<FVertexID> CornerVerticesIDs;
+			TArray<FVector, TInlineAllocator<3>> P;
 
-				FbxNode* MeshNode = Mesh->GetNodeCount() > 0 ? Mesh->GetNode() : nullptr; //Get the first node using the mesh to setup default asset materials
-				//
-				// create materials
-				//
-				//Create a material name array in the node order, also fill the Meshdescription PolygonGroup
-				TArray<FName> MaterialNames;
-				int32 MaterialCount = (MeshNode != nullptr) ? MeshNode->GetMaterialCount() : Mesh->GetElementMaterialCount();
-				if (MeshNode)
+			//Polygons
+			for (int32 PolygonIndex = 0; PolygonIndex < PolygonCount; PolygonIndex++)
+			{
+				int32 PolygonVertexCount = Mesh->GetPolygonSize(PolygonIndex);
+				//Verify if the polygon is degenerate, in this case do not add them
 				{
-					MaterialCount = MeshNode->GetMaterialCount();
-					MaterialNames.Reserve(MaterialCount);
-					for (int32 MaterialIndex = 0; MaterialIndex < MaterialCount; ++MaterialIndex)
+					float ComparisonThreshold = SMALL_NUMBER;
+					P.Reset();
+					P.AddUninitialized(PolygonVertexCount);
+					for (int32 CornerIndex = 0; CornerIndex < PolygonVertexCount; CornerIndex++)
 					{
-						FbxSurfaceMaterial* FbxMaterial = MeshNode->GetMaterial(MaterialIndex);
-						MaterialNames.Add(*FFbxHelper::GetFbxObjectName(FbxMaterial));
+						const int32 ControlPointIndex = Mesh->GetPolygonVertex(PolygonIndex, CornerIndex);
+						const FVertexID VertexID(VertexOffset + ControlPointIndex);
+						P[CornerIndex] = (FVector)VertexPositions[VertexID];
 					}
-				}
-
-				// Must do this before triangulating the mesh due to an FBX bug in TriangulateMeshAdvance
-				int32 LayerSmoothingCount = Mesh->GetLayerCount(FbxLayerElement::eSmoothing);
-				for (int32 i = 0; i < LayerSmoothingCount; i++)
-				{
-					FbxLayerElementSmoothing const* SmoothingInfo = Mesh->GetLayer(0)->GetSmoothing();
-					if (SmoothingInfo && SmoothingInfo->GetMappingMode() != FbxLayerElement::eByPolygon)
+					check(P.Num() > 2); //triangle is the smallest polygon we can have
+					const FVector Normal = ((P[1] - P[2]) ^ (P[0] - P[2])).GetSafeNormal(ComparisonThreshold);
+					//Check for degenerated polygons, avoid NAN
+					if (Normal.IsNearlyZero(ComparisonThreshold) || Normal.ContainsNaN())
 					{
-						SDKGeometryConverter->ComputePolygonSmoothingFromEdgeSmoothing(Mesh, i);
+						SkippedVertexInstance += PolygonVertexCount;
+						continue;
 					}
 				}
 
-				if (!Mesh->IsTriangleMesh())
+				int32 RealPolygonIndex = PolygonOffset + PolygonIndex;
+				CornerInstanceIDs.Reset();
+				CornerInstanceIDs.AddUninitialized(PolygonVertexCount);
+				CornerVerticesIDs.Reset();
+				CornerVerticesIDs.AddUninitialized(PolygonVertexCount);
+				for (int32 CornerIndex = 0; CornerIndex < PolygonVertexCount; CornerIndex++)
 				{
-					const bool bReplace = true;
-					FbxNodeAttribute* ConvertedNode = SDKGeometryConverter->Triangulate(Mesh, bReplace);
+					int32 VertexInstanceIndex = VertexInstanceOffset + CurrentVertexInstanceIndex + CornerIndex;
+					int32 RealFbxVertexIndex = SkippedVertexInstance + CurrentVertexInstanceIndex + CornerIndex;
+					const FVertexInstanceID VertexInstanceID(VertexInstanceIndex);
+					CornerInstanceIDs[CornerIndex] = VertexInstanceID;
+					const int32 ControlPointIndex = Mesh->GetPolygonVertex(PolygonIndex, CornerIndex);
+					const FVertexID VertexID(VertexOffset + ControlPointIndex);
+					const FVector VertexPosition = (FVector)VertexPositions[VertexID];
+					CornerVerticesIDs[CornerIndex] = VertexID;
 
-					if (ConvertedNode != NULL && ConvertedNode->GetAttributeType() == FbxNodeAttribute::eMesh)
-					{
-						Mesh = (fbxsdk::FbxMesh*)ConvertedNode;
-					}
-					else
+					FVertexInstanceID AddedVertexInstanceId = MeshDescription->CreateVertexInstance(VertexID);
+
+					//Make sure the Added vertex instance ID is matching the expected vertex instance ID
+					check(AddedVertexInstanceId == VertexInstanceID);
+
+					if (AddedVertexInstanceId.GetValue() != VertexInstanceIndex)
 					{
 						UInterchangeResultMeshError_Generic* Message = AddMessage<UInterchangeResultMeshError_Generic>(Mesh);
-						Message->Text = LOCTEXT("TriangulationError", "Unable to triangulate the mesh '{MeshName}'.");
+						Message->Text = LOCTEXT("CannotCreateVertexInstance", "Cannot create valid vertex instance for the mesh '{MeshName}'.");
 
-						MeshDescription->Empty();
 						return false;
 					}
-				}
 
-				// renew the base layer
-				BaseLayer = Mesh->GetLayer(0);
-
-				//
-				// get the smoothing group layer
-				//
-				bool bSmoothingAvailable = false;
-
-				FbxLayerElementSmoothing* SmoothingInfo = BaseLayer->GetSmoothing();
-				FbxLayerElement::EReferenceMode SmoothingReferenceMode(FbxLayerElement::eDirect);
-				FbxLayerElement::EMappingMode SmoothingMappingMode(FbxLayerElement::eByEdge);
-				if (SmoothingInfo)
-				{
-					if (SmoothingInfo->GetMappingMode() == FbxLayerElement::eByPolygon)
+					//UVs attributes
+					for (int32 UVLayerIndex = 0; UVLayerIndex < FBXUVs.UniqueUVCount; UVLayerIndex++)
 					{
-						//Convert the base layer to edge smoothing
-						SDKGeometryConverter->ComputeEdgeSmoothingFromPolygonSmoothing(Mesh, 0);
-						BaseLayer = Mesh->GetLayer(0);
-						SmoothingInfo = BaseLayer->GetSmoothing();
+						FVector2D FinalUVVector(0.0f, 0.0f);
+						if (FBXUVs.LayerElementUV[UVLayerIndex] != NULL)
+						{
+							int UVMapIndex = (FBXUVs.UVMappingMode[UVLayerIndex] == FbxLayerElement::eByControlPoint) ? ControlPointIndex : RealFbxVertexIndex;
+							int32 UVIndex = (FBXUVs.UVReferenceMode[UVLayerIndex] == FbxLayerElement::eDirect) ?
+								UVMapIndex : FBXUVs.LayerElementUV[UVLayerIndex]->GetIndexArray().GetAt(UVMapIndex);
+
+							FbxVector2	UVVector = FBXUVs.LayerElementUV[UVLayerIndex]->GetDirectArray().GetAt(UVIndex);
+							FinalUVVector.X = static_cast<float>(UVVector[0]);
+							FinalUVVector.Y = 1.f - static_cast<float>(UVVector[1]);   //flip the Y of UVs for DirectX
+						}
+						VertexInstanceUVs.Set(AddedVertexInstanceId, UVLayerIndex, FVector2f(FinalUVVector));	// LWC_TODO: Precision loss
 					}
 
-					if (SmoothingInfo->GetMappingMode() == FbxLayerElement::eByEdge)
+					//Color attribute
+					if (LayerElementVertexColor)
 					{
-						bSmoothingAvailable = true;
-					}
+						int32 VertexColorMappingIndex = (VertexColorMappingMode == FbxLayerElement::eByControlPoint) ?
+							Mesh->GetPolygonVertex(PolygonIndex, CornerIndex) : (RealFbxVertexIndex);
 
-					SmoothingReferenceMode = SmoothingInfo->GetReferenceMode();
-					SmoothingMappingMode = SmoothingInfo->GetMappingMode();
+						int32 VectorColorIndex = (VertexColorReferenceMode == FbxLayerElement::eDirect) ?
+							VertexColorMappingIndex : LayerElementVertexColor->GetIndexArray().GetAt(VertexColorMappingIndex);
 
-					//If all smooth group are 0's, we can say that there is no valid smoothing data, but only for skeletal mesh
-					//Static mesh will be faceted in this case since we will not compute smooth group from normals like we do here
-					//This is legacy code to get the same result has before
-					if (MeshType == EMeshType::Skinned)
-					{
-						// Check and see if the smooothing data is valid.  If not generate it from the normals
-						BaseLayer = Mesh->GetLayer(0);
-						if (BaseLayer)
-						{
-							const FbxLayerElementSmoothing* SmoothingLayer = BaseLayer->GetSmoothing();
+						FbxColor VertexColor = LayerElementVertexColor->GetDirectArray().GetAt(VectorColorIndex);
 
-							if (SmoothingLayer)
-							{
-								bool bValidSmoothingData = false;
-								FbxLayerElementArrayTemplate<int32>& Array = SmoothingLayer->GetDirectArray();
-								for (int32 SmoothingIndex = 0; SmoothingIndex < Array.GetCount(); ++SmoothingIndex)
-								{
-									if (Array[SmoothingIndex] != 0)
-									{
-										bValidSmoothingData = true;
-										break;
-									}
-								}
+						FColor VertexInstanceColor(
+							uint8(255.f * VertexColor.mRed),
+							uint8(255.f * VertexColor.mGreen),
+							uint8(255.f * VertexColor.mBlue),
+							uint8(255.f * VertexColor.mAlpha)
+						);
+						VertexInstanceColors[AddedVertexInstanceId] = FVector4f(FLinearColor(VertexInstanceColor));
+					}
 
-								if (!bValidSmoothingData && Mesh->GetPolygonVertexCount() > 0)
-								{
-									SDKGeometryConverter->ComputeEdgeSmoothingFromNormals(Mesh);
-									BaseLayer = Mesh->GetLayer(0);
-									SmoothingInfo = BaseLayer->GetSmoothing();
-								}
-							}
+					if (LayerElementNormal)
+					{
+						//normals may have different reference and mapping mode than tangents and binormals
+						int NormalMapIndex = (NormalMappingMode == FbxLayerElement::eByControlPoint) ?
+							ControlPointIndex : RealFbxVertexIndex;
+						int NormalValueIndex = (NormalReferenceMode == FbxLayerElement::eDirect) ?
+							NormalMapIndex : LayerElementNormal->GetIndexArray().GetAt(NormalMapIndex);
+
+						FbxVector4 TempValue = LayerElementNormal->GetDirectArray().GetAt(NormalValueIndex);
+						TempValue = TotalMatrixForNormal.MultT(TempValue);
+						FVector TangentZ = FFbxConvert::ConvertDir(TempValue);
+						VertexInstanceNormals[AddedVertexInstanceId] = (FVector3f)TangentZ.GetSafeNormal();
+						//tangents and binormals share the same reference, mapping mode and index array
+						if (bHasNTBInformation)
+						{
+							int TangentMapIndex = (TangentMappingMode == FbxLayerElement::eByControlPoint) ?
+								ControlPointIndex : RealFbxVertexIndex;
+							int TangentValueIndex = (TangentReferenceMode == FbxLayerElement::eDirect) ?
+								TangentMapIndex : LayerElementTangent->GetIndexArray().GetAt(TangentMapIndex);
+
+							TempValue = LayerElementTangent->GetDirectArray().GetAt(TangentValueIndex);
+							TempValue = TotalMatrixForNormal.MultT(TempValue);
+							FVector TangentX = FFbxConvert::ConvertDir(TempValue);
+							VertexInstanceTangents[AddedVertexInstanceId] = (FVector3f)TangentX.GetSafeNormal();
+
+							int BinormalMapIndex = (BinormalMappingMode == FbxLayerElement::eByControlPoint) ?
+								ControlPointIndex : RealFbxVertexIndex;
+							int BinormalValueIndex = (BinormalReferenceMode == FbxLayerElement::eDirect) ?
+								BinormalMapIndex : LayerElementBinormal->GetIndexArray().GetAt(BinormalMapIndex);
+
+							TempValue = LayerElementBinormal->GetDirectArray().GetAt(BinormalValueIndex);
+							TempValue = TotalMatrixForNormal.MultT(TempValue);
+							FVector TangentY = -FFbxConvert::ConvertDir(TempValue);
+							VertexInstanceBinormalSigns[AddedVertexInstanceId] = FbxGetBasisDeterminantSign(TangentX.GetSafeNormal(), TangentY.GetSafeNormal(), TangentZ.GetSafeNormal());
 						}
 					}
 				}
 
-				//
-				//	get the "material index" layer.  Do this AFTER the triangulation step as that may reorder material indices
-				//
-				FbxLayerElementMaterial* LayerElementMaterial = BaseLayer->GetMaterials();
-				FbxLayerElement::EMappingMode MaterialMappingMode = LayerElementMaterial ?
-					LayerElementMaterial->GetMappingMode() : FbxLayerElement::eByPolygon;
-
-				//	todo second phase UV, ok to put in first phase?
-				FBXUVs.Phase2(Mesh);
-
-				//
-				// get the first vertex color layer
-				//
-				FbxLayerElementVertexColor* LayerElementVertexColor = BaseLayer->GetVertexColors();
-				FbxLayerElement::EReferenceMode VertexColorReferenceMode(FbxLayerElement::eDirect);
-				FbxLayerElement::EMappingMode VertexColorMappingMode(FbxLayerElement::eByControlPoint);
-				if (LayerElementVertexColor)
+				// Check if the polygon just discovered is non-degenerate if we haven't found one yet
+				//TODO check all polygon vertex, not just the first 3 vertex
+				if (!bHasNonDegeneratePolygons)
 				{
-					VertexColorReferenceMode = LayerElementVertexColor->GetReferenceMode();
-					VertexColorMappingMode = LayerElementVertexColor->GetMappingMode();
+					float TriangleComparisonThreshold = THRESH_POINTS_ARE_SAME;
+					FVector VertexPosition[3];
+					VertexPosition[0] = (FVector)VertexPositions[CornerVerticesIDs[0]];
+					VertexPosition[1] = (FVector)VertexPositions[CornerVerticesIDs[1]];
+					VertexPosition[2] = (FVector)VertexPositions[CornerVerticesIDs[2]];
+					if (!(VertexPosition[0].Equals(VertexPosition[1], TriangleComparisonThreshold)
+						|| VertexPosition[0].Equals(VertexPosition[2], TriangleComparisonThreshold)
+						|| VertexPosition[1].Equals(VertexPosition[2], TriangleComparisonThreshold)))
+					{
+						bHasNonDegeneratePolygons = true;
+					}
 				}
 
 				//
-				// get the first normal layer
+				// material index
 				//
-				FbxLayerElementNormal* LayerElementNormal = BaseLayer->GetNormals();
-				FbxLayerElementTangent* LayerElementTangent = BaseLayer->GetTangents();
-				FbxLayerElementBinormal* LayerElementBinormal = BaseLayer->GetBinormals();
-
-				//whether there is normal, tangent and binormal data in this mesh
-				bool bHasNTBInformation = LayerElementNormal && LayerElementTangent && LayerElementBinormal;
-
-				FbxLayerElement::EReferenceMode NormalReferenceMode(FbxLayerElement::eDirect);
-				FbxLayerElement::EMappingMode NormalMappingMode(FbxLayerElement::eByControlPoint);
-				if (LayerElementNormal)
+				int32 MaterialIndex = 0;
+				if (MaterialCount > 0)
 				{
-					NormalReferenceMode = LayerElementNormal->GetReferenceMode();
-					NormalMappingMode = LayerElementNormal->GetMappingMode();
+					if (LayerElementMaterial)
+					{
+						switch (MaterialMappingMode)
+						{
+							// material index is stored in the IndexArray, not the DirectArray (which is irrelevant with 2009.1)
+							case FbxLayerElement::eAllSame:
+							{
+								MaterialIndex = LayerElementMaterial->GetIndexArray().GetAt(0);
+							}
+							break;
+							case FbxLayerElement::eByPolygon:
+							{
+								MaterialIndex = LayerElementMaterial->GetIndexArray().GetAt(PolygonIndex);
+							}
+							break;
+						}
+					}
 				}
 
-				FbxLayerElement::EReferenceMode TangentReferenceMode(FbxLayerElement::eDirect);
-				FbxLayerElement::EMappingMode TangentMappingMode(FbxLayerElement::eByControlPoint);
-				if (LayerElementTangent)
+				if (MaterialIndex < 0)
 				{
-					TangentReferenceMode = LayerElementTangent->GetReferenceMode();
-					TangentMappingMode = LayerElementTangent->GetMappingMode();
+					//TODO log an error message
+					MaterialIndex = 0;
 				}
-
-				FbxLayerElement::EReferenceMode BinormalReferenceMode(FbxLayerElement::eDirect);
-				FbxLayerElement::EMappingMode BinormalMappingMode(FbxLayerElement::eByControlPoint);
-				if (LayerElementBinormal)
+				//Add missing material name with generated name
+				while (MaterialIndex >= MaterialNames.Num())
 				{
-					BinormalReferenceMode = LayerElementBinormal->GetReferenceMode();
-					BinormalMappingMode = LayerElementBinormal->GetMappingMode();
+					FString MaterialName = TEXT("Material_") + FString::FromInt(MaterialNames.Num());
+					MaterialNames.Add((*MaterialName));
 				}
 
-				bool bHasNonDegeneratePolygons = false;
+				//Material name should always exist since we create up to index materials
+				check(MaterialNames.IsValidIndex(MaterialIndex));
 				{
-					MeshDescription->SuspendVertexInstanceIndexing();
-					MeshDescription->SuspendEdgeIndexing();
-					MeshDescription->SuspendPolygonIndexing();
-					MeshDescription->SuspendPolygonGroupIndexing();
-					MeshDescription->SuspendUVIndexing();
-
-					TRACE_CPUPROFILER_EVENT_SCOPE(BuildTriangles);
-
-					// Construct the matrices for the conversion from right handed to left handed system
-					FbxAMatrix TotalMatrix;
-					FbxAMatrix TotalMatrixForNormal;
-					TotalMatrix.SetIdentity();
-					TotalMatrixForNormal = TotalMatrix.Inverse();
-					TotalMatrixForNormal = TotalMatrixForNormal.Transpose();
-					int32 PolygonCount = Mesh->GetPolygonCount();
-
-					if (PolygonCount == 0)
-					{
-						UInterchangeResultMeshError_Generic* Message = AddMessage<UInterchangeResultMeshError_Generic>(Mesh);
-						Message->Text = LOCTEXT("NoPolygons", "No polygons were found in the mesh '{MeshName}'.");
-
-						return false;
-					}
-
-					int32 VertexCount = Mesh->GetControlPointsCount();
-					bool OddNegativeScale = IsOddNegativeScale(TotalMatrix);
-
-					TVertexAttributesRef<FVector3f> VertexPositions = Attributes.GetVertexPositions();
-					TVertexInstanceAttributesRef<FVector3f> VertexInstanceNormals = Attributes.GetVertexInstanceNormals();
-					TVertexInstanceAttributesRef<FVector3f> VertexInstanceTangents = Attributes.GetVertexInstanceTangents();
-					TVertexInstanceAttributesRef<float> VertexInstanceBinormalSigns = Attributes.GetVertexInstanceBinormalSigns();
-					TVertexInstanceAttributesRef<FVector4f> VertexInstanceColors = Attributes.GetVertexInstanceColors();
-					TVertexInstanceAttributesRef<FVector2f> VertexInstanceUVs = Attributes.GetVertexInstanceUVs();
-					TEdgeAttributesRef<bool> EdgeHardnesses = Attributes.GetEdgeHardnesses();
-					TPolygonGroupAttributesRef<FName> PolygonGroupImportedMaterialSlotNames = Attributes.GetPolygonGroupMaterialSlotNames();
-
-					int32 VertexOffset = MeshDescription->Vertices().Num();
-					int32 VertexInstanceOffset = MeshDescription->VertexInstances().Num();
-					int32 PolygonOffset = MeshDescription->Polygons().Num();
-
-					// The below code expects Num() to be equivalent to GetArraySize(), i.e. that all added elements are appended, not inserted into existing gaps
-					check(VertexOffset == MeshDescription->Vertices().GetArraySize());
-					check(VertexInstanceOffset == MeshDescription->VertexInstances().GetArraySize());
-					check(PolygonOffset == MeshDescription->Polygons().GetArraySize());
-
-					TMap<int32, FPolygonGroupID> PolygonGroupMapping;
-
-					// When importing multiple mesh pieces to the same static mesh.  Ensure each mesh piece has the same number of Uv's
-					int32 ExistingUVCount = VertexInstanceUVs.GetNumChannels();
-
-					int32 NumUVs = FMath::Max(FBXUVs.UniqueUVCount, ExistingUVCount);
-					NumUVs = FMath::Min<int32>(MAX_MESH_TEXTURE_COORDS_MD, NumUVs);
-					// At least one UV set must exist.  
-					// @todo: this needn't be true; we should be able to handle zero UV channels
-					NumUVs = FMath::Max(1, NumUVs);
-
-					//Make sure all Vertex instance have the correct number of UVs
-					VertexInstanceUVs.SetNumChannels(NumUVs);
-					MeshDescription->SetNumUVChannels(NumUVs);
-
-					TArray<int32> UVOffsets;
-					UVOffsets.SetNumUninitialized(NumUVs);
-					for (int UVChannel = 0; UVChannel < NumUVs; UVChannel++)
+					//Create a polygon with the 3 vertex instances Add it to the material group
+					FName MaterialName = MaterialNames[MaterialIndex];
+					if (!PolygonGroupMapping.Contains(MaterialIndex))
 					{
-						UVOffsets[UVChannel] = MeshDescription->UVs(UVChannel).GetArraySize();
-					}
-
-					//Fill the vertex array
-					MeshDescription->ReserveNewVertices(VertexCount);
-					for (int32 VertexIndex = 0; VertexIndex < VertexCount; ++VertexIndex)
-					{
-						int32 RealVertexIndex = VertexOffset + VertexIndex;
-						FbxVector4 FbxPosition = Mesh->GetControlPoints()[VertexIndex];
-						FbxPosition = TotalMatrix.MultT(FbxPosition);
-						const FVector VertexPosition = FFbxConvert::ConvertPos(FbxPosition);
-
-						FVertexID AddedVertexId = MeshDescription->CreateVertex();
-						VertexPositions[AddedVertexId] = (FVector3f)VertexPosition;
-						if (AddedVertexId.GetValue() != RealVertexIndex)
+						FPolygonGroupID ExistingPolygonGroup = INDEX_NONE;
+						for (const FPolygonGroupID PolygonGroupID : MeshDescription->PolygonGroups().GetElementIDs())
 						{
-							UInterchangeResultMeshError_Generic* Message = AddMessage<UInterchangeResultMeshError_Generic>(Mesh);
-							Message->Text = LOCTEXT("CannotCreateVertex", "Cannot create valid vertex for mesh '{MeshName}'.");
-
-							return false;
-						}
-					}
-
-					// Fill the UV arrays
-					for (int32 UVLayerIndex = 0; UVLayerIndex < FBXUVs.UniqueUVCount; UVLayerIndex++)
-					{
-						check(FBXUVs.LayerElementUV[UVLayerIndex]);
-						if (FBXUVs.LayerElementUV[UVLayerIndex] != nullptr)
-						{
-							int32 UVCount = FBXUVs.LayerElementUV[UVLayerIndex]->GetDirectArray().GetCount();
-							TUVAttributesRef<FVector2f> UVCoordinates = MeshDescription->UVAttributes(UVLayerIndex).GetAttributesRef<FVector2f>(MeshAttribute::UV::UVCoordinate);
-							MeshDescription->ReserveNewUVs(UVCount, UVLayerIndex);
-							for (int32 UVIndex = 0; UVIndex < UVCount; UVIndex++)
+							if (PolygonGroupImportedMaterialSlotNames[PolygonGroupID] == MaterialName)
 							{
-								FUVID UVID = MeshDescription->CreateUV(UVLayerIndex);
-								FbxVector2 UVVector = FBXUVs.LayerElementUV[UVLayerIndex]->GetDirectArray().GetAt(UVIndex);
-								UVCoordinates[UVID] = FVector2f(static_cast<float>(UVVector[0]), 1.0f - static_cast<float>(UVVector[1]));	// flip the Y of UVs for DirectX
+								ExistingPolygonGroup = PolygonGroupID;
+								break;
 							}
 						}
-					}
-
-					TMap<uint64, int32> RemapEdgeID;
-					{
-						TRACE_CPUPROFILER_EVENT_SCOPE(BuildMeshEdgeVertices);
-						Mesh->BeginGetMeshEdgeVertices();
-
-						//Fill the edge array
-						int32 FbxEdgeCount = Mesh->GetMeshEdgeCount();
-						RemapEdgeID.Reserve(FbxEdgeCount * 2);
-						for (int32 FbxEdgeIndex = 0; FbxEdgeIndex < FbxEdgeCount; ++FbxEdgeIndex)
+						if (ExistingPolygonGroup == INDEX_NONE)
 						{
-							int32 EdgeStartVertexIndex = -1;
-							int32 EdgeEndVertexIndex = -1;
-							Mesh->GetMeshEdgeVertices(FbxEdgeIndex, EdgeStartVertexIndex, EdgeEndVertexIndex);
-							// Skip invalid edges, i.e. one of the ends is invalid, or degenerated ones
-							if (EdgeStartVertexIndex == -1 || EdgeEndVertexIndex == -1 || EdgeStartVertexIndex == EdgeEndVertexIndex)
-							{
-								UInterchangeResultMeshWarning_Generic* Message = AddMessage<UInterchangeResultMeshWarning_Generic>(Mesh);
-								Message->Text = LOCTEXT("SkippingEdge", "Skipping invalid edge on mesh '{MeshName}'.");
-								continue;
-							}
-							FVertexID EdgeVertexStart(EdgeStartVertexIndex + VertexOffset);
-							check(MeshDescription->Vertices().IsValid(EdgeVertexStart));
-							FVertexID EdgeVertexEnd(EdgeEndVertexIndex + VertexOffset);
-							check(MeshDescription->Vertices().IsValid(EdgeVertexEnd));
-							uint64 CompactedKey = (((uint64)EdgeVertexStart.GetValue()) << 32) | ((uint64)EdgeVertexEnd.GetValue());
-							RemapEdgeID.Add(CompactedKey, FbxEdgeIndex);
-							//Add the other edge side
-							CompactedKey = (((uint64)EdgeVertexEnd.GetValue()) << 32) | ((uint64)EdgeVertexStart.GetValue());
-							RemapEdgeID.Add(CompactedKey, FbxEdgeIndex);
+							ExistingPolygonGroup = MeshDescription->CreatePolygonGroup();
+							PolygonGroupImportedMaterialSlotNames[ExistingPolygonGroup] = MaterialName;
 						}
-						//Call this after all GetMeshEdgeIndexForPolygon call this is for optimization purpose.
-						Mesh->EndGetMeshEdgeVertices();
+						PolygonGroupMapping.Add(MaterialIndex, ExistingPolygonGroup);
 					}
+				}
 
+				// Create polygon edges
+				{
+					// Add the edges of this polygon
+					for (uint32 PolygonEdgeNumber = 0; PolygonEdgeNumber < (uint32)PolygonVertexCount; ++PolygonEdgeNumber)
 					{
-						TRACE_CPUPROFILER_EVENT_SCOPE(BuildMeshEdgeIndexForPolygon);
+						//Find the matching edge ID
+						uint32 CornerIndices[2];
+						CornerIndices[0] = (PolygonEdgeNumber + 0) % PolygonVertexCount;
+						CornerIndices[1] = (PolygonEdgeNumber + 1) % PolygonVertexCount;
 
-						// Compute and reserve memory to be used for vertex instances
-						{
-							int32 TotalVertexCount = 0;
-							for (int32 PolygonIndex = 0; PolygonIndex < PolygonCount; PolygonIndex++)
-							{
-								TotalVertexCount += Mesh->GetPolygonSize(PolygonIndex);
-							}
+						FVertexID EdgeVertexIDs[2];
+						EdgeVertexIDs[0] = CornerVerticesIDs[CornerIndices[0]];
+						EdgeVertexIDs[1] = CornerVerticesIDs[CornerIndices[1]];
 
-							MeshDescription->ReserveNewPolygons(PolygonCount);
-							MeshDescription->ReserveNewVertexInstances(TotalVertexCount);
-							MeshDescription->ReserveNewEdges(TotalVertexCount);
+						FEdgeID MatchEdgeId = MeshDescription->GetVertexPairEdge(EdgeVertexIDs[0], EdgeVertexIDs[1]);
+						if (MatchEdgeId == INDEX_NONE)
+						{
+							MatchEdgeId = MeshDescription->CreateEdge(EdgeVertexIDs[0], EdgeVertexIDs[1]);
 						}
 
-						bool  bBeginGetMeshEdgeIndexForPolygonCalled = false;
-						bool  bBeginGetMeshEdgeIndexForPolygonRequired = true;
-						int32 CurrentVertexInstanceIndex = 0;
-						int32 SkippedVertexInstance = 0;
-
-						// keep those for all iterations to avoid heap allocations
-						TArray<FVertexInstanceID> CornerInstanceIDs;
-						TArray<FVertexID> CornerVerticesIDs;
-						TArray<FVector, TInlineAllocator<3>> P;
-
-						//Polygons
-						for (int32 PolygonIndex = 0; PolygonIndex < PolygonCount; PolygonIndex++)
+						//RawMesh do not have edges, so by ordering the edge with the triangle construction we can ensure back and forth conversion with RawMesh
+						//When raw mesh will be completely remove we can create the edges right after the vertex creation.
+						int32 EdgeIndex = INDEX_NONE;
+						uint64 CompactedKey = (((uint64)EdgeVertexIDs[0].GetValue()) << 32) | ((uint64)EdgeVertexIDs[1].GetValue());
+						if (RemapEdgeID.Contains(CompactedKey))
 						{
-							int32 PolygonVertexCount = Mesh->GetPolygonSize(PolygonIndex);
-							//Verify if the polygon is degenerate, in this case do not add them
-							{
-								float ComparisonThreshold = SMALL_NUMBER;
-								P.Reset();
-								P.AddUninitialized(PolygonVertexCount);
-								for (int32 CornerIndex = 0; CornerIndex < PolygonVertexCount; CornerIndex++)
-								{
-									const int32 ControlPointIndex = Mesh->GetPolygonVertex(PolygonIndex, CornerIndex);
-									const FVertexID VertexID(VertexOffset + ControlPointIndex);
-									P[CornerIndex] = (FVector)VertexPositions[VertexID];
-								}
-								check(P.Num() > 2); //triangle is the smallest polygon we can have
-								const FVector Normal = ((P[1] - P[2]) ^ (P[0] - P[2])).GetSafeNormal(ComparisonThreshold);
-								//Check for degenerated polygons, avoid NAN
-								if (Normal.IsNearlyZero(ComparisonThreshold) || Normal.ContainsNaN())
-								{
-									SkippedVertexInstance += PolygonVertexCount;
-									continue;
-								}
-							}
-
-							int32 RealPolygonIndex = PolygonOffset + PolygonIndex;
-							CornerInstanceIDs.Reset();
-							CornerInstanceIDs.AddUninitialized(PolygonVertexCount);
-							CornerVerticesIDs.Reset();
-							CornerVerticesIDs.AddUninitialized(PolygonVertexCount);
-							for (int32 CornerIndex = 0; CornerIndex < PolygonVertexCount; CornerIndex++)
-							{
-								int32 VertexInstanceIndex = VertexInstanceOffset + CurrentVertexInstanceIndex + CornerIndex;
-								int32 RealFbxVertexIndex = SkippedVertexInstance + CurrentVertexInstanceIndex + CornerIndex;
-								const FVertexInstanceID VertexInstanceID(VertexInstanceIndex);
-								CornerInstanceIDs[CornerIndex] = VertexInstanceID;
-								const int32 ControlPointIndex = Mesh->GetPolygonVertex(PolygonIndex, CornerIndex);
-								const FVertexID VertexID(VertexOffset + ControlPointIndex);
-								const FVector VertexPosition = (FVector)VertexPositions[VertexID];
-								CornerVerticesIDs[CornerIndex] = VertexID;
-
-								FVertexInstanceID AddedVertexInstanceId = MeshDescription->CreateVertexInstance(VertexID);
-
-								//Make sure the Added vertex instance ID is matching the expected vertex instance ID
-								check(AddedVertexInstanceId == VertexInstanceID);
-
-								if (AddedVertexInstanceId.GetValue() != VertexInstanceIndex)
-								{
-									UInterchangeResultMeshError_Generic* Message = AddMessage<UInterchangeResultMeshError_Generic>(Mesh);
-									Message->Text = LOCTEXT("CannotCreateVertexInstance", "Cannot create valid vertex instance for the mesh '{MeshName}'.");
-
-									return false;
-								}
-
-								//UVs attributes
-								for (int32 UVLayerIndex = 0; UVLayerIndex < FBXUVs.UniqueUVCount; UVLayerIndex++)
-								{
-									FVector2D FinalUVVector(0.0f, 0.0f);
-									if (FBXUVs.LayerElementUV[UVLayerIndex] != NULL)
-									{
-										int UVMapIndex = (FBXUVs.UVMappingMode[UVLayerIndex] == FbxLayerElement::eByControlPoint) ? ControlPointIndex : RealFbxVertexIndex;
-										int32 UVIndex = (FBXUVs.UVReferenceMode[UVLayerIndex] == FbxLayerElement::eDirect) ?
-											UVMapIndex : FBXUVs.LayerElementUV[UVLayerIndex]->GetIndexArray().GetAt(UVMapIndex);
-
-										FbxVector2	UVVector = FBXUVs.LayerElementUV[UVLayerIndex]->GetDirectArray().GetAt(UVIndex);
-										FinalUVVector.X = static_cast<float>(UVVector[0]);
-										FinalUVVector.Y = 1.f - static_cast<float>(UVVector[1]);   //flip the Y of UVs for DirectX
-									}
-									VertexInstanceUVs.Set(AddedVertexInstanceId, UVLayerIndex, FVector2f(FinalUVVector));	// LWC_TODO: Precision loss
-								}
-
-								//Color attribute
-								if (LayerElementVertexColor)
-								{
-									int32 VertexColorMappingIndex = (VertexColorMappingMode == FbxLayerElement::eByControlPoint) ?
-										Mesh->GetPolygonVertex(PolygonIndex, CornerIndex) : (RealFbxVertexIndex);
-
-									int32 VectorColorIndex = (VertexColorReferenceMode == FbxLayerElement::eDirect) ?
-										VertexColorMappingIndex : LayerElementVertexColor->GetIndexArray().GetAt(VertexColorMappingIndex);
-
-									FbxColor VertexColor = LayerElementVertexColor->GetDirectArray().GetAt(VectorColorIndex);
-
-									FColor VertexInstanceColor(
-										uint8(255.f * VertexColor.mRed),
-										uint8(255.f * VertexColor.mGreen),
-										uint8(255.f * VertexColor.mBlue),
-										uint8(255.f * VertexColor.mAlpha)
-									);
-									VertexInstanceColors[AddedVertexInstanceId] = FVector4f(FLinearColor(VertexInstanceColor));
-								}
-
-								if (LayerElementNormal)
-								{
-									//normals may have different reference and mapping mode than tangents and binormals
-									int NormalMapIndex = (NormalMappingMode == FbxLayerElement::eByControlPoint) ?
-										ControlPointIndex : RealFbxVertexIndex;
-									int NormalValueIndex = (NormalReferenceMode == FbxLayerElement::eDirect) ?
-										NormalMapIndex : LayerElementNormal->GetIndexArray().GetAt(NormalMapIndex);
-
-									FbxVector4 TempValue = LayerElementNormal->GetDirectArray().GetAt(NormalValueIndex);
-									TempValue = TotalMatrixForNormal.MultT(TempValue);
-									FVector TangentZ = FFbxConvert::ConvertDir(TempValue);
-									VertexInstanceNormals[AddedVertexInstanceId] = (FVector3f)TangentZ.GetSafeNormal();
-									//tangents and binormals share the same reference, mapping mode and index array
-									if (bHasNTBInformation)
-									{
-										int TangentMapIndex = (TangentMappingMode == FbxLayerElement::eByControlPoint) ?
-											ControlPointIndex : RealFbxVertexIndex;
-										int TangentValueIndex = (TangentReferenceMode == FbxLayerElement::eDirect) ?
-											TangentMapIndex : LayerElementTangent->GetIndexArray().GetAt(TangentMapIndex);
-
-										TempValue = LayerElementTangent->GetDirectArray().GetAt(TangentValueIndex);
-										TempValue = TotalMatrixForNormal.MultT(TempValue);
-										FVector TangentX = FFbxConvert::ConvertDir(TempValue);
-										VertexInstanceTangents[AddedVertexInstanceId] = (FVector3f)TangentX.GetSafeNormal();
-
-										int BinormalMapIndex = (BinormalMappingMode == FbxLayerElement::eByControlPoint) ?
-											ControlPointIndex : RealFbxVertexIndex;
-										int BinormalValueIndex = (BinormalReferenceMode == FbxLayerElement::eDirect) ?
-											BinormalMapIndex : LayerElementBinormal->GetIndexArray().GetAt(BinormalMapIndex);
-
-										TempValue = LayerElementBinormal->GetDirectArray().GetAt(BinormalValueIndex);
-										TempValue = TotalMatrixForNormal.MultT(TempValue);
-										FVector TangentY = -FFbxConvert::ConvertDir(TempValue);
-										VertexInstanceBinormalSigns[AddedVertexInstanceId] = FbxGetBasisDeterminantSign(TangentX.GetSafeNormal(), TangentY.GetSafeNormal(), TangentZ.GetSafeNormal());
-									}
-								}
-							}
-
-							// Check if the polygon just discovered is non-degenerate if we haven't found one yet
-							//TODO check all polygon vertex, not just the first 3 vertex
-							if (!bHasNonDegeneratePolygons)
-							{
-								float TriangleComparisonThreshold = THRESH_POINTS_ARE_SAME;
-								FVector VertexPosition[3];
-								VertexPosition[0] = (FVector)VertexPositions[CornerVerticesIDs[0]];
-								VertexPosition[1] = (FVector)VertexPositions[CornerVerticesIDs[1]];
-								VertexPosition[2] = (FVector)VertexPositions[CornerVerticesIDs[2]];
-								if (!(VertexPosition[0].Equals(VertexPosition[1], TriangleComparisonThreshold)
-									|| VertexPosition[0].Equals(VertexPosition[2], TriangleComparisonThreshold)
-									|| VertexPosition[1].Equals(VertexPosition[2], TriangleComparisonThreshold)))
-								{
-									bHasNonDegeneratePolygons = true;
-								}
-							}
-
-							//
-							// material index
-							//
-							int32 MaterialIndex = 0;
-							if (MaterialCount > 0)
+							EdgeIndex = RemapEdgeID[CompactedKey];
+						}
+						else
+						{
+							// Call BeginGetMeshEdgeIndexForPolygon lazily only if we enter the code path calling GetMeshEdgeIndexForPolygon
+							if (bBeginGetMeshEdgeIndexForPolygonRequired)
 							{
-								if (LayerElementMaterial)
+								//Call this before all GetMeshEdgeIndexForPolygon for optimization purpose.
+								//But do not spend time precomputing stuff if the mesh has no edge since 
+								//GetMeshEdgeIndexForPolygon is going to always returns -1 anyway without going into the slow path.
+								if (Mesh->GetMeshEdgeCount() > 0)
 								{
-									switch (MaterialMappingMode)
-									{
-										// material index is stored in the IndexArray, not the DirectArray (which is irrelevant with 2009.1)
-										case FbxLayerElement::eAllSame:
-										{
-											MaterialIndex = LayerElementMaterial->GetIndexArray().GetAt(0);
-										}
-										break;
-										case FbxLayerElement::eByPolygon:
-										{
-											MaterialIndex = LayerElementMaterial->GetIndexArray().GetAt(PolygonIndex);
-										}
-										break;
-									}
+									Mesh->BeginGetMeshEdgeIndexForPolygon();
+									bBeginGetMeshEdgeIndexForPolygonCalled = true;
 								}
+								bBeginGetMeshEdgeIndexForPolygonRequired = false;
 							}
 
-							if (MaterialIndex < 0)
-							{
-								//TODO log an error message
-								MaterialIndex = 0;
-							}
-							//Add missing material name with generated name
-							while (MaterialIndex >= MaterialNames.Num())
-							{
-								FString MaterialName = TEXT("Material_") + FString::FromInt(MaterialNames.Num());
-								MaterialNames.Add((*MaterialName));
-							}
+							EdgeIndex = Mesh->GetMeshEdgeIndexForPolygon(PolygonIndex, PolygonEdgeNumber);
+						}
 
-							//Material name should always exist since we create up to index materials
-							check(MaterialNames.IsValidIndex(MaterialIndex));
+						if (!EdgeHardnesses[MatchEdgeId])
+						{
+							if (bSmoothingAvailable && SmoothingInfo)
 							{
-								//Create a polygon with the 3 vertex instances Add it to the material group
-								FName MaterialName = MaterialNames[MaterialIndex];
-								if (!PolygonGroupMapping.Contains(MaterialIndex))
+								if (SmoothingMappingMode == FbxLayerElement::eByEdge)
 								{
-									FPolygonGroupID ExistingPolygonGroup = INDEX_NONE;
-									for (const FPolygonGroupID PolygonGroupID : MeshDescription->PolygonGroups().GetElementIDs())
-									{
-										if (PolygonGroupImportedMaterialSlotNames[PolygonGroupID] == MaterialName)
-										{
-											ExistingPolygonGroup = PolygonGroupID;
-											break;
-										}
-									}
-									if (ExistingPolygonGroup == INDEX_NONE)
-									{
-										ExistingPolygonGroup = MeshDescription->CreatePolygonGroup();
-										PolygonGroupImportedMaterialSlotNames[ExistingPolygonGroup] = MaterialName;
-									}
-									PolygonGroupMapping.Add(MaterialIndex, ExistingPolygonGroup);
+									int32 lSmoothingIndex = (SmoothingReferenceMode == FbxLayerElement::eDirect) ? EdgeIndex : SmoothingInfo->GetIndexArray().GetAt(EdgeIndex);
+									//Set the hard edges
+									int32 SmoothingFlag = SmoothingInfo->GetDirectArray().GetAt(lSmoothingIndex);
+									EdgeHardnesses[MatchEdgeId] = (SmoothingFlag == 0);
 								}
-							}
-
-							// Create polygon edges
-							{
-								// Add the edges of this polygon
-								for (uint32 PolygonEdgeNumber = 0; PolygonEdgeNumber < (uint32)PolygonVertexCount; ++PolygonEdgeNumber)
+								else
 								{
-									//Find the matching edge ID
-									uint32 CornerIndices[2];
-									CornerIndices[0] = (PolygonEdgeNumber + 0) % PolygonVertexCount;
-									CornerIndices[1] = (PolygonEdgeNumber + 1) % PolygonVertexCount;
-
-									FVertexID EdgeVertexIDs[2];
-									EdgeVertexIDs[0] = CornerVerticesIDs[CornerIndices[0]];
-									EdgeVertexIDs[1] = CornerVerticesIDs[CornerIndices[1]];
-
-									FEdgeID MatchEdgeId = MeshDescription->GetVertexPairEdge(EdgeVertexIDs[0], EdgeVertexIDs[1]);
-									if (MatchEdgeId == INDEX_NONE)
-									{
-										MatchEdgeId = MeshDescription->CreateEdge(EdgeVertexIDs[0], EdgeVertexIDs[1]);
-									}
-
-									//RawMesh do not have edges, so by ordering the edge with the triangle construction we can ensure back and forth conversion with RawMesh
-									//When raw mesh will be completely remove we can create the edges right after the vertex creation.
-									int32 EdgeIndex = INDEX_NONE;
-									uint64 CompactedKey = (((uint64)EdgeVertexIDs[0].GetValue()) << 32) | ((uint64)EdgeVertexIDs[1].GetValue());
-									if (RemapEdgeID.Contains(CompactedKey))
-									{
-										EdgeIndex = RemapEdgeID[CompactedKey];
-									}
-									else
-									{
-										// Call BeginGetMeshEdgeIndexForPolygon lazily only if we enter the code path calling GetMeshEdgeIndexForPolygon
-										if (bBeginGetMeshEdgeIndexForPolygonRequired)
-										{
-											//Call this before all GetMeshEdgeIndexForPolygon for optimization purpose.
-											//But do not spend time precomputing stuff if the mesh has no edge since 
-											//GetMeshEdgeIndexForPolygon is going to always returns -1 anyway without going into the slow path.
-											if (Mesh->GetMeshEdgeCount() > 0)
-											{
-												Mesh->BeginGetMeshEdgeIndexForPolygon();
-												bBeginGetMeshEdgeIndexForPolygonCalled = true;
-											}
-											bBeginGetMeshEdgeIndexForPolygonRequired = false;
-										}
-
-										EdgeIndex = Mesh->GetMeshEdgeIndexForPolygon(PolygonIndex, PolygonEdgeNumber);
-									}
-
-									if (!EdgeHardnesses[MatchEdgeId])
-									{
-										if (bSmoothingAvailable && SmoothingInfo)
-										{
-											if (SmoothingMappingMode == FbxLayerElement::eByEdge)
-											{
-												int32 lSmoothingIndex = (SmoothingReferenceMode == FbxLayerElement::eDirect) ? EdgeIndex : SmoothingInfo->GetIndexArray().GetAt(EdgeIndex);
-												//Set the hard edges
-												int32 SmoothingFlag = SmoothingInfo->GetDirectArray().GetAt(lSmoothingIndex);
-												EdgeHardnesses[MatchEdgeId] = (SmoothingFlag == 0);
-											}
-											else
-											{
-												EdgeHardnesses[MatchEdgeId] = false;
-												//TODO add an error log
-												//AddTokenizedErrorMessage(FTokenizedMessage::Create(EMessageSeverity::Warning, FText::Format(LOCTEXT("Error_UnsupportedSmoothingGroup", "Unsupported Smoothing group mapping mode on mesh  '{0}'"), FText::FromString(Mesh->GetName()))), FFbxErrors::Generic_Mesh_UnsupportingSmoothingGroup);
-											}
-										}
-										else
-										{
-											//When there is no smoothing group we set all edge to: hard (faceted mesh) for static mesh and smooth for skinned and rigid
-											EdgeHardnesses[MatchEdgeId] = MeshType == EMeshType::Static ? true : false;
-										}
-									}
+									EdgeHardnesses[MatchEdgeId] = false;
+									//TODO add an error log
+									//AddTokenizedErrorMessage(FTokenizedMessage::Create(EMessageSeverity::Warning, FText::Format(LOCTEXT("Error_UnsupportedSmoothingGroup", "Unsupported Smoothing group mapping mode on mesh  '{0}'"), FText::FromString(Mesh->GetName()))), FFbxErrors::Generic_Mesh_UnsupportingSmoothingGroup);
 								}
 							}
-
-							FPolygonGroupID PolygonGroupID = PolygonGroupMapping[MaterialIndex];
-							// Insert a triangle into the mesh
-							// @note: we only ever import triangulated meshes currently. This could easily change as we have the infrastructure set up for arbitrary ngons.
-							// However, I think either we need to triangulate in the exact same way as Maya/Max if we do that.
-							check(CornerInstanceIDs.Num() == 3);
-							check(PolygonVertexCount == 3);
-							TArray<FEdgeID> NewEdgeIDs;
-							const FTriangleID NewTriangleID = MeshDescription->CreateTriangle(PolygonGroupID, CornerInstanceIDs, &NewEdgeIDs);
-							check(NewEdgeIDs.Num() == 0);
-
-							for (int32 UVLayerIndex = 0; UVLayerIndex < FBXUVs.UniqueUVCount; UVLayerIndex++)
+							else
 							{
-								FUVID UVIDs[3] = { INDEX_NONE, INDEX_NONE, INDEX_NONE };
-								if (FBXUVs.LayerElementUV[UVLayerIndex] != NULL)
-								{
-									for (int32 VertexIndex = 0; VertexIndex < 3; ++VertexIndex)
-									{
-										int UVMapIndex = (FBXUVs.UVMappingMode[UVLayerIndex] == FbxLayerElement::eByControlPoint)
-											? Mesh->GetPolygonVertex(PolygonIndex, VertexIndex)
-											: SkippedVertexInstance + CurrentVertexInstanceIndex + VertexIndex;
-										int32 UVIndex = (FBXUVs.UVReferenceMode[UVLayerIndex] == FbxLayerElement::eDirect)
-											? UVMapIndex
-											: FBXUVs.LayerElementUV[UVLayerIndex]->GetIndexArray().GetAt(UVMapIndex);
-
-										if (UVIndex != -1)
-										{
-											UVIDs[VertexIndex] = UVIndex + UVOffsets[UVLayerIndex];
-
-											check(MeshDescription->VertexInstanceAttributes().GetAttribute<FVector2f>(CornerInstanceIDs[VertexIndex], MeshAttribute::VertexInstance::TextureCoordinate, UVLayerIndex) ==
-												  MeshDescription->UVAttributes(UVLayerIndex).GetAttribute<FVector2f>(UVIndex + UVOffsets[UVLayerIndex], MeshAttribute::UV::UVCoordinate));
-										}
-										else
-										{
-											// TODO: what does it mean to have a UV index of -1?
-											// Investigate this case more carefully and handle it properly.
-										}
-									}
-								}
-
-								MeshDescription->SetTriangleUVIndices(NewTriangleID, UVIDs, UVLayerIndex);
+								//When there is no smoothing group we set all edge to: hard (faceted mesh) for static mesh and smooth for skinned and rigid
+								EdgeHardnesses[MatchEdgeId] = MeshType == EMeshType::Static ? true : false;
 							}
-
-							CurrentVertexInstanceIndex += PolygonVertexCount;
-						}
-
-						//Call this after all GetMeshEdgeIndexForPolygon call this is for optimization purpose.
-						if (bBeginGetMeshEdgeIndexForPolygonCalled)
-						{
-							Mesh->EndGetMeshEdgeIndexForPolygon();
-						}
-
-						MeshDescription->ResumeVertexInstanceIndexing();
-						MeshDescription->ResumeEdgeIndexing();
-						MeshDescription->ResumePolygonIndexing();
-						MeshDescription->ResumePolygonGroupIndexing();
-						MeshDescription->ResumeUVIndexing();
-
-						if (SkippedVertexInstance > 0)
-						{
-							//We must compact the sparse array before reserving new space
-							//When we reserve it will make a hole in the sparse array if the last reserve was not fully use
-							//The importer assume there will be no hole when importing a mesh
-							FElementIDRemappings OutRemappings;
-							MeshDescription->Compact(OutRemappings);
 						}
 					}
+				}
 
-					if (MeshType == EMeshType::Skinned)
+				FPolygonGroupID PolygonGroupID = PolygonGroupMapping[MaterialIndex];
+				// Insert a triangle into the mesh
+				// @note: we only ever import triangulated meshes currently. This could easily change as we have the infrastructure set up for arbitrary ngons.
+				// However, I think either we need to triangulate in the exact same way as Maya/Max if we do that.
+				check(CornerInstanceIDs.Num() == 3);
+				check(PolygonVertexCount == 3);
+				TArray<FEdgeID> NewEdgeIDs;
+				const FTriangleID NewTriangleID = MeshDescription->CreateTriangle(PolygonGroupID, CornerInstanceIDs, &NewEdgeIDs);
+				check(NewEdgeIDs.Num() == 0);
+
+				for (int32 UVLayerIndex = 0; UVLayerIndex < FBXUVs.UniqueUVCount; UVLayerIndex++)
+				{
+					FUVID UVIDs[3] = { INDEX_NONE, INDEX_NONE, INDEX_NONE };
+					if (FBXUVs.LayerElementUV[UVLayerIndex] != NULL)
 					{
-						FSkeletalMeshAttributes SkeletalMeshAttributes(*MeshDescription);
-						SkeletalMeshAttributes.Register();
-
-						using namespace UE::AnimationCore;
-						TMap<FVertexID, TArray<FBoneWeight>> RawBoneWeights;
-						
-						//Add the influence data in the skeletalmesh description
-						FSkinWeightsVertexAttributesRef VertexSkinWeights = SkeletalMeshAttributes.GetVertexSkinWeights();
-
-						FbxSkin* Skin = (FbxSkin*)Mesh->GetDeformer(0, FbxDeformer::eSkin);
-						if (!ensure(Skin))
+						for (int32 VertexIndex = 0; VertexIndex < 3; ++VertexIndex)
 						{
-							//TODO log an error
-							return false;
-						}
-						const int32 ClusterCount = Skin->GetClusterCount();
-						TArray<FbxNode*> SortedJoints;
-						SortedJoints.Reserve(ClusterCount);
-						OutJointUniqueNames.Reserve(ClusterCount);
-						// create influences for each cluster
-						for (int32 ClusterIndex = 0; ClusterIndex < ClusterCount; ClusterIndex++)
-						{
-							FbxCluster* Cluster = Skin->GetCluster(ClusterIndex);
-							// When Maya plug-in exports rigid binding, it will generate "CompensationCluster" for each ancestor links.
-							// FBX writes these "CompensationCluster" out. The CompensationCluster also has weight 1 for vertices.
-							// Unreal importer should skip these clusters.
-							if (!Cluster || (FCStringAnsi::Strcmp(Cluster->GetUserDataID(), "Maya_ClusterHint") == 0 && FCStringAnsi::Strcmp(Cluster->GetUserData(), "CompensationCluster") == 0))
+							int UVMapIndex = (FBXUVs.UVMappingMode[UVLayerIndex] == FbxLayerElement::eByControlPoint)
+								? Mesh->GetPolygonVertex(PolygonIndex, VertexIndex)
+								: SkippedVertexInstance + CurrentVertexInstanceIndex + VertexIndex;
+							int32 UVIndex = (FBXUVs.UVReferenceMode[UVLayerIndex] == FbxLayerElement::eDirect)
+								? UVMapIndex
+								: FBXUVs.LayerElementUV[UVLayerIndex]->GetIndexArray().GetAt(UVMapIndex);
+
+							if (UVIndex != -1)
 							{
-								continue;
-							}
+								UVIDs[VertexIndex] = UVIndex + UVOffsets[UVLayerIndex];
 
-							FbxNode* Link = Cluster->GetLink();
-							// find the bone index
-							int32 BoneIndex = -1;
-							if (!SortedJoints.Find(Link, BoneIndex))
-							{
-								BoneIndex = SortedJoints.Add(Link);
-								FString JointNodeUniqueID = FFbxHelper::GetFbxObjectName(Link);
-								OutJointUniqueNames.Add(JointNodeUniqueID);
+								check(MeshDescription->VertexInstanceAttributes().GetAttribute<FVector2f>(CornerInstanceIDs[VertexIndex], MeshAttribute::VertexInstance::TextureCoordinate, UVLayerIndex) ==
+										MeshDescription->UVAttributes(UVLayerIndex).GetAttribute<FVector2f>(UVIndex + UVOffsets[UVLayerIndex], MeshAttribute::UV::UVCoordinate));
 							}
-
-							//	get the vertex indices
-							int32 ControlPointIndicesCount = Cluster->GetControlPointIndicesCount();
-							int32* ControlPointIndices = Cluster->GetControlPointIndices();
-							double* Weights = Cluster->GetControlPointWeights();
-
-							//	for each vertex index in the cluster
-							for (int32 ControlPointIndex = 0; ControlPointIndex < ControlPointIndicesCount; ++ControlPointIndex)
+							else
 							{
-								FVertexID VertexID(ControlPointIndices[ControlPointIndex] + VertexOffset);
-								if (!ensure(MeshDescription->IsVertexValid(VertexID)))
-								{
-									//Invalid Influence
-									continue;
-								}
-								float Weight = static_cast<float>(Weights[ControlPointIndex]);
-
-								TArray<FBoneWeight> *BoneWeights = RawBoneWeights.Find(VertexID);
-								if (BoneWeights)
-								{
-									// Do we already have a weight for this bone? Keep the greater weight.
-									bool bShouldAdd = true;
-									for (int32 WeightIndex = 0; WeightIndex < BoneWeights->Num(); WeightIndex++)
-									{
-										FBoneWeight &BoneWeight = (*BoneWeights)[WeightIndex];
-										if (BoneWeight.GetBoneIndex() == BoneIndex)
-										{
-											if (BoneWeight.GetWeight() < Weight)
-											{
-												BoneWeight.SetWeight(Weight);
-											}
-											bShouldAdd = false;
-											break;
-										}
-									}
-									if (bShouldAdd)
-									{
-										BoneWeights->Add(FBoneWeight(BoneIndex, Weight));
-									}
-								}
-								else
-								{
-									RawBoneWeights.Add(VertexID).Add(FBoneWeight(BoneIndex, Weight));
-								}
+								// TODO: what does it mean to have a UV index of -1?
+								// Investigate this case more carefully and handle it properly.
 							}
 						}
-
-						// Add all the raw bone weights. This will cause the weights to be sorted and re-normalized after culling to max influences.
-						for (const TTuple<FVertexID, TArray<FBoneWeight>> &Item: RawBoneWeights)
-						{
-							VertexSkinWeights.Set(Item.Key, Item.Value);
-						}
 					}
-				}
-				// needed?
-				FBXUVs.Cleanup();
 
-				if (!bHasNonDegeneratePolygons)
-				{
-					//TODO log an error we have imported no valid triangle in this fbx mesh node
-					//CreateTokenizedErrorForDegeneratedPart(this, StaticMesh->GetName(), FbxNodeName);
+					MeshDescription->SetTriangleUVIndices(NewTriangleID, UVIDs, UVLayerIndex);
 				}
 
-				return bHasNonDegeneratePolygons;
+				CurrentVertexInstanceIndex += PolygonVertexCount;
 			}
 
-			bool FMeshDescriptionImporter::IsOddNegativeScale(FbxAMatrix& TotalMatrix)
+			//Call this after all GetMeshEdgeIndexForPolygon call this is for optimization purpose.
+			if (bBeginGetMeshEdgeIndexForPolygonCalled)
 			{
-				FbxVector4 Scale = TotalMatrix.GetS();
-				int32 NegativeNum = 0;
+				Mesh->EndGetMeshEdgeIndexForPolygon();
+			}
 
-				if (Scale[0] < 0) NegativeNum++;
-				if (Scale[1] < 0) NegativeNum++;
-				if (Scale[2] < 0) NegativeNum++;
+			MeshDescription->ResumeVertexInstanceIndexing();
+			MeshDescription->ResumeEdgeIndexing();
+			MeshDescription->ResumePolygonIndexing();
+			MeshDescription->ResumePolygonGroupIndexing();
+			MeshDescription->ResumeUVIndexing();
 
-				return NegativeNum == 1 || NegativeNum == 3;
+			if (SkippedVertexInstance > 0)
+			{
+				//We must compact the sparse array before reserving new space
+				//When we reserve it will make a hole in the sparse array if the last reserve was not fully use
+				//The importer assume there will be no hole when importing a mesh
+				FElementIDRemappings OutRemappings;
+				MeshDescription->Compact(OutRemappings);
 			}
+		}
 
+		if (MeshType == EMeshType::Skinned)
+		{
+			FSkeletalMeshAttributes SkeletalMeshAttributes(*MeshDescription);
+			SkeletalMeshAttributes.Register();
 
-			//////////////////////////////////////////////////////////////////////////
-			/// FMeshPayloadContext implementation
+			using namespace UE::AnimationCore;
+			TMap<FVertexID, TArray<FBoneWeight>> RawBoneWeights;
+						
+			//Add the influence data in the skeletalmesh description
+			FSkinWeightsVertexAttributesRef VertexSkinWeights = SkeletalMeshAttributes.GetVertexSkinWeights();
 
-			bool FMeshPayloadContext::FetchPayloadToFile(FFbxParser& Parser, const FString& PayloadFilepath)
+			FbxSkin* Skin = (FbxSkin*)Mesh->GetDeformer(0, FbxDeformer::eSkin);
+			if (!ensure(Skin))
 			{
-				if (!ensure(SDKScene != nullptr))
-				{
-					UInterchangeResultError_Generic* Message = Parser.AddMessage<UInterchangeResultError_Generic>();
-					Message->InterchangeKey = FFbxHelper::GetMeshUniqueID(Mesh);
-					Message->Text = LOCTEXT("FBXSceneNull", "Cannot fetch FBX mesh payload because the FBX scene is null.");
-					return false;
-				}
-
-				if (!ensure(Mesh != nullptr))
+				//TODO log an error
+				return false;
+			}
+			const int32 ClusterCount = Skin->GetClusterCount();
+			TArray<FbxNode*> SortedJoints;
+			SortedJoints.Reserve(ClusterCount);
+			OutJointUniqueNames.Reserve(ClusterCount);
+			// create influences for each cluster
+			for (int32 ClusterIndex = 0; ClusterIndex < ClusterCount; ClusterIndex++)
+			{
+				FbxCluster* Cluster = Skin->GetCluster(ClusterIndex);
+				// When Maya plug-in exports rigid binding, it will generate "CompensationCluster" for each ancestor links.
+				// FBX writes these "CompensationCluster" out. The CompensationCluster also has weight 1 for vertices.
+				// Unreal importer should skip these clusters.
+				if (!Cluster || (FCStringAnsi::Strcmp(Cluster->GetUserDataID(), "Maya_ClusterHint") == 0 && FCStringAnsi::Strcmp(Cluster->GetUserData(), "CompensationCluster") == 0))
 				{
-					UInterchangeResultError_Generic* Message = Parser.AddMessage<UInterchangeResultError_Generic>();
-					Message->InterchangeKey = FFbxHelper::GetMeshUniqueID(Mesh);
-					Message->Text = LOCTEXT("FBXMeshNull", "Cannot fetch FBX mesh payload because the FBX mesh is null.");
-					return false;
+					continue;
 				}
 
-				if (!ensure(SDKGeometryConverter != nullptr))
+				FbxNode* Link = Cluster->GetLink();
+				// find the bone index
+				int32 BoneIndex = -1;
+				if (!SortedJoints.Find(Link, BoneIndex))
 				{
-					UInterchangeResultError_Generic* Message = Parser.AddMessage<UInterchangeResultError_Generic>();
-					Message->InterchangeKey = FFbxHelper::GetMeshUniqueID(Mesh);
-					Message->Text = LOCTEXT("FBXConverterNull", "Cannot fetch FBX mesh payload because the FBX geometry converter is null.");
-					return false;
+					BoneIndex = SortedJoints.Add(Link);
+					FString JointNodeUniqueID = FFbxHelper::GetFbxObjectName(Link);
+					OutJointUniqueNames.Add(JointNodeUniqueID);
 				}
 
-				bool bFetchSkinnedData = (Mesh->GetDeformerCount(FbxDeformer::eSkin) > 0);
-				
-				TArray<FString> JointUniqueNames;
+				//	get the vertex indices
+				int32 ControlPointIndicesCount = Cluster->GetControlPointIndicesCount();
+				int32* ControlPointIndices = Cluster->GetControlPointIndices();
+				double* Weights = Cluster->GetControlPointWeights();
 
-				FMeshDescription MeshDescription;
-				if (bFetchSkinnedData)
+				//	for each vertex index in the cluster
+				for (int32 ControlPointIndex = 0; ControlPointIndex < ControlPointIndicesCount; ++ControlPointIndex)
 				{
-					FSkeletalMeshAttributes SkeletalMeshAttribute(MeshDescription);
-					SkeletalMeshAttribute.Register();
-					FMeshDescriptionImporter MeshDescriptionImporter(Parser, &MeshDescription, SDKScene, SDKGeometryConverter);
-					if (!MeshDescriptionImporter.FillSkinnedMeshDescriptionFromFbxMesh(Mesh, JointUniqueNames))
+					FVertexID VertexID(ControlPointIndices[ControlPointIndex] + VertexOffset);
+					if (!ensure(MeshDescription->IsVertexValid(VertexID)))
 					{
-						UInterchangeResultError_Generic* Message = Parser.AddMessage<UInterchangeResultError_Generic>();
-						Message->InterchangeKey = FFbxHelper::GetMeshUniqueID(Mesh);
-						Message->Text = LOCTEXT("SkinnedMeshDescriptionError", "Cannot fetch skinned mesh payload because there was an error when creating the MeshDescription.");
-						return false;
+						//Invalid Influence
+						continue;
 					}
-				}
-				else
-				{
-					FStaticMeshAttributes StaticMeshAttribute(MeshDescription);
-					StaticMeshAttribute.Register();
-					FMeshDescriptionImporter MeshDescriptionImporter(Parser, &MeshDescription, SDKScene, SDKGeometryConverter);
-					if (!MeshDescriptionImporter.FillStaticMeshDescriptionFromFbxMesh(Mesh))
+					float Weight = static_cast<float>(Weights[ControlPointIndex]);
+
+					TArray<FBoneWeight> *BoneWeights = RawBoneWeights.Find(VertexID);
+					if (BoneWeights)
 					{
-						UInterchangeResultError_Generic* Message = Parser.AddMessage<UInterchangeResultError_Generic>();
-						Message->InterchangeKey = FFbxHelper::GetMeshUniqueID(Mesh);
-						Message->Text = LOCTEXT("StaticMeshDescriptionError", "Cannot fetch static mesh payload because there was an error when creating the MeshDescription.");
-						return false;
+						// Do we already have a weight for this bone? Keep the greater weight.
+						bool bShouldAdd = true;
+						for (int32 WeightIndex = 0; WeightIndex < BoneWeights->Num(); WeightIndex++)
+						{
+							FBoneWeight &BoneWeight = (*BoneWeights)[WeightIndex];
+							if (BoneWeight.GetBoneIndex() == BoneIndex)
+							{
+								if (BoneWeight.GetWeight() < Weight)
+								{
+									BoneWeight.SetWeight(Weight);
+								}
+								bShouldAdd = false;
+								break;
+							}
+						}
+						if (bShouldAdd)
+						{
+							BoneWeights->Add(FBoneWeight(BoneIndex, Weight));
+						}
 					}
-				}
-				//Dump the MeshDescription to a file
-				{
-					FLargeMemoryWriter Ar;
-					MeshDescription.Serialize(Ar);
-
-					Ar << bFetchSkinnedData;
-					if (bFetchSkinnedData)
+					else
 					{
-						//When passing a skinned MeshDescription, We want to pass the joint Node ID so we can know what the influence bone index refer to
-						Ar << JointUniqueNames;
+						RawBoneWeights.Add(VertexID).Add(FBoneWeight(BoneIndex, Weight));
 					}
-					uint8* ArchiveData = Ar.GetData();
-					int64 ArchiveSize = Ar.TotalSize();
-					TArray64<uint8> Buffer(ArchiveData, ArchiveSize);
-					FFileHelper::SaveArrayToFile(Buffer, *PayloadFilepath);
 				}
+			}
 
-				return true;
+			// Add all the raw bone weights. This will cause the weights to be sorted and re-normalized after culling to max influences.
+			for (const TTuple<FVertexID, TArray<FBoneWeight>> &Item: RawBoneWeights)
+			{
+				VertexSkinWeights.Set(Item.Key, Item.Value);
 			}
+		}
+	}
+	// needed?
+	FBXUVs.Cleanup();
 
-			//////////////////////////////////////////////////////////////////////////
-			/// FShapePayloadContext implementation
+	if (!bHasNonDegeneratePolygons)
+	{
+		//TODO log an error we have imported no valid triangle in this fbx mesh node
+		//CreateTokenizedErrorForDegeneratedPart(this, StaticMesh->GetName(), FbxNodeName);
+	}
 
-			bool FShapePayloadContext::FetchPayloadToFile(FFbxParser& Parser, const FString& PayloadFilepath)
-			{
-				if (!ensure(Shape))
-				{
-					//Todo log an error
-					return false;
-				}
-				//Import the BlendShape
-				FMeshDescription ShapeMeshDescription;
-				FStaticMeshAttributes StaticMeshAttribute(ShapeMeshDescription);
-				StaticMeshAttribute.Register();
-				FMeshDescriptionImporter MeshDescriptionImporter(Parser, &ShapeMeshDescription, SDKScene, SDKGeometryConverter);
-				if (!MeshDescriptionImporter.FillMeshDescriptionFromFbxShape(Shape))
-				{
-					UInterchangeResultError_Generic* Message = Parser.AddMessage<UInterchangeResultError_Generic>();
-					Message->InterchangeKey = FFbxHelper::GetMeshUniqueID(Shape);
-					Message->Text = LOCTEXT("MeshDescriptionBlendShapeError", "Unable to create MeshDescription from FBX blend shape.");
+	return bHasNonDegeneratePolygons;
+}
 
-					return false;
-				}
+bool FMeshDescriptionImporter::IsOddNegativeScale(FbxAMatrix& TotalMatrix)
+{
+	FbxVector4 Scale = TotalMatrix.GetS();
+	int32 NegativeNum = 0;
+
+	if (Scale[0] < 0) NegativeNum++;
+	if (Scale[1] < 0) NegativeNum++;
+	if (Scale[2] < 0) NegativeNum++;
+
+	return NegativeNum == 1 || NegativeNum == 3;
+}
+
+
+//////////////////////////////////////////////////////////////////////////
+/// FMeshPayloadContext implementation
+
+bool FMeshPayloadContext::FetchPayloadToFile(FFbxParser& Parser, const FString& PayloadFilepath)
+{
+	if (!ensure(SDKScene != nullptr))
+	{
+		UInterchangeResultError_Generic* Message = Parser.AddMessage<UInterchangeResultError_Generic>();
+		Message->InterchangeKey = FFbxHelper::GetMeshUniqueID(Mesh);
+		Message->Text = LOCTEXT("FBXSceneNull", "Cannot fetch FBX mesh payload because the FBX scene is null.");
+		return false;
+	}
+
+	if (!ensure(Mesh != nullptr))
+	{
+		UInterchangeResultError_Generic* Message = Parser.AddMessage<UInterchangeResultError_Generic>();
+		Message->InterchangeKey = FFbxHelper::GetMeshUniqueID(Mesh);
+		Message->Text = LOCTEXT("FBXMeshNull", "Cannot fetch FBX mesh payload because the FBX mesh is null.");
+		return false;
+	}
+
+	if (!ensure(SDKGeometryConverter != nullptr))
+	{
+		UInterchangeResultError_Generic* Message = Parser.AddMessage<UInterchangeResultError_Generic>();
+		Message->InterchangeKey = FFbxHelper::GetMeshUniqueID(Mesh);
+		Message->Text = LOCTEXT("FBXConverterNull", "Cannot fetch FBX mesh payload because the FBX geometry converter is null.");
+		return false;
+	}
+
+	bool bFetchSkinnedData = (Mesh->GetDeformerCount(FbxDeformer::eSkin) > 0);
+				
+	TArray<FString> JointUniqueNames;
 
-				//Dump the MeshDescription to a file
+	FMeshDescription MeshDescription;
+	if (bFetchSkinnedData)
+	{
+		FSkeletalMeshAttributes SkeletalMeshAttribute(MeshDescription);
+		SkeletalMeshAttribute.Register();
+		FMeshDescriptionImporter MeshDescriptionImporter(Parser, &MeshDescription, SDKScene, SDKGeometryConverter);
+		if (!MeshDescriptionImporter.FillSkinnedMeshDescriptionFromFbxMesh(Mesh, JointUniqueNames))
+		{
+			UInterchangeResultError_Generic* Message = Parser.AddMessage<UInterchangeResultError_Generic>();
+			Message->InterchangeKey = FFbxHelper::GetMeshUniqueID(Mesh);
+			Message->Text = LOCTEXT("SkinnedMeshDescriptionError", "Cannot fetch skinned mesh payload because there was an error when creating the MeshDescription.");
+			return false;
+		}
+	}
+	else
+	{
+		FStaticMeshAttributes StaticMeshAttribute(MeshDescription);
+		StaticMeshAttribute.Register();
+		FMeshDescriptionImporter MeshDescriptionImporter(Parser, &MeshDescription, SDKScene, SDKGeometryConverter);
+		if (!MeshDescriptionImporter.FillStaticMeshDescriptionFromFbxMesh(Mesh))
+		{
+			UInterchangeResultError_Generic* Message = Parser.AddMessage<UInterchangeResultError_Generic>();
+			Message->InterchangeKey = FFbxHelper::GetMeshUniqueID(Mesh);
+			Message->Text = LOCTEXT("StaticMeshDescriptionError", "Cannot fetch static mesh payload because there was an error when creating the MeshDescription.");
+			return false;
+		}
+	}
+	//Dump the MeshDescription to a file
+	{
+		FLargeMemoryWriter Ar;
+		MeshDescription.Serialize(Ar);
+
+		Ar << bFetchSkinnedData;
+		if (bFetchSkinnedData)
+		{
+			//When passing a skinned MeshDescription, We want to pass the joint Node ID so we can know what the influence bone index refer to
+			Ar << JointUniqueNames;
+		}
+		uint8* ArchiveData = Ar.GetData();
+		int64 ArchiveSize = Ar.TotalSize();
+		TArray64<uint8> Buffer(ArchiveData, ArchiveSize);
+		FFileHelper::SaveArrayToFile(Buffer, *PayloadFilepath);
+	}
+
+	return true;
+}
+
+//////////////////////////////////////////////////////////////////////////
+/// FShapePayloadContext implementation
+
+bool FShapePayloadContext::FetchPayloadToFile(FFbxParser& Parser, const FString& PayloadFilepath)
+{
+	if (!ensure(Shape))
+	{
+		//Todo log an error
+		return false;
+	}
+	//Import the BlendShape
+	FMeshDescription ShapeMeshDescription;
+	FStaticMeshAttributes StaticMeshAttribute(ShapeMeshDescription);
+	StaticMeshAttribute.Register();
+	FMeshDescriptionImporter MeshDescriptionImporter(Parser, &ShapeMeshDescription, SDKScene, SDKGeometryConverter);
+	if (!MeshDescriptionImporter.FillMeshDescriptionFromFbxShape(Shape))
+	{
+		UInterchangeResultError_Generic* Message = Parser.AddMessage<UInterchangeResultError_Generic>();
+		Message->InterchangeKey = FFbxHelper::GetMeshUniqueID(Shape);
+		Message->Text = LOCTEXT("MeshDescriptionBlendShapeError", "Unable to create MeshDescription from FBX blend shape.");
+
+		return false;
+	}
+
+	//Dump the MeshDescription to a file
+	{
+		FLargeMemoryWriter Ar;
+		ShapeMeshDescription.Serialize(Ar);
+		uint8* ArchiveData = Ar.GetData();
+		int64 ArchiveSize = Ar.TotalSize();
+		TArray64<uint8> Buffer(ArchiveData, ArchiveSize);
+		FFileHelper::SaveArrayToFile(Buffer, *PayloadFilepath);
+	}
+	return true;
+}
+
+//////////////////////////////////////////////////////////////////////////
+/// FFbxMesh implementation
+
+void GetMaterialIndex(FbxMesh* Mesh, TArray<int32>& MaterialIndexes)
+{
+	int32 PolygonCount = Mesh->GetPolygonCount();
+	if (FbxGeometryElementMaterial* GeometryElementMaterial = Mesh->GetElementMaterial())
+	{
+		FbxLayerElementArrayTemplate<int32>& IndexArray = GeometryElementMaterial->GetIndexArray();
+		MaterialIndexes.Reset();
+		switch (GeometryElementMaterial->GetMappingMode())
+		{
+			case FbxGeometryElement::eByPolygon:
+			{
+				if (IndexArray.GetCount() == PolygonCount)
 				{
-					FLargeMemoryWriter Ar;
-					ShapeMeshDescription.Serialize(Ar);
-					uint8* ArchiveData = Ar.GetData();
-					int64 ArchiveSize = Ar.TotalSize();
-					TArray64<uint8> Buffer(ArchiveData, ArchiveSize);
-					FFileHelper::SaveArrayToFile(Buffer, *PayloadFilepath);
+					for (int i = 0; i < PolygonCount; ++i)
+					{
+						MaterialIndexes.AddUnique(IndexArray.GetAt(i));
+					}
 				}
-				return true;
 			}
+			break;
+			case FbxGeometryElement::eAllSame:
+			{
+				MaterialIndexes.AddUnique(IndexArray.GetAt(0));
+			}
+			break;
+		}
+	}
+}
 
-			//////////////////////////////////////////////////////////////////////////
-			/// FFbxMesh implementation
-
-			void GetMaterialIndex(FbxMesh* Mesh, TArray<int32>& MaterialIndexes)
+void FFbxMesh::AddAllMeshes(FbxScene* SDKScene, FbxGeometryConverter* SDKGeometryConverter, UInterchangeBaseNodeContainer& NodeContainer, TMap<FString, TSharedPtr<FPayloadContextBase>>& PayloadContexts)
+{
+	int32 GeometryCount = SDKScene->GetGeometryCount();
+	for (int32 GeometryIndex = 0; GeometryIndex < GeometryCount; ++GeometryIndex)
+	{
+		FbxGeometry* Geometry = SDKScene->GetGeometry(GeometryIndex);
+		if (Geometry->GetAttributeType() != FbxNodeAttribute::eMesh)
+		{
+			continue;
+		}
+		FbxMesh* Mesh = static_cast<FbxMesh*>(Geometry);
+		if (!Mesh)
+		{
+			continue;
+		}
+		FString MeshName = FFbxHelper::GetMeshName(Mesh);
+		FString MeshUniqueID = FFbxHelper::GetMeshUniqueID(Mesh);
+		const UInterchangeMeshNode* ExistingMeshNode = Cast<UInterchangeMeshNode>(NodeContainer.GetNode(MeshUniqueID));
+		UInterchangeMeshNode* MeshNode = nullptr;
+		if (!ExistingMeshNode)
+		{
+			MeshNode = CreateMeshNode(NodeContainer, MeshName, MeshUniqueID);
+			if (Geometry->GetDeformerCount(FbxDeformer::eSkin) > 0)
 			{
+				//Set the skinned mesh attribute
+				MeshNode->SetSkinnedMesh(true);
+				ExtractSkinnedMeshNodeJoints(SDKScene, NodeContainer, Mesh, MeshNode);
+			}
+			Mesh->ComputeBBox();
+			const int32 MeshVertexCount = Mesh->GetControlPointsCount();
+			MeshNode->SetCustomVertexCount(MeshVertexCount);
+			const int32 MeshPolygonCount = Mesh->GetPolygonCount();
+			MeshNode->SetCustomPolygonCount(MeshPolygonCount);
+			const FBox MeshBoundingBox = FBox(FFbxConvert::ConvertPos(Mesh->BBoxMin.Get()), FFbxConvert::ConvertPos(Mesh->BBoxMax.Get()));
+			MeshNode->SetCustomBoundingBox(MeshBoundingBox);
+			const bool bMeshHasVertexNormal = Mesh->GetElementNormalCount() > 0;
+			MeshNode->SetCustomHasVertexNormal(bMeshHasVertexNormal);
+			const bool bMeshHasVertexBinormal = Mesh->GetElementBinormalCount() > 0;
+			MeshNode->SetCustomHasVertexBinormal(bMeshHasVertexBinormal);
+			const bool bMeshHasVertexTangent = Mesh->GetElementTangentCount() > 0;
+			MeshNode->SetCustomHasVertexTangent(bMeshHasVertexTangent);
+			const bool bMeshHasSmoothGroup = Mesh->GetElementSmoothingCount() > 0;
+			MeshNode->SetCustomHasSmoothGroup(bMeshHasSmoothGroup);
+			const bool bMeshHasVertexColor = Mesh->GetElementVertexColorCount() > 0;
+			MeshNode->SetCustomHasVertexColor(bMeshHasVertexColor);
+			const int32 MeshUVCount = Mesh->GetElementUVCount();
+			MeshNode->SetCustomUVCount(MeshUVCount);
+						
+			//Add Material dependencies, we use always the first fbx node that instance the fbx geometry to grab the fbx surface materials.
+			{
+				//Grab all Material indexes use by the mesh
+				TArray<int32> MaterialIndexes;
 				int32 PolygonCount = Mesh->GetPolygonCount();
 				if (FbxGeometryElementMaterial* GeometryElementMaterial = Mesh->GetElementMaterial())
 				{
 					FbxLayerElementArrayTemplate<int32>& IndexArray = GeometryElementMaterial->GetIndexArray();
-					MaterialIndexes.Reset();
 					switch (GeometryElementMaterial->GetMappingMode())
 					{
 						case FbxGeometryElement::eByPolygon:
@@ -1253,329 +1328,253 @@ namespace UE
 							}
 						}
 						break;
+
 						case FbxGeometryElement::eAllSame:
 						{
-							MaterialIndexes.AddUnique(IndexArray.GetAt(0));
+							if (IndexArray.GetCount() > 0)
+							{
+								MaterialIndexes.AddUnique(IndexArray.GetAt(0));
+							}
 						}
 						break;
 					}
 				}
+				if (FbxNode* FbxMeshNode = Mesh->GetNode())
+				{
+					bool bAddAllNodeMaterials = (MaterialIndexes.Num() == 0);
+					const int32 MaterialCount = FbxMeshNode->GetMaterialCount();
+					for (int32 MaterialIndex = 0; MaterialIndex < MaterialCount; ++MaterialIndex)
+					{
+						FbxSurfaceMaterial* FbxMaterial = FbxMeshNode->GetMaterial(MaterialIndex);
+						FString MaterialName = FFbxHelper::GetFbxObjectName(FbxMaterial);
+						FString MaterialUid = TEXT("\\Material\\") + MaterialName;
+						if (bAddAllNodeMaterials || MaterialIndexes.Contains(MaterialIndex))
+						{
+							MeshNode->SetMaterialDependencyUid(MaterialUid);
+						}
+					}
+				}
 			}
 
-			void FFbxMesh::AddAllMeshes(FbxScene* SDKScene, FbxGeometryConverter* SDKGeometryConverter, UInterchangeBaseNodeContainer& NodeContainer, TMap<FString, TSharedPtr<FPayloadContextBase>>& PayloadContexts)
+			FString PayLoadKey = MeshUniqueID;
+			if (ensure(!PayloadContexts.Contains(PayLoadKey)))
 			{
-				int32 GeometryCount = SDKScene->GetGeometryCount();
-				for (int32 GeometryIndex = 0; GeometryIndex < GeometryCount; ++GeometryIndex)
+				TSharedPtr<FMeshPayloadContext> GeoPayload = MakeShared<FMeshPayloadContext>();
+				GeoPayload->Mesh = Mesh;
+				GeoPayload->SDKScene = SDKScene;
+				GeoPayload->SDKGeometryConverter = SDKGeometryConverter;
+				PayloadContexts.Add(PayLoadKey, GeoPayload);
+			}
+			MeshNode->SetPayLoadKey(PayLoadKey);
+
+			//Add all BlendShapes for this geometry node
+
+			TMap<FString, FbxShape*> ShapeNameToFbxShape;
+			const int32 BlendShapeCount = Geometry->GetDeformerCount(FbxDeformer::eBlendShape);
+
+			for (int32 BlendShapeIndex = 0; BlendShapeIndex < BlendShapeCount; ++BlendShapeIndex)
+			{
+				FbxBlendShape* BlendShape = (FbxBlendShape*)Geometry->GetDeformer(BlendShapeIndex, FbxDeformer::eBlendShape);
+				const int32 BlendShapeChannelCount = BlendShape->GetBlendShapeChannelCount();
+				FString BlendShapeName = FFbxHelper::GetFbxObjectName(BlendShape);
+				// see below where this is used for explanation...
+				const bool bMightBeBadMAXFile = (BlendShapeName == FString("Morpher"));
+				for (int32 ChannelIndex = 0; ChannelIndex < BlendShapeChannelCount; ++ChannelIndex)
 				{
-					FbxGeometry* Geometry = SDKScene->GetGeometry(GeometryIndex);
-					if (Geometry->GetAttributeType() != FbxNodeAttribute::eMesh)
-					{
-						continue;
-					}
-					FbxMesh* Mesh = static_cast<FbxMesh*>(Geometry);
-					if (!Mesh)
-					{
-						continue;
-					}
-					FString MeshName = FFbxHelper::GetMeshName(Mesh);
-					FString MeshUniqueID = FFbxHelper::GetMeshUniqueID(Mesh);
-					UInterchangeMeshNode* MeshNode = Cast<UInterchangeMeshNode>(NodeContainer.GetNode(MeshUniqueID));
-					if (!MeshNode)
+					FbxBlendShapeChannel* Channel = BlendShape->GetBlendShapeChannel(ChannelIndex);
+					if (Channel)
 					{
-						MeshNode = CreateMeshNode(NodeContainer, MeshName, MeshUniqueID);
-						if (Geometry->GetDeformerCount(FbxDeformer::eSkin) > 0)
+						//Find which shape should we use according to the weight.
+						const int32 CurrentChannelShapeCount = Channel->GetTargetShapeCount();
+						FString ChannelName = FFbxHelper::GetFbxObjectName(Channel);
+						// Maya adds the name of the blendshape and an underscore to the front of the channel name, so remove it
+						if (ChannelName.StartsWith(BlendShapeName))
 						{
-							//Set the skinned mesh attribute
-							MeshNode->SetSkinnedMesh(true);
-							ExtractSkinnedMeshNodeJoints(SDKScene, NodeContainer, Mesh, MeshNode);
+							ChannelName.RightInline(ChannelName.Len() - (BlendShapeName.Len() + 1), false);
 						}
-						Mesh->ComputeBBox();
-						const int32 MeshVertexCount = Mesh->GetControlPointsCount();
-						MeshNode->SetCustomVertexCount(MeshVertexCount);
-						const int32 MeshPolygonCount = Mesh->GetPolygonCount();
-						MeshNode->SetCustomPolygonCount(MeshPolygonCount);
-						const FBox MeshBoundingBox = FBox(FFbxConvert::ConvertPos(Mesh->BBoxMin.Get()), FFbxConvert::ConvertPos(Mesh->BBoxMax.Get()));
-						MeshNode->SetCustomBoundingBox(MeshBoundingBox);
-						const bool bMeshHasVertexNormal = Mesh->GetElementNormalCount() > 0;
-						MeshNode->SetCustomHasVertexNormal(bMeshHasVertexNormal);
-						const bool bMeshHasVertexBinormal = Mesh->GetElementBinormalCount() > 0;
-						MeshNode->SetCustomHasVertexBinormal(bMeshHasVertexBinormal);
-						const bool bMeshHasVertexTangent = Mesh->GetElementTangentCount() > 0;
-						MeshNode->SetCustomHasVertexTangent(bMeshHasVertexTangent);
-						const bool bMeshHasSmoothGroup = Mesh->GetElementSmoothingCount() > 0;
-						MeshNode->SetCustomHasSmoothGroup(bMeshHasSmoothGroup);
-						const bool bMeshHasVertexColor = Mesh->GetElementVertexColorCount() > 0;
-						MeshNode->SetCustomHasVertexColor(bMeshHasVertexColor);
-						const int32 MeshUVCount = Mesh->GetElementUVCount();
-						MeshNode->SetCustomUVCount(MeshUVCount);
-						
-						//Add Material dependencies, we use always the first fbx node that instance the fbx geometry to grab the fbx surface materials.
+						for (int32 ShapeIndex = 0; ShapeIndex < CurrentChannelShapeCount; ++ShapeIndex)
 						{
-							//Grab all Material indexes use by the mesh
-							TArray<int32> MaterialIndexes;
-							int32 PolygonCount = Mesh->GetPolygonCount();
-							if (FbxGeometryElementMaterial* GeometryElementMaterial = Mesh->GetElementMaterial())
+							FbxShape* Shape = Channel->GetTargetShape(ShapeIndex);
+							FString ShapeName;
+							if (CurrentChannelShapeCount > 1)
 							{
-								FbxLayerElementArrayTemplate<int32>& IndexArray = GeometryElementMaterial->GetIndexArray();
-								switch (GeometryElementMaterial->GetMappingMode())
-								{
-									case FbxGeometryElement::eByPolygon:
-									{
-										if (IndexArray.GetCount() == PolygonCount)
-										{
-											for (int i = 0; i < PolygonCount; ++i)
-											{
-												MaterialIndexes.AddUnique(IndexArray.GetAt(i));
-											}
-										}
-									}
-									break;
-
-									case FbxGeometryElement::eAllSame:
-									{
-										if (IndexArray.GetCount() > 0)
-										{
-											MaterialIndexes.AddUnique(IndexArray.GetAt(0));
-										}
-									}
-									break;
-								}
+								ShapeName = FFbxHelper::GetFbxObjectName(Shape);
 							}
-							if (FbxNode* FbxMeshNode = Mesh->GetNode())
+							else
 							{
-								bool bAddAllNodeMaterials = (MaterialIndexes.Num() == 0);
-								const int32 MaterialCount = FbxMeshNode->GetMaterialCount();
-								for (int32 MaterialIndex = 0; MaterialIndex < MaterialCount; ++MaterialIndex)
+								if (bMightBeBadMAXFile)
+								{
+									ShapeName = FFbxHelper::GetFbxObjectName(Shape);
+								}
+								else
 								{
-									FbxSurfaceMaterial* FbxMaterial = FbxMeshNode->GetMaterial(MaterialIndex);
-									FString MaterialName = FFbxHelper::GetFbxObjectName(FbxMaterial);
-									FString MaterialUid = TEXT("\\Material\\") + MaterialName;
-									if (bAddAllNodeMaterials || MaterialIndexes.Contains(MaterialIndex))
-									{
-										MeshNode->SetMaterialDependencyUid(MaterialUid);
-									}
+									// Maya concatenates the number of the shape to the end of its name, so instead use the name of the channel
+									ShapeName = ChannelName;
 								}
 							}
-						}
-					}
-					FString PayLoadKey = MeshUniqueID;
-					if (ensure(!PayloadContexts.Contains(PayLoadKey)))
-					{
-						TSharedPtr<FMeshPayloadContext> GeoPayload = MakeShared<FMeshPayloadContext>();
-						GeoPayload->Mesh = Mesh;
-						GeoPayload->SDKScene = SDKScene;
-						GeoPayload->SDKGeometryConverter = SDKGeometryConverter;
-						PayloadContexts.Add(PayLoadKey, GeoPayload);
-					}
-					MeshNode->SetPayLoadKey(PayLoadKey);
-					
-					//Add all BlendShapes for this geometry node
-					TMap<FString, FbxShape*> ShapeNameToFbxShape;
-					const int32 BlendShapeCount = Geometry->GetDeformerCount(FbxDeformer::eBlendShape);
-					if (BlendShapeCount > 0)
-					{
-						for (int32 BlendShapeIndex = 0; BlendShapeIndex < BlendShapeCount; ++BlendShapeIndex)
-						{
-							FbxBlendShape* BlendShape = (FbxBlendShape*)Geometry->GetDeformer(BlendShapeIndex, FbxDeformer::eBlendShape);
-							const int32 BlendShapeChannelCount = BlendShape->GetBlendShapeChannelCount();
-							FString BlendShapeName = FFbxHelper::GetFbxObjectName(BlendShape);
-							// see below where this is used for explanation...
-							const bool bMightBeBadMAXFile = (BlendShapeName == FString("Morpher"));
-							for (int32 ChannelIndex = 0; ChannelIndex < BlendShapeChannelCount; ++ChannelIndex)
+							ensure(!ShapeNameToFbxShape.Contains(ShapeName));
+							ShapeNameToFbxShape.Add(ShapeName, Shape);
+							FString ShapeAttributeName = FFbxHelper::GetMeshName(Shape);
+							FString ShapeUniqueID = FFbxHelper::GetMeshUniqueID(Shape);
+							const UInterchangeMeshNode* ExistingShapeNode = Cast<const UInterchangeMeshNode>(NodeContainer.GetNode(ShapeUniqueID));
+							if (!ExistingShapeNode)
 							{
-								FbxBlendShapeChannel* Channel = BlendShape->GetBlendShapeChannel(ChannelIndex);
-								if (Channel)
+								UInterchangeMeshNode* ShapeNode = CreateMeshNode(NodeContainer, ShapeAttributeName, ShapeUniqueID);
+								const bool bIsBlendShape = true;
+								ShapeNode->SetBlendShape(bIsBlendShape);
+								ShapeNode->SetBlendShapeName(ShapeName);
+
+								//Create a Mesh node dependency, so the mesh node can retrieve is associate blendshape
+								MeshNode->SetShapeDependencyUid(ShapeUniqueID);
+
+								FString ShapePayLoadKey = ShapeUniqueID;
+								if (ensure(!PayloadContexts.Contains(ShapePayLoadKey)))
 								{
-									//Find which shape should we use according to the weight.
-									const int32 CurrentChannelShapeCount = Channel->GetTargetShapeCount();
-									FString ChannelName = FFbxHelper::GetFbxObjectName(Channel);
-									// Maya adds the name of the blendshape and an underscore to the front of the channel name, so remove it
-									if (ChannelName.StartsWith(BlendShapeName))
-									{
-										ChannelName.RightInline(ChannelName.Len() - (BlendShapeName.Len() + 1), false);
-									}
-									for (int32 ShapeIndex = 0; ShapeIndex < CurrentChannelShapeCount; ++ShapeIndex)
-									{
-										FbxShape* Shape = Channel->GetTargetShape(ShapeIndex);
-										FString ShapeName;
-										if (CurrentChannelShapeCount > 1)
-										{
-											ShapeName = FFbxHelper::GetFbxObjectName(Shape);
-										}
-										else
-										{
-											if (bMightBeBadMAXFile)
-											{
-												ShapeName = FFbxHelper::GetFbxObjectName(Shape);
-											}
-											else
-											{
-												// Maya concatenates the number of the shape to the end of its name, so instead use the name of the channel
-												ShapeName = ChannelName;
-											}
-										}
-										ensure(!ShapeNameToFbxShape.Contains(ShapeName));
-										ShapeNameToFbxShape.Add(ShapeName, Shape);
-										FString ShapeAttributeName = FFbxHelper::GetMeshName(Shape);
-										FString ShapeUniqueID = FFbxHelper::GetMeshUniqueID(Shape);
-										UInterchangeMeshNode* ShapeNode = Cast<UInterchangeMeshNode>(NodeContainer.GetNode(ShapeUniqueID));
-										if (!ShapeNode)
-										{
-											ShapeNode = CreateMeshNode(NodeContainer, ShapeAttributeName, ShapeUniqueID);
-											const bool bIsBlendShape = true;
-											ShapeNode->SetBlendShape(bIsBlendShape);
-											ShapeNode->SetBlendShapeName(ShapeName);
-											//Create a Mesh node dependency, so the mesh node can retrieve is associate blendshape
-											MeshNode->SetShapeDependencyUid(ShapeUniqueID);
-										}
-										FString ShapePayLoadKey = ShapeUniqueID;
-										if (ensure(!PayloadContexts.Contains(ShapePayLoadKey)))
-										{
-											TSharedPtr<FShapePayloadContext> GeoPayload = MakeShared<FShapePayloadContext>();
-											GeoPayload->Shape = Shape;
-											GeoPayload->SDKScene = SDKScene;
-											GeoPayload->SDKGeometryConverter = SDKGeometryConverter;
-											PayloadContexts.Add(ShapePayLoadKey, GeoPayload);
-										}
-										ShapeNode->SetPayLoadKey(ShapePayLoadKey);
-									}
+									TSharedPtr<FShapePayloadContext> GeoPayload = MakeShared<FShapePayloadContext>();
+									GeoPayload->Shape = Shape;
+									GeoPayload->SDKScene = SDKScene;
+									GeoPayload->SDKGeometryConverter = SDKGeometryConverter;
+									PayloadContexts.Add(ShapePayLoadKey, GeoPayload);
 								}
+								ShapeNode->SetPayLoadKey(ShapePayLoadKey);
 							}
 						}
 					}
 				}
 			}
+		}					
+	}
+}
 
-			bool FFbxMesh::GetGlobalJointBindPoseTransform(FbxScene* SDKScene, FbxNode* Joint, FbxAMatrix& GlobalBindPoseJointMatrix)
+bool FFbxMesh::GetGlobalJointBindPoseTransform(FbxScene* SDKScene, FbxNode* Joint, FbxAMatrix& GlobalBindPoseJointMatrix)
+{
+	const FString LinkName = FFbxHelper::GetFbxObjectName(Joint);
+
+	const int32 PoseCount = SDKScene->GetPoseCount();
+	for (int32 PoseIndex = 0; PoseIndex < PoseCount; PoseIndex++)
+	{
+		FbxPose* CurrentPose = SDKScene->GetPose(PoseIndex);
+		if (CurrentPose && CurrentPose->IsBindPose())
+		{
+			FString PoseName = FFbxHelper::GetFbxObjectName(CurrentPose);
+			int32 PoseLinkIndex = CurrentPose->Find(Joint);
+			if (PoseLinkIndex >= 0)
 			{
-				const FString LinkName = FFbxHelper::GetFbxObjectName(Joint);
+				FbxMatrix NoneAffineMatrix = CurrentPose->GetMatrix(PoseLinkIndex);
+				GlobalBindPoseJointMatrix = *(FbxAMatrix*)(double*)&NoneAffineMatrix;
+				return true;
+			}
+		}
+	}
 
-				const int32 PoseCount = SDKScene->GetPoseCount();
-				for (int32 PoseIndex = 0; PoseIndex < PoseCount; PoseIndex++)
+	//Search all skeletalmesh(FbxGeometry with valid deformer) using this joint and see if there is a valid FbxCluster
+	//containing a TransformLinkMatrix for this joint
+	const int32 GeometryCount = SDKScene->GetGeometryCount();
+	for (int32 GeometryIndex = 0; GeometryIndex < GeometryCount; ++GeometryIndex)
+	{
+		const FbxGeometry* Geometry = SDKScene->GetGeometry(GeometryIndex);
+		if (!ensure(Geometry))
+		{
+			continue;
+		}
+		const int32 GeometryDeformerCount = Geometry->GetDeformerCount(FbxDeformer::eSkin);
+		for (int32 GeometryDeformerIndex = 0; GeometryDeformerIndex < GeometryDeformerCount; ++GeometryDeformerIndex)
+		{
+			const FbxSkin* Skin = (FbxSkin*)Geometry->GetDeformer(GeometryDeformerIndex, FbxDeformer::eSkin);
+			if (!ensure(Skin))
+			{
+				continue;
+			}
+			const int32 ClusterCount = Skin->GetClusterCount();
+			for (int32 ClusterIndex = 0; ClusterIndex < ClusterCount; ClusterIndex++)
+			{
+				const FbxCluster* Cluster = Skin->GetCluster(ClusterIndex);
+				// When Maya plug-in exports rigid binding, it will generate "CompensationCluster" for each ancestor links.
+				// FBX writes these "CompensationCluster" out. The CompensationCluster also has weight 1 for vertices.
+				// Unreal importer should skip these clusters.
+				if (!Cluster || (FCStringAnsi::Strcmp(Cluster->GetUserDataID(), "Maya_ClusterHint") == 0 && FCStringAnsi::Strcmp(Cluster->GetUserData(), "CompensationCluster") == 0))
 				{
-					FbxPose* CurrentPose = SDKScene->GetPose(PoseIndex);
-					if (CurrentPose && CurrentPose->IsBindPose())
-					{
-						FString PoseName = FFbxHelper::GetFbxObjectName(CurrentPose);
-						int32 PoseLinkIndex = CurrentPose->Find(Joint);
-						if (PoseLinkIndex >= 0)
-						{
-							FbxMatrix NoneAffineMatrix = CurrentPose->GetMatrix(PoseLinkIndex);
-							GlobalBindPoseJointMatrix = *(FbxAMatrix*)(double*)&NoneAffineMatrix;
-							return true;
-						}
-					}
+					continue;
 				}
-
-				//Search all skeletalmesh(FbxGeometry with valid deformer) using this joint and see if there is a valid FbxCluster
-				//containing a TransformLinkMatrix for this joint
-				const int32 GeometryCount = SDKScene->GetGeometryCount();
-				for (int32 GeometryIndex = 0; GeometryIndex < GeometryCount; ++GeometryIndex)
+				if (Joint == Cluster->GetLink())
 				{
-					const FbxGeometry* Geometry = SDKScene->GetGeometry(GeometryIndex);
-					if (!ensure(Geometry))
-					{
-						continue;
-					}
-					const int32 GeometryDeformerCount = Geometry->GetDeformerCount(FbxDeformer::eSkin);
-					for (int32 GeometryDeformerIndex = 0; GeometryDeformerIndex < GeometryDeformerCount; ++GeometryDeformerIndex)
-					{
-						const FbxSkin* Skin = (FbxSkin*)Geometry->GetDeformer(GeometryDeformerIndex, FbxDeformer::eSkin);
-						if (!ensure(Skin))
-						{
-							continue;
-						}
-						const int32 ClusterCount = Skin->GetClusterCount();
-						for (int32 ClusterIndex = 0; ClusterIndex < ClusterCount; ClusterIndex++)
-						{
-							const FbxCluster* Cluster = Skin->GetCluster(ClusterIndex);
-							// When Maya plug-in exports rigid binding, it will generate "CompensationCluster" for each ancestor links.
-							// FBX writes these "CompensationCluster" out. The CompensationCluster also has weight 1 for vertices.
-							// Unreal importer should skip these clusters.
-							if (!Cluster || (FCStringAnsi::Strcmp(Cluster->GetUserDataID(), "Maya_ClusterHint") == 0 && FCStringAnsi::Strcmp(Cluster->GetUserData(), "CompensationCluster") == 0))
-							{
-								continue;
-							}
-							if (Joint == Cluster->GetLink())
-							{
-								Cluster->GetTransformLinkMatrix(GlobalBindPoseJointMatrix);
-								return true;
-							}
-						}
-					}
+					Cluster->GetTransformLinkMatrix(GlobalBindPoseJointMatrix);
+					return true;
 				}
-				return false;
 			}
+		}
+	}
+	return false;
+}
 
-			void FFbxMesh::ExtractSkinnedMeshNodeJoints(FbxScene* SDKScene, UInterchangeBaseNodeContainer& NodeContainer, FbxMesh* Mesh, UInterchangeMeshNode* MeshNode)
-			{
-				TArray<FString> JointNodeUniqueIDs;
+void FFbxMesh::ExtractSkinnedMeshNodeJoints(FbxScene* SDKScene, UInterchangeBaseNodeContainer& NodeContainer, FbxMesh* Mesh, UInterchangeMeshNode* MeshNode)
+{
+	TArray<FString> JointNodeUniqueIDs;
 
-				//Get the bind pose from fbx
-				const int32 PoseCount = SDKScene->GetPoseCount();
-				TArray<FbxPose*> BindPoses;
-				for (int32 PoseIndex = 0; PoseIndex < PoseCount; PoseIndex++)
-				{
-					FbxPose* CurrentPose = SDKScene->GetPose(PoseIndex);
+	//Get the bind pose from fbx
+	const int32 PoseCount = SDKScene->GetPoseCount();
+	TArray<FbxPose*> BindPoses;
+	for (int32 PoseIndex = 0; PoseIndex < PoseCount; PoseIndex++)
+	{
+		FbxPose* CurrentPose = SDKScene->GetPose(PoseIndex);
 
-					// current pose is bind pose, 
-					if (CurrentPose && CurrentPose->IsBindPose())
-					{
-						BindPoses.Add(CurrentPose);
-					}
-				}
+		// current pose is bind pose, 
+		if (CurrentPose && CurrentPose->IsBindPose())
+		{
+			BindPoses.Add(CurrentPose);
+		}
+	}
 
-				const int32 SkinDeformerCount = Mesh->GetDeformerCount(FbxDeformer::eSkin);
-				for (int32 DeformerIndex = 0; DeformerIndex < SkinDeformerCount; DeformerIndex++)
-				{
-					FbxSkin* Skin = (FbxSkin*)Mesh->GetDeformer(DeformerIndex, FbxDeformer::eSkin);
-					if (!ensure(Skin))
-					{
-						continue;
-					}
-					const int32 ClusterCount = Skin->GetClusterCount();
-					JointNodeUniqueIDs.Reserve(JointNodeUniqueIDs.Num() + ClusterCount);
-					for (int32 ClusterIndex = 0; ClusterIndex < ClusterCount; ClusterIndex++)
-					{
-						FbxCluster* Cluster = Skin->GetCluster(ClusterIndex);
-						// When Maya plug-in exports rigid binding, it will generate "CompensationCluster" for each ancestor links.
-						// FBX writes these "CompensationCluster" out. The CompensationCluster also has weight 1 for vertices.
-						// Unreal importer should skip these clusters.
-						if (!Cluster || (FCStringAnsi::Strcmp(Cluster->GetUserDataID(), "Maya_ClusterHint") == 0 && FCStringAnsi::Strcmp(Cluster->GetUserData(), "CompensationCluster") == 0))
-						{
-							continue;
-						}
-						FbxNode* Link = Cluster->GetLink();
-						FString JointNodeUniqueID = FFbxHelper::GetFbxNodeHierarchyName(Link);
-						// find the bone index
-						if (!JointNodeUniqueIDs.Contains(JointNodeUniqueID))
-						{
-							JointNodeUniqueIDs.Add(JointNodeUniqueID);
-							MeshNode->SetSkeletonDependencyUid(JointNodeUniqueID);
-						}
-					}
-				}
+	const int32 SkinDeformerCount = Mesh->GetDeformerCount(FbxDeformer::eSkin);
+	for (int32 DeformerIndex = 0; DeformerIndex < SkinDeformerCount; DeformerIndex++)
+	{
+		FbxSkin* Skin = (FbxSkin*)Mesh->GetDeformer(DeformerIndex, FbxDeformer::eSkin);
+		if (!ensure(Skin))
+		{
+			continue;
+		}
+		const int32 ClusterCount = Skin->GetClusterCount();
+		JointNodeUniqueIDs.Reserve(JointNodeUniqueIDs.Num() + ClusterCount);
+		for (int32 ClusterIndex = 0; ClusterIndex < ClusterCount; ClusterIndex++)
+		{
+			FbxCluster* Cluster = Skin->GetCluster(ClusterIndex);
+			// When Maya plug-in exports rigid binding, it will generate "CompensationCluster" for each ancestor links.
+			// FBX writes these "CompensationCluster" out. The CompensationCluster also has weight 1 for vertices.
+			// Unreal importer should skip these clusters.
+			if (!Cluster || (FCStringAnsi::Strcmp(Cluster->GetUserDataID(), "Maya_ClusterHint") == 0 && FCStringAnsi::Strcmp(Cluster->GetUserData(), "CompensationCluster") == 0))
+			{
+				continue;
 			}
-
-			UInterchangeMeshNode* FFbxMesh::CreateMeshNode(UInterchangeBaseNodeContainer& NodeContainer, const FString& NodeName, const FString& NodeUniqueID)
+			FbxNode* Link = Cluster->GetLink();
+			FString JointNodeUniqueID = FFbxHelper::GetFbxNodeHierarchyName(Link);
+			// find the bone index
+			if (!JointNodeUniqueIDs.Contains(JointNodeUniqueID))
 			{
-				FString DisplayLabel(NodeName);
-				FString NodeUid(NodeUniqueID);
-				UInterchangeMeshNode* MeshNode = NewObject<UInterchangeMeshNode>(&NodeContainer, NAME_None);
-				if (!ensure(MeshNode))
-				{
-					UInterchangeResultError_Generic* Message = Parser.AddMessage<UInterchangeResultError_Generic>();
-					Message->InterchangeKey = NodeUniqueID;
-					Message->Text = LOCTEXT("NodeAllocationError", "Mesh node allocation failed when importing FBX.");
-					return nullptr;
-				}
-				// Creating a UMaterialInterface
-				MeshNode->InitializeNode(NodeUid, DisplayLabel, EInterchangeNodeContainerType::TranslatedAsset);
-				NodeContainer.AddNode(MeshNode);
-				return MeshNode;
+				JointNodeUniqueIDs.Add(JointNodeUniqueID);
+				MeshNode->SetSkeletonDependencyUid(JointNodeUniqueID);
 			}
-		} //ns Private
-	} //ns Interchange
-}//ns UE
+		}
+	}
+}
+
+UInterchangeMeshNode* FFbxMesh::CreateMeshNode(UInterchangeBaseNodeContainer& NodeContainer, const FString& NodeName, const FString& NodeUniqueID)
+{
+	FString DisplayLabel(NodeName);
+	FString NodeUid(NodeUniqueID);
+	UInterchangeMeshNode* MeshNode = NewObject<UInterchangeMeshNode>(&NodeContainer, NAME_None);
+	if (!ensure(MeshNode))
+	{
+		UInterchangeResultError_Generic* Message = Parser.AddMessage<UInterchangeResultError_Generic>();
+		Message->InterchangeKey = NodeUniqueID;
+		Message->Text = LOCTEXT("NodeAllocationError", "Mesh node allocation failed when importing FBX.");
+		return nullptr;
+	}
+	// Creating a UMaterialInterface
+	MeshNode->InitializeNode(NodeUid, DisplayLabel, EInterchangeNodeContainerType::TranslatedAsset);
+	NodeContainer.AddNode(MeshNode);
+	return MeshNode;
+}
+
+} // namespace UE::Interchange::Private
 
 #undef LOCTEXT_NAMESPACE
diff --git a/Engine/Plugins/Experimental/Interchange/Runtime/Source/Parsers/Fbx/Private/FbxScene.cpp b/Engine/Plugins/Experimental/Interchange/Runtime/Source/Parsers/Fbx/Private/FbxScene.cpp
--- a/Engine/Plugins/Experimental/Interchange/Runtime/Source/Parsers/Fbx/Private/FbxScene.cpp
+++ b/Engine/Plugins/Experimental/Interchange/Runtime/Source/Parsers/Fbx/Private/FbxScene.cpp
@@ -30,7 +30,7 @@ namespace UE
 
 			void FFbxScene::CreateMeshNodeReference(UInterchangeSceneNode* UnrealSceneNode, FbxNodeAttribute* NodeAttribute, UInterchangeBaseNodeContainer& NodeContainer, const FTransform& GeometricTransform)
 			{
-				UInterchangeMeshNode* MeshNode = nullptr;
+				const UInterchangeMeshNode* MeshNode = nullptr;
 				if (NodeAttribute->GetAttributeType() == FbxNodeAttribute::eMesh)
 				{
 					FbxMesh* Mesh = static_cast<FbxMesh*>(NodeAttribute);
@@ -48,8 +48,15 @@ namespace UE
 				if (MeshNode)
 				{
 					UnrealSceneNode->SetCustomAssetInstanceUid(MeshNode->GetUniqueID());
-					MeshNode->SetSceneInstanceUid(UnrealSceneNode->GetUniqueID());
 					UnrealSceneNode->SetCustomGeometricTransform(GeometricTransform);
+
+					// @todo: Nothing is using the SceneInstanceUid in the MeshNode. Do we even need to support it?
+					// For the moment an ugly const_cast so we can mutate it (it was fetched from the NodeContainer and is hence const).
+					// Possible solutions:
+					// - keep track in some other way of MeshNodes which we are in the process of maintaining / modifying
+					// - a derived UInterchangeMutableBaseNodeContainer which overrides node accessors and makes them mutable, to be passed only to translators
+					// - get rid of this attribute, and provide an alternate method for getting the scene instance UIDs which reference the mesh (by iterating scene instance nodes)
+					const_cast<UInterchangeMeshNode*>(MeshNode)->SetSceneInstanceUid(UnrealSceneNode->GetUniqueID());
 				}
 			}
 
@@ -57,7 +64,7 @@ namespace UE
 			{
 				const FString AssetUniqueID = FFbxHelper::GetNodeAttributeUniqueID(NodeAttribute, TypeName);
 
-				if (UInterchangeBaseNode* AssetNode = NodeContainer.GetNode(AssetUniqueID))
+				if (const UInterchangeBaseNode* AssetNode = NodeContainer.GetNode(AssetUniqueID))
 				{
 					UnrealSceneNode->SetCustomAssetInstanceUid(AssetNode->GetUniqueID());
 				}
diff --git a/Engine/Plugins/Experimental/Interchange/Runtime/Source/Pipelines/Private/InterchangeGenericAnimationPipeline.cpp b/Engine/Plugins/Experimental/Interchange/Runtime/Source/Pipelines/Private/InterchangeGenericAnimationPipeline.cpp
--- a/Engine/Plugins/Experimental/Interchange/Runtime/Source/Pipelines/Private/InterchangeGenericAnimationPipeline.cpp
+++ b/Engine/Plugins/Experimental/Interchange/Runtime/Source/Pipelines/Private/InterchangeGenericAnimationPipeline.cpp
@@ -82,7 +82,7 @@ void UInterchangeGenericAnimationPipeline::ExecutePreImportPipeline(UInterchange
 	BaseNodeContainer->GetNodes(UInterchangeSourceNode::StaticClass(), SourceNodeUids);
 	if (SourceNodeUids.Num() > 0)
 	{
-		if (UInterchangeSourceNode* SourceNode = Cast<UInterchangeSourceNode>(BaseNodeContainer->GetNode(SourceNodeUids[0])))
+		if (const UInterchangeSourceNode* SourceNode = Cast<UInterchangeSourceNode>(BaseNodeContainer->GetNode(SourceNodeUids[0])))
 		{
 			if (bImportBoneTracks)
 			{
@@ -141,7 +141,7 @@ void UInterchangeGenericAnimationPipeline::ExecutePreImportPipeline(UInterchange
 			BaseNodeContainer->GetNodes(UInterchangeSceneNode::StaticClass(), SceneNodeUids);
 			for (const FString& SceneNodeUid : SceneNodeUids)
 			{
-				if (UInterchangeSceneNode* SceneNode = Cast<UInterchangeSceneNode>(BaseNodeContainer->GetNode(SceneNodeUid)))
+				if (const UInterchangeSceneNode* SceneNode = Cast<const UInterchangeSceneNode>(BaseNodeContainer->GetNode(SceneNodeUid)))
 				{
 					double SceneNodeAnimStart;
 					double SceneNodeAnimStop;
@@ -223,7 +223,7 @@ void UInterchangeGenericAnimationPipeline::ExecutePreImportPipeline(UInterchange
 			BaseNodeContainer->GetNodes(UInterchangeSkeletalMeshFactoryNode::StaticClass(), SkeletalMeshNodeUids);
 			for (const FString& SkelMeshFactoryNodeUid : SkeletalMeshNodeUids)
 			{
-				if (UInterchangeSkeletalMeshFactoryNode* SkeletalMeshFactoryNode = Cast<UInterchangeSkeletalMeshFactoryNode>(BaseNodeContainer->GetNode(SkelMeshFactoryNodeUid)))
+				if (const UInterchangeSkeletalMeshFactoryNode* SkeletalMeshFactoryNode = Cast<const UInterchangeSkeletalMeshFactoryNode>(BaseNodeContainer->GetFactoryNode(SkelMeshFactoryNodeUid)))
 				{
 					TArray<FString> SkeletalMeshDependencies;
 					SkeletalMeshFactoryNode->GetFactoryDependencies(SkeletalMeshDependencies);
@@ -273,7 +273,7 @@ void UInterchangeGenericAnimationPipeline::ExecutePostImportPipeline(const UInte
 		return;
 	}
 
-	UInterchangeBaseNode* Node = BaseNodeContainer->GetNode(NodeKey);
+	const UInterchangeBaseNode* Node = BaseNodeContainer->GetNode(NodeKey);
 	if (!Node)
 	{
 		return;
diff --git a/Engine/Plugins/Experimental/Interchange/Runtime/Source/Pipelines/Private/InterchangeGenericAssetsPipeline.cpp b/Engine/Plugins/Experimental/Interchange/Runtime/Source/Pipelines/Private/InterchangeGenericAssetsPipeline.cpp
--- a/Engine/Plugins/Experimental/Interchange/Runtime/Source/Pipelines/Private/InterchangeGenericAssetsPipeline.cpp
+++ b/Engine/Plugins/Experimental/Interchange/Runtime/Source/Pipelines/Private/InterchangeGenericAssetsPipeline.cpp
@@ -486,7 +486,7 @@ void UInterchangeGenericAssetsPipeline::ImplementUseSourceNameForAssetOption()
 		//StaticMesh
 		if (MeshesImportedNodeCount == 1 && StaticMeshNodeUids.Num() > 0)
 		{
-			UInterchangeStaticMeshFactoryNode* StaticMeshNode = Cast<UInterchangeStaticMeshFactoryNode>(BaseNodeContainer->GetNode(StaticMeshNodeUids[0]));
+			UInterchangeStaticMeshFactoryNode* StaticMeshNode = Cast<UInterchangeStaticMeshFactoryNode>(BaseNodeContainer->GetFactoryNode(StaticMeshNodeUids[0]));
 			const FString DisplayLabelName = FPaths::GetBaseFilename(SourceDatas[0]->GetFilename());
 			StaticMeshNode->SetDisplayLabel(DisplayLabelName);
 		}
@@ -494,7 +494,7 @@ void UInterchangeGenericAssetsPipeline::ImplementUseSourceNameForAssetOption()
 		//Animation, simply look if we import only 1 animation before applying the option to animation
 		if (AnimSequenceNodeUids.Num() == 1)
 		{
-			UInterchangeAnimSequenceFactoryNode* AnimSequenceNode = Cast<UInterchangeAnimSequenceFactoryNode>(BaseNodeContainer->GetNode(AnimSequenceNodeUids[0]));
+			UInterchangeAnimSequenceFactoryNode* AnimSequenceNode = Cast<UInterchangeAnimSequenceFactoryNode>(BaseNodeContainer->GetFactoryNode(AnimSequenceNodeUids[0]));
 			const FString DisplayLabelName = FPaths::GetBaseFilename(SourceDatas[0]->GetFilename()) + TEXT("_Anim");
 			AnimSequenceNode->SetDisplayLabel(DisplayLabelName);
 		}
diff --git a/Engine/Plugins/Experimental/Interchange/Runtime/Source/Pipelines/Private/InterchangeGenericMaterialPipeline.cpp b/Engine/Plugins/Experimental/Interchange/Runtime/Source/Pipelines/Private/InterchangeGenericMaterialPipeline.cpp
--- a/Engine/Plugins/Experimental/Interchange/Runtime/Source/Pipelines/Private/InterchangeGenericMaterialPipeline.cpp
+++ b/Engine/Plugins/Experimental/Interchange/Runtime/Source/Pipelines/Private/InterchangeGenericMaterialPipeline.cpp
@@ -102,7 +102,7 @@ UInterchangeBaseMaterialFactoryNode* UInterchangeGenericMaterialPipeline::Create
 	if (BaseNodeContainer->IsNodeUidValid(NodeUid))
 	{
 		//The node already exist, just return it
-		MaterialFactoryNode = Cast<UInterchangeBaseMaterialFactoryNode>(BaseNodeContainer->GetNode(NodeUid));
+		MaterialFactoryNode = Cast<UInterchangeBaseMaterialFactoryNode>(BaseNodeContainer->GetFactoryNode(NodeUid));
 		if (!ensure(MaterialFactoryNode))
 		{
 			//Log an error
@@ -742,7 +742,7 @@ void UInterchangeGenericMaterialPipeline::HandleTextureSampleNode(const UInterch
 	FString ExpressionClassName;
 	FString TextureFactoryUid;
 
-	if (UInterchangeTextureNode* TextureNode = Cast<UInterchangeTextureNode>(BaseNodeContainer->GetNode(TextureUid)))
+	if (const UInterchangeTextureNode* TextureNode = Cast<const UInterchangeTextureNode>(BaseNodeContainer->GetNode(TextureUid)))
 	{
 		if (TextureNode->IsA<UInterchangeTextureCubeNode>())
 		{
@@ -774,14 +774,14 @@ void UInterchangeGenericMaterialPipeline::HandleTextureSampleNode(const UInterch
 
 		if (bParsingForNormalInput)
 		{
-			if (UInterchangeTextureFactoryNode* TextureFactoryNode = Cast<UInterchangeTextureFactoryNode>(BaseNodeContainer->GetNode(TextureFactoryUid)))
+			if (UInterchangeTextureFactoryNode* TextureFactoryNode = Cast<UInterchangeTextureFactoryNode>(BaseNodeContainer->GetFactoryNode(TextureFactoryUid)))
 			{
 				TextureFactoryNode->SetCustomCompressionSettings(TextureCompressionSettings::TC_Normalmap);
 			}
 		}
 		else if (bParsingForLinearInput)
 		{
-			if (UInterchangeTextureFactoryNode* TextureFactoryNode = Cast<UInterchangeTextureFactoryNode>(BaseNodeContainer->GetNode(TextureFactoryUid)))
+			if (UInterchangeTextureFactoryNode* TextureFactoryNode = Cast<UInterchangeTextureFactoryNode>(BaseNodeContainer->GetFactoryNode(TextureFactoryUid)))
 			{
 				TextureFactoryNode->SetCustomSRGB(false);
 			}
@@ -1032,7 +1032,7 @@ UInterchangeMaterialExpressionFactoryNode* UInterchangeGenericMaterialPipeline::
 
 	const FString MaterialExpressionUid = TEXT("Factory_") + ShaderNode->GetUniqueID();
 
-	UInterchangeMaterialExpressionFactoryNode* MaterialExpression = Cast<UInterchangeMaterialExpressionFactoryNode>(BaseNodeContainer->GetNode(MaterialExpressionUid));
+	UInterchangeMaterialExpressionFactoryNode* MaterialExpression = Cast<UInterchangeMaterialExpressionFactoryNode>(BaseNodeContainer->GetFactoryNode(MaterialExpressionUid));
 	if (MaterialExpression != nullptr)
 	{
 		return MaterialExpression;
@@ -1098,7 +1098,7 @@ UInterchangeMaterialExpressionFactoryNode* UInterchangeGenericMaterialPipeline::
 		ShaderNode->GetStringAttribute(UInterchangeShaderPortsAPI::MakeInputValueKey(Nodes::TextureSample::Inputs::Texture.ToString()), TextureUid);
 
 		// Make the material factory node have a dependency on the texture factory node so that the texture asset gets created first
-		if (UInterchangeTextureNode* TextureNode = Cast<UInterchangeTextureNode>(BaseNodeContainer->GetNode(TextureUid)))
+		if (const UInterchangeTextureNode* TextureNode = Cast<const UInterchangeTextureNode>(BaseNodeContainer->GetNode(TextureUid)))
 		{
 			TArray<FString> TextureNodeTargets;
 			TextureNode->GetTargetNodeUids(TextureNodeTargets);
@@ -1209,7 +1209,7 @@ TTuple<UInterchangeMaterialExpressionFactoryNode*, FString> UInterchangeGenericM
 	FString OutputName;
 	if (UInterchangeShaderPortsAPI::GetInputConnection(ShaderNode, InputName, ConnectedShaderNodeUid, OutputName))
 	{
-		if (UInterchangeShaderNode* ConnectedShaderNode = Cast<UInterchangeShaderNode>(BaseNodeContainer->GetNode(ConnectedShaderNodeUid)))
+		if (const UInterchangeShaderNode* ConnectedShaderNode = Cast<const UInterchangeShaderNode>(BaseNodeContainer->GetNode(ConnectedShaderNodeUid)))
 		{
 			MaterialExpressionFactoryNode = CreateMaterialExpressionForShaderNode(MaterialFactoryNode, ConnectedShaderNode, ParentUid);
 		}
@@ -1314,7 +1314,7 @@ UInterchangeMaterialInstanceFactoryNode* UInterchangeGenericMaterialPipeline::Cr
 		FString OutputName;
 		if (UInterchangeShaderPortsAPI::GetInputConnection(ShaderGraphNode, InputName, ConnectedShaderNodeUid, OutputName))
 		{
-			if (UInterchangeShaderNode* ConnectedShaderNode = Cast<UInterchangeShaderNode>(BaseNodeContainer->GetNode(ConnectedShaderNodeUid)))
+			if (const UInterchangeShaderNode* ConnectedShaderNode = Cast<const UInterchangeShaderNode>(BaseNodeContainer->GetNode(ConnectedShaderNodeUid)))
 			{
 				InputValue = VisitShaderNode(ConnectedShaderNode);
 			}
@@ -1410,7 +1410,7 @@ TVariant<FString, FLinearColor, float> UInterchangeGenericMaterialPipeline::Visi
 	FString OutputName;
 	if (UInterchangeShaderPortsAPI::GetInputConnection(ShaderNode, InputName, ConnectedShaderNodeUid, OutputName))
 	{
-		if (UInterchangeShaderNode* ConnectedShaderNode = Cast<UInterchangeShaderNode>(BaseNodeContainer->GetNode(ConnectedShaderNodeUid)))
+		if (const UInterchangeShaderNode* ConnectedShaderNode = Cast<const UInterchangeShaderNode>(BaseNodeContainer->GetNode(ConnectedShaderNodeUid)))
 		{
 			Result = VisitShaderNode(ConnectedShaderNode);
 		}
@@ -1578,7 +1578,7 @@ TVariant<FString, FLinearColor, float> UInterchangeGenericMaterialPipeline::Visi
 		if (!TextureUid.IsEmpty())
 		{
 			FString TextureFactoryUid;
-			if (UInterchangeTextureNode* TextureNode = Cast<UInterchangeTextureNode>(BaseNodeContainer->GetNode(TextureUid)))
+			if (const UInterchangeTextureNode* TextureNode = Cast<const UInterchangeTextureNode>(BaseNodeContainer->GetNode(TextureUid)))
 			{
 				TArray<FString> TextureTargetNodes;
 				TextureNode->GetTargetNodeUids(TextureTargetNodes);
diff --git a/Engine/Plugins/Experimental/Interchange/Runtime/Source/Pipelines/Private/InterchangeGenericMeshPipeline.cpp b/Engine/Plugins/Experimental/Interchange/Runtime/Source/Pipelines/Private/InterchangeGenericMeshPipeline.cpp
--- a/Engine/Plugins/Experimental/Interchange/Runtime/Source/Pipelines/Private/InterchangeGenericMeshPipeline.cpp
+++ b/Engine/Plugins/Experimental/Interchange/Runtime/Source/Pipelines/Private/InterchangeGenericMeshPipeline.cpp
@@ -39,7 +39,7 @@ void UInterchangeGenericMeshPipeline::ExecutePreImportPipeline(UInterchangeBaseN
 	ExecutePreImportPipelineStaticMesh();
 }
 
-void UInterchangeGenericMeshPipeline::ExecutePostImportPipeline(const UInterchangeBaseNodeContainer* InBaseNodeContainer, const FString& NodeKey, UObject* CreatedAsset, bool bIsAReimport)
+void UInterchangeGenericMeshPipeline::ExecutePostImportPipeline(const UInterchangeBaseNodeContainer* InBaseNodeContainer, const FString& FactoryNodeKey, UObject* CreatedAsset, bool bIsAReimport)
 {
 	//We do not use the provided base container since ExecutePreImportPipeline cache it
 	//We just make sure the same one is pass in parameter
@@ -48,8 +48,8 @@ void UInterchangeGenericMeshPipeline::ExecutePostImportPipeline(const UInterchan
 		return;
 	}
 
-	UInterchangeBaseNode* Node = BaseNodeContainer->GetNode(NodeKey);
-	if (!Node)
+	const UInterchangeFactoryBaseNode* FactoryNode = BaseNodeContainer->GetFactoryNode(FactoryNodeKey);
+	if (!FactoryNode)
 	{
 		return;
 	}
@@ -57,10 +57,10 @@ void UInterchangeGenericMeshPipeline::ExecutePostImportPipeline(const UInterchan
 	//Set the last content type import
 	LastSkeletalMeshImportContentType = SkeletalMeshImportContentType;
 
-	PostImportSkeletalMesh(CreatedAsset, Node);
+	PostImportSkeletalMesh(CreatedAsset, FactoryNode);
 
 	//Finish the physics asset import, it need the skeletal mesh render data to create the physics collision geometry
-	PostImportPhysicsAssetImport(CreatedAsset, Node);
+	PostImportPhysicsAssetImport(CreatedAsset, FactoryNode);
 }
 
 void UInterchangeGenericMeshPipeline::SetReimportSourceIndex(UClass* ReimportObjectClass, const int32 SourceFileIndex)
diff --git a/Engine/Plugins/Experimental/Interchange/Runtime/Source/Pipelines/Private/InterchangeGenericSkeletalMeshPipeline.cpp b/Engine/Plugins/Experimental/Interchange/Runtime/Source/Pipelines/Private/InterchangeGenericSkeletalMeshPipeline.cpp
--- a/Engine/Plugins/Experimental/Interchange/Runtime/Source/Pipelines/Private/InterchangeGenericSkeletalMeshPipeline.cpp
+++ b/Engine/Plugins/Experimental/Interchange/Runtime/Source/Pipelines/Private/InterchangeGenericSkeletalMeshPipeline.cpp
@@ -262,7 +262,7 @@ UInterchangeSkeletonFactoryNode* UInterchangeGenericMeshPipeline::CreateSkeleton
 	if (BaseNodeContainer->IsNodeUidValid(SkeletonUid))
 	{
 		//The node already exist, just return it
-		SkeletonFactoryNode = Cast<UInterchangeSkeletonFactoryNode>(BaseNodeContainer->GetNode(SkeletonUid));
+		SkeletonFactoryNode = Cast<UInterchangeSkeletonFactoryNode>(BaseNodeContainer->GetFactoryNode(SkeletonUid));
 		if (!ensure(SkeletonFactoryNode))
 		{
 			//Log an error
@@ -314,7 +314,7 @@ UInterchangeSkeletalMeshFactoryNode* UInterchangeGenericMeshPipeline::CreateSkel
 		return nullptr;
 	}
 	const FString SkeletonUid = TEXT("\\Skeleton\\") + RootJointNode->GetUniqueID();
-	UInterchangeSkeletonFactoryNode* SkeletonFactoryNode = Cast<UInterchangeSkeletonFactoryNode>(BaseNodeContainer->GetNode(SkeletonUid));
+	UInterchangeSkeletonFactoryNode* SkeletonFactoryNode = Cast<UInterchangeSkeletonFactoryNode>(BaseNodeContainer->GetFactoryNode(SkeletonUid));
 	if (!ensure(SkeletonFactoryNode))
 	{
 		//Log an error
@@ -342,13 +342,13 @@ UInterchangeSkeletalMeshFactoryNode* UInterchangeGenericMeshPipeline::CreateSkel
 				if (MeshUids.Num() > 0)
 				{
 					const FString& MeshUid = MeshUids[0];
-					UInterchangeMeshNode* MeshNode = Cast<UInterchangeMeshNode>(BaseNodeContainer->GetNode(MeshUid));
+					const UInterchangeMeshNode* MeshNode = Cast<const UInterchangeMeshNode>(BaseNodeContainer->GetNode(MeshUid));
 					if (MeshNode)
 					{
 						OutFirstMeshNodeUid = MeshUid;
 						return MeshNode;
 					}
-					UInterchangeSceneNode* SceneNode = Cast<UInterchangeSceneNode>(BaseNodeContainer->GetNode(MeshUid));
+					const UInterchangeSceneNode* SceneNode = Cast<const UInterchangeSceneNode>(BaseNodeContainer->GetNode(MeshUid));
 					if (SceneNode)
 					{
 						FString MeshNodeUid;
@@ -524,7 +524,7 @@ void UInterchangeGenericMeshPipeline::AddLodDataToSkeletalMesh(const UInterchang
 		const FString LODDataPrefix = TEXT("\\LodData") + (LodIndex > 0 ? FString::FromInt(LodIndex) : TEXT(""));
 		const FString SkeletalMeshLodDataUniqueID = LODDataPrefix + SkeletalMeshUid + SkeletonUid;
 		//The LodData already exist
-		UInterchangeSkeletalMeshLodDataNode* LodDataNode = Cast<UInterchangeSkeletalMeshLodDataNode>(BaseNodeContainer->GetNode(SkeletalMeshLodDataUniqueID));
+		UInterchangeSkeletalMeshLodDataNode* LodDataNode = Cast<UInterchangeSkeletalMeshLodDataNode>(BaseNodeContainer->GetFactoryNode(SkeletalMeshLodDataUniqueID));
 		if (!LodDataNode)
 		{
 			//Add the data for the LOD (skeleton Unique ID and all the mesh node fbx path, so we can find them when we will create the payload data)
@@ -568,7 +568,7 @@ void UInterchangeGenericMeshPipeline::AddLodDataToSkeletalMesh(const UInterchang
 				const FString MaterialFactoryNodeUid = UInterchangeMaterialFactoryNode::GetMaterialFactoryNodeUidFromMaterialNodeUid(MaterialDependencies[MaterialIndex]);
 				if (BaseNodeContainer->IsNodeUidValid(MaterialFactoryNodeUid))
 				{
-					BaseNodeContainer->GetNode(MaterialFactoryNodeUid)->SetEnabled(true);
+					BaseNodeContainer->GetFactoryNode(MaterialFactoryNodeUid)->SetEnabled(true);
 					//Create a factory dependency so Material asset are import before the skeletal mesh asset
 					TArray<FString> FactoryDependencies;
 					SkeletalMeshFactoryNode->GetFactoryDependencies(FactoryDependencies);
@@ -583,7 +583,7 @@ void UInterchangeGenericMeshPipeline::AddLodDataToSkeletalMesh(const UInterchang
 	}
 }
 
-void UInterchangeGenericMeshPipeline::PostImportSkeletalMesh(UObject* CreatedAsset, UInterchangeBaseNode* Node)
+void UInterchangeGenericMeshPipeline::PostImportSkeletalMesh(UObject* CreatedAsset, const UInterchangeFactoryBaseNode* FactoryNode)
 {
 	check(!CommonSkeletalMeshesAndAnimationsProperties.IsNull());
 
@@ -607,7 +607,7 @@ void UInterchangeGenericMeshPipeline::PostImportSkeletalMesh(UObject* CreatedAss
 	}
 }
 
-void UInterchangeGenericMeshPipeline::PostImportPhysicsAssetImport(UObject* CreatedAsset, UInterchangeBaseNode* Node)
+void UInterchangeGenericMeshPipeline::PostImportPhysicsAssetImport(UObject* CreatedAsset, const UInterchangeFactoryBaseNode* FactoryNode)
 {
 #if WITH_EDITOR
 	if (!bCreatePhysicsAsset || !BaseNodeContainer)
@@ -620,12 +620,12 @@ void UInterchangeGenericMeshPipeline::PostImportPhysicsAssetImport(UObject* Crea
 	{
 		return;
 	}
-	if (UInterchangePhysicsAssetFactoryNode* PhysicsAssetFactoryNode = Cast<UInterchangePhysicsAssetFactoryNode>(Node))
+	if (const UInterchangePhysicsAssetFactoryNode* PhysicsAssetFactoryNode = Cast<const UInterchangePhysicsAssetFactoryNode>(FactoryNode))
 	{
 		FString SkeletalMeshFactoryNodeUid;
 		if (PhysicsAssetFactoryNode->GetCustomSkeletalMeshUid(SkeletalMeshFactoryNodeUid))
 		{
-			if (UInterchangeSkeletalMeshFactoryNode* SkeletalMeshFactoryNode = Cast<UInterchangeSkeletalMeshFactoryNode>(BaseNodeContainer->GetNode(SkeletalMeshFactoryNodeUid)))
+			if (const UInterchangeSkeletalMeshFactoryNode* SkeletalMeshFactoryNode = Cast<const UInterchangeSkeletalMeshFactoryNode>(BaseNodeContainer->GetFactoryNode(SkeletalMeshFactoryNodeUid)))
 			{
 				if (SkeletalMeshFactoryNode->ReferenceObject.IsValid())
 				{
@@ -675,7 +675,7 @@ void UInterchangeGenericMeshPipeline::ImplementUseSourceNameForAssetOptionSkelet
 	//If we import only one asset, and bUseSourceNameForAsset is true, we want to rename the asset using the file name.
 	const bool bShouldChangeAssetName = (bUseSourceNameForAsset && MeshesImportedNodeCount == 1);
 	const FString SkeletalMeshUid = SkeletalMeshNodeUids[0];
-	UInterchangeSkeletalMeshFactoryNode* SkeletalMeshNode = Cast<UInterchangeSkeletalMeshFactoryNode>(BaseNodeContainer->GetNode(SkeletalMeshUid));
+	UInterchangeSkeletalMeshFactoryNode* SkeletalMeshNode = Cast<UInterchangeSkeletalMeshFactoryNode>(BaseNodeContainer->GetFactoryNode(SkeletalMeshUid));
 	if (!SkeletalMeshNode)
 	{
 		return;
@@ -695,14 +695,14 @@ void UInterchangeGenericMeshPipeline::ImplementUseSourceNameForAssetOptionSkelet
 	if (LodDataUids.Num() > 0)
 	{
 		//Get the skeleton from the base LOD, skeleton is shared with all LODs
-		if (UInterchangeSkeletalMeshLodDataNode* SkeletalMeshLodDataNode = Cast<UInterchangeSkeletalMeshLodDataNode>(BaseNodeContainer->GetNode(LodDataUids[0])))
+		if (const UInterchangeSkeletalMeshLodDataNode* SkeletalMeshLodDataNode = Cast<const UInterchangeSkeletalMeshLodDataNode>(BaseNodeContainer->GetFactoryNode(LodDataUids[0])))
 		{
 			//If the user did not specify any skeleton
 			if (CommonSkeletalMeshesAndAnimationsProperties->Skeleton.IsNull())
 			{
 				FString SkeletalMeshSkeletonUid;
 				SkeletalMeshLodDataNode->GetCustomSkeletonUid(SkeletalMeshSkeletonUid);
-				UInterchangeSkeletonFactoryNode* SkeletonFactoryNode = Cast<UInterchangeSkeletonFactoryNode>(BaseNodeContainer->GetNode(SkeletalMeshSkeletonUid));
+				UInterchangeSkeletonFactoryNode* SkeletonFactoryNode = Cast<UInterchangeSkeletonFactoryNode>(BaseNodeContainer->GetFactoryNode(SkeletalMeshSkeletonUid));
 				if (SkeletonFactoryNode)
 				{
 					const FString SkeletonName = DisplayLabelName + TEXT("_Skeleton");
@@ -716,7 +716,7 @@ void UInterchangeGenericMeshPipeline::ImplementUseSourceNameForAssetOptionSkelet
 	BaseNodeContainer->GetNodes(PhysicsAssetFactoryNodeClass, PhysicsAssetNodeUids);
 	for (const FString& PhysicsAssetNodeUid : PhysicsAssetNodeUids)
 	{
-		UInterchangePhysicsAssetFactoryNode* PhysicsAssetFactoryNode = Cast<UInterchangePhysicsAssetFactoryNode>(BaseNodeContainer->GetNode(PhysicsAssetNodeUid));
+		UInterchangePhysicsAssetFactoryNode* PhysicsAssetFactoryNode = Cast<UInterchangePhysicsAssetFactoryNode>(BaseNodeContainer->GetFactoryNode(PhysicsAssetNodeUid));
 		if (!ensure(PhysicsAssetFactoryNode))
 		{
 			continue;
diff --git a/Engine/Plugins/Experimental/Interchange/Runtime/Source/Pipelines/Private/InterchangeGenericStaticMeshPipeline.cpp b/Engine/Plugins/Experimental/Interchange/Runtime/Source/Pipelines/Private/InterchangeGenericStaticMeshPipeline.cpp
--- a/Engine/Plugins/Experimental/Interchange/Runtime/Source/Pipelines/Private/InterchangeGenericStaticMeshPipeline.cpp
+++ b/Engine/Plugins/Experimental/Interchange/Runtime/Source/Pipelines/Private/InterchangeGenericStaticMeshPipeline.cpp
@@ -326,21 +326,21 @@ bool UInterchangeGenericMeshPipeline::MakeMeshFactoryNodeUidAndDisplayLabel(cons
 			if (Uids.Num() > 0)
 			{
 				const FString& Uid = Uids[0];
-				UInterchangeBaseNode* Node = BaseNodeContainer->GetNode(Uid);
+				const UInterchangeBaseNode* Node = BaseNodeContainer->GetNode(Uid);
 
-				if (UInterchangeMeshNode* MeshNode = Cast<UInterchangeMeshNode>(Node))
+				if (const UInterchangeMeshNode* MeshNode = Cast<const UInterchangeMeshNode>(Node))
 				{
 					DisplayLabel = Node->GetDisplayLabel();
 					NewNodeUid = Uid;
 					return true;
 				}
 
-				if (UInterchangeSceneNode* SceneNode = Cast<UInterchangeSceneNode>(Node))
+				if (const UInterchangeSceneNode* SceneNode = Cast<const UInterchangeSceneNode>(Node))
 				{
 					FString RefMeshUid;
 					if (SceneNode->GetCustomAssetInstanceUid(RefMeshUid))
 					{
-						UInterchangeBaseNode* MeshNode = BaseNodeContainer->GetNode(RefMeshUid);
+						const UInterchangeBaseNode* MeshNode = BaseNodeContainer->GetNode(RefMeshUid);
 						if (MeshNode)
 						{
 							DisplayLabel = MeshNode->GetDisplayLabel();
@@ -459,7 +459,7 @@ void UInterchangeGenericMeshPipeline::AddLodDataToStaticMesh(UInterchangeStaticM
 		const FString StaticMeshLodDataUniqueID = LODDataPrefix + StaticMeshFactoryUid;
 
 		// Create LodData node if it doesn't already exist
-		UInterchangeStaticMeshLodDataNode* LodDataNode = Cast<UInterchangeStaticMeshLodDataNode>(BaseNodeContainer->GetNode(StaticMeshLodDataUniqueID));
+		UInterchangeStaticMeshLodDataNode* LodDataNode = Cast<UInterchangeStaticMeshLodDataNode>(BaseNodeContainer->GetFactoryNode(StaticMeshLodDataUniqueID));
 		if (!LodDataNode)
 		{
 			// Add the data for the LOD (all the mesh node fbx path, so we can find them when we will create the payload data)
@@ -506,7 +506,7 @@ void UInterchangeGenericMeshPipeline::AddLodDataToStaticMesh(UInterchangeStaticM
 				const FString MaterialFactoryNodeUid = UInterchangeMaterialFactoryNode::GetMaterialFactoryNodeUidFromMaterialNodeUid(MaterialDependencies[MaterialIndex]);
 				if (BaseNodeContainer->IsNodeUidValid(MaterialFactoryNodeUid))
 				{
-					BaseNodeContainer->GetNode(MaterialFactoryNodeUid)->SetEnabled(true);
+					BaseNodeContainer->GetFactoryNode(MaterialFactoryNodeUid)->SetEnabled(true);
 					// Create a factory dependency so Material asset are imported before the skeletal mesh asset
 					TArray<FString> FactoryDependencies;
 					StaticMeshFactoryNode->GetFactoryDependencies(FactoryDependencies);
diff --git a/Engine/Plugins/Experimental/Interchange/Runtime/Source/Pipelines/Private/InterchangeGenericTexturePipeline.cpp b/Engine/Plugins/Experimental/Interchange/Runtime/Source/Pipelines/Private/InterchangeGenericTexturePipeline.cpp
--- a/Engine/Plugins/Experimental/Interchange/Runtime/Source/Pipelines/Private/InterchangeGenericTexturePipeline.cpp
+++ b/Engine/Plugins/Experimental/Interchange/Runtime/Source/Pipelines/Private/InterchangeGenericTexturePipeline.cpp
@@ -126,7 +126,7 @@ void UInterchangeGenericTexturePipeline::ExecutePostImportPipeline(const UInterc
 		return;
 	}
 
-	UInterchangeBaseNode* Node = BaseNodeContainer->GetNode(NodeKey);
+	const UInterchangeFactoryBaseNode* Node = BaseNodeContainer->GetFactoryNode(NodeKey);
 	if (!Node)
 	{
 		return;
@@ -194,7 +194,7 @@ UInterchangeTextureFactoryNode* UInterchangeGenericTexturePipeline::CreateTextur
 	UInterchangeTextureFactoryNode* TextureFactoryNode = nullptr;
 	if (BaseNodeContainer->IsNodeUidValid(NodeUid))
 	{
-		TextureFactoryNode = Cast<UInterchangeTextureFactoryNode>(BaseNodeContainer->GetNode(NodeUid));
+		TextureFactoryNode = Cast<UInterchangeTextureFactoryNode>(BaseNodeContainer->GetFactoryNode(NodeUid));
 		if (!ensure(TextureFactoryNode))
 		{
 			//Log an error
diff --git a/Engine/Plugins/Experimental/Interchange/Runtime/Source/Pipelines/Private/InterchangePipelineMeshesUtilities.cpp b/Engine/Plugins/Experimental/Interchange/Runtime/Source/Pipelines/Private/InterchangePipelineMeshesUtilities.cpp
--- a/Engine/Plugins/Experimental/Interchange/Runtime/Source/Pipelines/Private/InterchangePipelineMeshesUtilities.cpp
+++ b/Engine/Plugins/Experimental/Interchange/Runtime/Source/Pipelines/Private/InterchangePipelineMeshesUtilities.cpp
@@ -33,7 +33,7 @@ UInterchangePipelineMeshesUtilities* UInterchangePipelineMeshesUtilities::Create
 		{
 			if (Node->GetNodeContainerType() == EInterchangeNodeContainerType::TranslatedAsset)
 			{
-				if (UInterchangeMeshNode* MeshNode = Cast<UInterchangeMeshNode>(Node))
+				if (const UInterchangeMeshNode* MeshNode = Cast<UInterchangeMeshNode>(Node))
 				{
 					FInterchangeMeshGeometry& MeshGeometry = PipelineMeshesUtilities->MeshGeometriesPerMeshUid.FindOrAdd(NodeUid);
 					MeshGeometry.MeshUid = NodeUid;
@@ -95,7 +95,7 @@ UInterchangePipelineMeshesUtilities* UInterchangePipelineMeshesUtilities::Create
 										break;
 									}
 									LastChildUid = ParentMeshSceneNode->GetUniqueID();
-									ParentMeshSceneNode = Cast<UInterchangeSceneNode>(BaseNodeContainer->GetNode(ParentMeshSceneNode->GetParentUid()));
+									ParentMeshSceneNode = Cast<const UInterchangeSceneNode>(BaseNodeContainer->GetNode(ParentMeshSceneNode->GetParentUid()));
 								} while (ParentMeshSceneNode);
 
 								FInterchangeMeshGeometry& MeshGeometry = PipelineMeshesUtilities->MeshGeometriesPerMeshUid.FindChecked(MeshUid);
diff --git a/Engine/Plugins/Experimental/Interchange/Runtime/Source/Pipelines/Public/InterchangeGenericMeshPipeline.h b/Engine/Plugins/Experimental/Interchange/Runtime/Source/Pipelines/Public/InterchangeGenericMeshPipeline.h
--- a/Engine/Plugins/Experimental/Interchange/Runtime/Source/Pipelines/Public/InterchangeGenericMeshPipeline.h
+++ b/Engine/Plugins/Experimental/Interchange/Runtime/Source/Pipelines/Public/InterchangeGenericMeshPipeline.h
@@ -153,12 +153,12 @@ private:
 	/**
 	 * This function will finish creating the skeletalmesh asset
 	 */
-	void PostImportSkeletalMesh(UObject* CreatedAsset, UInterchangeBaseNode* Node);
+	void PostImportSkeletalMesh(UObject* CreatedAsset, const UInterchangeFactoryBaseNode* FactoryNode);
 
 	/**
 	 * This function will finish creating the physics asset with the skeletalmesh render data
 	 */
-	void PostImportPhysicsAssetImport(UObject* CreatedAsset, UInterchangeBaseNode* Node);
+	void PostImportPhysicsAssetImport(UObject* CreatedAsset, const UInterchangeFactoryBaseNode* FactoryNode);
 public:
 	
 	/** Specialize for skeletalmesh */
diff --git a/Engine/Source/Runtime/Experimental/Interchange/Core/Private/Nodes/InterchangeBaseNode.cpp b/Engine/Source/Runtime/Experimental/Interchange/Core/Private/Nodes/InterchangeBaseNode.cpp
--- a/Engine/Source/Runtime/Experimental/Interchange/Core/Private/Nodes/InterchangeBaseNode.cpp
+++ b/Engine/Source/Runtime/Experimental/Interchange/Core/Private/Nodes/InterchangeBaseNode.cpp
@@ -42,12 +42,6 @@ namespace UE
 			return Key;
 		}
 
-		const FString& FBaseNodeStaticData::GetFactoryDependenciesBaseKey()
-		{
-			static FString BaseNodeFactoryDependencies_BaseKey = TEXT("__BaseNodeFactoryDependencies__");
-			return BaseNodeFactoryDependencies_BaseKey;
-		}
-
 		const FAttributeKey& FBaseNodeStaticData::ClassTypeAttributeKey()
 		{
 			static FAttributeKey AttributeKey(TEXT("__ClassTypeAttribute__"));
@@ -78,7 +72,6 @@ namespace UE
 UInterchangeBaseNode::UInterchangeBaseNode()
 {
 	Attributes = MakeShared<UE::Interchange::FAttributeStorage, ESPMode::ThreadSafe>();
-	FactoryDependencies.Initialize(Attributes, UE::Interchange::FBaseNodeStaticData::GetFactoryDependenciesBaseKey());
 	TargetNodes.Initialize(Attributes, UE::Interchange::FBaseNodeStaticData::TargetAssetIDsKey());
 	RegisterAttribute<bool>(UE::Interchange::FBaseNodeStaticData::IsEnabledKey(), true);
 	RegisterAttribute<uint8>(UE::Interchange::FBaseNodeStaticData::NodeContainerTypeKey(), static_cast<uint8>(EInterchangeNodeContainerType::None));
@@ -146,20 +139,6 @@ FString UInterchangeBaseNode::GetKeyDisplayName(const UE::Interchange::FAttribut
 	{
 		KeyDisplayName = TEXT("Imported Asset Name");
 	}
-	else if (NodeAttributeKey.Key.Equals(UE::Interchange::FBaseNodeStaticData::GetFactoryDependenciesBaseKey()))
-	{
-		KeyDisplayName = TEXT("Factory Dependencies Count");
-	}
-	else if (NodeAttributeKey.Key.StartsWith(UE::Interchange::FBaseNodeStaticData::GetFactoryDependenciesBaseKey()))
-	{
-		KeyDisplayName = TEXT("Factory Dependencies Index ");
-		const FString IndexKey = UE::Interchange::TArrayAttributeHelper<FString>::IndexKey();
-		int32 IndexPosition = NodeAttributeKey.Key.Find(IndexKey) + IndexKey.Len();
-		if (IndexPosition < NodeAttributeKey.Key.Len())
-		{
-			KeyDisplayName += NodeAttributeKey.Key.RightChop(IndexPosition);
-		}
-	}
 	else if (NodeAttributeKey.Key.Equals(UE::Interchange::FBaseNodeStaticData::TargetAssetIDsKey()))
 	{
 		KeyDisplayName = TEXT("Target Asset Count");
@@ -193,10 +172,6 @@ bool UInterchangeBaseNode::ShouldHideAttribute(const UE::Interchange::FAttribute
 FString UInterchangeBaseNode::GetAttributeCategory(const UE::Interchange::FAttributeKey& NodeAttributeKey) const
 {
 	FString CategoryName = TEXT("Attributes");
-	if (NodeAttributeKey.Key.StartsWith(UE::Interchange::FBaseNodeStaticData::GetFactoryDependenciesBaseKey()))
-	{
-		CategoryName = TEXT("FactoryDependencies");
-	}
 	return CategoryName;
 }
 
@@ -346,31 +321,6 @@ bool UInterchangeBaseNode::SetParentUid(const FString& ParentUid)
 	return false;
 }
 
-int32 UInterchangeBaseNode::GetFactoryDependenciesCount() const
-{
-	return FactoryDependencies.GetCount();
-}
-
-void UInterchangeBaseNode::GetFactoryDependency(const int32 Index, FString& OutDependency) const
-{
-	FactoryDependencies.GetItem(Index, OutDependency);
-}
-
-void UInterchangeBaseNode::GetFactoryDependencies(TArray<FString>& OutDependencies ) const
-{
-	FactoryDependencies.GetItems(OutDependencies);
-}
-
-bool UInterchangeBaseNode::AddFactoryDependencyUid(const FString& DependencyUid)
-{
-	return FactoryDependencies.AddItem(DependencyUid);
-}
-
-bool UInterchangeBaseNode::RemoveFactoryDependencyUid(const FString& DependencyUid)
-{
-	return FactoryDependencies.RemoveItem(DependencyUid);
-}
-
 bool UInterchangeBaseNode::IsEnabled() const
 {
 	if (!Attributes->ContainAttribute(UE::Interchange::FBaseNodeStaticData::IsEnabledKey()))
@@ -420,11 +370,6 @@ FGuid UInterchangeBaseNode::GetHash() const
 	return Attributes->GetStorageHash();
 }
 
-UClass* UInterchangeBaseNode::GetObjectClass() const
-{
-	return nullptr;
-}
-
 FString UInterchangeBaseNode::GetAssetName() const
 {
 	if (!Attributes->ContainAttribute(UE::Interchange::FBaseNodeStaticData::AssetNameKey()))
diff --git a/Engine/Source/Runtime/Experimental/Interchange/Core/Private/Nodes/InterchangeBaseNodeContainer.cpp b/Engine/Source/Runtime/Experimental/Interchange/Core/Private/Nodes/InterchangeBaseNodeContainer.cpp
--- a/Engine/Source/Runtime/Experimental/Interchange/Core/Private/Nodes/InterchangeBaseNodeContainer.cpp
+++ b/Engine/Source/Runtime/Experimental/Interchange/Core/Private/Nodes/InterchangeBaseNodeContainer.cpp
@@ -99,7 +99,7 @@ void UInterchangeBaseNodeContainer::GetNodes(const UClass* ClassNode, TArray<FSt
 	});
 }
 
-UInterchangeBaseNode* UInterchangeBaseNodeContainer::GetNode(const FString& NodeUniqueID) const
+const UInterchangeBaseNode* UInterchangeBaseNodeContainer::GetNode(const FString& NodeUniqueID) const
 {
 	if (NodeUniqueID == UInterchangeBaseNode::InvalidNodeUid())
 	{
@@ -113,6 +113,20 @@ UInterchangeBaseNode* UInterchangeBaseNodeContainer::GetNode(const FString& Node
 	return Node;
 }
 
+UInterchangeFactoryBaseNode* UInterchangeBaseNodeContainer::GetFactoryNode(const FString& NodeUniqueID) const
+{
+	if (NodeUniqueID == UInterchangeBaseNode::InvalidNodeUid())
+	{
+		return nullptr;
+	}
+	if (!Nodes.Contains(NodeUniqueID))
+	{
+		return nullptr;
+	}
+	UInterchangeFactoryBaseNode* FactoryNode = Cast< UInterchangeFactoryBaseNode>(Nodes.FindChecked(NodeUniqueID));
+	return FactoryNode;
+}
+
 bool UInterchangeBaseNodeContainer::SetNodeParentUid(const FString& NodeUniqueID, const FString& NewParentNodeUid)
 {
 	if (!Nodes.Contains(NodeUniqueID))
diff --git a/Engine/Source/Runtime/Experimental/Interchange/Core/Private/Nodes/InterchangeFactoryBaseNode.cpp b/Engine/Source/Runtime/Experimental/Interchange/Core/Private/Nodes/InterchangeFactoryBaseNode.cpp
--- a/Engine/Source/Runtime/Experimental/Interchange/Core/Private/Nodes/InterchangeFactoryBaseNode.cpp
+++ b/Engine/Source/Runtime/Experimental/Interchange/Core/Private/Nodes/InterchangeFactoryBaseNode.cpp
@@ -7,6 +7,27 @@
 #include "UObject/Object.h"
 #include "UObject/ObjectMacros.h"
 
+
+namespace UE
+{
+	namespace Interchange
+	{
+		const FString& FFactoryBaseNodeStaticData::FactoryDependenciesBaseKey()
+		{
+			static FString BaseNodeFactoryDependencies_BaseKey = TEXT("__BaseNodeFactoryDependencies__");
+			return BaseNodeFactoryDependencies_BaseKey;
+		}
+
+	} //ns Interchange
+} //ns UE
+
+
+UInterchangeFactoryBaseNode::UInterchangeFactoryBaseNode()
+{
+	FactoryDependencies.Initialize(Attributes, UE::Interchange::FFactoryBaseNodeStaticData::FactoryDependenciesBaseKey());
+}
+
+
 FString UInterchangeFactoryBaseNode::GetKeyDisplayName(const UE::Interchange::FAttributeKey& NodeAttributeKey) const
 {
 	FString KeyDisplayName = NodeAttributeKey.ToString();
@@ -14,6 +35,20 @@ FString UInterchangeFactoryBaseNode::GetKeyDisplayName(const UE::Interchange::FA
 	{
 		KeyDisplayName = TEXT("Import Sub-Path");
 	}
+	else if (NodeAttributeKey.Key.Equals(UE::Interchange::FFactoryBaseNodeStaticData::FactoryDependenciesBaseKey()))
+	{
+		KeyDisplayName = TEXT("Factory Dependencies Count");
+	}
+	else if (NodeAttributeKey.Key.StartsWith(UE::Interchange::FFactoryBaseNodeStaticData::FactoryDependenciesBaseKey()))
+	{
+		KeyDisplayName = TEXT("Factory Dependencies Index ");
+		const FString IndexKey = UE::Interchange::TArrayAttributeHelper<FString>::IndexKey();
+		int32 IndexPosition = NodeAttributeKey.Key.Find(IndexKey) + IndexKey.Len();
+		if (IndexPosition < NodeAttributeKey.Key.Len())
+		{
+			KeyDisplayName += NodeAttributeKey.Key.RightChop(IndexPosition);
+		}
+	}
 	else
 	{
 		KeyDisplayName = Super::GetKeyDisplayName(NodeAttributeKey);
@@ -22,6 +57,25 @@ FString UInterchangeFactoryBaseNode::GetKeyDisplayName(const UE::Interchange::FA
 	return KeyDisplayName;
 }
 
+
+FString UInterchangeFactoryBaseNode::GetAttributeCategory(const UE::Interchange::FAttributeKey& NodeAttributeKey) const
+{
+	if (NodeAttributeKey.Key.StartsWith(UE::Interchange::FFactoryBaseNodeStaticData::FactoryDependenciesBaseKey()))
+	{
+		return TEXT("FactoryDependencies");
+	}
+	else
+	{
+		return Super::GetAttributeCategory(NodeAttributeKey);
+	}
+}
+
+
+UClass* UInterchangeFactoryBaseNode::GetObjectClass() const
+{
+	return nullptr;
+}
+
 bool UInterchangeFactoryBaseNode::GetCustomSubPath(FString& AttributeValue) const
 {
 	IMPLEMENT_NODE_ATTRIBUTE_GETTER(SubPath, FString)
@@ -31,3 +85,29 @@ bool UInterchangeFactoryBaseNode::SetCustomSubPath(const FString& AttributeValue
 {
 	IMPLEMENT_NODE_ATTRIBUTE_SETTER_NODELEGATE(SubPath, FString)
 }
+
+int32 UInterchangeFactoryBaseNode::GetFactoryDependenciesCount() const
+{
+	return FactoryDependencies.GetCount();
+}
+
+void UInterchangeFactoryBaseNode::GetFactoryDependency(const int32 Index, FString& OutDependency) const
+{
+	FactoryDependencies.GetItem(Index, OutDependency);
+}
+
+void UInterchangeFactoryBaseNode::GetFactoryDependencies(TArray<FString>& OutDependencies) const
+{
+	FactoryDependencies.GetItems(OutDependencies);
+}
+
+bool UInterchangeFactoryBaseNode::AddFactoryDependencyUid(const FString& DependencyUid)
+{
+	return FactoryDependencies.AddItem(DependencyUid);
+}
+
+bool UInterchangeFactoryBaseNode::RemoveFactoryDependencyUid(const FString& DependencyUid)
+{
+	return FactoryDependencies.RemoveItem(DependencyUid);
+}
+
diff --git a/Engine/Source/Runtime/Experimental/Interchange/Core/Public/InterchangeFactoryBase.h b/Engine/Source/Runtime/Experimental/Interchange/Core/Public/InterchangeFactoryBase.h
--- a/Engine/Source/Runtime/Experimental/Interchange/Core/Public/InterchangeFactoryBase.h
+++ b/Engine/Source/Runtime/Experimental/Interchange/Core/Public/InterchangeFactoryBase.h
@@ -11,8 +11,8 @@
 #include "InterchangeFactoryBase.generated.h"
 
 class ULevel;
-class UInterchangeBaseNode;
 class UInterchangeBaseNodeContainer;
+class UInterchangeFactoryBaseNode;
 class UInterchangePipelineBase;
 class UInterchangeSourceData;
 class UInterchangeTranslatorBase;
@@ -44,7 +44,7 @@ public:
 		FString AssetName = FString();
 
 		/** The base node that describe how to create the asset */
-		UInterchangeBaseNode* AssetNode = nullptr;
+		UInterchangeFactoryBaseNode* AssetNode = nullptr;
 
 		/** The translator is use to retrieve the PayLoad data in case the factory need it */
 		const UInterchangeTranslatorBase* Translator = nullptr;
@@ -102,7 +102,7 @@ public:
 		FString ObjectName;
 
 		/** The base node that describe how to create the asset */
-		UInterchangeBaseNode* FactoryNode = nullptr;
+		UInterchangeFactoryBaseNode* FactoryNode = nullptr;
 
 		/** The node container associated with the current source index */
 		const UInterchangeBaseNodeContainer* NodeContainer = nullptr;
@@ -132,7 +132,7 @@ public:
 	{
 		/** The source data, mainly use to set the asset import data file. TODO: we have to refactor UAssetImportData, the source data should be the base class for this now */
 		const UInterchangeSourceData* SourceData = nullptr;
-		UInterchangeBaseNode* FactoryNode = nullptr;
+		UInterchangeFactoryBaseNode* FactoryNode = nullptr;
 
 
 		/** The UObject  we want to execute code on*/
diff --git a/Engine/Source/Runtime/Experimental/Interchange/Core/Public/InterchangePipelineBase.h b/Engine/Source/Runtime/Experimental/Interchange/Core/Public/InterchangePipelineBase.h
--- a/Engine/Source/Runtime/Experimental/Interchange/Core/Public/InterchangePipelineBase.h
+++ b/Engine/Source/Runtime/Experimental/Interchange/Core/Public/InterchangePipelineBase.h
@@ -4,6 +4,7 @@
 
 #include "CoreMinimal.h"
 #include "Nodes/InterchangeBaseNodeContainer.h"
+#include "Nodes/InterchangeFactoryBaseNode.h"
 #include "InterchangeResultsContainer.h"
 #include "InterchangeSourceData.h"
 #include "UObject/Object.h"
@@ -49,12 +50,12 @@ public:
 	 * @note - the Interchange manager is calling this function not the virtual one that is call by the default implementation.
 	 */
 	UFUNCTION(BlueprintNativeEvent, BlueprintCallable, Category = "Interchange | Pipeline")
-	void ScriptedExecutePostImportPipeline(const UInterchangeBaseNodeContainer* BaseNodeContainer, const FString& NodeKey, UObject* CreatedAsset, bool bIsAReimport);
+	void ScriptedExecutePostImportPipeline(const UInterchangeBaseNodeContainer* BaseNodeContainer, const FString& FactoryNodeKey, UObject* CreatedAsset, bool bIsAReimport);
 	/** The default implementation (call if the blueprint do not have any implementation) will call the virtual ExecuteImportPipeline */
-	void ScriptedExecutePostImportPipeline_Implementation(const UInterchangeBaseNodeContainer* BaseNodeContainer, const FString& NodeKey, UObject* CreatedAsset, bool bIsAReimport)
+	void ScriptedExecutePostImportPipeline_Implementation(const UInterchangeBaseNodeContainer* BaseNodeContainer, const FString& FactoryNodeKey, UObject* CreatedAsset, bool bIsAReimport)
 	{
 		//By default we call the virtual import pipeline execution
-		ExecutePostImportPipeline(BaseNodeContainer, NodeKey, CreatedAsset, bIsAReimport);
+		ExecutePostImportPipeline(BaseNodeContainer, FactoryNodeKey, CreatedAsset, bIsAReimport);
 	}
 
 	/**
diff --git a/Engine/Source/Runtime/Experimental/Interchange/Core/Public/Nodes/InterchangeBaseNode.h b/Engine/Source/Runtime/Experimental/Interchange/Core/Public/Nodes/InterchangeBaseNode.h
--- a/Engine/Source/Runtime/Experimental/Interchange/Core/Public/Nodes/InterchangeBaseNode.h
+++ b/Engine/Source/Runtime/Experimental/Interchange/Core/Public/Nodes/InterchangeBaseNode.h
@@ -197,7 +197,6 @@ namespace UE
 			static const FAttributeKey& ParentIDKey();
 			static const FAttributeKey& IsEnabledKey();
 			static const FString& TargetAssetIDsKey();
-			static const FString& GetFactoryDependenciesBaseKey();
 			static const FAttributeKey& ClassTypeAttributeKey();
 			static const FAttributeKey& AssetNameKey();
 			static const FAttributeKey& NodeContainerTypeKey();
@@ -485,36 +484,6 @@ public:
 	UFUNCTION(BlueprintCallable, Category = "Interchange | Node")
 	bool SetParentUid(const FString& ParentUid);
 
-	/**
-	 * This function allow to retrieve the number of factory dependencies for this object.
-	 */
-	UFUNCTION(BlueprintCallable, Category = "Interchange | Node")
-	int32 GetFactoryDependenciesCount() const;
-
-	/**
-	 * This function allow to retrieve the dependency for this object.
-	 */
-	UFUNCTION(BlueprintCallable, Category = "Interchange | Node")
-	void GetFactoryDependencies(TArray<FString>& OutDependencies ) const;
-
-	/**
-	 * This function allow to retrieve one dependency for this object.
-	 */
-	UFUNCTION(BlueprintCallable, Category = "Interchange | Node")
-	void GetFactoryDependency(const int32 Index, FString& OutDependency) const;
-
-	/**
-	 * Add one dependency to this object.
-	 */
-	UFUNCTION(BlueprintCallable, Category = "Interchange | Node")
-	bool AddFactoryDependencyUid(const FString& DependencyUid);
-
-	/**
-	 * Remove one dependency from this object.
-	 */
-	UFUNCTION(BlueprintCallable, Category = "Interchange | Node")
-	bool RemoveFactoryDependencyUid(const FString& DependencyUid);
-
 	/**
 	 * Get number of target assets relating to this object.
 	 */
@@ -569,11 +538,6 @@ public:
 	 */
 	virtual FGuid GetHash() const;
 
-	/**
-	 * Optional, Any node that can import/export an object should return the UClass of the object so we can find factory/writer
-	 */
-	virtual class UClass* GetObjectClass() const;
-
 	/**
 	 * Optional, Any node that can import/export an asset should set the proper name we will give to the asset.
 	 * If the attribute was never set, it will return GetDisplayLabel.
@@ -615,8 +579,6 @@ public:
 		return;
 	}
 
-	UPROPERTY()
-	mutable FSoftObjectPath ReferenceObject;
 protected:
 	/** The storage use to store the Key value attribute for this node. */
 	TSharedPtr<UE::Interchange::FAttributeStorage, ESPMode::ThreadSafe> Attributes;
@@ -628,13 +590,6 @@ protected:
 
 	bool bIsInitialized = false;
 
-	/**
-	 * Those dependencies are use by the interchange parsing task to make sure the asset are created in the correct order.
-	 * Example: Mesh factory node will have dependencies on material factory node
-	 *          Material factory node will have dependencies on texture factory node
-	 */
-	UE::Interchange::TArrayAttributeHelper<FString> FactoryDependencies;
-
 	/**
 	 * This tracks the IDs of asset nodes which are the target of factories
 	 */
diff --git a/Engine/Source/Runtime/Experimental/Interchange/Core/Public/Nodes/InterchangeBaseNodeContainer.h b/Engine/Source/Runtime/Experimental/Interchange/Core/Public/Nodes/InterchangeBaseNodeContainer.h
--- a/Engine/Source/Runtime/Experimental/Interchange/Core/Public/Nodes/InterchangeBaseNodeContainer.h
+++ b/Engine/Source/Runtime/Experimental/Interchange/Core/Public/Nodes/InterchangeBaseNodeContainer.h
@@ -4,7 +4,7 @@
 
 #include "CoreMinimal.h"
 #include "Misc/FileHelper.h"
-#include "Nodes/InterchangeBaseNode.h"
+#include "Nodes/InterchangeFactoryBaseNode.h"
 #include "Templates/UniquePtr.h"
 #include "UObject/Class.h"
 #include "UObject/Object.h"
@@ -43,9 +43,36 @@ public:
 	/** Unordered iteration of the all nodes */
 	void IterateNodes(TFunctionRef<void(const FString&, UInterchangeBaseNode*)> IterationLambda) const;
 
+	template <typename T>
+	void IterateNodesOfType(TFunctionRef<void(const FString&, T*)> IterationLambda) const
+	{
+		for (const TPair<FString, TObjectPtr<UInterchangeBaseNode>>& NodeKeyValue : Nodes)
+		{
+			if (T* Node = Cast<T>(NodeKeyValue.Value))
+			{
+				IterationLambda(NodeKeyValue.Key, Node);
+			}
+		}
+	}
+
 	/** Unordered iteration of the all nodes, but I can be stop early by returning true */
 	void BreakableIterateNodes(TFunctionRef<bool(const FString&, UInterchangeBaseNode*)> IterationLambda) const;
 
+	template <typename T>
+	void BreakableIterateNodesOfType(TFunctionRef<bool(const FString&, T*)> IterationLambda) const
+	{
+		for (const TPair<FString, TObjectPtr<UInterchangeBaseNode>>& NodeKeyValue : Nodes)
+		{
+			if (T* Node = Cast<T>(NodeKeyValue.Value))
+			{
+				if (IterationLambda(NodeKeyValue.Key, Node))
+				{
+					break;
+				}
+			}
+		}
+	}
+
 	/** Return all nodes that do not have any parent */
 	UFUNCTION(BlueprintCallable, Category = "Interchange | Node Container")
 	void GetRoots(TArray<FString>& RootNodes) const;
@@ -54,9 +81,13 @@ public:
 	UFUNCTION(BlueprintCallable, Category = "Interchange | Node Container")
 	void GetNodes(const UClass* ClassNode, TArray<FString>& OutNodes) const;
 
-	/** Get an node pointer */
+	/** Get a node pointer. Once added to the container, nodes are considered const */
+	UFUNCTION(BlueprintCallable, Category = "Interchange | Node Container")
+	const UInterchangeBaseNode* GetNode(const FString& NodeUniqueID) const;
+
+	/** Get a factory node pointer */
 	UFUNCTION(BlueprintCallable, Category = "Interchange | Node Container")
-	UInterchangeBaseNode* GetNode(const FString& NodeUniqueID) const;
+	UInterchangeFactoryBaseNode* GetFactoryNode(const FString& NodeUniqueID) const;
 
 	/** Set node ParentUid */
 	UFUNCTION(BlueprintCallable, Category = "Interchange | Node Container")
diff --git a/Engine/Source/Runtime/Experimental/Interchange/Core/Public/Nodes/InterchangeFactoryBaseNode.h b/Engine/Source/Runtime/Experimental/Interchange/Core/Public/Nodes/InterchangeFactoryBaseNode.h
--- a/Engine/Source/Runtime/Experimental/Interchange/Core/Public/Nodes/InterchangeFactoryBaseNode.h
+++ b/Engine/Source/Runtime/Experimental/Interchange/Core/Public/Nodes/InterchangeFactoryBaseNode.h
@@ -12,6 +12,20 @@
 
 #include "InterchangeFactoryBaseNode.generated.h"
 
+
+namespace UE
+{
+	namespace Interchange
+	{
+		struct INTERCHANGECORE_API FFactoryBaseNodeStaticData : public FBaseNodeStaticData
+		{
+			static const FString& FactoryDependenciesBaseKey();
+		};
+
+	} // namespace Interchange
+} // namespace UE
+
+
 /**
  * This struct is used to store and retrieve key value attributes. The attributes are store in a generic FAttributeStorage which serialize the value in a TArray64<uint8>
  * See UE::Interchange::EAttributeTypes to know the supported template types
@@ -23,7 +37,15 @@ class INTERCHANGECORE_API UInterchangeFactoryBaseNode : public UInterchangeBaseN
 	GENERATED_BODY()
 
 public:
+	UInterchangeFactoryBaseNode();
+
 	virtual FString GetKeyDisplayName(const UE::Interchange::FAttributeKey& NodeAttributeKey) const override;
+	virtual FString GetAttributeCategory(const UE::Interchange::FAttributeKey& NodeAttributeKey) const override;
+
+	/**
+	 * Return the UClass of the object we represent so we can find factory/writer
+	 */
+	virtual class UClass* GetObjectClass() const;
 
 	/**
 	 * Return the custom sub-path under PackageBasePath, where the assets will be created.
@@ -37,6 +59,49 @@ public:
 	UFUNCTION(BlueprintCallable, Category = "Interchange | Node")
 	bool SetCustomSubPath(const FString& AttributeValue);
 
+	/**
+	 * This function allow to retrieve the number of factory dependencies for this object.
+	 */
+	UFUNCTION(BlueprintCallable, Category = "Interchange | Node")
+	int32 GetFactoryDependenciesCount() const;
+
+	/**
+	 * This function allow to retrieve the dependency for this object.
+	 */
+	UFUNCTION(BlueprintCallable, Category = "Interchange | Node")
+	void GetFactoryDependencies(TArray<FString>& OutDependencies ) const;
+
+	/**
+	 * This function allow to retrieve one dependency for this object.
+	 */
+	UFUNCTION(BlueprintCallable, Category = "Interchange | Node")
+	void GetFactoryDependency(const int32 Index, FString& OutDependency) const;
+
+	/**
+	 * Add one dependency to this object.
+	 */
+	UFUNCTION(BlueprintCallable, Category = "Interchange | Node")
+	bool AddFactoryDependencyUid(const FString& DependencyUid);
+
+	/**
+	 * Remove one dependency from this object.
+	 */
+	UFUNCTION(BlueprintCallable, Category = "Interchange | Node")
+	bool RemoveFactoryDependencyUid(const FString& DependencyUid);
+
+
+	UPROPERTY()
+	mutable FSoftObjectPath ReferenceObject;
+
+protected:
+	/**
+	 * Those dependencies are use by the interchange parsing task to make sure the asset are created in the correct order.
+	 * Example: Mesh factory node will have dependencies on material factory node
+	 *          Material factory node will have dependencies on texture factory node
+	 */
+	UE::Interchange::TArrayAttributeHelper<FString> FactoryDependencies;
+
+
 private:
 
 	const UE::Interchange::FAttributeKey Macro_CustomSubPathKey = UE::Interchange::FAttributeKey(TEXT("SubPath"));
diff --git a/Engine/Source/Runtime/Experimental/Interchange/Engine/Private/InterchangeImportCommon.cpp b/Engine/Source/Runtime/Experimental/Interchange/Engine/Private/InterchangeImportCommon.cpp
--- a/Engine/Source/Runtime/Experimental/Interchange/Engine/Private/InterchangeImportCommon.cpp
+++ b/Engine/Source/Runtime/Experimental/Interchange/Engine/Private/InterchangeImportCommon.cpp
@@ -7,6 +7,7 @@
 #include "InterchangePipelineBase.h"
 #include "InterchangeSourceData.h"
 #include "Nodes/InterchangeBaseNodeContainer.h"
+#include "Nodes/InterchangeFactoryBaseNode.h"
 #include "UObject/Object.h"
 #include "EditorFramework/AssetImportData.h"
 #include "HAL/FileManager.h"
@@ -184,9 +185,9 @@ namespace UE
 
 
 		void FFactoryCommon::ApplyReimportStrategyToAsset(UObject* Asset
-										  , UInterchangeBaseNode* PreviousAssetNode
-										  , UInterchangeBaseNode* CurrentAssetNode
-										  , UInterchangeBaseNode* PipelineAssetNode)
+										  , UInterchangeFactoryBaseNode* PreviousAssetNode
+										  , UInterchangeFactoryBaseNode* CurrentAssetNode
+										  , UInterchangeFactoryBaseNode* PipelineAssetNode)
 		{
 			if (!ensure(PreviousAssetNode) || !ensure(PipelineAssetNode) || !ensure(CurrentAssetNode))
 			{
diff --git a/Engine/Source/Runtime/Experimental/Interchange/Engine/Private/Tasks/InterchangeTaskCompletion.cpp b/Engine/Source/Runtime/Experimental/Interchange/Engine/Private/Tasks/InterchangeTaskCompletion.cpp
--- a/Engine/Source/Runtime/Experimental/Interchange/Engine/Private/Tasks/InterchangeTaskCompletion.cpp
+++ b/Engine/Source/Runtime/Experimental/Interchange/Engine/Private/Tasks/InterchangeTaskCompletion.cpp
@@ -10,6 +10,7 @@
 #include "InterchangeFactoryBase.h"
 #include "InterchangeManager.h"
 #include "InterchangeResultsContainer.h"
+#include "Nodes/InterchangeFactoryBaseNode.h"
 #include "Stats/Stats.h"
 #include "Templates/SharedPointer.h"
 #include "UObject/GarbageCollection.h"
diff --git a/Engine/Source/Runtime/Experimental/Interchange/Engine/Private/Tasks/InterchangeTaskCreateAsset.cpp b/Engine/Source/Runtime/Experimental/Interchange/Engine/Private/Tasks/InterchangeTaskCreateAsset.cpp
--- a/Engine/Source/Runtime/Experimental/Interchange/Engine/Private/Tasks/InterchangeTaskCreateAsset.cpp
+++ b/Engine/Source/Runtime/Experimental/Interchange/Engine/Private/Tasks/InterchangeTaskCreateAsset.cpp
@@ -18,8 +18,6 @@
 #include "UObject/UObjectGlobals.h"
 #include "UObject/ObjectMacros.h"
 #include "UObject/WeakObjectPtrTemplates.h"
-#include "Nodes/InterchangeBaseNode.h"
-#include "Nodes/InterchangeFactoryBaseNode.h"
 
 namespace UE
 {
@@ -27,32 +25,29 @@ namespace UE
 	{
 		namespace Private
 		{
-			void InternalGetPackageName(const UE::Interchange::FImportAsyncHelper& AsyncHelper, const int32 SourceIndex, const FString& PackageBasePath, const UInterchangeBaseNode* Node, FString& OutPackageName, FString& OutAssetName)
+			void InternalGetPackageName(const UE::Interchange::FImportAsyncHelper& AsyncHelper, const int32 SourceIndex, const FString& PackageBasePath, const UInterchangeFactoryBaseNode* FactoryNode, FString& OutPackageName, FString& OutAssetName)
 			{
 				const UInterchangeSourceData* SourceData = AsyncHelper.SourceDatas[SourceIndex];
 				check(SourceData);
-				FString NodeDisplayName = Node->GetAssetName();
+				FString NodeDisplayName = FactoryNode->GetAssetName();
 
 				// Set the asset name and the package name
 				OutAssetName = NodeDisplayName;
 				SanitizeObjectName(OutAssetName);
 
-				FString SanitazedPackageBasePath = PackageBasePath;
-				SanitizeObjectPath(SanitazedPackageBasePath);
+				FString SanitizedPackageBasePath = PackageBasePath;
+				SanitizeObjectPath(SanitizedPackageBasePath);
 
 				FString SubPath;
-				if (const UInterchangeFactoryBaseNode* FactoryBaseNode = Cast<UInterchangeFactoryBaseNode>(Node))
+				if (FactoryNode->GetCustomSubPath(SubPath))
 				{
-					if(FactoryBaseNode->GetCustomSubPath(SubPath))
-					{
-						SanitizeObjectPath(SubPath);
-					}
+					SanitizeObjectPath(SubPath);
 				}
 
-				OutPackageName = FPaths::Combine(*SanitazedPackageBasePath, *SubPath, *OutAssetName);
+				OutPackageName = FPaths::Combine(*SanitizedPackageBasePath, *SubPath, *OutAssetName);
 			}
 
-			UObject* GetExistingObjectFromAssetImportData(TSharedPtr<UE::Interchange::FImportAsyncHelper, ESPMode::ThreadSafe> AsyncHelper, UInterchangeBaseNode* Node)
+			UObject* GetExistingObjectFromAssetImportData(TSharedPtr<UE::Interchange::FImportAsyncHelper, ESPMode::ThreadSafe> AsyncHelper, UInterchangeFactoryBaseNode* FactoryNode)
 			{
 				UInterchangeAssetImportData* OriginalAssetImportData = nullptr;
 				TArray<UObject*> SubObjects;
@@ -70,7 +65,7 @@ namespace UE
 				{
 					if (UInterchangeBaseNodeContainer* NodeContainer = OriginalAssetImportData->NodeContainer)
 					{
-						UClass* FactoryNodeClass = Node->GetClass();
+						UClass* FactoryNodeClass = FactoryNode->GetClass();
 						UInterchangeBaseNode* SelectedOriginalNode = nullptr;
 						NodeContainer->BreakableIterateNodes([FactoryNodeClass, &SelectedOriginalNode](const FString&, UInterchangeBaseNode* OriginalNode)
 							{
@@ -120,7 +115,7 @@ void UE::Interchange::FTaskCreatePackage::DoTask(ENamedThreads::Type CurrentThre
 
 	{
 		FScopeLock Lock(&AsyncHelper->CreatedFactoriesLock);
-		AsyncHelper->CreatedFactories.Add(Node->GetUniqueID(), Factory);
+		AsyncHelper->CreatedFactories.Add(FactoryNode->GetUniqueID(), Factory);
 	}
 
 	UPackage* Pkg = nullptr;
@@ -129,7 +124,7 @@ void UE::Interchange::FTaskCreatePackage::DoTask(ENamedThreads::Type CurrentThre
 	//If we do a reimport no need to create a package
 	if (AsyncHelper->TaskData.ReimportObject)
 	{
-		Private::InternalGetPackageName(*AsyncHelper, SourceIndex, PackageBasePath, Node, PackageName, AssetName);
+		Private::InternalGetPackageName(*AsyncHelper, SourceIndex, PackageBasePath, FactoryNode, PackageName, AssetName);
 		UPackage* ResultFindPackage = FindPackage(nullptr, *PackageName);
 		UObject* FindOuter = (ResultFindPackage == nullptr ? ANY_PACKAGE : ResultFindPackage);
 		UObject* ExistingObject = FindObject<UObject>(FindOuter, *AssetName);
@@ -137,7 +132,7 @@ void UE::Interchange::FTaskCreatePackage::DoTask(ENamedThreads::Type CurrentThre
 		if (ExistingObject != AsyncHelper->TaskData.ReimportObject)
 		{
 			// Try to see if we can do a mapping from the source data (we should revisit this for the MVP)
-			ExistingObject = UE::Interchange::Private::GetExistingObjectFromAssetImportData(AsyncHelper, Node);
+			ExistingObject = UE::Interchange::Private::GetExistingObjectFromAssetImportData(AsyncHelper, FactoryNode);
 		}
 		
 
@@ -151,7 +146,7 @@ void UE::Interchange::FTaskCreatePackage::DoTask(ENamedThreads::Type CurrentThre
 			UInterchangeResultError_Generic* Message = Factory->AddMessage<UInterchangeResultError_Generic>();
 			Message->SourceAssetName = AsyncHelper->SourceDatas[SourceIndex]->GetFilename();
 			Message->DestinationAssetName = AssetName;
-			Message->AssetType = Node->GetObjectClass();
+			Message->AssetType = FactoryNode->GetObjectClass();
 			Message->Text = NSLOCTEXT("Interchange", "CannotFindPackageDuringReimport", "Cannot find an existing package.");
 
 			//Skip this asset
@@ -160,14 +155,14 @@ void UE::Interchange::FTaskCreatePackage::DoTask(ENamedThreads::Type CurrentThre
 	}
 	else
 	{
-		Private::InternalGetPackageName(*AsyncHelper, SourceIndex, PackageBasePath, Node, PackageName, AssetName);
+		Private::InternalGetPackageName(*AsyncHelper, SourceIndex, PackageBasePath, FactoryNode, PackageName, AssetName);
 		// We can not create assets that share the name of a map file in the same location
 		if (UE::Interchange::FPackageUtils::IsMapPackageAsset(PackageName))
 		{
 			UInterchangeResultError_Generic* Message = Factory->AddMessage<UInterchangeResultError_Generic>();
 			Message->SourceAssetName = AsyncHelper->SourceDatas[SourceIndex]->GetFilename();
 			Message->DestinationAssetName = AssetName;
-			Message->AssetType = Node->GetObjectClass();
+			Message->AssetType = FactoryNode->GetObjectClass();
 			Message->Text = NSLOCTEXT("Interchange", "MapExistsWithSameName", "You cannot create an asset with this name, as there is already a map file with the same name in this folder.");
 
 			//Skip this asset
@@ -180,7 +175,7 @@ void UE::Interchange::FTaskCreatePackage::DoTask(ENamedThreads::Type CurrentThre
 			UInterchangeResultError_Generic* Message = Factory->AddMessage<UInterchangeResultError_Generic>();
 			Message->SourceAssetName = AsyncHelper->SourceDatas[SourceIndex]->GetFilename();
 			Message->DestinationAssetName = AssetName;
-			Message->AssetType = Node->GetObjectClass();
+			Message->AssetType = FactoryNode->GetObjectClass();
 			Message->Text = FText::Format(NSLOCTEXT("Interchange", "CouldntCreatePackage", "It was not possible to create a package named '{0}'; the asset will not be imported."), FText::FromString(PackageName));
 
 			//Skip this asset
@@ -190,7 +185,7 @@ void UE::Interchange::FTaskCreatePackage::DoTask(ENamedThreads::Type CurrentThre
 		//Import Asset describe by the node
 		UInterchangeFactoryBase::FCreateAssetParams CreateAssetParams;
 		CreateAssetParams.AssetName = AssetName;
-		CreateAssetParams.AssetNode = Node;
+		CreateAssetParams.AssetNode = FactoryNode;
 		CreateAssetParams.Parent = Pkg;
 		CreateAssetParams.SourceData = AsyncHelper->SourceDatas[SourceIndex];
 		CreateAssetParams.Translator = AsyncHelper->Translators[SourceIndex];
@@ -214,9 +209,9 @@ void UE::Interchange::FTaskCreatePackage::DoTask(ENamedThreads::Type CurrentThre
 			UE::Interchange::FImportAsyncHelper::FImportedObjectInfo& AssetInfo = ImportedInfos.AddDefaulted_GetRef();
 			AssetInfo.ImportedObject = NodeAsset;
 			AssetInfo.Factory = Factory;
-			AssetInfo.FactoryNode = Node;
+			AssetInfo.FactoryNode = FactoryNode;
 			AssetInfo.bIsReimport = bool(AsyncHelper->TaskData.ReimportObject);
-			Node->ReferenceObject = FSoftObjectPath(NodeAsset);
+			FactoryNode->ReferenceObject = FSoftObjectPath(NodeAsset);
 		}
 	}
 
@@ -246,13 +241,13 @@ void UE::Interchange::FTaskCreateAsset::DoTask(ENamedThreads::Type CurrentThread
 	UInterchangeFactoryBase* Factory = nullptr;
 	{
 		FScopeLock Lock(&AsyncHelper->CreatedFactoriesLock);
-		Factory = AsyncHelper->CreatedFactories.FindChecked(Node->GetUniqueID());
+		Factory = AsyncHelper->CreatedFactories.FindChecked(FactoryNode->GetUniqueID());
 	}
 
 	UPackage* Pkg = nullptr;
 	FString PackageName;
 	FString AssetName;
-	Private::InternalGetPackageName(*AsyncHelper, SourceIndex, PackageBasePath, Node, PackageName, AssetName);
+	Private::InternalGetPackageName(*AsyncHelper, SourceIndex, PackageBasePath, FactoryNode, PackageName, AssetName);
 	bool bSkipAsset = false;
 	UObject* ExistingObject = nullptr;
 	if (AsyncHelper->TaskData.ReimportObject)
@@ -266,7 +261,7 @@ void UE::Interchange::FTaskCreateAsset::DoTask(ENamedThreads::Type CurrentThread
 		if(ExistingObject != AsyncHelper->TaskData.ReimportObject)
 		{
 			// Try to see if we can do a mapping from the source data (we should revisit this for the MVP)
-			ExistingObject = UE::Interchange::Private::GetExistingObjectFromAssetImportData(AsyncHelper, Node);
+			ExistingObject = UE::Interchange::Private::GetExistingObjectFromAssetImportData(AsyncHelper, FactoryNode);
 		}
 
 		bSkipAsset = !ExistingObject || ExistingObject != AsyncHelper->TaskData.ReimportObject;
@@ -292,7 +287,7 @@ void UE::Interchange::FTaskCreateAsset::DoTask(ENamedThreads::Type CurrentThread
 			UInterchangeResultError_Generic* Message = Factory->AddMessage<UInterchangeResultError_Generic>();
 			Message->SourceAssetName = AsyncHelper->SourceDatas[SourceIndex]->GetFilename();
 			Message->DestinationAssetName = AssetName;
-			Message->AssetType = Node->GetObjectClass();
+			Message->AssetType = FactoryNode->GetObjectClass();
 			Message->Text = NSLOCTEXT("Interchange", "BadPackage", "It was not possible to create the asset as its package was not created correctly.");
 
 			return;
@@ -302,7 +297,7 @@ void UE::Interchange::FTaskCreateAsset::DoTask(ENamedThreads::Type CurrentThread
 		{
 			UInterchangeResultError_Generic* Message = Factory->AddMessage<UInterchangeResultError_Generic>();
 			Message->DestinationAssetName = AssetName;
-			Message->AssetType = Node->GetObjectClass();
+			Message->AssetType = FactoryNode->GetObjectClass();
 			Message->Text = NSLOCTEXT("Interchange", "SourceDataOrTranslatorInvalid", "It was not possible to create the asset as its translator was not created correctly.");
 
 			return;
@@ -322,7 +317,7 @@ void UE::Interchange::FTaskCreateAsset::DoTask(ENamedThreads::Type CurrentThread
 		//Import Asset describe by the node
 		UInterchangeFactoryBase::FCreateAssetParams CreateAssetParams;
 		CreateAssetParams.AssetName = AssetName;
-		CreateAssetParams.AssetNode = Node;
+		CreateAssetParams.AssetNode = FactoryNode;
 		CreateAssetParams.Parent = Pkg;
 		CreateAssetParams.SourceData = AsyncHelper->SourceDatas[SourceIndex];
 		CreateAssetParams.Translator = Translator;
@@ -350,7 +345,7 @@ void UE::Interchange::FTaskCreateAsset::DoTask(ENamedThreads::Type CurrentThread
 				UE::Interchange::FImportAsyncHelper::FImportedObjectInfo& AssetInfo = ImportedInfos.AddDefaulted_GetRef();
 				AssetInfo.ImportedObject = NodeAsset;
 				AssetInfo.Factory = Factory;
-				AssetInfo.FactoryNode = Node;
+				AssetInfo.FactoryNode = FactoryNode;
 				AssetInfo.bIsReimport = bool(AsyncHelper->TaskData.ReimportObject);
 			}
 
@@ -361,7 +356,7 @@ void UE::Interchange::FTaskCreateAsset::DoTask(ENamedThreads::Type CurrentThread
 				if (!Result->InterchangeKey.IsEmpty() && (Result->DestinationAssetName.IsEmpty() || Result->AssetType == nullptr))
 				{
 					TArray<FString> TargetAssets;
-					Node->GetTargetNodeUids(TargetAssets);
+					FactoryNode->GetTargetNodeUids(TargetAssets);
 					if (TargetAssets.Contains(Result->InterchangeKey))
 					{
 						Result->DestinationAssetName = NodeAsset->GetPathName();
@@ -371,6 +366,6 @@ void UE::Interchange::FTaskCreateAsset::DoTask(ENamedThreads::Type CurrentThread
 			}
 		}
 
-		Node->ReferenceObject = FSoftObjectPath(NodeAsset);
+		FactoryNode->ReferenceObject = FSoftObjectPath(NodeAsset);
 	}
 }
\ No newline at end of file
diff --git a/Engine/Source/Runtime/Experimental/Interchange/Engine/Private/Tasks/InterchangeTaskCreateAsset.h b/Engine/Source/Runtime/Experimental/Interchange/Engine/Private/Tasks/InterchangeTaskCreateAsset.h
--- a/Engine/Source/Runtime/Experimental/Interchange/Engine/Private/Tasks/InterchangeTaskCreateAsset.h
+++ b/Engine/Source/Runtime/Experimental/Interchange/Engine/Private/Tasks/InterchangeTaskCreateAsset.h
@@ -8,7 +8,7 @@
 #include "InterchangeManager.h"
 #include "Stats/Stats.h"
 #include "UObject/WeakObjectPtrTemplates.h"
-#include "Nodes/InterchangeBaseNode.h"
+#include "Nodes/InterchangeFactoryBaseNode.h"
 
 namespace UE
 {
@@ -24,18 +24,18 @@ namespace UE
 			FString PackageBasePath;
 			int32 SourceIndex;
 			TWeakPtr<FImportAsyncHelper, ESPMode::ThreadSafe> WeakAsyncHelper;
-			UInterchangeBaseNode* Node;
+			UInterchangeFactoryBaseNode* FactoryNode;
 			const UClass* FactoryClass;
 
 		public:
-			FTaskCreatePackage(const FString& InPackageBasePath, const int32 InSourceIndex, TWeakPtr<FImportAsyncHelper, ESPMode::ThreadSafe> InAsyncHelper, UInterchangeBaseNode* InNode, const UClass* InFactoryClass)
+			FTaskCreatePackage(const FString& InPackageBasePath, const int32 InSourceIndex, TWeakPtr<FImportAsyncHelper, ESPMode::ThreadSafe> InAsyncHelper, UInterchangeFactoryBaseNode* InFactoryNode, const UClass* InFactoryClass)
 				: PackageBasePath(InPackageBasePath)
 				, SourceIndex(InSourceIndex)
 				, WeakAsyncHelper(InAsyncHelper)
-				, Node(InNode)
+				, FactoryNode(InFactoryNode)
 				, FactoryClass(InFactoryClass)
 			{
-				check(Node);
+				check(FactoryNode);
 				check(FactoryClass);
 			}
 
@@ -69,18 +69,18 @@ namespace UE
 			FString PackageBasePath;
 			int32 SourceIndex;
 			TWeakPtr<FImportAsyncHelper, ESPMode::ThreadSafe> WeakAsyncHelper;
-			UInterchangeBaseNode* Node;
+			UInterchangeFactoryBaseNode* FactoryNode;
 			bool bCanRunOnAnyThread;
 
 		public:
-			FTaskCreateAsset(const FString& InPackageBasePath, const int32 InSourceIndex, TWeakPtr<FImportAsyncHelper, ESPMode::ThreadSafe> InAsyncHelper, UInterchangeBaseNode* InNode, bool bInCanRunOnAnyThread)
+			FTaskCreateAsset(const FString& InPackageBasePath, const int32 InSourceIndex, TWeakPtr<FImportAsyncHelper, ESPMode::ThreadSafe> InAsyncHelper, UInterchangeFactoryBaseNode* InFactoryNode, bool bInCanRunOnAnyThread)
 				: PackageBasePath(InPackageBasePath)
 				, SourceIndex(InSourceIndex)
 				, WeakAsyncHelper(InAsyncHelper)
-				, Node(InNode)
+				, FactoryNode(InFactoryNode)
 				, bCanRunOnAnyThread(bInCanRunOnAnyThread)
 			{
-				check(Node);
+				check(FactoryNode);
 			}
 
 			ENamedThreads::Type GetDesiredThread() const
diff --git a/Engine/Source/Runtime/Experimental/Interchange/Engine/Private/Tasks/InterchangeTaskCreateSceneObjects.cpp b/Engine/Source/Runtime/Experimental/Interchange/Engine/Private/Tasks/InterchangeTaskCreateSceneObjects.cpp
--- a/Engine/Source/Runtime/Experimental/Interchange/Engine/Private/Tasks/InterchangeTaskCreateSceneObjects.cpp
+++ b/Engine/Source/Runtime/Experimental/Interchange/Engine/Private/Tasks/InterchangeTaskCreateSceneObjects.cpp
@@ -11,7 +11,7 @@
 #include "InterchangeSourceData.h"
 #include "InterchangeTranslatorBase.h"
 #include "Misc/Paths.h"
-#include "Nodes/InterchangeBaseNode.h"
+#include "Nodes/InterchangeFactoryBaseNode.h"
 #include "PackageUtils/PackageUtils.h"
 #include "Stats/Stats.h"
 #include "Templates/SharedPointer.h"
@@ -21,11 +21,11 @@
 #include "UObject/UObjectGlobals.h"
 #include "UObject/WeakObjectPtrTemplates.h"
 
-UE::Interchange::FTaskCreateSceneObjects::FTaskCreateSceneObjects(const FString& InPackageBasePath, const int32 InSourceIndex, TWeakPtr<FImportAsyncHelper> InAsyncHelper, TArrayView<UInterchangeBaseNode*> InNodes, const UClass* InFactoryClass)
+UE::Interchange::FTaskCreateSceneObjects::FTaskCreateSceneObjects(const FString& InPackageBasePath, const int32 InSourceIndex, TWeakPtr<FImportAsyncHelper> InAsyncHelper, TArrayView<UInterchangeFactoryBaseNode*> InFactoryNodes, const UClass* InFactoryClass)
 	: PackageBasePath(InPackageBasePath)
 	, SourceIndex(InSourceIndex)
 	, WeakAsyncHelper(InAsyncHelper)
-	, Nodes(InNodes)
+	, FactoryNodes(InFactoryNodes)
 	, FactoryClass(InFactoryClass)
 {
 	check(FactoryClass);
@@ -52,21 +52,21 @@ void UE::Interchange::FTaskCreateSceneObjects::DoTask(ENamedThreads::Type Curren
 		return;
 	}
 
-	for (UInterchangeBaseNode* Node : Nodes)
+	for (UInterchangeFactoryBaseNode* FactoryNode : FactoryNodes)
 	{
 		UInterchangeFactoryBase* Factory = NewObject<UInterchangeFactoryBase>(GetTransientPackage(), FactoryClass);
 		Factory->SetResultsContainer(AsyncHelper->AssetImportResult->GetResults());
 		{
 			FScopeLock Lock(&AsyncHelper->CreatedFactoriesLock);
-			AsyncHelper->CreatedFactories.Add(Node->GetUniqueID(), Factory);
+			AsyncHelper->CreatedFactories.Add(FactoryNode->GetUniqueID(), Factory);
 		}
 
-		FString NodeDisplayName = Node->GetDisplayLabel();
+		FString NodeDisplayName = FactoryNode->GetDisplayLabel();
 		SanitizeObjectName(NodeDisplayName);
 
 		UInterchangeFactoryBase::FCreateSceneObjectsParams CreateSceneObjectsParams;
 		CreateSceneObjectsParams.ObjectName = NodeDisplayName;
-		CreateSceneObjectsParams.FactoryNode = Node;
+		CreateSceneObjectsParams.FactoryNode = FactoryNode;
 		CreateSceneObjectsParams.Level = GWorld->GetCurrentLevel();
 
 		if (AsyncHelper->BaseNodeContainers.IsValidIndex(SourceIndex))
@@ -89,10 +89,10 @@ void UE::Interchange::FTaskCreateSceneObjects::DoTask(ENamedThreads::Type Curren
 				UE::Interchange::FImportAsyncHelper::FImportedObjectInfo& ObjectInfo = ImportedInfos.AddDefaulted_GetRef();
 				ObjectInfo.ImportedObject = SceneObject;
 				ObjectInfo.Factory = Factory;
-				ObjectInfo.FactoryNode = Node;
+				ObjectInfo.FactoryNode = FactoryNode;
 			}
 
-			Node->ReferenceObject = FSoftObjectPath(SceneObject);
+			FactoryNode->ReferenceObject = FSoftObjectPath(SceneObject);
 		}
 	}
 }
\ No newline at end of file
diff --git a/Engine/Source/Runtime/Experimental/Interchange/Engine/Private/Tasks/InterchangeTaskCreateSceneObjects.h b/Engine/Source/Runtime/Experimental/Interchange/Engine/Private/Tasks/InterchangeTaskCreateSceneObjects.h
--- a/Engine/Source/Runtime/Experimental/Interchange/Engine/Private/Tasks/InterchangeTaskCreateSceneObjects.h
+++ b/Engine/Source/Runtime/Experimental/Interchange/Engine/Private/Tasks/InterchangeTaskCreateSceneObjects.h
@@ -9,6 +9,7 @@
 #include "Nodes/InterchangeBaseNode.h"
 
 class UInterchangeBaseNode;
+class UInterchangeFactoryBaseNode;
 
 namespace UE
 {
@@ -22,11 +23,11 @@ namespace UE
 			FString PackageBasePath;
 			int32 SourceIndex;
 			TWeakPtr<FImportAsyncHelper> WeakAsyncHelper;
-			TArray<UInterchangeBaseNode*> Nodes;
+			TArray<UInterchangeFactoryBaseNode*> FactoryNodes;
 			const UClass* FactoryClass;
 
 		public:
-			explicit FTaskCreateSceneObjects(const FString& InPackageBasePath, const int32 InSourceIndex, TWeakPtr<FImportAsyncHelper> InAsyncHelper, TArrayView<UInterchangeBaseNode*> InNodes, const UClass* InFactoryClass);
+			explicit FTaskCreateSceneObjects(const FString& InPackageBasePath, const int32 InSourceIndex, TWeakPtr<FImportAsyncHelper> InAsyncHelper, TArrayView<UInterchangeFactoryBaseNode*> InNodes, const UClass* InFactoryClass);
 
 			ENamedThreads::Type GetDesiredThread()
 			{
diff --git a/Engine/Source/Runtime/Experimental/Interchange/Engine/Private/Tasks/InterchangeTaskParsing.cpp b/Engine/Source/Runtime/Experimental/Interchange/Engine/Private/Tasks/InterchangeTaskParsing.cpp
--- a/Engine/Source/Runtime/Experimental/Interchange/Engine/Private/Tasks/InterchangeTaskParsing.cpp
+++ b/Engine/Source/Runtime/Experimental/Interchange/Engine/Private/Tasks/InterchangeTaskParsing.cpp
@@ -35,7 +35,7 @@ struct FNodeDependencyCache
 {
 	const TSet<FString>& GetAccumulatedDependencies(const UInterchangeBaseNodeContainer* NodeContainer, const FString& NodeID)
 	{
-		if (TSet<FString>* DependenciesPtr = CachedDependencies.Find(NodeID))
+		if (const TSet<FString>* DependenciesPtr = CachedDependencies.Find(NodeID))
 		{
 			return *DependenciesPtr;
 		}
@@ -49,14 +49,14 @@ private:
 
 	void AccumulateDependencies(const UInterchangeBaseNodeContainer* NodeContainer, const FString& NodeID, TSet<FString>& OutDependenciesSet)
 	{
-		const UInterchangeBaseNode* Node = NodeContainer->GetNode(NodeID);
-		if (!Node)
+		const UInterchangeFactoryBaseNode* FactoryNode = NodeContainer->GetFactoryNode(NodeID);
+		if (!FactoryNode)
 		{
 			return;
 		}
 
 		TArray<FString> FactoryDependencies;
-		Node->GetFactoryDependencies(FactoryDependencies);
+		FactoryNode->GetFactoryDependencies(FactoryDependencies);
 		for (const FString& DependencyID : FactoryDependencies)
 		{
 			bool bAlreadyInSet = false;
@@ -93,7 +93,7 @@ void UE::Interchange::FTaskParsing::DoTask(ENamedThreads::Type CurrentThread, co
 		FGraphEventArray Prerequisites;
 		const UClass* FactoryClass;
 
-		TArray<UInterchangeBaseNode*, TInlineAllocator<1>> Nodes; // For scenes, we can group multiple nodes into a single task as they are usually very light
+		TArray<UInterchangeFactoryBaseNode*, TInlineAllocator<1>> Nodes; // For scenes, we can group multiple nodes into a single task as they are usually very light
 	};
 
 	TArray<FTaskData> AssetTaskDatas;
@@ -117,19 +117,20 @@ void UE::Interchange::FTaskParsing::DoTask(ENamedThreads::Type CurrentThread, co
 				continue;
 			}
 			const bool bCanImportSceneNode = AsyncHelper->TaskData.ImportType == EImportType::ImportType_Scene;
-			BaseNodeContainer->IterateNodes([&](const FString& NodeUID, UInterchangeBaseNode* Node)
+			BaseNodeContainer->IterateNodesOfType<UInterchangeFactoryBaseNode>([&](const FString& NodeUID, UInterchangeFactoryBaseNode* FactoryNode)
 			{
-				if (!Node->IsEnabled())
+				if (!FactoryNode->IsEnabled())
 				{
 					//Do not call factory for a disabled node
 					return;
 				}
 
-				if (Node->GetObjectClass() != nullptr)
+				UClass* ObjectClass = FactoryNode->GetObjectClass();
+				if (ObjectClass != nullptr)
 				{
-					const UClass* RegisteredFactoryClass = InterchangeManager->GetRegisteredFactoryClass(Node->GetObjectClass());
+					const UClass* RegisteredFactoryClass = InterchangeManager->GetRegisteredFactoryClass(ObjectClass);
 
-					const bool bIsAsset = !(Node->GetObjectClass()->IsChildOf<AActor>() || Node->GetObjectClass()->IsChildOf<UActorComponent>());
+					const bool bIsAsset = !(FactoryNode->GetObjectClass()->IsChildOf<AActor>() || FactoryNode->GetObjectClass()->IsChildOf<UActorComponent>());
 
 					if (!RegisteredFactoryClass || (!bIsAsset && !bCanImportSceneNode))
 					{
@@ -138,10 +139,10 @@ void UE::Interchange::FTaskParsing::DoTask(ENamedThreads::Type CurrentThread, co
 					}
 
 					FTaskData& NodeTaskData = bIsAsset ? SourceAssetTaskDatas.AddDefaulted_GetRef() : SourceSceneTaskDatas.AddDefaulted_GetRef();
-					NodeTaskData.UniqueID = Node->GetUniqueID();
+					NodeTaskData.UniqueID = FactoryNode->GetUniqueID();
 					NodeTaskData.SourceIndex = SourceIndex;
-					NodeTaskData.Nodes.Add(Node);
-					Node->GetFactoryDependencies(NodeTaskData.Dependencies);
+					NodeTaskData.Nodes.Add(FactoryNode);
+					FactoryNode->GetFactoryDependencies(NodeTaskData.Dependencies);
 					NodeTaskData.FactoryClass = RegisteredFactoryClass;
 				}
 			});
@@ -214,7 +215,7 @@ void UE::Interchange::FTaskParsing::DoTask(ENamedThreads::Type CurrentThread, co
 
 			const int32 SourceIndex = TaskData.SourceIndex;
 			const UClass* const FactoryClass = TaskData.FactoryClass;
-			UInterchangeBaseNode* const FactoryNode = TaskData.Nodes[0];
+			UInterchangeFactoryBaseNode* const FactoryNode = TaskData.Nodes[0];
 			const bool bFactoryCanRunOnAnyThread = FactoryClass->GetDefaultObject<UInterchangeFactoryBase>()->CanExecuteOnAnyThread();
 
 			//Add create package task has a prerequisite of FTaskCreateAsset. Create package task is a game thread task
diff --git a/Engine/Source/Runtime/Experimental/Interchange/Engine/Private/Tasks/InterchangeTaskPipeline.cpp b/Engine/Source/Runtime/Experimental/Interchange/Engine/Private/Tasks/InterchangeTaskPipeline.cpp
--- a/Engine/Source/Runtime/Experimental/Interchange/Engine/Private/Tasks/InterchangeTaskPipeline.cpp
+++ b/Engine/Source/Runtime/Experimental/Interchange/Engine/Private/Tasks/InterchangeTaskPipeline.cpp
@@ -7,6 +7,7 @@
 #include "InterchangeManager.h"
 #include "InterchangePipelineBase.h"
 #include "Nodes/InterchangeBaseNodeContainer.h"
+#include "Nodes/InterchangeFactoryBaseNode.h"
 #include "Stats/Stats.h"
 #include "Templates/SharedPointer.h"
 #include "UObject/GarbageCollection.h"
diff --git a/Engine/Source/Runtime/Experimental/Interchange/Engine/Public/InterchangeImportCommon.h b/Engine/Source/Runtime/Experimental/Interchange/Engine/Public/InterchangeImportCommon.h
--- a/Engine/Source/Runtime/Experimental/Interchange/Engine/Public/InterchangeImportCommon.h
+++ b/Engine/Source/Runtime/Experimental/Interchange/Engine/Public/InterchangeImportCommon.h
@@ -7,6 +7,7 @@
 class UAssetImportData;
 class UInterchangeAssetImportData;
 class UInterchangeBaseNode;
+class UInterchangeFactoryBaseNode;
 class UInterchangeBaseNodeContainer;
 class UInterchangePipelineBase;
 class UInterchangeSourceData;
@@ -82,9 +83,9 @@ namespace UE
 			 * Apply the current strategy to the PipelineAssetNode
 			 */
 			static void ApplyReimportStrategyToAsset(UObject* Asset
-											  , UInterchangeBaseNode* PreviousAssetNode
-											  , UInterchangeBaseNode* CurrentAssetNode
-											  , UInterchangeBaseNode* PipelineAssetNode);
+											  , UInterchangeFactoryBaseNode* PreviousAssetNode
+											  , UInterchangeFactoryBaseNode* CurrentAssetNode
+											  , UInterchangeFactoryBaseNode* PipelineAssetNode);
 		};
 	}
 }
diff --git a/Engine/Source/Runtime/Experimental/Interchange/Engine/Public/InterchangeManager.h b/Engine/Source/Runtime/Experimental/Interchange/Engine/Public/InterchangeManager.h
--- a/Engine/Source/Runtime/Experimental/Interchange/Engine/Public/InterchangeManager.h
+++ b/Engine/Source/Runtime/Experimental/Interchange/Engine/Public/InterchangeManager.h
@@ -28,6 +28,7 @@
 
 class FAsyncTaskNotification;
 class UInterchangeFactoryBase;
+class UInterchangeFactoryBaseNode;
 class UInterchangePipelineBase;
 
 /** Some utilities delegates for the automation of interchange */
@@ -206,7 +207,7 @@ namespace UE
 			{
 				UObject* ImportedObject = nullptr; // The object that was imported
 				UInterchangeFactoryBase* Factory = nullptr; //The factory that created the imported object
-				UInterchangeBaseNode* FactoryNode = nullptr; //The node that describes the object
+				UInterchangeFactoryBaseNode* FactoryNode = nullptr; //The node that describes the object
 				bool bIsReimport;
 			};
 
