DatasmithInterchange - Adding support for DatasmithAreaLight actor.

In addition of the expected TranslatedNode/FactoryNode/Factory classes, we are now also using our dedicated DatasmithActorPipeline which derives from the generic level pipeline.
The generic level pipeline was refactored to allow for factory node creation override. That way it is now possible to support additional actor types that are not directly known by the generic level pipeline.

diff --git a/Engine/Plugins/Experimental/Interchange/Runtime/Source/Import/Public/Scene/InterchangeActorHelper.h b/Engine/Plugins/Experimental/Interchange/Runtime/Source/Import/Public/Scene/InterchangeActorHelper.h
--- a/Engine/Plugins/Experimental/Interchange/Runtime/Source/Import/Public/Scene/InterchangeActorHelper.h
+++ b/Engine/Plugins/Experimental/Interchange/Runtime/Source/Import/Public/Scene/InterchangeActorHelper.h
@@ -13,16 +13,16 @@ namespace UE::Interchange::ActorHelper
 	/**
 	 * Returns the parent actor of the actor described by the FactoryNode using the node hierarchy.
 	 */
-	AActor* GetSpawnedParentActor(const UInterchangeBaseNodeContainer* NodeContainer, const UInterchangeActorFactoryNode* FactoryNode);
+	INTERCHANGEIMPORT_API AActor* GetSpawnedParentActor(const UInterchangeBaseNodeContainer* NodeContainer, const UInterchangeActorFactoryNode* FactoryNode);
 
 	/**
 	 * Spawns and return a AAActor using the given SceneObjectsParam.
 	 * The actor's parent will be determine by the node hierarchy and its type by the factory node's ObjectClass.
 	 */
-	AActor* SpawnFactoryActor(const UInterchangeFactoryBase::FCreateSceneObjectsParams& CreateSceneObjectsParams);
+	INTERCHANGEIMPORT_API AActor* SpawnFactoryActor(const UInterchangeFactoryBase::FCreateSceneObjectsParams& CreateSceneObjectsParams);
 	
 	/**
 	 * Returns the factory node of the asset instanced by ActorFactoryNode.
 	 */
-	const UInterchangeFactoryBaseNode* FindAssetInstanceFactoryNode(const UInterchangeBaseNodeContainer* NodeContainer, const UInterchangeFactoryBaseNode* ActorFactoryNode);
+	INTERCHANGEIMPORT_API const UInterchangeFactoryBaseNode* FindAssetInstanceFactoryNode(const UInterchangeBaseNodeContainer* NodeContainer, const UInterchangeFactoryBaseNode* ActorFactoryNode);
 }
diff --git a/Engine/Plugins/Experimental/Interchange/Runtime/Source/Pipelines/Private/InterchangeGenericScenesPipeline.cpp b/Engine/Plugins/Experimental/Interchange/Runtime/Source/Pipelines/Private/InterchangeGenericScenesPipeline.cpp
--- a/Engine/Plugins/Experimental/Interchange/Runtime/Source/Pipelines/Private/InterchangeGenericScenesPipeline.cpp
+++ b/Engine/Plugins/Experimental/Interchange/Runtime/Source/Pipelines/Private/InterchangeGenericScenesPipeline.cpp
@@ -58,19 +58,18 @@ void UInterchangeGenericLevelPipeline::ExecutePreImportPipeline(UInterchangeBase
 					continue;
 				}
 			}
-			CreateActorFactoryNode(InBaseNodeContainer, SceneNode, InBaseNodeContainer);
+			ExecuteSceneNodePreImport(InBaseNodeContainer, SceneNode, InBaseNodeContainer);
 		}
 	}
 }
 
-void UInterchangeGenericLevelPipeline::CreateActorFactoryNode(UInterchangeBaseNodeContainer* InBaseNodeContainer, const UInterchangeSceneNode* SceneNode, UInterchangeBaseNodeContainer* FactoryNodeContainer)
+void UInterchangeGenericLevelPipeline::ExecuteSceneNodePreImport(UInterchangeBaseNodeContainer* InBaseNodeContainer, const UInterchangeSceneNode* SceneNode, UInterchangeBaseNodeContainer* FactoryNodeContainer)
 {
 	if (!SceneNode)
 	{
 		return;
 	}
 
-	UInterchangeActorFactoryNode* ActorFactoryNode = nullptr;
 	const UInterchangeBaseNode* TranslatedAssetNode = nullptr;
 
 	FString AssetInstanceUid;
@@ -79,14 +78,7 @@ void UInterchangeGenericLevelPipeline::CreateActorFactoryNode(UInterchangeBaseNo
 		TranslatedAssetNode = FactoryNodeContainer->GetNode(AssetInstanceUid);
 	}
 
-	if (TranslatedAssetNode && TranslatedAssetNode->IsA<UInterchangeCameraNode>())
-	{
-		ActorFactoryNode = NewObject<UInterchangeCineCameraFactoryNode>(FactoryNodeContainer, NAME_None);
-	}
-	else
-	{
-		ActorFactoryNode = NewObject<UInterchangeActorFactoryNode>(FactoryNodeContainer, NAME_None);
-	}
+	UInterchangeActorFactoryNode* ActorFactoryNode = CreateActorFactoryNode(SceneNode, TranslatedAssetNode, FactoryNodeContainer);
 
 	if (!ensure(ActorFactoryNode))
 	{
@@ -114,70 +106,86 @@ void UInterchangeGenericLevelPipeline::CreateActorFactoryNode(UInterchangeBaseNo
 
 	if (TranslatedAssetNode)
 	{
-		if (const UInterchangeMeshNode* MeshNode = Cast<UInterchangeMeshNode>(TranslatedAssetNode))
+		SetUpFactoryNode(ActorFactoryNode, TranslatedAssetNode, FactoryNodeContainer);
+	}
+
+	FactoryNodeContainer->AddNode(ActorFactoryNode);
+}
+
+UInterchangeActorFactoryNode* UInterchangeGenericLevelPipeline::CreateActorFactoryNode(const UInterchangeSceneNode* SceneNode, const UInterchangeBaseNode* TranslatedAssetNode, UInterchangeBaseNodeContainer* FactoryNodeContainer) const
+{
+	if (TranslatedAssetNode && TranslatedAssetNode->IsA<UInterchangeCameraNode>())
+	{
+		return NewObject<UInterchangeCineCameraFactoryNode>(FactoryNodeContainer, NAME_None);
+	}
+	else
+	{
+		return NewObject<UInterchangeActorFactoryNode>(FactoryNodeContainer, NAME_None);
+	}
+}
+
+void UInterchangeGenericLevelPipeline::SetUpFactoryNode(UInterchangeActorFactoryNode* ActorFactoryNode, const UInterchangeBaseNode* TranslatedAssetNode, UInterchangeBaseNodeContainer* FactoryNodeContainer) const
+{
+	if (const UInterchangeMeshNode* MeshNode = Cast<UInterchangeMeshNode>(TranslatedAssetNode))
+	{
+		if (MeshNode->IsSkinnedMesh())
 		{
-			if (MeshNode->IsSkinnedMesh())
-			{
-				ActorFactoryNode->SetCustomActorClassName(ASkeletalMeshActor::StaticClass()->GetPathName());
-				ActorFactoryNode->SetCustomMobility(EComponentMobility::Movable);
-			}
-			else
-			{
-				ActorFactoryNode->SetCustomActorClassName(AStaticMeshActor::StaticClass()->GetPathName());
-			}
+			ActorFactoryNode->SetCustomActorClassName(ASkeletalMeshActor::StaticClass()->GetPathName());
+			ActorFactoryNode->SetCustomMobility(EComponentMobility::Movable);
 		}
-		else if (const UInterchangeLightNode* LightNode = Cast<UInterchangeLightNode>(TranslatedAssetNode))
+		else
 		{
-			//Test for spot before point since a spot light is a point light
-			if (LightNode->IsA<UInterchangeSpotLightNode>())
-			{
-				ActorFactoryNode->SetCustomActorClassName(ASpotLight::StaticClass()->GetPathName());
-			}
-			else if (LightNode->IsA<UInterchangePointLightNode>())
-			{
-				ActorFactoryNode->SetCustomActorClassName(APointLight::StaticClass()->GetPathName());
-			} 
-			else if (LightNode->IsA<UInterchangeRectLightNode>())
-			{
-				ActorFactoryNode->SetCustomActorClassName(ARectLight::StaticClass()->GetPathName());
-			}
-			else if (LightNode->IsA<UInterchangeDirectionalLightNode>())
+			ActorFactoryNode->SetCustomActorClassName(AStaticMeshActor::StaticClass()->GetPathName());
+		}
+	}
+	else if (const UInterchangeLightNode* LightNode = Cast<UInterchangeLightNode>(TranslatedAssetNode))
+	{
+		//Test for spot before point since a spot light is a point light
+		if (LightNode->IsA<UInterchangeSpotLightNode>())
+		{
+			ActorFactoryNode->SetCustomActorClassName(ASpotLight::StaticClass()->GetPathName());
+		}
+		else if (LightNode->IsA<UInterchangePointLightNode>())
+		{
+			ActorFactoryNode->SetCustomActorClassName(APointLight::StaticClass()->GetPathName());
+		}
+		else if (LightNode->IsA<UInterchangeRectLightNode>())
+		{
+			ActorFactoryNode->SetCustomActorClassName(ARectLight::StaticClass()->GetPathName());
+		}
+		else if (LightNode->IsA<UInterchangeDirectionalLightNode>())
+		{
+			ActorFactoryNode->SetCustomActorClassName(ADirectionalLight::StaticClass()->GetPathName());
+		}
+		else
+		{
+			ActorFactoryNode->SetCustomActorClassName(APointLight::StaticClass()->GetPathName());
+		}
+	}
+	else if (const UInterchangeCameraNode* CameraNode = Cast<UInterchangeCameraNode>(TranslatedAssetNode))
+	{
+		ActorFactoryNode->SetCustomActorClassName(ACineCameraActor::StaticClass()->GetPathName());
+		ActorFactoryNode->SetCustomMobility(EComponentMobility::Movable);
+
+		if (UInterchangeCineCameraFactoryNode* CineCameraFactoryNode = Cast<UInterchangeCineCameraFactoryNode>(ActorFactoryNode))
+		{
+			float FocalLength;
+			if (CameraNode->GetCustomFocalLength(FocalLength))
 			{
-				ActorFactoryNode->SetCustomActorClassName(ADirectionalLight::StaticClass()->GetPathName());
+				CineCameraFactoryNode->SetCustomFocalLength(FocalLength);
 			}
-			else
+
+			float SensorHeight;
+			if (CameraNode->GetCustomSensorHeight(SensorHeight))
 			{
-				ActorFactoryNode->SetCustomActorClassName(APointLight::StaticClass()->GetPathName());
+				CineCameraFactoryNode->SetCustomSensorHeight(SensorHeight);
 			}
-		}
-		else if (const UInterchangeCameraNode* CameraNode = Cast<UInterchangeCameraNode>(TranslatedAssetNode))
-		{
-			ActorFactoryNode->SetCustomActorClassName(ACineCameraActor::StaticClass()->GetPathName());
-			ActorFactoryNode->SetCustomMobility(EComponentMobility::Movable);
 
-			if (UInterchangeCineCameraFactoryNode* CineCameraFactoryNode = Cast<UInterchangeCineCameraFactoryNode>(ActorFactoryNode))
+			float SensorWidth;
+			if (CameraNode->GetCustomSensorWidth(SensorWidth))
 			{
-				float FocalLength;
-				if (CameraNode->GetCustomFocalLength(FocalLength))
-				{
-					CineCameraFactoryNode->SetCustomFocalLength(FocalLength);
-				}
-				
-				float SensorHeight;
-				if (CameraNode->GetCustomSensorHeight(SensorHeight))
-				{
-					CineCameraFactoryNode->SetCustomSensorHeight(SensorHeight);
-				}
-
-				float SensorWidth;
-				if (CameraNode->GetCustomSensorWidth(SensorWidth))
-				{
-					CineCameraFactoryNode->SetCustomSensorWidth(SensorWidth);
-				}
+				CineCameraFactoryNode->SetCustomSensorWidth(SensorWidth);
 			}
 		}
 	}
-
-	FactoryNodeContainer->AddNode(ActorFactoryNode);
 }
-
diff --git a/Engine/Plugins/Experimental/Interchange/Runtime/Source/Pipelines/Public/InterchangeGenericScenesPipeline.h b/Engine/Plugins/Experimental/Interchange/Runtime/Source/Pipelines/Public/InterchangeGenericScenesPipeline.h
--- a/Engine/Plugins/Experimental/Interchange/Runtime/Source/Pipelines/Public/InterchangeGenericScenesPipeline.h
+++ b/Engine/Plugins/Experimental/Interchange/Runtime/Source/Pipelines/Public/InterchangeGenericScenesPipeline.h
@@ -20,8 +20,6 @@ class INTERCHANGEPIPELINES_API UInterchangeGenericLevelPipeline : public UInterc
 {
 	GENERATED_BODY()
 
-public:
-
 protected:
 
 	virtual void ExecutePreImportPipeline(UInterchangeBaseNodeContainer* InBaseNodeContainer, const TArray<UInterchangeSourceData*>& InSourceDatas) override;
@@ -31,7 +29,20 @@ protected:
 		return true;
 	}
 
-	void CreateActorFactoryNode(UInterchangeBaseNodeContainer* InBaseNodeContainer, const UInterchangeSceneNode* SceneNode, UInterchangeBaseNodeContainer* FactoryNodeContainer);
+	/**
+	 * PreImport step called for each translated SceneNode.
+	 */
+	virtual void ExecuteSceneNodePreImport(UInterchangeBaseNodeContainer* InBaseNodeContainer, const UInterchangeSceneNode* SceneNode, UInterchangeBaseNodeContainer* FactoryNodeContainer);
+
+	/**
+	 * Return a new Actor Factory Node to be used for the given SceneNode.
+	 */
+	virtual UInterchangeActorFactoryNode* CreateActorFactoryNode(const UInterchangeSceneNode* SceneNode, const UInterchangeBaseNode* TranslatedAssetNode, UInterchangeBaseNodeContainer* FactoryNodeContainer) const;
+
+	/**
+	 * Use to set up the given factory node's attributes after its initialization.
+	 */
+	virtual void SetUpFactoryNode(UInterchangeActorFactoryNode* ActorFactoryNode, const UInterchangeBaseNode* TranslatedAssetNode, UInterchangeBaseNodeContainer* FactoryNodeContainer) const;
 	
 };
 
