Fix crash when we have unsupported material type when translating fbx file
#rnx

diff --git a/Engine/Plugins/Experimental/Interchange/Runtime/Source/Parsers/Fbx/Private/FbxMaterial.cpp b/Engine/Plugins/Experimental/Interchange/Runtime/Source/Parsers/Fbx/Private/FbxMaterial.cpp
--- a/Engine/Plugins/Experimental/Interchange/Runtime/Source/Parsers/Fbx/Private/FbxMaterial.cpp
+++ b/Engine/Plugins/Experimental/Interchange/Runtime/Source/Parsers/Fbx/Private/FbxMaterial.cpp
@@ -209,14 +209,14 @@ namespace UE
 
 						// Convert FBX property names to standard Interchange Phong parameter names when possible
 						{
-							if (PropertyName == FbxSurfaceMaterial::sDiffuse)
+							if (PropertyName == FbxSurfaceMaterial::sDiffuse && SurfaceMaterial->Is<FbxSurfaceLambert>())
 							{
 								const float Factor = (float)((FbxSurfaceLambert*)SurfaceMaterial)->DiffuseFactor.Get();
 								TVariant<FLinearColor, float> DefaultValue;
 								DefaultValue.Set<FLinearColor>(FLinearColor::Black);
 								ConvertPropertyToShaderNode(NodeContainer, ShaderGraphNode, CurrentProperty, Factor, Phong::Parameters::DiffuseColor, DefaultValue);
 							}
-							else if (PropertyName == FbxSurfaceMaterial::sSpecular)
+							else if (PropertyName == FbxSurfaceMaterial::sSpecular && SurfaceMaterial->Is<FbxSurfacePhong>())
 							{
 								const float Factor = (float)((FbxSurfacePhong*)SurfaceMaterial)->SpecularFactor.Get();
 								TVariant<FLinearColor, float> DefaultValue;
@@ -230,7 +230,7 @@ namespace UE
 								DefaultValue.Set<float>(0.2f);
 								ConvertPropertyToShaderNode(NodeContainer, ShaderGraphNode, CurrentProperty, Factor, Phong::Parameters::Shininess, DefaultValue);
 							}
-							else if (PropertyName == FbxSurfaceMaterial::sEmissive)
+							else if (PropertyName == FbxSurfaceMaterial::sEmissive && SurfaceMaterial->Is<FbxSurfaceLambert>())
 							{
 								const float Factor = (float)((FbxSurfaceLambert*)SurfaceMaterial)->EmissiveFactor.Get();
 								TVariant<FLinearColor, float> DefaultValue;
@@ -244,7 +244,7 @@ namespace UE
 								DefaultValue.Set<FLinearColor>(FLinearColor(FVector::UpVector));
 								ConvertPropertyToShaderNode(NodeContainer, ShaderGraphNode, CurrentProperty, Factor, Phong::Parameters::Normal, DefaultValue);
 							}
-							else if (PropertyName == FbxSurfaceMaterial::sTransparentColor)
+							else if (PropertyName == FbxSurfaceMaterial::sTransparentColor && SurfaceMaterial->Is<FbxSurfaceLambert>())
 							{
 								const float TransparencyFactor = (float)((FbxSurfaceLambert*)SurfaceMaterial)->TransparencyFactor.Get();
 								TVariant<FLinearColor, float> DefaultValue;
