-Adding MaterialX Standard Surface shader model into the UInterchangeGenericMaterialPipeline.
-Adding specific inputs for the Standard Surface to avoid clash with other inputs of shading models.
-Adding material functions for Autodesk's standard surface shader

diff --git a/Engine/Plugins/Experimental/Interchange/Runtime/Source/Nodes/Public/InterchangeMaterialDefinitions.h b/Engine/Plugins/Experimental/Interchange/Runtime/Source/Nodes/Public/InterchangeMaterialDefinitions.h
--- a/Engine/Plugins/Experimental/Interchange/Runtime/Source/Nodes/Public/InterchangeMaterialDefinitions.h
+++ b/Engine/Plugins/Experimental/Interchange/Runtime/Source/Nodes/Public/InterchangeMaterialDefinitions.h
@@ -202,6 +202,57 @@ namespace Interchange
 				const FName SheenRoughness = TEXT("SheenRoughness"); // Type: float
 			}
 		}
+
+		namespace StandardSurface
+		{
+			const FName Name = TEXT("standard_surface");
+
+			namespace Parameters
+			{
+				const FName Base = TEXT("base");
+				const FName BaseColor = TEXT("base_color");
+				const FName DiffuseRoughness = TEXT("diffuse_roughness");
+				const FName Metalness = TEXT("metalness");
+				const FName Specular = TEXT("specular");
+				const FName SpecularColor = TEXT("specular_color");
+				const FName SpecularRoughness = TEXT("specular_roughness");
+				const FName SpecularIOR = TEXT("specular_IOR");
+				const FName SpecularAnisotropy = TEXT("specular_anisotropy");
+				const FName SpecularRotation = TEXT("specular_rotation");
+				const FName Transmission = TEXT("transmission");
+				const FName TransmissionColor = TEXT("transmission_color");
+				const FName TransmissionDepth = TEXT("transmission_depth");
+				const FName TransmissionScatter = TEXT("transmission_scatter");
+				const FName TransmissionScatterAnisotropy = TEXT("transmission_scatter_anisotropy");
+				const FName TransmissionDispersion = TEXT("transmission_dispersion");
+				const FName TransmissionExtraRoughness = TEXT("transmission_extra_roughness");
+				const FName Subsurface = TEXT("subsurface");
+				const FName SubsurfaceColor = TEXT("subsurface_color");
+				const FName SubsurfaceRadius = TEXT("subsurface_radius");
+				const FName SubsurfaceScale = TEXT("subsurface_scale");
+				const FName SubsurfaceAnisotropy = TEXT("subsurface_anisotropy");
+				const FName Sheen = TEXT("sheen");
+				const FName SheenColor = TEXT("sheen_color");
+				const FName SheenRoughness = TEXT("sheen_roughness");
+				const FName Coat = TEXT("coat");
+				const FName CoatColor = TEXT("coat_color");
+				const FName CoatRoughness = TEXT("coat_roughness");
+				const FName CoatAnisotropy = TEXT("coat_anisotropy");
+				const FName CoatRotation = TEXT("coat_rotation");
+				const FName CoatIOR = TEXT("coat_IOR");
+				const FName CoatNormal = TEXT("coat_normal");
+				const FName CoatAffectColor = TEXT("coat_affect_color");
+				const FName CoatAffectRoughness = TEXT("coat_affect_roughness");
+				const FName ThinFilmThickness = TEXT("thin_film_thickness");
+				const FName ThinFilmIOR = TEXT("thin_film_IOR");
+				const FName Emission = TEXT("emission");
+				const FName EmissionColor = TEXT("emission_color");
+				const FName Opacity = TEXT("opacity");
+				const FName ThinWalled = TEXT("thin_walled");
+				const FName Normal = TEXT("normal");
+				const FName Tangent = TEXT("tangent");
+			}
+		}
 	}
 }
 }
diff --git a/Engine/Plugins/Experimental/Interchange/Runtime/Source/Pipelines/Private/InterchangeGenericMaterialPipeline.cpp b/Engine/Plugins/Experimental/Interchange/Runtime/Source/Pipelines/Private/InterchangeGenericMaterialPipeline.cpp
--- a/Engine/Plugins/Experimental/Interchange/Runtime/Source/Pipelines/Private/InterchangeGenericMaterialPipeline.cpp
+++ b/Engine/Plugins/Experimental/Interchange/Runtime/Source/Pipelines/Private/InterchangeGenericMaterialPipeline.cpp
@@ -35,6 +35,7 @@
 #include "Misc/Paths.h"
 #include "Nodes/InterchangeBaseNode.h"
 #include "Nodes/InterchangeBaseNodeContainer.h"
+#include "Templates/Function.h"
 #include "Templates/SubclassOf.h"
 #include "UObject/Object.h"
 #include "UObject/ObjectMacros.h"
@@ -77,7 +78,7 @@ void UInterchangeGenericMaterialPipeline::ExecutePreImportPipeline(UInterchangeB
 			if (UInterchangeMaterialFactoryNode* MaterialFactoryNode = CreateMaterialFactoryNode(ShaderGraphNode))
 			{
 				//By default we do not create the materials, every node with mesh attribute can enable them. So we wont create unused materials.
-				MaterialFactoryNode->SetEnabled(false);
+				MaterialFactoryNode->SetEnabled(true);
 			}
 		}
 	}
@@ -181,6 +182,20 @@ bool UInterchangeGenericMaterialPipeline::IsLambertModel(const UInterchangeShade
 	return bHasDiffuseInput;
 }
 
+bool UInterchangeGenericMaterialPipeline::IsStandardSurfaceModel(const UInterchangeShaderGraphNode* ShaderGraphNode) const
+{
+	using namespace UE::Interchange::Materials;
+	FString ShaderType;
+	ShaderGraphNode->GetCustomShaderType(ShaderType);
+
+	if(ShaderType == StandardSurface::Name.ToString())
+	{
+		return true;
+	}
+
+	return false;
+}
+
 bool UInterchangeGenericMaterialPipeline::HandlePhongModel(const UInterchangeShaderGraphNode* ShaderGraphNode, UInterchangeMaterialFactoryNode* MaterialFactoryNode)
 {
 	using namespace UE::Interchange::Materials::Phong;
@@ -382,6 +397,92 @@ bool UInterchangeGenericMaterialPipeline::HandlePBRModel(const UInterchangeShade
 	return bShadingModelHandled;
 }
 
+bool UInterchangeGenericMaterialPipeline::HandleStandardSurfaceModel(const UInterchangeShaderGraphNode* ShaderGraphNode, UInterchangeMaterialFactoryNode* MaterialFactoryNode)
+{
+	if(IsStandardSurfaceModel(ShaderGraphNode))
+	{
+		UInterchangeMaterialExpressionFactoryNode* FunctionCallExpression = NewObject<UInterchangeMaterialExpressionFactoryNode>(BaseNodeContainer, NAME_None);
+		FunctionCallExpression->SetCustomExpressionClassName(UMaterialExpressionMaterialFunctionCall::StaticClass()->GetName());
+		FString FunctionCallExpressionUid = MaterialFactoryNode->GetUniqueID() + TEXT("StandardSurface");
+		FunctionCallExpression->InitializeNode(FunctionCallExpressionUid, TEXT("StandardSurface"), EInterchangeNodeContainerType::FactoryData);
+
+		BaseNodeContainer->AddNode(FunctionCallExpression);
+		BaseNodeContainer->SetNodeParentUid(FunctionCallExpressionUid, MaterialFactoryNode->GetUniqueID());
+
+		const FName MaterialFunctionMemberName = GET_MEMBER_NAME_CHECKED(UMaterialExpressionMaterialFunctionCall, MaterialFunction);
+
+		FunctionCallExpression->AddStringAttribute(
+			MaterialFunctionMemberName.ToString(),
+			TEXT("MaterialFunction'/Interchange/Functions/MX_StandardSurface.MX_StandardSurface'"));
+		FunctionCallExpression->AddApplyAndFillDelegates<FString>(MaterialFunctionMemberName.ToString(), UMaterialExpressionMaterialFunctionCall::StaticClass(), MaterialFunctionMemberName);
+
+		using namespace UE::Interchange::Materials;
+
+		auto ConnectMaterialExpressionOutputToInput = [&](const FString & InputName, bool bLinearParsingValue, bool bNormalParsingValue)
+		{
+			TGuardValue<bool> ParsingForLinearInputGuard(bParsingForLinearInput, bLinearParsingValue);
+			TGuardValue<bool> ParsingForNormalInputGuard(bParsingForNormalInput, bNormalParsingValue);
+
+			TTuple<UInterchangeMaterialExpressionFactoryNode*, FString> ExpressionFactoryNode =
+				CreateMaterialExpressionForInput(MaterialFactoryNode, ShaderGraphNode, InputName, FunctionCallExpression->GetUniqueID());
+
+			if(ExpressionFactoryNode.Get<0>())
+			{
+				UInterchangeShaderPortsAPI::ConnectOuputToInput(FunctionCallExpression, InputName,
+																ExpressionFactoryNode.Get<0>()->GetUniqueID(), ExpressionFactoryNode.Get<1>());
+			}
+		};
+
+		using ArgumentsType = TTuple<FString, bool, bool>;
+		// TODO: We need to support all standard surface inputs
+		static ArgumentsType ArgumentsPerInputs[]
+		{
+			ArgumentsType{StandardSurface::Parameters::Base.ToString(), false, false},				// BaseWeight
+			ArgumentsType{StandardSurface::Parameters::BaseColor.ToString(), false, false},			// BaseColor
+			ArgumentsType{StandardSurface::Parameters::Metalness.ToString(), true, false},			// Metallic
+			ArgumentsType{StandardSurface::Parameters::Specular.ToString(), true, false},			// SpecularWeight
+			ArgumentsType{StandardSurface::Parameters::SpecularRoughness.ToString(), true, false},	// SpecularRoughness
+			ArgumentsType{StandardSurface::Parameters::Normal.ToString(), false, true}				// Normal
+		};
+
+		for(const ArgumentsType & Arguments : ArgumentsPerInputs)
+		{
+			const FString& InputName = Arguments.Get<0>();
+			bool bLinearParsingValue = Arguments.Get<1>();
+			bool bNormalParsingValue = Arguments.Get<2>();
+			ConnectMaterialExpressionOutputToInput(InputName, bLinearParsingValue, bNormalParsingValue);
+		}
+
+		MaterialFactoryNode->ConnectOutputToBaseColor(FunctionCallExpressionUid, PBR::Parameters::BaseColor.ToString());
+		MaterialFactoryNode->ConnectOutputToMetallic(FunctionCallExpressionUid, PBR::Parameters::Metallic.ToString());
+		MaterialFactoryNode->ConnectOutputToRoughness(FunctionCallExpressionUid, PBR::Parameters::Roughness.ToString());
+		MaterialFactoryNode->ConnectOutputToSpecular(FunctionCallExpressionUid, PBR::Parameters::Specular.ToString());
+		MaterialFactoryNode->ConnectOutputToNormal(FunctionCallExpressionUid, PBR::Parameters::Normal.ToString());
+
+		// We can't have all shading models at once, so we have to make a choice here
+		if(UInterchangeShaderPortsAPI::HasInput(ShaderGraphNode, StandardSurface::Parameters::Sheen))
+		{
+			MaterialFactoryNode->SetCustomShadingModel(MSM_Cloth);
+		}
+		else if(UInterchangeShaderPortsAPI::HasInput(ShaderGraphNode, StandardSurface::Parameters::Coat))
+		{
+			MaterialFactoryNode->SetCustomShadingModel(MSM_ClearCoat);
+		}
+		else if(UInterchangeShaderPortsAPI::HasInput(ShaderGraphNode, StandardSurface::Parameters::Subsurface))
+		{
+			MaterialFactoryNode->SetCustomShadingModel(MSM_Subsurface);
+		}
+		else
+		{
+			MaterialFactoryNode->SetCustomShadingModel(MSM_DefaultLit);
+		}
+
+		return true;
+	}
+
+	return false;
+}
+
 bool UInterchangeGenericMaterialPipeline::HandleClearCoat(const UInterchangeShaderGraphNode* ShaderGraphNode, UInterchangeMaterialFactoryNode* MaterialFactoryNode)
 {
 	using namespace UE::Interchange::Materials::ClearCoat;
@@ -1237,6 +1338,11 @@ UInterchangeMaterialFactoryNode* UInterchangeGenericMaterialPipeline::CreateMate
 {
 	UInterchangeMaterialFactoryNode* MaterialFactoryNode = Cast<UInterchangeMaterialFactoryNode>( CreateBaseMaterialFactoryNode(ShaderGraphNode, UInterchangeMaterialFactoryNode::StaticClass()) );
 
+	if(HandleStandardSurfaceModel(ShaderGraphNode, MaterialFactoryNode))
+	{
+		return MaterialFactoryNode;
+	}
+
 	if (!HandlePhongModel(ShaderGraphNode, MaterialFactoryNode))
 	{
 		HandleLambertModel(ShaderGraphNode, MaterialFactoryNode);
diff --git a/Engine/Plugins/Experimental/Interchange/Runtime/Source/Pipelines/Public/InterchangeGenericMaterialPipeline.h b/Engine/Plugins/Experimental/Interchange/Runtime/Source/Pipelines/Public/InterchangeGenericMaterialPipeline.h
--- a/Engine/Plugins/Experimental/Interchange/Runtime/Source/Pipelines/Public/InterchangeGenericMaterialPipeline.h
+++ b/Engine/Plugins/Experimental/Interchange/Runtime/Source/Pipelines/Public/InterchangeGenericMaterialPipeline.h
@@ -77,9 +77,13 @@ private:
 	/** True if the shader graph has a diffuse color input. */
 	bool IsLambertModel(const UInterchangeShaderGraphNode* ShaderGraphNode) const;
 
+	/** True if the shader graph has the standard surface's shader type name */
+	bool IsStandardSurfaceModel(const UInterchangeShaderGraphNode* ShaderGraphNode) const;
+ 
 	bool HandlePhongModel(const UInterchangeShaderGraphNode* ShaderGraphNode, UInterchangeMaterialFactoryNode* MaterialFactoryNode);
 	bool HandleLambertModel(const UInterchangeShaderGraphNode* ShaderGraphNode, UInterchangeMaterialFactoryNode* MaterialFactoryNode);
 	bool HandlePBRModel(const UInterchangeShaderGraphNode* ShaderGraphNode, UInterchangeMaterialFactoryNode* MaterialFactoryNode);
+	bool HandleStandardSurfaceModel(const UInterchangeShaderGraphNode* ShaderGraphNode, UInterchangeMaterialFactoryNode* MaterialFactoryNode);
 	bool HandleClearCoat(const UInterchangeShaderGraphNode* ShaderGraphNode, UInterchangeMaterialFactoryNode* MaterialFactoryNode);
 	bool HandleSheen(const UInterchangeShaderGraphNode* ShaderGraphNode, UInterchangeMaterialFactoryNode* MaterialFactoryNode);
 	bool HandleThinTranslucent(const UInterchangeShaderGraphNode* ShaderGraphNode, UInterchangeMaterialFactoryNode* MaterialFactoryNode);
