Reverting MaterialFactoryNode to SetEnabled(false)

diff --git a/Engine/Plugins/Experimental/Interchange/Runtime/Source/Pipelines/Private/InterchangeGenericMaterialPipeline.cpp b/Engine/Plugins/Experimental/Interchange/Runtime/Source/Pipelines/Private/InterchangeGenericMaterialPipeline.cpp
--- a/Engine/Plugins/Experimental/Interchange/Runtime/Source/Pipelines/Private/InterchangeGenericMaterialPipeline.cpp
+++ b/Engine/Plugins/Experimental/Interchange/Runtime/Source/Pipelines/Private/InterchangeGenericMaterialPipeline.cpp
@@ -78,7 +78,7 @@ void UInterchangeGenericMaterialPipeline::ExecutePreImportPipeline(UInterchangeB
 			if (UInterchangeMaterialFactoryNode* MaterialFactoryNode = CreateMaterialFactoryNode(ShaderGraphNode))
 			{
 				//By default we do not create the materials, every node with mesh attribute can enable them. So we wont create unused materials.
-				MaterialFactoryNode->SetEnabled(true);
+				MaterialFactoryNode->SetEnabled(false);
 			}
 		}
 	}
