Fix issue with bind pose, the fbx translator will now always use cluster deformation instead of bake poses. Since sometime the bake pose are not synchronise properly with the real cluster deformation.
#rnx

diff --git a/Engine/Plugins/Experimental/Interchange/Runtime/Source/Dispatcher/Public/InterchangeCommands.h b/Engine/Plugins/Experimental/Interchange/Runtime/Source/Dispatcher/Public/InterchangeCommands.h
--- a/Engine/Plugins/Experimental/Interchange/Runtime/Source/Dispatcher/Public/InterchangeCommands.h
+++ b/Engine/Plugins/Experimental/Interchange/Runtime/Source/Dispatcher/Public/InterchangeCommands.h
@@ -17,7 +17,7 @@ namespace UE
 		{
 		public:
 			//Major version should be updated when there is an existing API that has been change.
-			static int32 GetMajor() { return 6; }
+			static int32 GetMajor() { return 7; }
 			//Minor version must be updated when there is an addition to the API.
 			static int32 GetMinor() { return 0; }
 			//Patch version should be update if there is some bug fixes in the private code.
diff --git a/Engine/Plugins/Experimental/Interchange/Runtime/Source/Parsers/Fbx/Private/FbxMesh.cpp b/Engine/Plugins/Experimental/Interchange/Runtime/Source/Parsers/Fbx/Private/FbxMesh.cpp
--- a/Engine/Plugins/Experimental/Interchange/Runtime/Source/Parsers/Fbx/Private/FbxMesh.cpp
+++ b/Engine/Plugins/Experimental/Interchange/Runtime/Source/Parsers/Fbx/Private/FbxMesh.cpp
@@ -279,8 +279,6 @@ bool FMeshDescriptionImporter::FillMeshDescriptionFromFbxShape(FbxShape* Shape)
 			VertexPositions[AddedVertexId] = (FVector3f)VertexPosition;// -MeshVertexPosition;
 		}
 		MeshDescription->ResumeVertexIndexing();
-		FElementIDRemappings OutRemappings;
-		MeshDescription->Compact(OutRemappings);
 	}
 	return true;
 }
@@ -982,11 +980,7 @@ bool FMeshDescriptionImporter::FillMeshDescriptionFromFbxMesh(FbxMesh* Mesh, TAr
 
 			if (SkippedVertexInstance > 0)
 			{
-				//We must compact the sparse array before reserving new space
-				//When we reserve it will make a hole in the sparse array if the last reserve was not fully use
-				//The importer assume there will be no hole when importing a mesh
-				FElementIDRemappings OutRemappings;
-				MeshDescription->Compact(OutRemappings);
+				check(MeshDescription->Triangles().Num() == MeshDescription->Triangles().GetArraySize());
 			}
 		}
 
@@ -1448,25 +1442,6 @@ void FFbxMesh::AddAllMeshes(FbxScene* SDKScene, FbxGeometryConverter* SDKGeometr
 
 bool FFbxMesh::GetGlobalJointBindPoseTransform(FbxScene* SDKScene, FbxNode* Joint, FbxAMatrix& GlobalBindPoseJointMatrix)
 {
-	const FString LinkName = FFbxHelper::GetFbxObjectName(Joint);
-
-	const int32 PoseCount = SDKScene->GetPoseCount();
-	for (int32 PoseIndex = 0; PoseIndex < PoseCount; PoseIndex++)
-	{
-		FbxPose* CurrentPose = SDKScene->GetPose(PoseIndex);
-		if (CurrentPose && CurrentPose->IsBindPose())
-		{
-			FString PoseName = FFbxHelper::GetFbxObjectName(CurrentPose);
-			int32 PoseLinkIndex = CurrentPose->Find(Joint);
-			if (PoseLinkIndex >= 0)
-			{
-				FbxMatrix NoneAffineMatrix = CurrentPose->GetMatrix(PoseLinkIndex);
-				GlobalBindPoseJointMatrix = *(FbxAMatrix*)(double*)&NoneAffineMatrix;
-				return true;
-			}
-		}
-	}
-
 	//Search all skeletalmesh(FbxGeometry with valid deformer) using this joint and see if there is a valid FbxCluster
 	//containing a TransformLinkMatrix for this joint
 	const int32 GeometryCount = SDKScene->GetGeometryCount();
@@ -1511,20 +1486,6 @@ void FFbxMesh::ExtractSkinnedMeshNodeJoints(FbxScene* SDKScene, UInterchangeBase
 {
 	TArray<FString> JointNodeUniqueIDs;
 
-	//Get the bind pose from fbx
-	const int32 PoseCount = SDKScene->GetPoseCount();
-	TArray<FbxPose*> BindPoses;
-	for (int32 PoseIndex = 0; PoseIndex < PoseCount; PoseIndex++)
-	{
-		FbxPose* CurrentPose = SDKScene->GetPose(PoseIndex);
-
-		// current pose is bind pose, 
-		if (CurrentPose && CurrentPose->IsBindPose())
-		{
-			BindPoses.Add(CurrentPose);
-		}
-	}
-
 	const int32 SkinDeformerCount = Mesh->GetDeformerCount(FbxDeformer::eSkin);
 	for (int32 DeformerIndex = 0; DeformerIndex < SkinDeformerCount; DeformerIndex++)
 	{
