Avoid importing texture that do not have filepath or if the filepath is invalid.
#rnx

diff --git a/Engine/Plugins/Experimental/Interchange/Runtime/Source/Messages/Private/Fbx/InterchangeFbxMessages.cpp b/Engine/Plugins/Experimental/Interchange/Runtime/Source/Messages/Private/Fbx/InterchangeFbxMessages.cpp
--- a/Engine/Plugins/Experimental/Interchange/Runtime/Source/Messages/Private/Fbx/InterchangeFbxMessages.cpp
+++ b/Engine/Plugins/Experimental/Interchange/Runtime/Source/Messages/Private/Fbx/InterchangeFbxMessages.cpp
@@ -40,5 +40,21 @@ FText UInterchangeResultMeshWarning_TooManyUVs::GetText() const
 	return FText::Format(LOCTEXT("TooManyUVs", "Reached the maximum number of UV Channels for mesh '{MeshName}' - discarding {ExcessUVs} {ExcessUVs}|plural(one=channel,other=channels)."), Args);
 }
 
+FText UInterchangeResultTextureWarning_TextureFileDoNotExist::GetText() const
+{
+	FFormatNamedArguments Args
+	{
+		{ TEXT("TextureName"), FText::FromString(TextureName) },
+		{ TEXT("MaterialName"), FText::FromString(MaterialName) }
+	};
+
+	if (MaterialName.IsEmpty())
+	{
+		return FText::Format(LOCTEXT("TextureFileDoNotExist", "Cannot open texture file '{TextureName}'."), Args);
+	}
+	
+	return FText::Format(LOCTEXT("TextureFileDoNotExist", "Cannot open texture file '{TextureName}' when importing material '{MaterialName}'."), Args);
+}
+
 
 #undef LOCTEXT_NAMESPACE
diff --git a/Engine/Plugins/Experimental/Interchange/Runtime/Source/Messages/Public/Fbx/InterchangeFbxMessages.h b/Engine/Plugins/Experimental/Interchange/Runtime/Source/Messages/Public/Fbx/InterchangeFbxMessages.h
--- a/Engine/Plugins/Experimental/Interchange/Runtime/Source/Messages/Public/Fbx/InterchangeFbxMessages.h
+++ b/Engine/Plugins/Experimental/Interchange/Runtime/Source/Messages/Public/Fbx/InterchangeFbxMessages.h
@@ -22,6 +22,19 @@ public:
 	FString MeshName;
 };
 
+/**
+ * Base class for FBX parser warnings
+ */
+UCLASS()
+class INTERCHANGEMESSAGES_API UInterchangeResultTextureWarning : public UInterchangeResultWarning
+{
+	GENERATED_BODY()
+
+public:
+
+	UPROPERTY()
+	FString TextureName;
+};
 
 /**
  * Base class for FBX parser errors
@@ -85,3 +98,20 @@ public:
 	int32 ExcessUVs;
 };
 
+/**
+ * A generic class for FBX parser warnings, with no additional metadata, and where the text is specified by the user
+ */
+UCLASS()
+class INTERCHANGEMESSAGES_API UInterchangeResultTextureWarning_TextureFileDoNotExist : public UInterchangeResultTextureWarning
+{
+	GENERATED_BODY()
+
+public:
+	virtual FText GetText() const override;
+
+	UPROPERTY()
+	FText Text;
+
+	UPROPERTY()
+	FString MaterialName;
+};
\ No newline at end of file
diff --git a/Engine/Plugins/Experimental/Interchange/Runtime/Source/Parsers/Fbx/Private/FbxMaterial.cpp b/Engine/Plugins/Experimental/Interchange/Runtime/Source/Parsers/Fbx/Private/FbxMaterial.cpp
--- a/Engine/Plugins/Experimental/Interchange/Runtime/Source/Parsers/Fbx/Private/FbxMaterial.cpp
+++ b/Engine/Plugins/Experimental/Interchange/Runtime/Source/Parsers/Fbx/Private/FbxMaterial.cpp
@@ -6,6 +6,7 @@
 #include "FbxConvert.h"
 #include "FbxHelper.h"
 #include "FbxInclude.h"
+#include "Fbx/InterchangeFbxMessages.h"
 #include "InterchangeMaterialDefinitions.h"
 #include "InterchangeResultsContainer.h"
 #include "InterchangeSceneNode.h"
@@ -15,6 +16,8 @@
 #include "Misc/Paths.h"
 #include "Nodes/InterchangeBaseNodeContainer.h"
 
+
+
 #define LOCTEXT_NAMESPACE "InterchangeFbxMaterial"
 
 namespace UE
@@ -110,48 +113,55 @@ namespace UE
 				if (TextureCount > 0)
 				{
 					FbxFileTexture* FbxTexture = Property.GetSrcObject<FbxFileTexture>(0);
-
 					FString TextureFilename = UTF8_TO_TCHAR(FbxTexture->GetFileName());
-					FString TextureName = FPaths::GetBaseFilename(TextureFilename);
+					//Only import texture that exist on disk
+					if(!TextureFilename.IsEmpty() && FPaths::FileExists(TextureFilename))
+					{
+						FString TextureName = FPaths::GetBaseFilename(TextureFilename);
+						UInterchangeShaderNode* TextureSampleShader = NewObject<UInterchangeShaderNode>(&NodeContainer);
+						TextureSampleShader->SetCustomShaderType(TextureSample::Name.ToString());
 
-					UInterchangeShaderNode* TextureSampleShader = NewObject<UInterchangeShaderNode>(&NodeContainer);
-					TextureSampleShader->SetCustomShaderType(TextureSample::Name.ToString());
+						FString TextureSampleShaderUid = ShaderGraphNode->GetUniqueID() + TEXT("\\Textures\\") + TextureName;
+						TextureSampleShader->InitializeNode(TextureSampleShaderUid, TextureName, EInterchangeNodeContainerType::TranslatedAsset);
 
-					FString TextureSampleShaderUid = ShaderGraphNode->GetUniqueID() + TEXT("\\Textures\\") + TextureName;
-					TextureSampleShader->InitializeNode(TextureSampleShaderUid, TextureName, EInterchangeNodeContainerType::TranslatedAsset);
+						NodeContainer.AddNode(TextureSampleShader);
+						NodeContainer.SetNodeParentUid(TextureSampleShaderUid, ShaderGraphNode->GetUniqueID());
 
-					NodeContainer.AddNode(TextureSampleShader);
-					NodeContainer.SetNodeParentUid(TextureSampleShaderUid, ShaderGraphNode->GetUniqueID());
+						FString TextureNodeUid = TEXT("\\Texture\\") + TextureFilename;
+						const UInterchangeTexture2DNode* TextureNode = Cast<const UInterchangeTexture2DNode>(NodeContainer.GetNode(TextureNodeUid));
 
-					FString TextureNodeUid = TEXT("\\Texture\\") + TextureFilename;
-					const UInterchangeTexture2DNode* TextureNode = Cast<const UInterchangeTexture2DNode>(NodeContainer.GetNode(TextureNodeUid));
+						if (!TextureNode)
+						{
+							TextureNode = CreateTexture2DNode(NodeContainer, TextureNodeUid, TextureFilename);
+						}
 
-					if (!TextureNode)
-					{
-						TextureNode = CreateTexture2DNode(NodeContainer, TextureNodeUid, TextureFilename);
-					}
+						TextureSampleShader->AddStringAttribute(UInterchangeShaderPortsAPI::MakeInputValueKey(TextureSample::Inputs::Texture.ToString()), TextureNode->GetUniqueID());
 
-					TextureSampleShader->AddStringAttribute(UInterchangeShaderPortsAPI::MakeInputValueKey(TextureSample::Inputs::Texture.ToString()), TextureNode->GetUniqueID());
+						if (!FMath::IsNearlyEqual(FbxTexture->GetScaleU(), 1.0) || !FMath::IsNearlyEqual(FbxTexture->GetScaleV(), 1.0))
+						{
+							UInterchangeShaderNode* TextureCoordinateShader = NewObject<UInterchangeShaderNode>(&NodeContainer);
+							TextureCoordinateShader->SetCustomShaderType(TextureCoordinate::Name.ToString());
 
-					if (!FMath::IsNearlyEqual(FbxTexture->GetScaleU(), 1.0) || !FMath::IsNearlyEqual(FbxTexture->GetScaleV(), 1.0))
-					{
-						UInterchangeShaderNode* TextureCoordinateShader = NewObject<UInterchangeShaderNode>(&NodeContainer);
-						TextureCoordinateShader->SetCustomShaderType(TextureCoordinate::Name.ToString());
+							FString TextureCoordinateShaderUid = TextureSampleShader->GetUniqueID() + TEXT("_Coordinate");
+							TextureCoordinateShader->InitializeNode(TextureCoordinateShaderUid, TextureName + TEXT("_Coordinate"), EInterchangeNodeContainerType::TranslatedAsset);
 
-						FString TextureCoordinateShaderUid = TextureSampleShader->GetUniqueID() + TEXT("_Coordinate");
-						TextureCoordinateShader->InitializeNode(TextureCoordinateShaderUid, TextureName + TEXT("_Coordinate"), EInterchangeNodeContainerType::TranslatedAsset);
+							NodeContainer.AddNode(TextureCoordinateShader);
+							NodeContainer.SetNodeParentUid(TextureCoordinateShaderUid, ShaderGraphNode->GetUniqueID());
 
-						NodeContainer.AddNode(TextureCoordinateShader);
-						NodeContainer.SetNodeParentUid(TextureCoordinateShaderUid, ShaderGraphNode->GetUniqueID());
+							TextureCoordinateShader->AddFloatAttribute(UInterchangeShaderPortsAPI::MakeInputValueKey(TextureCoordinate::Inputs::UTiling.ToString()), (float)FbxTexture->GetScaleU());
+							TextureCoordinateShader->AddFloatAttribute(UInterchangeShaderPortsAPI::MakeInputValueKey(TextureCoordinate::Inputs::VTiling.ToString()), (float)FbxTexture->GetScaleV());
 
-						TextureCoordinateShader->AddFloatAttribute(UInterchangeShaderPortsAPI::MakeInputValueKey(TextureCoordinate::Inputs::UTiling.ToString()), (float)FbxTexture->GetScaleU());
-						TextureCoordinateShader->AddFloatAttribute(UInterchangeShaderPortsAPI::MakeInputValueKey(TextureCoordinate::Inputs::VTiling.ToString()), (float)FbxTexture->GetScaleV());
+							UInterchangeShaderPortsAPI::ConnectDefaultOuputToInput(TextureSampleShader, TextureSample::Inputs::Coordinates.ToString(), TextureCoordinateShaderUid);
+						}
 
-						UInterchangeShaderPortsAPI::ConnectDefaultOuputToInput(TextureSampleShader, TextureSample::Inputs::Coordinates.ToString(), TextureCoordinateShaderUid);
+						UInterchangeShaderPortsAPI::ConnectDefaultOuputToInput(NodeToConnectTo, InputToConnectTo, TextureSampleShaderUid);
+					}
+					else
+					{
+						UInterchangeResultTextureWarning_TextureFileDoNotExist* Message = Parser.AddMessage<UInterchangeResultTextureWarning_TextureFileDoNotExist>();
+						Message->TextureName = TextureFilename;
+						Message->MaterialName = ShaderGraphNode->GetDisplayLabel();
 					}
-
-					UInterchangeShaderPortsAPI::ConnectDefaultOuputToInput(NodeToConnectTo, InputToConnectTo, TextureSampleShaderUid);
-
 				}
 				else if (DataType.GetType() == eFbxDouble || DataType.GetType() == eFbxFloat || DataType.GetType() == eFbxInt)
 				{
