Interchange AttributeKey will use FName instead of FString. FName are more appropriate since the keys are repeated for each node.
Add a children cache in the base node container to optimize the queries.
#rnx

diff --git a/Engine/Plugins/Experimental/Interchange/Editor/Source/Pipelines/Private/InterchangeEditorPipelineDetails.cpp b/Engine/Plugins/Experimental/Interchange/Editor/Source/Pipelines/Private/InterchangeEditorPipelineDetails.cpp
--- a/Engine/Plugins/Experimental/Interchange/Editor/Source/Pipelines/Private/InterchangeEditorPipelineDetails.cpp
+++ b/Engine/Plugins/Experimental/Interchange/Editor/Source/Pipelines/Private/InterchangeEditorPipelineDetails.cpp
@@ -896,7 +896,7 @@ void FInterchangeBaseNodeDetailsCustomization::BuildVectorValueContent(IDetailCa
 	};
 
 	const bool bAdvancedProperty = false;
-	FDetailWidgetRow& VectorRow = AttributeCategory.AddCustomRow(FText::FromString(AttributeKey.Key), bAdvancedProperty);
+	FDetailWidgetRow& VectorRow = AttributeCategory.AddCustomRow(FText::FromString(AttributeKey.ToString()), bAdvancedProperty);
 	VectorRow.NameContent()
 	[
 		CreateNameWidget(AttributeKey)
diff --git a/Engine/Plugins/Experimental/Interchange/Runtime/Source/FactoryNodes/Private/InterchangeSkeletalMeshFactoryNode.cpp b/Engine/Plugins/Experimental/Interchange/Runtime/Source/FactoryNodes/Private/InterchangeSkeletalMeshFactoryNode.cpp
--- a/Engine/Plugins/Experimental/Interchange/Runtime/Source/FactoryNodes/Private/InterchangeSkeletalMeshFactoryNode.cpp
+++ b/Engine/Plugins/Experimental/Interchange/Runtime/Source/FactoryNodes/Private/InterchangeSkeletalMeshFactoryNode.cpp
@@ -7,24 +7,21 @@
 #endif
 
 //Interchange namespace
-namespace UE
+namespace UE::Interchange
 {
-	namespace Interchange
+	const FAttributeKey& FSkeletalMeshNodeStaticData::GetLodDependenciesBaseKey()
 	{
-		const FString& FSkeletalMeshNodeStaticData::GetLodDependenciesBaseKey()
-		{
-			static FString LodDependencies_BaseKey = TEXT("Lod_Dependencies");
-			return LodDependencies_BaseKey;
-		}
-	}//ns Interchange
-}//ns UE
+		static FAttributeKey LodDependencies_BaseKey(TEXT("Lod_Dependencies"));
+		return LodDependencies_BaseKey;
+	}
+}//ns UE::Interchange
 
 UInterchangeSkeletalMeshFactoryNode::UInterchangeSkeletalMeshFactoryNode()
 {
 #if WITH_ENGINE
 	AssetClass = nullptr;
 #endif
-	LodDependencies.Initialize(Attributes, UE::Interchange::FSkeletalMeshNodeStaticData::GetLodDependenciesBaseKey());
+	LodDependencies.Initialize(Attributes, UE::Interchange::FSkeletalMeshNodeStaticData::GetLodDependenciesBaseKey().ToString());
 }
 
 void UInterchangeSkeletalMeshFactoryNode::InitializeSkeletalMeshNode(const FString& UniqueID, const FString& DisplayLabel, const FString& InAssetClass)
@@ -58,20 +55,21 @@ FString UInterchangeSkeletalMeshFactoryNode::GetTypeName() const
 
 FString UInterchangeSkeletalMeshFactoryNode::GetKeyDisplayName(const UE::Interchange::FAttributeKey& NodeAttributeKey) const
 {
-	FString KeyDisplayName = NodeAttributeKey.Key;
-	if (NodeAttributeKey.Key.Equals(UE::Interchange::FSkeletalMeshNodeStaticData::GetLodDependenciesBaseKey()))
+	FString KeyDisplayName = NodeAttributeKey.ToString();
+	const FString NodeAttributeKeyString = KeyDisplayName;
+	if (NodeAttributeKey == UE::Interchange::FSkeletalMeshNodeStaticData::GetLodDependenciesBaseKey())
 	{
 		KeyDisplayName = TEXT("LOD Dependencies Count");
 		return KeyDisplayName;
 	}
-	else if (NodeAttributeKey.Key.StartsWith(UE::Interchange::FSkeletalMeshNodeStaticData::GetLodDependenciesBaseKey()))
+	else if (NodeAttributeKeyString.StartsWith(UE::Interchange::FSkeletalMeshNodeStaticData::GetLodDependenciesBaseKey().ToString()))
 	{
 		KeyDisplayName = TEXT("LOD Dependencies Index ");
 		const FString IndexKey = UE::Interchange::TArrayAttributeHelper<FString>::IndexKey();
-		int32 IndexPosition = NodeAttributeKey.Key.Find(IndexKey) + IndexKey.Len();
-		if (IndexPosition < NodeAttributeKey.Key.Len())
+		int32 IndexPosition = NodeAttributeKeyString.Find(IndexKey) + IndexKey.Len();
+		if (IndexPosition < NodeAttributeKeyString.Len())
 		{
-			KeyDisplayName += NodeAttributeKey.Key.RightChop(IndexPosition);
+			KeyDisplayName += NodeAttributeKeyString.RightChop(IndexPosition);
 		}
 		return KeyDisplayName;
 	}
diff --git a/Engine/Plugins/Experimental/Interchange/Runtime/Source/FactoryNodes/Private/InterchangeSkeletalMeshLodDataNode.cpp b/Engine/Plugins/Experimental/Interchange/Runtime/Source/FactoryNodes/Private/InterchangeSkeletalMeshLodDataNode.cpp
--- a/Engine/Plugins/Experimental/Interchange/Runtime/Source/FactoryNodes/Private/InterchangeSkeletalMeshLodDataNode.cpp
+++ b/Engine/Plugins/Experimental/Interchange/Runtime/Source/FactoryNodes/Private/InterchangeSkeletalMeshLodDataNode.cpp
@@ -3,21 +3,18 @@
 #include "InterchangeSkeletalMeshLodDataNode.h"
 
 //Interchange namespace
-namespace UE
+namespace UE::Interchange
 {
-	namespace Interchange
+	const FAttributeKey& FSkeletalMeshNodeLodDataStaticData::GetMeshUidsBaseKey()
 	{
-		const FString& FSkeletalMeshNodeLodDataStaticData::GetMeshUidsBaseKey()
-		{
-			static FString MeshUids_BaseKey = TEXT("__MeshUids__Key");
-			return MeshUids_BaseKey;
-		}
-	}//ns Interchange
-}//ns UE
+		static FAttributeKey MeshUids_BaseKey(TEXT("__MeshUids__Key"));
+		return MeshUids_BaseKey;
+	}
+}//ns UE::Interchange
 
 UInterchangeSkeletalMeshLodDataNode::UInterchangeSkeletalMeshLodDataNode()
 {
-	MeshUids.Initialize(Attributes, UE::Interchange::FSkeletalMeshNodeLodDataStaticData::GetMeshUidsBaseKey());
+	MeshUids.Initialize(Attributes, UE::Interchange::FSkeletalMeshNodeLodDataStaticData::GetMeshUidsBaseKey().ToString());
 }
 
 /**
@@ -31,20 +28,21 @@ FString UInterchangeSkeletalMeshLodDataNode::GetTypeName() const
 
 FString UInterchangeSkeletalMeshLodDataNode::GetKeyDisplayName(const UE::Interchange::FAttributeKey& NodeAttributeKey) const
 {
-	FString KeyDisplayName = NodeAttributeKey.Key;
-	if (NodeAttributeKey.Key.Equals(UE::Interchange::FSkeletalMeshNodeLodDataStaticData::GetMeshUidsBaseKey()))
+	FString KeyDisplayName = NodeAttributeKey.ToString();
+	const FString NodeAttributeKeyString = KeyDisplayName;
+	if (NodeAttributeKey == UE::Interchange::FSkeletalMeshNodeLodDataStaticData::GetMeshUidsBaseKey())
 	{
 		KeyDisplayName = TEXT("Mesh count");
 		return KeyDisplayName;
 	}
-	else if (NodeAttributeKey.Key.StartsWith(UE::Interchange::FSkeletalMeshNodeLodDataStaticData::GetMeshUidsBaseKey()))
+	else if (NodeAttributeKeyString.StartsWith(UE::Interchange::FSkeletalMeshNodeLodDataStaticData::GetMeshUidsBaseKey().ToString()))
 	{
 		KeyDisplayName = TEXT("Mesh index ");
 		const FString IndexKey = UE::Interchange::TArrayAttributeHelper<FString>::IndexKey();
-		int32 IndexPosition = NodeAttributeKey.Key.Find(IndexKey) + IndexKey.Len();
-		if (IndexPosition < NodeAttributeKey.Key.Len())
+		int32 IndexPosition = NodeAttributeKeyString.Find(IndexKey) + IndexKey.Len();
+		if (IndexPosition < NodeAttributeKeyString.Len())
 		{
-			KeyDisplayName += NodeAttributeKey.Key.RightChop(IndexPosition);
+			KeyDisplayName += NodeAttributeKeyString.RightChop(IndexPosition);
 		}
 		return KeyDisplayName;
 	}
@@ -58,7 +56,7 @@ FString UInterchangeSkeletalMeshLodDataNode::GetKeyDisplayName(const UE::Interch
 
 FString UInterchangeSkeletalMeshLodDataNode::GetAttributeCategory(const UE::Interchange::FAttributeKey& NodeAttributeKey) const
 {
-	if (NodeAttributeKey.Key.StartsWith(UE::Interchange::FSkeletalMeshNodeLodDataStaticData::GetMeshUidsBaseKey()))
+	if (NodeAttributeKey.ToString().StartsWith(UE::Interchange::FSkeletalMeshNodeLodDataStaticData::GetMeshUidsBaseKey().ToString()))
 	{
 		return FString(TEXT("Meshes"));
 	}
diff --git a/Engine/Plugins/Experimental/Interchange/Runtime/Source/FactoryNodes/Private/InterchangeStaticMeshFactoryNode.cpp b/Engine/Plugins/Experimental/Interchange/Runtime/Source/FactoryNodes/Private/InterchangeStaticMeshFactoryNode.cpp
--- a/Engine/Plugins/Experimental/Interchange/Runtime/Source/FactoryNodes/Private/InterchangeStaticMeshFactoryNode.cpp
+++ b/Engine/Plugins/Experimental/Interchange/Runtime/Source/FactoryNodes/Private/InterchangeStaticMeshFactoryNode.cpp
@@ -10,15 +10,15 @@ namespace UE
 {
 	namespace Interchange
 	{
-		const FString& FStaticMeshNodeStaticData::GetLodDependenciesBaseKey()
+		const FAttributeKey& FStaticMeshNodeStaticData::GetLodDependenciesBaseKey()
 		{
-			static FString LodDependencies_BaseKey = TEXT("Lod_Dependencies");
+			static FAttributeKey LodDependencies_BaseKey = FAttributeKey(TEXT("Lod_Dependencies"));
 			return LodDependencies_BaseKey;
 		}
 
-		const FString& FStaticMeshNodeStaticData::GetSocketUidsBaseKey()
+		const FAttributeKey& FStaticMeshNodeStaticData::GetSocketUidsBaseKey()
 		{
-			static FString SocketUids_BaseKey = TEXT("SocketUids");
+			static FAttributeKey SocketUids_BaseKey = FAttributeKey(TEXT("SocketUids"));
 			return SocketUids_BaseKey;
 		}
 	} // namespace Interchange
@@ -30,8 +30,8 @@ UInterchangeStaticMeshFactoryNode::UInterchangeStaticMeshFactoryNode()
 #if WITH_ENGINE
 	AssetClass = nullptr;
 #endif
-	LodDependencies.Initialize(Attributes, UE::Interchange::FStaticMeshNodeStaticData::GetLodDependenciesBaseKey());
-	SocketUids.Initialize(Attributes, UE::Interchange::FStaticMeshNodeStaticData::GetSocketUidsBaseKey());
+	LodDependencies.Initialize(Attributes, UE::Interchange::FStaticMeshNodeStaticData::GetLodDependenciesBaseKey().ToString());
+	SocketUids.Initialize(Attributes, UE::Interchange::FStaticMeshNodeStaticData::GetSocketUidsBaseKey().ToString());
 }
 
 void UInterchangeStaticMeshFactoryNode::InitializeStaticMeshNode(const FString& UniqueID, const FString& DisplayLabel, const FString& InAssetClass)
@@ -64,20 +64,21 @@ FString UInterchangeStaticMeshFactoryNode::GetTypeName() const
 
 FString UInterchangeStaticMeshFactoryNode::GetKeyDisplayName(const UE::Interchange::FAttributeKey& NodeAttributeKey) const
 {
-	FString KeyDisplayName = NodeAttributeKey.Key;
-	if (NodeAttributeKey.Key.Equals(UE::Interchange::FStaticMeshNodeStaticData::GetLodDependenciesBaseKey()))
+	FString KeyDisplayName = NodeAttributeKey.ToString();
+	FString NodeAttributeKeyString = NodeAttributeKey.ToString();
+	if (NodeAttributeKey == UE::Interchange::FStaticMeshNodeStaticData::GetLodDependenciesBaseKey())
 	{
 		KeyDisplayName = TEXT("LOD Dependencies Count");
 		return KeyDisplayName;
 	}
-	else if (NodeAttributeKey.Key.StartsWith(UE::Interchange::FStaticMeshNodeStaticData::GetLodDependenciesBaseKey()))
+	else if (NodeAttributeKeyString.StartsWith(UE::Interchange::FStaticMeshNodeStaticData::GetLodDependenciesBaseKey().ToString()))
 	{
 		KeyDisplayName = TEXT("LOD Dependencies Index ");
 		const FString IndexKey = UE::Interchange::TArrayAttributeHelper<FString>::IndexKey();
-		int32 IndexPosition = NodeAttributeKey.Key.Find(IndexKey) + IndexKey.Len();
-		if (IndexPosition < NodeAttributeKey.Key.Len())
+		int32 IndexPosition = NodeAttributeKeyString.Find(IndexKey) + IndexKey.Len();
+		if (IndexPosition < NodeAttributeKeyString.Len())
 		{
-			KeyDisplayName += NodeAttributeKey.Key.RightChop(IndexPosition);
+			KeyDisplayName += NodeAttributeKeyString.RightChop(IndexPosition);
 		}
 		return KeyDisplayName;
 	}
diff --git a/Engine/Plugins/Experimental/Interchange/Runtime/Source/FactoryNodes/Private/InterchangeStaticMeshLodDataNode.cpp b/Engine/Plugins/Experimental/Interchange/Runtime/Source/FactoryNodes/Private/InterchangeStaticMeshLodDataNode.cpp
--- a/Engine/Plugins/Experimental/Interchange/Runtime/Source/FactoryNodes/Private/InterchangeStaticMeshLodDataNode.cpp
+++ b/Engine/Plugins/Experimental/Interchange/Runtime/Source/FactoryNodes/Private/InterchangeStaticMeshLodDataNode.cpp
@@ -6,33 +6,33 @@ namespace UE
 {
 	namespace Interchange
 	{
-		const FString& FStaticMeshNodeLodDataStaticData::GetMeshUidsBaseKey()
+		const FAttributeKey& FStaticMeshNodeLodDataStaticData::GetMeshUidsBaseKey()
 		{
-			static FString MeshUids_BaseKey = TEXT("__MeshUids__Key");
+			static FAttributeKey MeshUids_BaseKey(TEXT("__MeshUids__Key"));
 			return MeshUids_BaseKey;
 		}
 
-		const FString& FStaticMeshNodeLodDataStaticData::GetBoxCollisionMeshUidsBaseKey()
+		const FAttributeKey& FStaticMeshNodeLodDataStaticData::GetBoxCollisionMeshUidsBaseKey()
 		{
-			static FString CollisionMeshUids_BaseKey = TEXT("__BoxCollisionMeshUids__Key");
+			static FAttributeKey CollisionMeshUids_BaseKey(TEXT("__BoxCollisionMeshUids__Key"));
 			return CollisionMeshUids_BaseKey;
 		}
 
-		const FString& FStaticMeshNodeLodDataStaticData::GetCapsuleCollisionMeshUidsBaseKey()
+		const FAttributeKey& FStaticMeshNodeLodDataStaticData::GetCapsuleCollisionMeshUidsBaseKey()
 		{
-			static FString CollisionMeshUids_BaseKey = TEXT("__CapsuleCollisionMeshUids__Key");
+			static FAttributeKey CollisionMeshUids_BaseKey(TEXT("__CapsuleCollisionMeshUids__Key"));
 			return CollisionMeshUids_BaseKey;
 		}
 
-		const FString& FStaticMeshNodeLodDataStaticData::GetSphereCollisionMeshUidsBaseKey()
+		const FAttributeKey& FStaticMeshNodeLodDataStaticData::GetSphereCollisionMeshUidsBaseKey()
 		{
-			static FString CollisionMeshUids_BaseKey = TEXT("__SphereCollisionMeshUids__Key");
+			static FAttributeKey CollisionMeshUids_BaseKey(TEXT("__SphereCollisionMeshUids__Key"));
 			return CollisionMeshUids_BaseKey;
 		}
 
-		const FString& FStaticMeshNodeLodDataStaticData::GetConvexCollisionMeshUidsBaseKey()
+		const FAttributeKey& FStaticMeshNodeLodDataStaticData::GetConvexCollisionMeshUidsBaseKey()
 		{
-			static FString CollisionMeshUids_BaseKey = TEXT("__ConvexCollisionMeshUids__Key");
+			static FAttributeKey CollisionMeshUids_BaseKey(TEXT("__ConvexCollisionMeshUids__Key"));
 			return CollisionMeshUids_BaseKey;
 		}
 
@@ -42,11 +42,11 @@ namespace UE
 
 UInterchangeStaticMeshLodDataNode::UInterchangeStaticMeshLodDataNode()
 {
-	MeshUids.Initialize(Attributes, UE::Interchange::FStaticMeshNodeLodDataStaticData::GetMeshUidsBaseKey());
-	BoxCollisionMeshUids.Initialize(Attributes, UE::Interchange::FStaticMeshNodeLodDataStaticData::GetBoxCollisionMeshUidsBaseKey());
-	CapsuleCollisionMeshUids.Initialize(Attributes, UE::Interchange::FStaticMeshNodeLodDataStaticData::GetCapsuleCollisionMeshUidsBaseKey());
-	SphereCollisionMeshUids.Initialize(Attributes, UE::Interchange::FStaticMeshNodeLodDataStaticData::GetSphereCollisionMeshUidsBaseKey());
-	ConvexCollisionMeshUids.Initialize(Attributes, UE::Interchange::FStaticMeshNodeLodDataStaticData::GetConvexCollisionMeshUidsBaseKey());
+	MeshUids.Initialize(Attributes, UE::Interchange::FStaticMeshNodeLodDataStaticData::GetMeshUidsBaseKey().ToString());
+	BoxCollisionMeshUids.Initialize(Attributes, UE::Interchange::FStaticMeshNodeLodDataStaticData::GetBoxCollisionMeshUidsBaseKey().ToString());
+	CapsuleCollisionMeshUids.Initialize(Attributes, UE::Interchange::FStaticMeshNodeLodDataStaticData::GetCapsuleCollisionMeshUidsBaseKey().ToString());
+	SphereCollisionMeshUids.Initialize(Attributes, UE::Interchange::FStaticMeshNodeLodDataStaticData::GetSphereCollisionMeshUidsBaseKey().ToString());
+	ConvexCollisionMeshUids.Initialize(Attributes, UE::Interchange::FStaticMeshNodeLodDataStaticData::GetConvexCollisionMeshUidsBaseKey().ToString());
 }
 
 /**
@@ -60,20 +60,21 @@ FString UInterchangeStaticMeshLodDataNode::GetTypeName() const
 
 FString UInterchangeStaticMeshLodDataNode::GetKeyDisplayName(const UE::Interchange::FAttributeKey& NodeAttributeKey) const
 {
-	FString KeyDisplayName = NodeAttributeKey.Key;
-	if (NodeAttributeKey.Key.Equals(UE::Interchange::FStaticMeshNodeLodDataStaticData::GetMeshUidsBaseKey()))
+	FString KeyDisplayName = NodeAttributeKey.ToString();
+	const FString NodeAttributeKeyString = KeyDisplayName;
+	if (NodeAttributeKey == UE::Interchange::FStaticMeshNodeLodDataStaticData::GetMeshUidsBaseKey())
 	{
 		KeyDisplayName = TEXT("Mesh count");
 		return KeyDisplayName;
 	}
-	else if (NodeAttributeKey.Key.StartsWith(UE::Interchange::FStaticMeshNodeLodDataStaticData::GetMeshUidsBaseKey()))
+	else if (NodeAttributeKeyString.StartsWith(UE::Interchange::FStaticMeshNodeLodDataStaticData::GetMeshUidsBaseKey().ToString()))
 	{
 		KeyDisplayName = TEXT("Mesh index ");
 		const FString IndexKey = UE::Interchange::TArrayAttributeHelper<FString>::IndexKey();
-		int32 IndexPosition = NodeAttributeKey.Key.Find(IndexKey) + IndexKey.Len();
-		if (IndexPosition < NodeAttributeKey.Key.Len())
+		int32 IndexPosition = NodeAttributeKeyString.Find(IndexKey) + IndexKey.Len();
+		if (IndexPosition < NodeAttributeKeyString.Len())
 		{
-			KeyDisplayName += NodeAttributeKey.Key.RightChop(IndexPosition);
+			KeyDisplayName += NodeAttributeKeyString.RightChop(IndexPosition);
 		}
 		return KeyDisplayName;
 	}
@@ -83,7 +84,7 @@ FString UInterchangeStaticMeshLodDataNode::GetKeyDisplayName(const UE::Interchan
 
 FString UInterchangeStaticMeshLodDataNode::GetAttributeCategory(const UE::Interchange::FAttributeKey& NodeAttributeKey) const
 {
-	if (NodeAttributeKey.Key.StartsWith(UE::Interchange::FStaticMeshNodeLodDataStaticData::GetMeshUidsBaseKey()))
+	if (NodeAttributeKey.ToString().StartsWith(UE::Interchange::FStaticMeshNodeLodDataStaticData::GetMeshUidsBaseKey().ToString()))
 	{
 		return FString(TEXT("Meshes"));
 	}
diff --git a/Engine/Plugins/Experimental/Interchange/Runtime/Source/FactoryNodes/Public/InterchangeSkeletalMeshFactoryNode.h b/Engine/Plugins/Experimental/Interchange/Runtime/Source/FactoryNodes/Public/InterchangeSkeletalMeshFactoryNode.h
--- a/Engine/Plugins/Experimental/Interchange/Runtime/Source/FactoryNodes/Public/InterchangeSkeletalMeshFactoryNode.h
+++ b/Engine/Plugins/Experimental/Interchange/Runtime/Source/FactoryNodes/Public/InterchangeSkeletalMeshFactoryNode.h
@@ -12,18 +12,13 @@
 #include "InterchangeSkeletalMeshFactoryNode.generated.h"
 
 //Interchange namespace
-namespace UE
+namespace UE::Interchange
 {
-	namespace Interchange
+	struct INTERCHANGEFACTORYNODES_API FSkeletalMeshNodeStaticData : public FBaseNodeStaticData
 	{
-
-		struct INTERCHANGEFACTORYNODES_API FSkeletalMeshNodeStaticData : public FBaseNodeStaticData
-		{
-			static const FString& GetLodDependenciesBaseKey();
-		};
-
-	}//ns Interchange
-}//ns UE
+		static const FAttributeKey& GetLodDependenciesBaseKey();
+	};
+}//ns UE::Interchange
 
 UENUM(BlueprintType)
 enum class EInterchangeSkeletalMeshContentType : uint8
diff --git a/Engine/Plugins/Experimental/Interchange/Runtime/Source/FactoryNodes/Public/InterchangeSkeletalMeshLodDataNode.h b/Engine/Plugins/Experimental/Interchange/Runtime/Source/FactoryNodes/Public/InterchangeSkeletalMeshLodDataNode.h
--- a/Engine/Plugins/Experimental/Interchange/Runtime/Source/FactoryNodes/Public/InterchangeSkeletalMeshLodDataNode.h
+++ b/Engine/Plugins/Experimental/Interchange/Runtime/Source/FactoryNodes/Public/InterchangeSkeletalMeshLodDataNode.h
@@ -15,7 +15,7 @@ namespace UE
 
 		struct INTERCHANGEFACTORYNODES_API FSkeletalMeshNodeLodDataStaticData : public FBaseNodeStaticData
 		{
-			static const FString& GetMeshUidsBaseKey();
+			static const FAttributeKey& GetMeshUidsBaseKey();
 		};
 
 	}//ns Interchange
diff --git a/Engine/Plugins/Experimental/Interchange/Runtime/Source/FactoryNodes/Public/InterchangeStaticMeshFactoryNode.h b/Engine/Plugins/Experimental/Interchange/Runtime/Source/FactoryNodes/Public/InterchangeStaticMeshFactoryNode.h
--- a/Engine/Plugins/Experimental/Interchange/Runtime/Source/FactoryNodes/Public/InterchangeStaticMeshFactoryNode.h
+++ b/Engine/Plugins/Experimental/Interchange/Runtime/Source/FactoryNodes/Public/InterchangeStaticMeshFactoryNode.h
@@ -18,8 +18,8 @@ namespace UE
 	{
 		struct INTERCHANGEFACTORYNODES_API FStaticMeshNodeStaticData : public FBaseNodeStaticData
 		{
-			static const FString& GetLodDependenciesBaseKey();
-			static const FString& GetSocketUidsBaseKey();
+			static const FAttributeKey& GetLodDependenciesBaseKey();
+			static const FAttributeKey& GetSocketUidsBaseKey();
 		};
 	} // namespace Interchange
 } // namespace UE
diff --git a/Engine/Plugins/Experimental/Interchange/Runtime/Source/FactoryNodes/Public/InterchangeStaticMeshLodDataNode.h b/Engine/Plugins/Experimental/Interchange/Runtime/Source/FactoryNodes/Public/InterchangeStaticMeshLodDataNode.h
--- a/Engine/Plugins/Experimental/Interchange/Runtime/Source/FactoryNodes/Public/InterchangeStaticMeshLodDataNode.h
+++ b/Engine/Plugins/Experimental/Interchange/Runtime/Source/FactoryNodes/Public/InterchangeStaticMeshLodDataNode.h
@@ -15,11 +15,11 @@ namespace UE
 	{
 		struct FStaticMeshNodeLodDataStaticData : public FBaseNodeStaticData
 		{
-			static const FString& GetMeshUidsBaseKey();
-			static const FString& GetBoxCollisionMeshUidsBaseKey();
-			static const FString& GetCapsuleCollisionMeshUidsBaseKey();
-			static const FString& GetSphereCollisionMeshUidsBaseKey();
-			static const FString& GetConvexCollisionMeshUidsBaseKey();
+			static const FAttributeKey& GetMeshUidsBaseKey();
+			static const FAttributeKey& GetBoxCollisionMeshUidsBaseKey();
+			static const FAttributeKey& GetCapsuleCollisionMeshUidsBaseKey();
+			static const FAttributeKey& GetSphereCollisionMeshUidsBaseKey();
+			static const FAttributeKey& GetConvexCollisionMeshUidsBaseKey();
 		};
 	} // namespace Interchange
 } // namespace UE
diff --git a/Engine/Plugins/Experimental/Interchange/Runtime/Source/FactoryNodes/Public/InterchangeTexture2DFactoryNode.h b/Engine/Plugins/Experimental/Interchange/Runtime/Source/FactoryNodes/Public/InterchangeTexture2DFactoryNode.h
--- a/Engine/Plugins/Experimental/Interchange/Runtime/Source/FactoryNodes/Public/InterchangeTexture2DFactoryNode.h
+++ b/Engine/Plugins/Experimental/Interchange/Runtime/Source/FactoryNodes/Public/InterchangeTexture2DFactoryNode.h
@@ -13,10 +13,10 @@ namespace UE::Interchange
 { 
 	struct FTexture2DFactoryNodeStaticData : public FBaseNodeStaticData
 	{
-		static const FString& GetBaseSourceBlocksKey()
+		static const FAttributeKey& GetBaseSourceBlocksKey()
 		{
-			static FString StringKey(TEXT("SourceBlocks"));
-			return StringKey;
+			static FAttributeKey AttributeKey(TEXT("SourceBlocks"));
+			return AttributeKey;
 		}
 	};
 }//ns UE::Interchange
@@ -106,7 +106,7 @@ public:
 	virtual void PostInitProperties()
 	{
 		Super::PostInitProperties();
-		SourceBlocks.Initialize(Attributes.ToSharedRef(), UE::Interchange::FTexture2DFactoryNodeStaticData::GetBaseSourceBlocksKey());
+		SourceBlocks.Initialize(Attributes.ToSharedRef(), UE::Interchange::FTexture2DFactoryNodeStaticData::GetBaseSourceBlocksKey().ToString());
 	}
 
 private:
diff --git a/Engine/Plugins/Experimental/Interchange/Runtime/Source/Import/Private/Gltf/InterchangeGltfTranslator.cpp b/Engine/Plugins/Experimental/Interchange/Runtime/Source/Import/Private/Gltf/InterchangeGltfTranslator.cpp
--- a/Engine/Plugins/Experimental/Interchange/Runtime/Source/Import/Private/Gltf/InterchangeGltfTranslator.cpp
+++ b/Engine/Plugins/Experimental/Interchange/Runtime/Source/Import/Private/Gltf/InterchangeGltfTranslator.cpp
@@ -92,7 +92,9 @@ void UInterchangeGltfTranslator::HandleGltfNode( UInterchangeBaseNodeContainer&
 		}
 	}
 
-	InterchangeSceneNode->SetCustomLocalTransform(&NodeContainer, Transform );
+	constexpr bool bResetCache = false;
+
+	InterchangeSceneNode->SetCustomLocalTransform(&NodeContainer, Transform, bResetCache);
 
 	if ( !ParentNodeUid.IsEmpty() )
 	{
diff --git a/Engine/Plugins/Experimental/Interchange/Runtime/Source/Nodes/Private/InterchangeMeshNode.cpp b/Engine/Plugins/Experimental/Interchange/Runtime/Source/Nodes/Private/InterchangeMeshNode.cpp
--- a/Engine/Plugins/Experimental/Interchange/Runtime/Source/Nodes/Private/InterchangeMeshNode.cpp
+++ b/Engine/Plugins/Experimental/Interchange/Runtime/Source/Nodes/Private/InterchangeMeshNode.cpp
@@ -32,27 +32,27 @@ namespace UE
 			return AttributeKey;
 		}
 
-		const FString& FMeshNodeStaticData::GetSkeletonDependenciesKey()
+		const FAttributeKey& FMeshNodeStaticData::GetSkeletonDependenciesKey()
 		{
-			static FString Dependencies_BaseKey = TEXT("__MeshSkeletonDependencies__");
+			static FAttributeKey Dependencies_BaseKey(TEXT("__MeshSkeletonDependencies__"));
 			return Dependencies_BaseKey;
 		}
 
-		const FString& FMeshNodeStaticData::GetMaterialDependenciesKey()
+		const FAttributeKey& FMeshNodeStaticData::GetMaterialDependenciesKey()
 		{
-			static FString Dependencies_BaseKey = TEXT("__MeshMaterialDependencies__");
+			static FAttributeKey Dependencies_BaseKey(TEXT("__MeshMaterialDependencies__"));
 			return Dependencies_BaseKey;
 		}
 
-		const FString& FMeshNodeStaticData::GetShapeDependenciesKey()
+		const FAttributeKey& FMeshNodeStaticData::GetShapeDependenciesKey()
 		{
-			static FString Dependencies_BaseKey = TEXT("__MeshShapeDependencies__");
+			static FAttributeKey Dependencies_BaseKey(TEXT("__MeshShapeDependencies__"));
 			return Dependencies_BaseKey;
 		}
 
-		const FString& FMeshNodeStaticData::GetSceneInstancesUidsKey()
+		const FAttributeKey& FMeshNodeStaticData::GetSceneInstancesUidsKey()
 		{
-			static FString SceneInstanceUids_BaseKey = TEXT("__MeshSceneInstancesUids__");
+			static FAttributeKey SceneInstanceUids_BaseKey(TEXT("__MeshSceneInstancesUids__"));
 			return SceneInstanceUids_BaseKey;
 		}
 
@@ -61,15 +61,16 @@ namespace UE
 
 UInterchangeMeshNode::UInterchangeMeshNode()
 {
-	SkeletonDependencies.Initialize(Attributes, UE::Interchange::FMeshNodeStaticData::GetSkeletonDependenciesKey());
-	MaterialDependencies.Initialize(Attributes, UE::Interchange::FMeshNodeStaticData::GetMaterialDependenciesKey());
-	ShapeDependencies.Initialize(Attributes, UE::Interchange::FMeshNodeStaticData::GetShapeDependenciesKey());
-	SceneInstancesUids.Initialize(Attributes, UE::Interchange::FMeshNodeStaticData::GetSceneInstancesUidsKey());
+	SkeletonDependencies.Initialize(Attributes, UE::Interchange::FMeshNodeStaticData::GetSkeletonDependenciesKey().ToString());
+	MaterialDependencies.Initialize(Attributes, UE::Interchange::FMeshNodeStaticData::GetMaterialDependenciesKey().ToString());
+	ShapeDependencies.Initialize(Attributes, UE::Interchange::FMeshNodeStaticData::GetShapeDependenciesKey().ToString());
+	SceneInstancesUids.Initialize(Attributes, UE::Interchange::FMeshNodeStaticData::GetSceneInstancesUidsKey().ToString());
 }
 
 FString UInterchangeMeshNode::GetKeyDisplayName(const UE::Interchange::FAttributeKey& NodeAttributeKey) const
 {
 	FString KeyDisplayName = NodeAttributeKey.ToString();
+	const FString NodeAttributeKeyString = KeyDisplayName;
 	if (NodeAttributeKey == UE::Interchange::FMeshNodeStaticData::PayloadSourceFileKey())
 	{
 		return KeyDisplayName = TEXT("Payload Source Key");
@@ -86,63 +87,63 @@ FString UInterchangeMeshNode::GetKeyDisplayName(const UE::Interchange::FAttribut
 	{
 		return KeyDisplayName = TEXT("Blend Shape Name");
 	}
-	else if (NodeAttributeKey.Key.Equals(UE::Interchange::FMeshNodeStaticData::GetSkeletonDependenciesKey()))
+	else if (NodeAttributeKey == UE::Interchange::FMeshNodeStaticData::GetSkeletonDependenciesKey())
 	{
 		return KeyDisplayName = TEXT("Skeleton Dependencies count");
 	}
-	else if (NodeAttributeKey.Key.StartsWith(UE::Interchange::FMeshNodeStaticData::GetSkeletonDependenciesKey()))
+	else if (NodeAttributeKeyString.StartsWith(UE::Interchange::FMeshNodeStaticData::GetSkeletonDependenciesKey().ToString()))
 	{
 		KeyDisplayName = TEXT("Skeleton Dependencies Index ");
 		const FString IndexKey = UE::Interchange::TArrayAttributeHelper<FString>::IndexKey();
-		int32 IndexPosition = NodeAttributeKey.Key.Find(IndexKey) + IndexKey.Len();
-		if (IndexPosition < NodeAttributeKey.Key.Len())
+		int32 IndexPosition = NodeAttributeKeyString.Find(IndexKey) + IndexKey.Len();
+		if (IndexPosition < NodeAttributeKeyString.Len())
 		{
-			KeyDisplayName += NodeAttributeKey.Key.RightChop(IndexPosition);
+			KeyDisplayName += NodeAttributeKeyString.RightChop(IndexPosition);
 		}
 		return KeyDisplayName;
 	}
-	else if (NodeAttributeKey.Key.Equals(UE::Interchange::FMeshNodeStaticData::GetMaterialDependenciesKey()))
+	else if (NodeAttributeKey == UE::Interchange::FMeshNodeStaticData::GetMaterialDependenciesKey())
 	{
 		return KeyDisplayName = TEXT("Material Dependencies count");
 	}
-	else if (NodeAttributeKey.Key.StartsWith(UE::Interchange::FMeshNodeStaticData::GetMaterialDependenciesKey()))
+	else if (NodeAttributeKeyString.StartsWith(UE::Interchange::FMeshNodeStaticData::GetMaterialDependenciesKey().ToString()))
 	{
 		KeyDisplayName = TEXT("Material Dependencies Index ");
 		const FString IndexKey = UE::Interchange::TArrayAttributeHelper<FString>::IndexKey();
-		int32 IndexPosition = NodeAttributeKey.Key.Find(IndexKey) + IndexKey.Len();
-		if (IndexPosition < NodeAttributeKey.Key.Len())
+		int32 IndexPosition = NodeAttributeKeyString.Find(IndexKey) + IndexKey.Len();
+		if (IndexPosition < NodeAttributeKeyString.Len())
 		{
-			KeyDisplayName += NodeAttributeKey.Key.RightChop(IndexPosition);
+			KeyDisplayName += NodeAttributeKeyString.RightChop(IndexPosition);
 		}
 		return KeyDisplayName;
 	}
-	else if (NodeAttributeKey.Key.Equals(UE::Interchange::FMeshNodeStaticData::GetShapeDependenciesKey()))
+	else if (NodeAttributeKey == UE::Interchange::FMeshNodeStaticData::GetShapeDependenciesKey())
 	{
 		return KeyDisplayName = TEXT("Shape Dependencies count");
 	}
-	else if (NodeAttributeKey.Key.StartsWith(UE::Interchange::FMeshNodeStaticData::GetShapeDependenciesKey()))
+	else if (NodeAttributeKeyString.StartsWith(UE::Interchange::FMeshNodeStaticData::GetShapeDependenciesKey().ToString()))
 	{
 		KeyDisplayName = TEXT("Shape Dependencies Index ");
 		const FString IndexKey = UE::Interchange::TArrayAttributeHelper<FString>::IndexKey();
-		int32 IndexPosition = NodeAttributeKey.Key.Find(IndexKey) + IndexKey.Len();
-		if (IndexPosition < NodeAttributeKey.Key.Len())
+		int32 IndexPosition = NodeAttributeKeyString.Find(IndexKey) + IndexKey.Len();
+		if (IndexPosition < NodeAttributeKeyString.Len())
 		{
-			KeyDisplayName += NodeAttributeKey.Key.RightChop(IndexPosition);
+			KeyDisplayName += NodeAttributeKeyString.RightChop(IndexPosition);
 		}
 		return KeyDisplayName;
 	}
-	else if (NodeAttributeKey.Key.Equals(UE::Interchange::FMeshNodeStaticData::GetSceneInstancesUidsKey()))
+	else if (NodeAttributeKey == UE::Interchange::FMeshNodeStaticData::GetSceneInstancesUidsKey())
 	{
 		return KeyDisplayName = TEXT("Scene mesh instances count");
 	}
-	else if (NodeAttributeKey.Key.StartsWith(UE::Interchange::FMeshNodeStaticData::GetSceneInstancesUidsKey()))
+	else if (NodeAttributeKeyString.StartsWith(UE::Interchange::FMeshNodeStaticData::GetSceneInstancesUidsKey().ToString()))
 	{
 		KeyDisplayName = TEXT("Scene mesh instances Index ");
 		const FString IndexKey = UE::Interchange::TArrayAttributeHelper<FString>::IndexKey();
-		int32 IndexPosition = NodeAttributeKey.Key.Find(IndexKey) + IndexKey.Len();
-		if (IndexPosition < NodeAttributeKey.Key.Len())
+		int32 IndexPosition = NodeAttributeKeyString.Find(IndexKey) + IndexKey.Len();
+		if (IndexPosition < NodeAttributeKeyString.Len())
 		{
-			KeyDisplayName += NodeAttributeKey.Key.RightChop(IndexPosition);
+			KeyDisplayName += NodeAttributeKeyString.RightChop(IndexPosition);
 		}
 		return KeyDisplayName;
 	}
@@ -151,19 +152,20 @@ FString UInterchangeMeshNode::GetKeyDisplayName(const UE::Interchange::FAttribut
 
 FString UInterchangeMeshNode::GetAttributeCategory(const UE::Interchange::FAttributeKey& NodeAttributeKey) const
 {
-	if (NodeAttributeKey.Key.StartsWith(UE::Interchange::FMeshNodeStaticData::GetSkeletonDependenciesKey()))
+	const FString NodeAttributeKeyString = NodeAttributeKey.ToString();;
+	if (NodeAttributeKeyString.StartsWith(UE::Interchange::FMeshNodeStaticData::GetSkeletonDependenciesKey().ToString()))
 	{
 		return FString(TEXT("SkeletonDependencies"));
 	}
-	else if (NodeAttributeKey.Key.StartsWith(UE::Interchange::FMeshNodeStaticData::GetMaterialDependenciesKey()))
+	else if (NodeAttributeKeyString.StartsWith(UE::Interchange::FMeshNodeStaticData::GetMaterialDependenciesKey().ToString()))
 	{
 		return FString(TEXT("MaterialDependencies"));
 	}
-	else if (NodeAttributeKey.Key.StartsWith(UE::Interchange::FMeshNodeStaticData::GetShapeDependenciesKey()))
+	else if (NodeAttributeKeyString.StartsWith(UE::Interchange::FMeshNodeStaticData::GetShapeDependenciesKey().ToString()))
 	{
 		return FString(TEXT("ShapeDependencies"));
 	}
-	else if (NodeAttributeKey.Key.StartsWith(UE::Interchange::FMeshNodeStaticData::GetSceneInstancesUidsKey()))
+	else if (NodeAttributeKeyString.StartsWith(UE::Interchange::FMeshNodeStaticData::GetSceneInstancesUidsKey().ToString()))
 	{
 		return FString(TEXT("SceneInstances"));
 	}
diff --git a/Engine/Plugins/Experimental/Interchange/Runtime/Source/Nodes/Private/InterchangeSceneNode.cpp b/Engine/Plugins/Experimental/Interchange/Runtime/Source/Nodes/Private/InterchangeSceneNode.cpp
--- a/Engine/Plugins/Experimental/Interchange/Runtime/Source/Nodes/Private/InterchangeSceneNode.cpp
+++ b/Engine/Plugins/Experimental/Interchange/Runtime/Source/Nodes/Private/InterchangeSceneNode.cpp
@@ -9,33 +9,33 @@ namespace UE
 	namespace Interchange
 	{
 
-		const FString& FSceneNodeStaticData::GetNodeSpecializeTypeBaseKey()
+		const FAttributeKey& FSceneNodeStaticData::GetNodeSpecializeTypeBaseKey()
 		{
-			static FString SceneNodeSpecializeType_BaseKey = TEXT("SceneNodeSpecializeType");
+			static FAttributeKey SceneNodeSpecializeType_BaseKey(TEXT("SceneNodeSpecializeType"));
 			return SceneNodeSpecializeType_BaseKey;
 		}
 
-		const FString& FSceneNodeStaticData::GetMaterialDependencyUidsBaseKey()
+		const FAttributeKey& FSceneNodeStaticData::GetMaterialDependencyUidsBaseKey()
 		{
-			static FString MaterialDependencyUids_BaseKey = TEXT("__MaterialDependencyUidsBaseKey__");
+			static FAttributeKey MaterialDependencyUids_BaseKey(TEXT("__MaterialDependencyUidsBaseKey__"));
 			return MaterialDependencyUids_BaseKey;
 		}
 		
 		const FString& FSceneNodeStaticData::GetTransformSpecializeTypeString()
 		{
-			static FString TransformSpecializeTypeString = TEXT("Transform");
+			static FString TransformSpecializeTypeString(TEXT("Transform"));
 			return TransformSpecializeTypeString;
 		}
 
 		const FString& FSceneNodeStaticData::GetJointSpecializeTypeString()
 		{
-			static FString JointSpecializeTypeString = TEXT("Joint");
+			static FString JointSpecializeTypeString(TEXT("Joint"));
 			return JointSpecializeTypeString;
 		}
 
 		const FString& FSceneNodeStaticData::GetLodGroupSpecializeTypeString()
 		{
-			static FString JointSpecializeTypeString = TEXT("LodGroup");
+			static FString JointSpecializeTypeString(TEXT("LodGroup"));
 			return JointSpecializeTypeString;
 		}
 	}//ns Interchange
@@ -43,8 +43,8 @@ namespace UE
 
 UInterchangeSceneNode::UInterchangeSceneNode()
 {
-	NodeSpecializeTypes.Initialize(Attributes, UE::Interchange::FSceneNodeStaticData::GetNodeSpecializeTypeBaseKey());
-	MaterialDependencyUids.Initialize(Attributes, UE::Interchange::FSceneNodeStaticData::GetMaterialDependencyUidsBaseKey());
+	NodeSpecializeTypes.Initialize(Attributes, UE::Interchange::FSceneNodeStaticData::GetNodeSpecializeTypeBaseKey().ToString());
+	MaterialDependencyUids.Initialize(Attributes, UE::Interchange::FSceneNodeStaticData::GetMaterialDependencyUidsBaseKey().ToString());
 }
 
 /**
@@ -58,40 +58,41 @@ FString UInterchangeSceneNode::GetTypeName() const
 
 FString UInterchangeSceneNode::GetKeyDisplayName(const UE::Interchange::FAttributeKey& NodeAttributeKey) const
 {
-	FString KeyDisplayName = NodeAttributeKey.Key;
-	if (NodeAttributeKey.Key.Equals(UE::Interchange::FSceneNodeStaticData::GetNodeSpecializeTypeBaseKey()))
+	FString KeyDisplayName = NodeAttributeKey.ToString();
+	const FString NodeAttributeKeyString = KeyDisplayName;
+	if (NodeAttributeKey == UE::Interchange::FSceneNodeStaticData::GetNodeSpecializeTypeBaseKey())
 	{
 		KeyDisplayName = TEXT("Specialized type count");
 		return KeyDisplayName;
 	}
-	else if (NodeAttributeKey.Key.StartsWith(UE::Interchange::FSceneNodeStaticData::GetNodeSpecializeTypeBaseKey()))
+	else if (NodeAttributeKeyString.StartsWith(UE::Interchange::FSceneNodeStaticData::GetNodeSpecializeTypeBaseKey().ToString()))
 	{
 		KeyDisplayName = TEXT("Specialized type index ");
 		const FString IndexKey = UE::Interchange::TArrayAttributeHelper<FString>::IndexKey();
-		int32 IndexPosition = NodeAttributeKey.Key.Find(IndexKey) + IndexKey.Len();
-		if (IndexPosition < NodeAttributeKey.Key.Len())
+		int32 IndexPosition = NodeAttributeKeyString.Find(IndexKey) + IndexKey.Len();
+		if (IndexPosition < NodeAttributeKeyString.Len())
 		{
-			KeyDisplayName += NodeAttributeKey.Key.RightChop(IndexPosition);
+			KeyDisplayName += NodeAttributeKeyString.RightChop(IndexPosition);
 		}
 		return KeyDisplayName;
 	}
-	else if (NodeAttributeKey.Key.Equals(UE::Interchange::FSceneNodeStaticData::GetMaterialDependencyUidsBaseKey()))
+	else if (NodeAttributeKey == UE::Interchange::FSceneNodeStaticData::GetMaterialDependencyUidsBaseKey())
 	{
 		KeyDisplayName = TEXT("Material dependencies count");
 		return KeyDisplayName;
 	}
-	else if (NodeAttributeKey.Key.StartsWith(UE::Interchange::FSceneNodeStaticData::GetMaterialDependencyUidsBaseKey()))
+	else if (NodeAttributeKeyString.StartsWith(UE::Interchange::FSceneNodeStaticData::GetMaterialDependencyUidsBaseKey().ToString()))
 	{
 		KeyDisplayName = TEXT("Material dependency index ");
 		const FString IndexKey = UE::Interchange::TArrayAttributeHelper<FString>::IndexKey();
-		int32 IndexPosition = NodeAttributeKey.Key.Find(IndexKey) + IndexKey.Len();
-		if (IndexPosition < NodeAttributeKey.Key.Len())
+		int32 IndexPosition = NodeAttributeKeyString.Find(IndexKey) + IndexKey.Len();
+		if (IndexPosition < NodeAttributeKeyString.Len())
 		{
-			KeyDisplayName += NodeAttributeKey.Key.RightChop(IndexPosition);
+			KeyDisplayName += NodeAttributeKeyString.RightChop(IndexPosition);
 		}
 		return KeyDisplayName;
 	}
-	else if (NodeAttributeKey.Key.Equals(Macro_CustomTransformCurvePayloadKeyKey.Key))
+	else if (NodeAttributeKey == Macro_CustomTransformCurvePayloadKeyKey)
 	{
 		return FString(TEXT("Transform Curve Payload Key"));
 	}
@@ -100,15 +101,9 @@ FString UInterchangeSceneNode::GetKeyDisplayName(const UE::Interchange::FAttribu
 
 FString UInterchangeSceneNode::GetAttributeCategory(const UE::Interchange::FAttributeKey& NodeAttributeKey) const
 {
-	if (NodeAttributeKey.Key.StartsWith(UE::Interchange::FSceneNodeStaticData::GetNodeSpecializeTypeBaseKey()))
-	{
-		return FString(TEXT("SpecializeType"));
-	}
-	else if (NodeAttributeKey.Key.StartsWith(UE::Interchange::FSceneNodeStaticData::GetMaterialDependencyUidsBaseKey()))
-	{
-		return FString(TEXT("MaterialDependencies"));
-	}
-	else if (NodeAttributeKey == Macro_CustomLocalTransformKey
+	const FString NodeAttributeKeyString = NodeAttributeKey.ToString();
+
+	if (NodeAttributeKey == Macro_CustomLocalTransformKey
 		|| NodeAttributeKey == Macro_CustomAssetInstanceUidKey)
 	{
 		return FString(TEXT("Scene"));
@@ -119,6 +114,15 @@ FString UInterchangeSceneNode::GetAttributeCategory(const UE::Interchange::FAttr
 	{
 		return FString(TEXT("Joint"));
 	}
+	else if (NodeAttributeKeyString.StartsWith(UE::Interchange::FSceneNodeStaticData::GetNodeSpecializeTypeBaseKey().ToString()))
+	{
+		return FString(TEXT("SpecializeType"));
+	}
+	else if (NodeAttributeKeyString.StartsWith(UE::Interchange::FSceneNodeStaticData::GetMaterialDependencyUidsBaseKey().ToString()))
+	{
+		return FString(TEXT("MaterialDependencies"));
+	}
+	
 	return Super::GetAttributeCategory(NodeAttributeKey);
 }
 
@@ -203,9 +207,12 @@ bool UInterchangeSceneNode::GetCustomLocalTransform(FTransform& AttributeValue)
 	IMPLEMENT_NODE_ATTRIBUTE_GETTER(LocalTransform, FTransform);
 }
 
-bool UInterchangeSceneNode::SetCustomLocalTransform(const UInterchangeBaseNodeContainer* BaseNodeContainer, const FTransform& AttributeValue)
+bool UInterchangeSceneNode::SetCustomLocalTransform(const UInterchangeBaseNodeContainer* BaseNodeContainer, const FTransform& AttributeValue, bool bResetCache /*= true*/)
 {
-	ResetGlobalTransformCachesOfNodeAndAllChildren(BaseNodeContainer, this);
+	if(bResetCache)
+	{
+		ResetGlobalTransformCachesOfNodeAndAllChildren(BaseNodeContainer, this);
+	}
 	IMPLEMENT_NODE_ATTRIBUTE_SETTER_NODELEGATE(LocalTransform, FTransform);
 }
 
@@ -219,9 +226,12 @@ bool UInterchangeSceneNode::GetCustomBindPoseLocalTransform(FTransform& Attribut
 	IMPLEMENT_NODE_ATTRIBUTE_GETTER(BindPoseLocalTransform, FTransform);
 }
 
-bool UInterchangeSceneNode::SetCustomBindPoseLocalTransform(const UInterchangeBaseNodeContainer* BaseNodeContainer, const FTransform& AttributeValue)
+bool UInterchangeSceneNode::SetCustomBindPoseLocalTransform(const UInterchangeBaseNodeContainer* BaseNodeContainer, const FTransform& AttributeValue, bool bResetCache /*= true*/)
 {
-	ResetGlobalTransformCachesOfNodeAndAllChildren(BaseNodeContainer, this);
+	if (bResetCache)
+	{
+		ResetGlobalTransformCachesOfNodeAndAllChildren(BaseNodeContainer, this);
+	}
 	IMPLEMENT_NODE_ATTRIBUTE_SETTER_NODELEGATE(BindPoseLocalTransform, FTransform);
 }
 
@@ -235,9 +245,12 @@ bool UInterchangeSceneNode::GetCustomTimeZeroLocalTransform(FTransform& Attribut
 	IMPLEMENT_NODE_ATTRIBUTE_GETTER(TimeZeroLocalTransform, FTransform);
 }
 
-bool UInterchangeSceneNode::SetCustomTimeZeroLocalTransform(const UInterchangeBaseNodeContainer* BaseNodeContainer, const FTransform& AttributeValue)
+bool UInterchangeSceneNode::SetCustomTimeZeroLocalTransform(const UInterchangeBaseNodeContainer* BaseNodeContainer, const FTransform& AttributeValue, bool bResetCache /*= true*/)
 {
-	ResetGlobalTransformCachesOfNodeAndAllChildren(BaseNodeContainer, this);
+	if (bResetCache)
+	{
+		ResetGlobalTransformCachesOfNodeAndAllChildren(BaseNodeContainer, this);
+	}
 	IMPLEMENT_NODE_ATTRIBUTE_SETTER_NODELEGATE(TimeZeroLocalTransform, FTransform);
 }
 
diff --git a/Engine/Plugins/Experimental/Interchange/Runtime/Source/Nodes/Private/InterchangeShaderGraphNode.cpp b/Engine/Plugins/Experimental/Interchange/Runtime/Source/Nodes/Private/InterchangeShaderGraphNode.cpp
--- a/Engine/Plugins/Experimental/Interchange/Runtime/Source/Nodes/Private/InterchangeShaderGraphNode.cpp
+++ b/Engine/Plugins/Experimental/Interchange/Runtime/Source/Nodes/Private/InterchangeShaderGraphNode.cpp
@@ -5,7 +5,7 @@
 const TCHAR* UInterchangeShaderPortsAPI::InputPrefix = TEXT("Inputs");
 const TCHAR* UInterchangeShaderPortsAPI::InputSeparator = TEXT(":");
 
-FString UInterchangeShaderPortsAPI::MakeInputConnectionKey(const FString& InputName)
+FName UInterchangeShaderPortsAPI::MakeInputConnectionKey(const FString& InputName)
 {
 	TStringBuilder<128> StringBuilder;
 	StringBuilder.Append(InputPrefix);
@@ -14,10 +14,10 @@ FString UInterchangeShaderPortsAPI::MakeInputConnectionKey(const FString& InputN
 	StringBuilder.Append(InputSeparator);
 	StringBuilder.Append(TEXT("Connect"));
 
-	return StringBuilder.ToString();
+	return FName(StringBuilder.ToString());
 }
 
-FString UInterchangeShaderPortsAPI::MakeInputValueKey(const FString& InputName)
+FName UInterchangeShaderPortsAPI::MakeInputValueKey(const FString& InputName)
 {
 	TStringBuilder<128> StringBuilder;
 	StringBuilder.Append(InputPrefix);
@@ -25,8 +25,8 @@ FString UInterchangeShaderPortsAPI::MakeInputValueKey(const FString& InputName)
 	StringBuilder.Append(InputName);
 	StringBuilder.Append(InputSeparator);
 	StringBuilder.Append(TEXT("Value"));
-
-	return StringBuilder.ToString();
+	
+	return FName(StringBuilder.ToString());
 }
 
 FString UInterchangeShaderPortsAPI::MakeInputName(const FString& InputKey)
diff --git a/Engine/Plugins/Experimental/Interchange/Runtime/Source/Nodes/Public/InterchangeMeshNode.h b/Engine/Plugins/Experimental/Interchange/Runtime/Source/Nodes/Public/InterchangeMeshNode.h
--- a/Engine/Plugins/Experimental/Interchange/Runtime/Source/Nodes/Public/InterchangeMeshNode.h
+++ b/Engine/Plugins/Experimental/Interchange/Runtime/Source/Nodes/Public/InterchangeMeshNode.h
@@ -19,10 +19,10 @@ namespace UE
 			static const FAttributeKey& IsSkinnedMeshKey();
 			static const FAttributeKey& IsBlendShapeKey();
 			static const FAttributeKey& BlendShapeNameKey();
-			static const FString& GetSkeletonDependenciesKey();
-			static const FString& GetMaterialDependenciesKey();
-			static const FString& GetShapeDependenciesKey();
-			static const FString& GetSceneInstancesUidsKey();
+			static const FAttributeKey& GetSkeletonDependenciesKey();
+			static const FAttributeKey& GetMaterialDependenciesKey();
+			static const FAttributeKey& GetShapeDependenciesKey();
+			static const FAttributeKey& GetSceneInstancesUidsKey();
 		};
 
 	}//ns Interchange
diff --git a/Engine/Plugins/Experimental/Interchange/Runtime/Source/Nodes/Public/InterchangeSceneNode.h b/Engine/Plugins/Experimental/Interchange/Runtime/Source/Nodes/Public/InterchangeSceneNode.h
--- a/Engine/Plugins/Experimental/Interchange/Runtime/Source/Nodes/Public/InterchangeSceneNode.h
+++ b/Engine/Plugins/Experimental/Interchange/Runtime/Source/Nodes/Public/InterchangeSceneNode.h
@@ -17,8 +17,8 @@ namespace UE
 
 		struct INTERCHANGENODES_API FSceneNodeStaticData : public FBaseNodeStaticData
 		{
-			static const FString& GetNodeSpecializeTypeBaseKey();
-			static const FString& GetMaterialDependencyUidsBaseKey();
+			static const FAttributeKey& GetNodeSpecializeTypeBaseKey();
+			static const FAttributeKey& GetMaterialDependencyUidsBaseKey();
 			static const FString& GetTransformSpecializeTypeString();
 			static const FString& GetJointSpecializeTypeString();
 			static const FString& GetLodGroupSpecializeTypeString();
@@ -92,7 +92,7 @@ public:
 
 	/** Store the default scene node local transform. */
 	UFUNCTION(BlueprintCallable, Category = "Interchange | Node | Scene")
-	bool SetCustomLocalTransform(const UInterchangeBaseNodeContainer* BaseNodeContainer, const FTransform& AttributeValue);
+	bool SetCustomLocalTransform(const UInterchangeBaseNodeContainer* BaseNodeContainer, const FTransform& AttributeValue, bool bResetCache = true);
 
 	/** Return the default scene node global transform. This value is computed with all parent local transform. */
 	UFUNCTION(BlueprintCallable, Category = "Interchange | Node | Scene")
@@ -107,7 +107,7 @@ public:
 
 	/** Store the bind pose scene node local transform. */
 	UFUNCTION(BlueprintCallable, Category = "Interchange | Node | Joint")
-	bool SetCustomBindPoseLocalTransform(const UInterchangeBaseNodeContainer* BaseNodeContainer, const FTransform& AttributeValue);
+	bool SetCustomBindPoseLocalTransform(const UInterchangeBaseNodeContainer* BaseNodeContainer, const FTransform& AttributeValue, bool bResetCache = true);
 
 	/** Return the bind pose scene node global transform. This value is computed with all parent bind pose local transform. */
 	UFUNCTION(BlueprintCallable, Category = "Interchange | Node | Joint")
@@ -122,7 +122,7 @@ public:
 
 	/** Store the time zero scene node local transform. */
 	UFUNCTION(BlueprintCallable, Category = "Interchange | Node | Joint")
-	bool SetCustomTimeZeroLocalTransform(const UInterchangeBaseNodeContainer* BaseNodeContainer, const FTransform& AttributeValue);
+	bool SetCustomTimeZeroLocalTransform(const UInterchangeBaseNodeContainer* BaseNodeContainer, const FTransform& AttributeValue, bool bResetCache = true);
 
 	/** Return the time zero scene node global transform. This value is computed with all parent timezero local transform. */
 	UFUNCTION(BlueprintCallable, Category = "Interchange | Node | Joint")
diff --git a/Engine/Plugins/Experimental/Interchange/Runtime/Source/Nodes/Public/InterchangeShaderGraphNode.h b/Engine/Plugins/Experimental/Interchange/Runtime/Source/Nodes/Public/InterchangeShaderGraphNode.h
--- a/Engine/Plugins/Experimental/Interchange/Runtime/Source/Nodes/Public/InterchangeShaderGraphNode.h
+++ b/Engine/Plugins/Experimental/Interchange/Runtime/Source/Nodes/Public/InterchangeShaderGraphNode.h
@@ -21,13 +21,13 @@ public:
 	 * Makes an attribute key to represent a node being connected to an input (ie: Inputs:InputName:Connect).
 	 */
 	UFUNCTION(BlueprintCallable, Category = "Interchange | Node | Material")
-	static FString MakeInputConnectionKey(const FString& InputName);
+	static FName MakeInputConnectionKey(const FString& InputName);
  
 	/**
 	 * Makes an attribute key to represent a value being given to an input (ie: Inputs:InputName:Value).
 	 */
 	UFUNCTION(BlueprintCallable, Category = "Interchange | Node | Material")
-	static FString MakeInputValueKey(const FString& InputName);
+	static FName MakeInputValueKey(const FString& InputName);
  
 	/**
 	 * From an attribute key associated with an input (ie: Inputs:InputName:Value), retrieves the input name from it.
diff --git a/Engine/Plugins/Experimental/Interchange/Runtime/Source/Nodes/Public/InterchangeTexture2DNode.h b/Engine/Plugins/Experimental/Interchange/Runtime/Source/Nodes/Public/InterchangeTexture2DNode.h
--- a/Engine/Plugins/Experimental/Interchange/Runtime/Source/Nodes/Public/InterchangeTexture2DNode.h
+++ b/Engine/Plugins/Experimental/Interchange/Runtime/Source/Nodes/Public/InterchangeTexture2DNode.h
@@ -12,9 +12,9 @@ namespace UE::Interchange
 {
 	struct FTexture2DNodeStaticData : public FBaseNodeStaticData
 	{
-		static const FString& GetBaseSourceBlocksKey()
+		static const FAttributeKey& GetBaseSourceBlocksKey()
 		{
-			static FString StringKey(TEXT("SourceBlocks"));
+			static FAttributeKey StringKey(TEXT("SourceBlocks"));
 			return StringKey;
 		}
 	};
@@ -30,7 +30,7 @@ public:
 	virtual void PostInitProperties()
 	{
 		Super::PostInitProperties();
-		SourceBlocks.Initialize(Attributes.ToSharedRef(), UE::Interchange::FTexture2DNodeStaticData::GetBaseSourceBlocksKey());
+		SourceBlocks.Initialize(Attributes.ToSharedRef(), UE::Interchange::FTexture2DNodeStaticData::GetBaseSourceBlocksKey().ToString());
 	}
 
 	/**
diff --git a/Engine/Plugins/Experimental/Interchange/Runtime/Source/Parsers/Fbx/Private/FbxScene.cpp b/Engine/Plugins/Experimental/Interchange/Runtime/Source/Parsers/Fbx/Private/FbxScene.cpp
--- a/Engine/Plugins/Experimental/Interchange/Runtime/Source/Parsers/Fbx/Private/FbxScene.cpp
+++ b/Engine/Plugins/Experimental/Interchange/Runtime/Source/Parsers/Fbx/Private/FbxScene.cpp
@@ -104,6 +104,7 @@ namespace UE
 				, UInterchangeBaseNodeContainer& NodeContainer
 				, TMap<FString, TSharedPtr<FPayloadContextBase, ESPMode::ThreadSafe>>& PayloadContexts)
 			{
+				constexpr bool bResetCache = false;
 				FString NodeName = FFbxHelper::GetFbxObjectName(Node);
 				FString NodeUniqueID = FFbxHelper::GetFbxNodeHierarchyName(Node);
 
@@ -111,7 +112,7 @@ namespace UE
 				check(UnrealNode);
 				if (UnrealParentNode)
 				{
-					UnrealNode->SetParentUid(UnrealParentNode->GetUniqueID());
+					NodeContainer.SetNodeParentUid(UnrealNode->GetUniqueID(), UnrealParentNode->GetUniqueID());
 				}
 				
 				auto GetConvertedTransform = [Node](FbxAMatrix& NewFbxMatrix)
@@ -149,12 +150,12 @@ namespace UE
 						FbxAMatrix GlobalFbxParentMatrix = ParentNode->EvaluateGlobalTransform();
 						FbxAMatrix	LocalFbxMatrix = GlobalFbxParentMatrix.Inverse() * GlobalFbxMatrix;
 						FTransform LocalTransform = GetConvertedTransform(LocalFbxMatrix);
-						UnrealNode->SetCustomLocalTransform(&NodeContainer, LocalTransform);
+						UnrealNode->SetCustomLocalTransform(&NodeContainer, LocalTransform, bResetCache);
 					}
 					else
 					{
 						//No parent, set the same matrix has the global
-						UnrealNode->SetCustomLocalTransform(&NodeContainer, GlobalTransform);
+						UnrealNode->SetCustomLocalTransform(&NodeContainer, GlobalTransform, bResetCache);
 					}
 				}
 
@@ -210,12 +211,12 @@ namespace UE
 									FFbxMesh::GetGlobalJointBindPoseTransform(SDKScene, ParentNode, GlobalFbxParentMatrix);
 									FbxAMatrix	LocalFbxMatrix = GlobalFbxParentMatrix.Inverse() * GlobalBindPoseJointMatrix;
 									FTransform LocalBindPoseJointTransform = GetConvertedTransform(LocalFbxMatrix);
-									UnrealNode->SetCustomBindPoseLocalTransform(&NodeContainer, LocalBindPoseJointTransform);
+									UnrealNode->SetCustomBindPoseLocalTransform(&NodeContainer, LocalBindPoseJointTransform, bResetCache);
 								}
 								else
 								{
 									//No parent, set the same matrix has the global
-									UnrealNode->SetCustomBindPoseLocalTransform(&NodeContainer, GlobalBindPoseJointTransform);
+									UnrealNode->SetCustomBindPoseLocalTransform(&NodeContainer, GlobalBindPoseJointTransform, bResetCache);
 								}
 							}
 
@@ -229,12 +230,12 @@ namespace UE
 									FbxAMatrix GlobalFbxParentMatrix = ParentNode->EvaluateGlobalTransform(FBXSDK_TIME_ZERO);
 									FbxAMatrix	LocalFbxMatrix = GlobalFbxParentMatrix.Inverse() * GlobalFbxMatrix;
 									FTransform LocalTransform = GetConvertedTransform(LocalFbxMatrix);
-									UnrealNode->SetCustomTimeZeroLocalTransform(&NodeContainer, LocalTransform);
+									UnrealNode->SetCustomTimeZeroLocalTransform(&NodeContainer, LocalTransform, bResetCache);
 								}
 								else
 								{
 									//No parent, set the same matrix has the global
-									UnrealNode->SetCustomTimeZeroLocalTransform(&NodeContainer, GlobalTransform);
+									UnrealNode->SetCustomTimeZeroLocalTransform(&NodeContainer, GlobalTransform, bResetCache);
 								}
 							}
 							FFbxAnimation::AddJointAnimation(SDKScene, Node, NodeContainer, UnrealNode, PayloadContexts);
diff --git a/Engine/Plugins/Experimental/Interchange/Runtime/Source/Parsers/Fbx/Private/InterchangeFbxParser.cpp b/Engine/Plugins/Experimental/Interchange/Runtime/Source/Parsers/Fbx/Private/InterchangeFbxParser.cpp
--- a/Engine/Plugins/Experimental/Interchange/Runtime/Source/Parsers/Fbx/Private/InterchangeFbxParser.cpp
+++ b/Engine/Plugins/Experimental/Interchange/Runtime/Source/Parsers/Fbx/Private/InterchangeFbxParser.cpp
@@ -4,6 +4,8 @@
 
 #include "CoreMinimal.h"
 #include "FbxAPI.h"
+#include "HAL/CriticalSection.h"
+#include "Misc/ScopeLock.h"
 #include "Nodes/InterchangeBaseNodeContainer.h"
 #include "InterchangeTextureNode.h"
 #include "UObject/StrongObjectPtr.h"
@@ -64,12 +66,15 @@ namespace UE
 		{
 			check(FbxParserPrivate.IsValid());
 			ResultsContainer->Empty();
+			FString PayloadFilepathCopy;
+			{
+				FScopeLock Lock(&ResultPayloadsCriticalSection);
+				FString& PayloadFilepath = ResultPayloads.FindOrAdd(PayloadKey);
+				PayloadFilepath = ResultFolder + TEXT("/") + PayloadKey + TEXT(".payload");
 
-			FString& PayloadFilepath = ResultPayloads.FindOrAdd(PayloadKey);
-			PayloadFilepath = ResultFolder + TEXT("/") + PayloadKey + TEXT(".payload");
-
-			//Copy the map filename key because we are multithreaded and the TMap can be reallocated
-			FString PayloadFilepathCopy = PayloadFilepath;
+				//Copy the map filename key because we are multithreaded and the TMap can be reallocated
+				PayloadFilepathCopy = PayloadFilepath;
+			}
 			if (!FbxParserPrivate->FetchPayloadData(PayloadKey, PayloadFilepathCopy))
 			{
 				UInterchangeResultError_Generic* Error = AddMessage<UInterchangeResultError_Generic>();
diff --git a/Engine/Plugins/Experimental/Interchange/Runtime/Source/Parsers/Fbx/Public/InterchangeFbxParser.h b/Engine/Plugins/Experimental/Interchange/Runtime/Source/Parsers/Fbx/Public/InterchangeFbxParser.h
--- a/Engine/Plugins/Experimental/Interchange/Runtime/Source/Parsers/Fbx/Public/InterchangeFbxParser.h
+++ b/Engine/Plugins/Experimental/Interchange/Runtime/Source/Parsers/Fbx/Public/InterchangeFbxParser.h
@@ -3,6 +3,8 @@
 #pragma once
 
 #include "CoreMinimal.h"
+#include "HAL/CriticalSection.h"
+#include "Misc/ScopeLock.h"
 #include "Nodes/InterchangeBaseNodeContainer.h"
 #include "InterchangeResultsContainer.h"
 #include "UObject/StrongObjectPtr.h"
@@ -38,6 +40,7 @@ namespace UE
 			FString GetResultFilepath() const { return ResultFilepath; }
 			FString GetResultPayloadFilepath(const FString& PayloadKey) const
 			{
+				FScopeLock Lock(&ResultPayloadsCriticalSection);
 				if (const FString* PayloadPtr = ResultPayloads.Find(PayloadKey))
 				{
 					return *PayloadPtr;
@@ -61,6 +64,7 @@ namespace UE
 			TObjectPtr<UInterchangeResultsContainer> ResultsContainer = nullptr;
 			FString SourceFilename;
 			FString ResultFilepath;
+			mutable FCriticalSection ResultPayloadsCriticalSection;
 			TMap<FString, FString> ResultPayloads;
 			TUniquePtr<UE::Interchange::Private::FFbxParser> FbxParserPrivate;
 		};
diff --git a/Engine/Plugins/Experimental/Interchange/Runtime/Source/Pipelines/Private/InterchangeGenericMaterialPipeline.cpp b/Engine/Plugins/Experimental/Interchange/Runtime/Source/Pipelines/Private/InterchangeGenericMaterialPipeline.cpp
--- a/Engine/Plugins/Experimental/Interchange/Runtime/Source/Pipelines/Private/InterchangeGenericMaterialPipeline.cpp
+++ b/Engine/Plugins/Experimental/Interchange/Runtime/Source/Pipelines/Private/InterchangeGenericMaterialPipeline.cpp
@@ -14,6 +14,7 @@
 #include "InterchangeTextureNode.h"
 #include "InterchangePipelineLog.h"
 #include "InterchangeSourceData.h"
+#include "Nodes/InterchangeBaseNode.h"
 
 #include "Materials/MaterialExpressionAdd.h"
 #include "Materials/MaterialExpressionComponentMask.h"
@@ -227,9 +228,9 @@ bool UInterchangeGenericMaterialPipeline::HandlePhongModel(const UInterchangeSha
 		const FName MaterialFunctionMemberName = GET_MEMBER_NAME_CHECKED(UMaterialExpressionMaterialFunctionCall, MaterialFunction);
 
 		FunctionCallExpression->AddStringAttribute(
-			MaterialFunctionMemberName.ToString(),
+			MaterialFunctionMemberName,
 			TEXT("MaterialFunction'/Engine/Functions/Engine_MaterialFunctions01/Shading/ConvertFromDiffSpec.ConvertFromDiffSpec'"));
-		FunctionCallExpression->AddApplyAndFillDelegates<FString>(MaterialFunctionMemberName.ToString(), UMaterialExpressionMaterialFunctionCall::StaticClass(), MaterialFunctionMemberName);
+		FunctionCallExpression->AddApplyAndFillDelegates<FString>(MaterialFunctionMemberName, UMaterialExpressionMaterialFunctionCall::StaticClass(), MaterialFunctionMemberName);
 
 		MaterialFactoryNode->ConnectOutputToBaseColor(FunctionCallExpressionUid, UE::Interchange::Materials::PBR::Parameters::BaseColor.ToString());
 		MaterialFactoryNode->ConnectOutputToMetallic(FunctionCallExpressionUid, UE::Interchange::Materials::PBR::Parameters::Metallic.ToString());
@@ -277,7 +278,7 @@ bool UInterchangeGenericMaterialPipeline::HandlePhongModel(const UInterchangeSha
 						CreateExpressionNode(TEXT("DivideShininess"), ShininessExpressionFactoryNode.Get<0>()->GetUniqueID(), UMaterialExpressionDivide::StaticClass());
 
 					const float ShininessScale = 100.f; // Divide shininess by 100 to bring it into a 0-1 range for roughness.
-					const FString ShininessScaleParameterName = GET_MEMBER_NAME_CHECKED(UMaterialExpressionDivide, ConstB).ToString();
+					const FName ShininessScaleParameterName = GET_MEMBER_NAME_CHECKED(UMaterialExpressionDivide, ConstB);
 					DivideShininessNode->AddFloatAttribute(ShininessScaleParameterName, ShininessScale);
 					DivideShininessNode->AddApplyAndFillDelegates<float>(ShininessScaleParameterName, UMaterialExpressionDivide::StaticClass(),  GET_MEMBER_NAME_CHECKED(UMaterialExpressionDivide, ConstB));
 
@@ -425,9 +426,9 @@ bool UInterchangeGenericMaterialPipeline::HandleStandardSurfaceModel(const UInte
 		const FName MaterialFunctionMemberName = GET_MEMBER_NAME_CHECKED(UMaterialExpressionMaterialFunctionCall, MaterialFunction);
 
 		FunctionCallExpression->AddStringAttribute(
-			MaterialFunctionMemberName.ToString(),
+			MaterialFunctionMemberName,
 			TEXT("MaterialFunction'/Interchange/Functions/MX_StandardSurface.MX_StandardSurface'"));
-		FunctionCallExpression->AddApplyAndFillDelegates<FString>(MaterialFunctionMemberName.ToString(), UMaterialExpressionMaterialFunctionCall::StaticClass(), MaterialFunctionMemberName);
+		FunctionCallExpression->AddApplyAndFillDelegates<FString>(MaterialFunctionMemberName, UMaterialExpressionMaterialFunctionCall::StaticClass(), MaterialFunctionMemberName);
 
 		using namespace UE::Interchange::Materials;
 
@@ -793,9 +794,9 @@ void UInterchangeGenericMaterialPipeline::HandleCommonParameters(const UIntercha
 				UInterchangeMaterialExpressionFactoryNode* IORLerp = CreateExpressionNode(TEXT("IORLerp"), ShaderGraphNode->GetUniqueID(), UMaterialExpressionLinearInterpolate::StaticClass());
 
 				const float AirIOR = 1.f;
-				const FString ConstAMemberName = GET_MEMBER_NAME_CHECKED(UMaterialExpressionLinearInterpolate, ConstA).ToString();
+				const FName ConstAMemberName = GET_MEMBER_NAME_CHECKED(UMaterialExpressionLinearInterpolate, ConstA);
 				IORLerp->AddFloatAttribute(ConstAMemberName, AirIOR);
-				IORLerp->AddApplyAndFillDelegates<float>(ConstAMemberName, UMaterialExpressionLinearInterpolate::StaticClass(), *ConstAMemberName);
+				IORLerp->AddApplyAndFillDelegates<float>(ConstAMemberName, UMaterialExpressionLinearInterpolate::StaticClass(), ConstAMemberName);
 
 				UInterchangeShaderPortsAPI::ConnectOuputToInput(IORLerp, GET_MEMBER_NAME_CHECKED(UMaterialExpressionLinearInterpolate, B).ToString(),
 					ExpressionFactoryNode.Get<0>()->GetUniqueID(), ExpressionFactoryNode.Get<1>());
@@ -817,9 +818,9 @@ void UInterchangeGenericMaterialPipeline::HandleFlattenNormalNode(const UInterch
 
 	FlattenNormalFactoryNode->SetCustomExpressionClassName(UMaterialExpressionMaterialFunctionCall::StaticClass()->GetName());
 
-	const FString MaterialFunctionMemberName = GET_MEMBER_NAME_CHECKED(UMaterialExpressionMaterialFunctionCall, MaterialFunction).ToString();
+	const FName MaterialFunctionMemberName = GET_MEMBER_NAME_CHECKED(UMaterialExpressionMaterialFunctionCall, MaterialFunction);
 	FlattenNormalFactoryNode->AddStringAttribute(MaterialFunctionMemberName, TEXT("/Engine/Functions/Engine_MaterialFunctions01/Texturing/FlattenNormal.FlattenNormal"));
-	FlattenNormalFactoryNode->AddApplyAndFillDelegates<FString>(MaterialFunctionMemberName, UMaterialExpressionMaterialFunctionCall::StaticClass(), *MaterialFunctionMemberName);
+	FlattenNormalFactoryNode->AddApplyAndFillDelegates<FString>(MaterialFunctionMemberName, UMaterialExpressionMaterialFunctionCall::StaticClass(), MaterialFunctionMemberName);
 
 	// Normal
 	{
@@ -932,9 +933,9 @@ void UInterchangeGenericMaterialPipeline::HandleTextureCoordinateNode(const UInt
 		int32 CoordIndex;
 		if (ShaderNode->GetInt32Attribute(UInterchangeShaderPortsAPI::MakeInputValueKey(Nodes::TextureCoordinate::Inputs::Index.ToString()), CoordIndex))
 		{
-			const FString CoordinateIndexMemberName = GET_MEMBER_NAME_CHECKED(UMaterialExpressionTextureCoordinate, CoordinateIndex).ToString();
+			const FName CoordinateIndexMemberName = GET_MEMBER_NAME_CHECKED(UMaterialExpressionTextureCoordinate, CoordinateIndex);
 			TexCoordFactoryNode->AddInt32Attribute(CoordinateIndexMemberName, CoordIndex);
-			TexCoordFactoryNode->AddApplyAndFillDelegates<int32>(CoordinateIndexMemberName, UMaterialExpressionTextureCoordinate::StaticClass(), *CoordinateIndexMemberName);
+			TexCoordFactoryNode->AddApplyAndFillDelegates<int32>(CoordinateIndexMemberName, UMaterialExpressionTextureCoordinate::StaticClass(), CoordinateIndexMemberName);
 		}
 	}
 
@@ -944,9 +945,9 @@ void UInterchangeGenericMaterialPipeline::HandleTextureCoordinateNode(const UInt
 
 		if (UTilingValue.IsType<float>())
 		{
-			const FString UTilingMemberName = GET_MEMBER_NAME_CHECKED(UMaterialExpressionTextureCoordinate, UTiling).ToString();
+			const FName UTilingMemberName = GET_MEMBER_NAME_CHECKED(UMaterialExpressionTextureCoordinate, UTiling);
 			TexCoordFactoryNode->AddFloatAttribute(UTilingMemberName, UTilingValue.Get<float>());
-			TexCoordFactoryNode->AddApplyAndFillDelegates<float>(UTilingMemberName, UMaterialExpressionTextureCoordinate::StaticClass(), *UTilingMemberName);
+			TexCoordFactoryNode->AddApplyAndFillDelegates<float>(UTilingMemberName, UMaterialExpressionTextureCoordinate::StaticClass(), UTilingMemberName);
 		}
 	}
 
@@ -956,9 +957,9 @@ void UInterchangeGenericMaterialPipeline::HandleTextureCoordinateNode(const UInt
 
 		if (VTilingValue.IsType<float>())
 		{
-			const FString VTilingMemberName = GET_MEMBER_NAME_CHECKED(UMaterialExpressionTextureCoordinate, VTiling).ToString();
+			const FName VTilingMemberName = GET_MEMBER_NAME_CHECKED(UMaterialExpressionTextureCoordinate, VTiling);
 			TexCoordFactoryNode->AddFloatAttribute(VTilingMemberName, VTilingValue.Get<float>());
-			TexCoordFactoryNode->AddApplyAndFillDelegates<float>(VTilingMemberName, UMaterialExpressionTextureCoordinate::StaticClass(), *VTilingMemberName);
+			TexCoordFactoryNode->AddApplyAndFillDelegates<float>(VTilingMemberName, UMaterialExpressionTextureCoordinate::StaticClass(), VTilingMemberName);
 		}
 	}
 
@@ -991,9 +992,9 @@ void UInterchangeGenericMaterialPipeline::HandleTextureCoordinateNode(const UInt
 			UInterchangeMaterialExpressionFactoryNode* CallRotatorExpression =
 				CreateExpressionNode(RotateExpression.Get<0>()->GetDisplayLabel() + TEXT("_Rotator"), TexCoordFactoryNode->GetUniqueID(), UMaterialExpressionMaterialFunctionCall::StaticClass());
 
-			const FString MaterialFunctionMemberName = GET_MEMBER_NAME_CHECKED(UMaterialExpressionMaterialFunctionCall, MaterialFunction).ToString();
+			const FName MaterialFunctionMemberName = GET_MEMBER_NAME_CHECKED(UMaterialExpressionMaterialFunctionCall, MaterialFunction);
 			CallRotatorExpression->AddStringAttribute(MaterialFunctionMemberName, TEXT("/Engine/Functions/Engine_MaterialFunctions02/Texturing/CustomRotator.CustomRotator"));
-			CallRotatorExpression->AddApplyAndFillDelegates<FString>(MaterialFunctionMemberName, UMaterialExpressionMaterialFunctionCall::StaticClass(), *MaterialFunctionMemberName);
+			CallRotatorExpression->AddApplyAndFillDelegates<FString>(MaterialFunctionMemberName, UMaterialExpressionMaterialFunctionCall::StaticClass(), MaterialFunctionMemberName);
 
 			UInterchangeShaderPortsAPI::ConnectDefaultOuputToInput(CallRotatorExpression, TEXT("UVs"), TexCoordFactoryNode->GetUniqueID());
 			UInterchangeShaderPortsAPI::ConnectOuputToInput(CallRotatorExpression, TEXT("Rotation Angle (0-1)"), RotateExpression.Get<0>()->GetUniqueID(), RotateExpression.Get<1>());
@@ -1084,39 +1085,39 @@ void UInterchangeGenericMaterialPipeline::HandleMaskNode(const UInterchangeShade
 	bool bBChannel = false;
 	bool bAChannel = false;
 	bool bIsAnyMaskChannelSet = 
-		ShaderNode->GetBooleanAttribute(Mask::Attributes::R.ToString(), bRChannel) ||
-		ShaderNode->GetBooleanAttribute(Mask::Attributes::G.ToString(), bGChannel) ||
-		ShaderNode->GetBooleanAttribute(Mask::Attributes::B.ToString(), bBChannel) ||
-		ShaderNode->GetBooleanAttribute(Mask::Attributes::A.ToString(), bAChannel);
+		ShaderNode->GetBooleanAttribute(Mask::Attributes::R, bRChannel) ||
+		ShaderNode->GetBooleanAttribute(Mask::Attributes::G, bGChannel) ||
+		ShaderNode->GetBooleanAttribute(Mask::Attributes::B, bBChannel) ||
+		ShaderNode->GetBooleanAttribute(Mask::Attributes::A, bAChannel);
 
 	if(bIsAnyMaskChannelSet)
 	{
 		// R
 		{
-			const FString RMemberName = GET_MEMBER_NAME_CHECKED(UMaterialExpressionComponentMask, R).ToString();
+			const FName RMemberName = GET_MEMBER_NAME_CHECKED(UMaterialExpressionComponentMask, R);
 			MaskFactoryNode->AddBooleanAttribute(RMemberName, bRChannel);
-			MaskFactoryNode->AddApplyAndFillDelegates<bool>(RMemberName, UMaterialExpressionComponentMask::StaticClass(), *RMemberName);
+			MaskFactoryNode->AddApplyAndFillDelegates<bool>(RMemberName, UMaterialExpressionComponentMask::StaticClass(), RMemberName);
 		}
 
 		// G
 		{
-			const FString GMemberName = GET_MEMBER_NAME_CHECKED(UMaterialExpressionComponentMask, G).ToString();
+			const FName GMemberName = GET_MEMBER_NAME_CHECKED(UMaterialExpressionComponentMask, G);
 			MaskFactoryNode->AddBooleanAttribute(GMemberName, bGChannel);
-			MaskFactoryNode->AddApplyAndFillDelegates<bool>(GMemberName, UMaterialExpressionComponentMask::StaticClass(), *GMemberName);
+			MaskFactoryNode->AddApplyAndFillDelegates<bool>(GMemberName, UMaterialExpressionComponentMask::StaticClass(), GMemberName);
 		}
 
 		// B
 		{
-			const FString BMemberName = GET_MEMBER_NAME_CHECKED(UMaterialExpressionComponentMask, B).ToString();
+			const FName BMemberName = GET_MEMBER_NAME_CHECKED(UMaterialExpressionComponentMask, B);
 			MaskFactoryNode->AddBooleanAttribute(BMemberName, bBChannel);
-			MaskFactoryNode->AddApplyAndFillDelegates<bool>(BMemberName, UMaterialExpressionComponentMask::StaticClass(), *BMemberName);
+			MaskFactoryNode->AddApplyAndFillDelegates<bool>(BMemberName, UMaterialExpressionComponentMask::StaticClass(), BMemberName);
 		}
 
 		// A
 		{
-			const FString AMemberName = GET_MEMBER_NAME_CHECKED(UMaterialExpressionComponentMask, A).ToString();
+			const FName AMemberName = GET_MEMBER_NAME_CHECKED(UMaterialExpressionComponentMask, A);
 			MaskFactoryNode->AddBooleanAttribute(AMemberName, bAChannel);
-			MaskFactoryNode->AddApplyAndFillDelegates<bool>(AMemberName, UMaterialExpressionComponentMask::StaticClass(), *AMemberName);
+			MaskFactoryNode->AddApplyAndFillDelegates<bool>(AMemberName, UMaterialExpressionComponentMask::StaticClass(), AMemberName);
 		}
 	}
 
@@ -1258,8 +1259,8 @@ UInterchangeMaterialExpressionFactoryNode* UInterchangeGenericMaterialPipeline::
 	if (ShaderNode->GetFloatAttribute(UInterchangeShaderPortsAPI::MakeInputValueKey(InputName), InputValue))
 	{
 		const FName DefaultValueMemberName = GET_MEMBER_NAME_CHECKED(UMaterialExpressionScalarParameter, DefaultValue);
-		MaterialExpressionFactoryNode->AddFloatAttribute(DefaultValueMemberName.ToString(), InputValue);
-		MaterialExpressionFactoryNode->AddApplyAndFillDelegates<float>(DefaultValueMemberName.ToString(), UMaterialExpressionScalarParameter::StaticClass(), DefaultValueMemberName);
+		MaterialExpressionFactoryNode->AddFloatAttribute(DefaultValueMemberName, InputValue);
+		MaterialExpressionFactoryNode->AddApplyAndFillDelegates<float>(DefaultValueMemberName, UMaterialExpressionScalarParameter::StaticClass(), DefaultValueMemberName);
 	}
 
 	return MaterialExpressionFactoryNode;
@@ -1273,8 +1274,8 @@ UInterchangeMaterialExpressionFactoryNode* UInterchangeGenericMaterialPipeline::
 	if (ShaderNode->GetLinearColorAttribute(UInterchangeShaderPortsAPI::MakeInputValueKey(InputName), InputValue))
 	{
 		const FName DefaultValueMemberName = GET_MEMBER_NAME_CHECKED(UMaterialExpressionVectorParameter, DefaultValue);
-		MaterialExpressionFactoryNode->AddLinearColorAttribute(DefaultValueMemberName.ToString(), InputValue);
-		MaterialExpressionFactoryNode->AddApplyAndFillDelegates<FLinearColor>(DefaultValueMemberName.ToString(), UMaterialExpressionVectorParameter::StaticClass(), DefaultValueMemberName);
+		MaterialExpressionFactoryNode->AddLinearColorAttribute(DefaultValueMemberName, InputValue);
+		MaterialExpressionFactoryNode->AddApplyAndFillDelegates<FLinearColor>(DefaultValueMemberName, UMaterialExpressionVectorParameter::StaticClass(), DefaultValueMemberName);
 	}
 
 	return MaterialExpressionFactoryNode;
@@ -1288,8 +1289,8 @@ UInterchangeMaterialExpressionFactoryNode* UInterchangeGenericMaterialPipeline::
 		UInterchangeMaterialExpressionFactoryNode* VectorParameterFactoryNode = CreateExpressionNode(InputName, ParentUid, UMaterialExpressionVectorParameter::StaticClass());
 
 		const FName DefaultValueMemberName = GET_MEMBER_NAME_CHECKED(UMaterialExpressionVectorParameter, DefaultValue);
-		VectorParameterFactoryNode->AddLinearColorAttribute(DefaultValueMemberName.ToString(), FLinearColor(InputValue.X, InputValue.Y, 0.f));
-		VectorParameterFactoryNode->AddApplyAndFillDelegates<FLinearColor>(DefaultValueMemberName.ToString(), UMaterialExpressionVectorParameter::StaticClass(), DefaultValueMemberName);
+		VectorParameterFactoryNode->AddLinearColorAttribute(DefaultValueMemberName, FLinearColor(InputValue.X, InputValue.Y, 0.f));
+		VectorParameterFactoryNode->AddApplyAndFillDelegates<FLinearColor>(DefaultValueMemberName, UMaterialExpressionVectorParameter::StaticClass(), DefaultValueMemberName);
 
 		// Defaults to R&G
 		UInterchangeMaterialExpressionFactoryNode* ComponentMaskFactoryNode = CreateExpressionNode(InputName + TEXT("_Mask"), ParentUid, UMaterialExpressionComponentMask::StaticClass());
diff --git a/Engine/Plugins/Experimental/Interchange/Runtime/Source/Pipelines/Private/InterchangeGenericScenesPipeline.cpp b/Engine/Plugins/Experimental/Interchange/Runtime/Source/Pipelines/Private/InterchangeGenericScenesPipeline.cpp
--- a/Engine/Plugins/Experimental/Interchange/Runtime/Source/Pipelines/Private/InterchangeGenericScenesPipeline.cpp
+++ b/Engine/Plugins/Experimental/Interchange/Runtime/Source/Pipelines/Private/InterchangeGenericScenesPipeline.cpp
@@ -86,11 +86,11 @@ void UInterchangeGenericLevelPipeline::ExecuteSceneNodePreImport(UInterchangeBas
 	}
 
 	ActorFactoryNode->InitializeNode(TEXT("Factory_") + SceneNode->GetUniqueID(), SceneNode->GetDisplayLabel(), EInterchangeNodeContainerType::FactoryData);
-
+	const FString ActorFactoryNodeUid = FactoryNodeContainer->AddNode(ActorFactoryNode);
 	if (!SceneNode->GetParentUid().IsEmpty())
 	{
-		ActorFactoryNode->SetParentUid(TEXT("Factory_") + SceneNode->GetParentUid());
-		ActorFactoryNode->AddFactoryDependencyUid(ActorFactoryNode->GetParentUid());
+		FactoryNodeContainer->SetNodeParentUid(ActorFactoryNodeUid, TEXT("Factory_") + SceneNode->GetParentUid());
+		ActorFactoryNode->AddFactoryDependencyUid(ActorFactoryNodeUid);
 	}
 
 	ActorFactoryNode->AddTargetNodeUid(SceneNode->GetUniqueID());
diff --git a/Engine/Plugins/Experimental/Interchange/Runtime/Source/Pipelines/Private/InterchangeGenericSkeletalMeshPipeline.cpp b/Engine/Plugins/Experimental/Interchange/Runtime/Source/Pipelines/Private/InterchangeGenericSkeletalMeshPipeline.cpp
--- a/Engine/Plugins/Experimental/Interchange/Runtime/Source/Pipelines/Private/InterchangeGenericSkeletalMeshPipeline.cpp
+++ b/Engine/Plugins/Experimental/Interchange/Runtime/Source/Pipelines/Private/InterchangeGenericSkeletalMeshPipeline.cpp
@@ -529,7 +529,7 @@ void UInterchangeGenericMeshPipeline::AddLodDataToSkeletalMesh(const UInterchang
 		{
 			//Add the data for the LOD (skeleton Unique ID and all the mesh node fbx path, so we can find them when we will create the payload data)
 			LodDataNode = CreateSkeletalMeshLodDataNode(SkeletalMeshLodDataName, SkeletalMeshLodDataUniqueID);
-			LodDataNode->SetParentUid(SkeletalMeshUid);
+			BaseNodeContainer->SetNodeParentUid(SkeletalMeshLodDataUniqueID, SkeletalMeshUid);
 			LodDataNode->SetCustomSkeletonUid(SkeletonUid);
 			SkeletalMeshFactoryNode->AddLodDataUniqueId(SkeletalMeshLodDataUniqueID);
 		}
diff --git a/Engine/Plugins/Experimental/Interchange/Runtime/Source/Pipelines/Private/InterchangeGenericStaticMeshPipeline.cpp b/Engine/Plugins/Experimental/Interchange/Runtime/Source/Pipelines/Private/InterchangeGenericStaticMeshPipeline.cpp
--- a/Engine/Plugins/Experimental/Interchange/Runtime/Source/Pipelines/Private/InterchangeGenericStaticMeshPipeline.cpp
+++ b/Engine/Plugins/Experimental/Interchange/Runtime/Source/Pipelines/Private/InterchangeGenericStaticMeshPipeline.cpp
@@ -464,7 +464,7 @@ void UInterchangeGenericMeshPipeline::AddLodDataToStaticMesh(UInterchangeStaticM
 		{
 			// Add the data for the LOD (all the mesh node fbx path, so we can find them when we will create the payload data)
 			LodDataNode = CreateStaticMeshLodDataNode(StaticMeshLodDataName, StaticMeshLodDataUniqueID);
-			LodDataNode->SetParentUid(StaticMeshFactoryUid);
+			BaseNodeContainer->SetNodeParentUid(StaticMeshLodDataUniqueID, StaticMeshFactoryUid);
 			StaticMeshFactoryNode->AddLodDataUniqueId(StaticMeshLodDataUniqueID);
 		}
 		constexpr bool bAddSourceNodeName = true;
diff --git a/Engine/Source/Runtime/Experimental/Interchange/Core/Private/Nodes/InterchangeBaseNode.cpp b/Engine/Source/Runtime/Experimental/Interchange/Core/Private/Nodes/InterchangeBaseNode.cpp
--- a/Engine/Source/Runtime/Experimental/Interchange/Core/Private/Nodes/InterchangeBaseNode.cpp
+++ b/Engine/Source/Runtime/Experimental/Interchange/Core/Private/Nodes/InterchangeBaseNode.cpp
@@ -36,10 +36,10 @@ namespace UE
 			return AttributeKey;
 		}
 
-		const FString& FBaseNodeStaticData::TargetAssetIDsKey()
+		const FAttributeKey& FBaseNodeStaticData::TargetAssetIDsKey()
 		{
-			static FString Key(TEXT("__TARGET_ASSET_IDS_"));
-			return Key;
+			static FAttributeKey AttributeKey(TEXT("__TARGET_ASSET_IDS_"));
+			return AttributeKey;
 		}
 
 		const FAttributeKey& FBaseNodeStaticData::ClassTypeAttributeKey()
@@ -72,7 +72,7 @@ namespace UE
 UInterchangeBaseNode::UInterchangeBaseNode()
 {
 	Attributes = MakeShared<UE::Interchange::FAttributeStorage, ESPMode::ThreadSafe>();
-	TargetNodes.Initialize(Attributes, UE::Interchange::FBaseNodeStaticData::TargetAssetIDsKey());
+	TargetNodes.Initialize(Attributes, UE::Interchange::FBaseNodeStaticData::TargetAssetIDsKey().ToString());
 	RegisterAttribute<bool>(UE::Interchange::FBaseNodeStaticData::IsEnabledKey(), true);
 	RegisterAttribute<uint8>(UE::Interchange::FBaseNodeStaticData::NodeContainerTypeKey(), static_cast<uint8>(EInterchangeNodeContainerType::None));
 	RegisterAttribute<uint8>(UE::Interchange::FBaseNodeStaticData::ReimportStrategyFlagsKey(), static_cast<uint8>(EReimportStrategyFlags::ApplyNoProperties));
@@ -115,6 +115,7 @@ FName UInterchangeBaseNode::GetIconName() const
 FString UInterchangeBaseNode::GetKeyDisplayName(const UE::Interchange::FAttributeKey& NodeAttributeKey) const
 {
 	FString KeyDisplayName = NodeAttributeKey.ToString();
+	FString NodeAttributeKeyString = NodeAttributeKey.ToString();
 	if (NodeAttributeKey == UE::Interchange::FBaseNodeStaticData::ParentIDKey())
 	{
 		KeyDisplayName = TEXT("Parent Unique ID");
@@ -139,28 +140,29 @@ FString UInterchangeBaseNode::GetKeyDisplayName(const UE::Interchange::FAttribut
 	{
 		KeyDisplayName = TEXT("Imported Asset Name");
 	}
-	else if (NodeAttributeKey.Key.Equals(UE::Interchange::FBaseNodeStaticData::TargetAssetIDsKey()))
+	else if (NodeAttributeKey == UE::Interchange::FBaseNodeStaticData::NodeContainerTypeKey())
+	{
+		KeyDisplayName = TEXT("Node Container Type");
+	}
+	else if (NodeAttributeKey == UE::Interchange::FBaseNodeStaticData::ReimportStrategyFlagsKey())
+	{
+		KeyDisplayName = TEXT("Re-Import Strategy");
+	}
+	else if (NodeAttributeKey == UE::Interchange::FBaseNodeStaticData::TargetAssetIDsKey())
 	{
 		KeyDisplayName = TEXT("Target Asset Count");
 	}
-	else if (NodeAttributeKey.Key.StartsWith(UE::Interchange::FBaseNodeStaticData::TargetAssetIDsKey()))
+	else if (NodeAttributeKeyString.StartsWith(UE::Interchange::FBaseNodeStaticData::TargetAssetIDsKey().ToString()))
 	{
 		KeyDisplayName = TEXT("Target Asset Index ");
 		const FString IndexKey = UE::Interchange::TArrayAttributeHelper<FString>::IndexKey();
-		int32 IndexPosition = NodeAttributeKey.Key.Find(IndexKey) + IndexKey.Len();
-		if (IndexPosition < NodeAttributeKey.Key.Len())
+		int32 IndexPosition = NodeAttributeKeyString.Find(IndexKey) + IndexKey.Len();
+		if (IndexPosition < NodeAttributeKeyString.Len())
 		{
-			KeyDisplayName += NodeAttributeKey.Key.RightChop(IndexPosition);
+			KeyDisplayName += NodeAttributeKeyString.RightChop(IndexPosition);
 		}
 	}
-	else if (NodeAttributeKey == UE::Interchange::FBaseNodeStaticData::NodeContainerTypeKey())
-	{
-		KeyDisplayName = TEXT("Node Container Type");
-	}
-	else if (NodeAttributeKey == UE::Interchange::FBaseNodeStaticData::ReimportStrategyFlagsKey())
-	{
-		KeyDisplayName = TEXT("Re-Import Strategy");
-	}
+	
 	return KeyDisplayName;
 }
 
@@ -190,58 +192,68 @@ void UInterchangeBaseNode::GetAttributeKeys(TArray<UE::Interchange::FAttributeKe
 	Attributes->GetAttributeKeys(AttributeKeys);
 }
 
-bool UInterchangeBaseNode::RemoveAttribute(const FString& NodeAttributeKey)
+bool UInterchangeBaseNode::RemoveAttribute(const FName& NodeAttributeKey)
 {
 	Attributes->UnregisterAttribute(UE::Interchange::FAttributeKey(NodeAttributeKey));
 	return !HasAttribute(UE::Interchange::FAttributeKey(NodeAttributeKey));
 }
 
-bool UInterchangeBaseNode::AddBooleanAttribute(const FString& NodeAttributeKey, const bool& Value)
+bool UInterchangeBaseNode::AddBooleanAttribute(const FName& NodeAttributeKey, const bool& Value)
 {
 	INTERCHANGE_BASE_NODE_ADD_ATTRIBUTE(bool);
 }
 
-bool UInterchangeBaseNode::GetBooleanAttribute(const FString& NodeAttributeKey, bool& OutValue) const
+bool UInterchangeBaseNode::GetBooleanAttribute(const FName& NodeAttributeKey, bool& OutValue) const
 {
 	INTERCHANGE_BASE_NODE_GET_ATTRIBUTE(bool);
 }
 
-bool UInterchangeBaseNode::AddInt32Attribute(const FString& NodeAttributeKey, const int32& Value)
+bool UInterchangeBaseNode::AddInt32Attribute(const FName& NodeAttributeKey, const int32& Value)
 {
 	INTERCHANGE_BASE_NODE_ADD_ATTRIBUTE(int32);
 }
 
-bool UInterchangeBaseNode::GetInt32Attribute(const FString& NodeAttributeKey, int32& OutValue) const
+bool UInterchangeBaseNode::GetInt32Attribute(const FName& NodeAttributeKey, int32& OutValue) const
 {
 	INTERCHANGE_BASE_NODE_GET_ATTRIBUTE(int32);
 }
 
-bool UInterchangeBaseNode::AddFloatAttribute(const FString& NodeAttributeKey, const float& Value)
+bool UInterchangeBaseNode::AddFloatAttribute(const FName& NodeAttributeKey, const float& Value)
 {
 	INTERCHANGE_BASE_NODE_ADD_ATTRIBUTE(float);
 }
 
-bool UInterchangeBaseNode::GetFloatAttribute(const FString& NodeAttributeKey, float& OutValue) const
+bool UInterchangeBaseNode::GetFloatAttribute(const FName& NodeAttributeKey, float& OutValue) const
 {
 	INTERCHANGE_BASE_NODE_GET_ATTRIBUTE(float);
 }
 
-bool UInterchangeBaseNode::AddStringAttribute(const FString& NodeAttributeKey, const FString& Value)
+bool UInterchangeBaseNode::AddStringAttribute(const FName& NodeAttributeKey, const FString& Value)
 {
 	INTERCHANGE_BASE_NODE_ADD_ATTRIBUTE(FString);
 }
 
-bool UInterchangeBaseNode::GetStringAttribute(const FString& NodeAttributeKey, FString& OutValue) const
+bool UInterchangeBaseNode::GetStringAttribute(const FName& NodeAttributeKey, FString& OutValue) const
 {
 	INTERCHANGE_BASE_NODE_GET_ATTRIBUTE(FString);
 }
 
-bool UInterchangeBaseNode::AddLinearColorAttribute(const FString& NodeAttributeKey, const FLinearColor& Value)
+bool UInterchangeBaseNode::AddGuidAttribute(const FName& NodeAttributeKey, const FGuid& Value)
+{
+	INTERCHANGE_BASE_NODE_ADD_ATTRIBUTE(FGuid);
+}
+
+bool UInterchangeBaseNode::GetGuidAttribute(const FName& NodeAttributeKey, FGuid& OutValue) const
+{
+	INTERCHANGE_BASE_NODE_GET_ATTRIBUTE(FGuid);
+}
+
+bool UInterchangeBaseNode::AddLinearColorAttribute(const FName& NodeAttributeKey, const FLinearColor& Value)
 {
 	INTERCHANGE_BASE_NODE_ADD_ATTRIBUTE(FLinearColor);
 }
 
-bool UInterchangeBaseNode::GetLinearColorAttribute(const FString& NodeAttributeKey, FLinearColor& OutValue) const
+bool UInterchangeBaseNode::GetLinearColorAttribute(const FName& NodeAttributeKey, FLinearColor& OutValue) const
 {
 	INTERCHANGE_BASE_NODE_GET_ATTRIBUTE(FLinearColor);
 }
diff --git a/Engine/Source/Runtime/Experimental/Interchange/Core/Private/Nodes/InterchangeBaseNodeContainer.cpp b/Engine/Source/Runtime/Experimental/Interchange/Core/Private/Nodes/InterchangeBaseNodeContainer.cpp
--- a/Engine/Source/Runtime/Experimental/Interchange/Core/Private/Nodes/InterchangeBaseNodeContainer.cpp
+++ b/Engine/Source/Runtime/Experimental/Interchange/Core/Private/Nodes/InterchangeBaseNodeContainer.cpp
@@ -138,7 +138,15 @@ bool UInterchangeBaseNodeContainer::SetNodeParentUid(const FString& NodeUniqueID
 		return false;
 	}
 	UInterchangeBaseNode* Node = Nodes.FindChecked(NodeUniqueID);
+	
+	PRAGMA_DISABLE_DEPRECATION_WARNINGS
 	Node->SetParentUid(NewParentNodeUid);
+	PRAGMA_ENABLE_DEPRECATION_WARNINGS
+
+	//Update the children cache
+	TArray<FString>& Children = ChildrenCache.FindOrAdd(NewParentNodeUid);
+	Children.Add(NodeUniqueID);
+
 	return true;
 }
 
@@ -150,15 +158,21 @@ int32 UInterchangeBaseNodeContainer::GetNodeChildrenCount(const FString& NodeUni
 
 TArray<FString> UInterchangeBaseNodeContainer::GetNodeChildrenUids(const FString& NodeUniqueID) const
 {
-	TArray<FString> OutChildrenUids;
+	if(TArray<FString>* CacheChildrenPtr = ChildrenCache.Find(NodeUniqueID))
+	{
+		return *CacheChildrenPtr;
+	}
+
+	//Update the cache
+	TArray<FString>& CacheChildren = ChildrenCache.Add(NodeUniqueID);
 	for (const auto& NodeKeyValue : Nodes)
 	{
 		if (NodeKeyValue.Value->GetParentUid() == NodeUniqueID)
 		{
-			OutChildrenUids.Add(NodeKeyValue.Key);
+			CacheChildren.Add(NodeKeyValue.Key);
 		}
 	}
-	return OutChildrenUids;
+	return CacheChildren;
 }
 
 UInterchangeBaseNode* UInterchangeBaseNodeContainer::GetNodeChildren(const FString& NodeUniqueID, int32 ChildIndex)
@@ -224,6 +238,7 @@ void UInterchangeBaseNodeContainer::SerializeNodeContainerData(FArchive& Ar)
 			BaseNode->Serialize(Ar);
 			AddNode(BaseNode);
 		}
+		ComputeChildrenCache();
 	}
 }
 
@@ -255,6 +270,24 @@ void UInterchangeBaseNodeContainer::LoadFromFile(const FString& Filename)
 	SerializeNodeContainerData(Ar);
 }
 
+void UInterchangeBaseNodeContainer::ComputeChildrenCache()
+{
+	ResetChildrenCache();
+	for (const auto& NodeKeyValue : Nodes)
+	{
+		//Force all node to have a cache, even if there is no parent
+		ChildrenCache.FindOrAdd(NodeKeyValue.Key);
+
+		//Update the parent cache
+		const FString ParentUid = NodeKeyValue.Value->GetParentUid();
+		if (!ParentUid.IsEmpty())
+		{
+			TArray<FString>& Children = ChildrenCache.FindOrAdd(ParentUid);
+			Children.Add(NodeKeyValue.Key);
+		}
+	}
+}
+
 UInterchangeBaseNode* UInterchangeBaseNodeContainer::GetNodeChildrenInternal(const FString& NodeUniqueID, int32 ChildIndex)
 {
 	TArray<FString> ChildrenUids = GetNodeChildrenUids(NodeUniqueID);
diff --git a/Engine/Source/Runtime/Experimental/Interchange/Core/Private/Nodes/InterchangeFactoryBaseNode.cpp b/Engine/Source/Runtime/Experimental/Interchange/Core/Private/Nodes/InterchangeFactoryBaseNode.cpp
--- a/Engine/Source/Runtime/Experimental/Interchange/Core/Private/Nodes/InterchangeFactoryBaseNode.cpp
+++ b/Engine/Source/Runtime/Experimental/Interchange/Core/Private/Nodes/InterchangeFactoryBaseNode.cpp
@@ -31,22 +31,23 @@ UInterchangeFactoryBaseNode::UInterchangeFactoryBaseNode()
 FString UInterchangeFactoryBaseNode::GetKeyDisplayName(const UE::Interchange::FAttributeKey& NodeAttributeKey) const
 {
 	FString KeyDisplayName = NodeAttributeKey.ToString();
+	const FString OriginalKeyName = KeyDisplayName;
 	if (NodeAttributeKey == Macro_CustomSubPathKey)
 	{
 		KeyDisplayName = TEXT("Import Sub-Path");
 	}
-	else if (NodeAttributeKey.Key.Equals(UE::Interchange::FFactoryBaseNodeStaticData::FactoryDependenciesBaseKey()))
+	else if (OriginalKeyName.Equals(UE::Interchange::FFactoryBaseNodeStaticData::FactoryDependenciesBaseKey()))
 	{
 		KeyDisplayName = TEXT("Factory Dependencies Count");
 	}
-	else if (NodeAttributeKey.Key.StartsWith(UE::Interchange::FFactoryBaseNodeStaticData::FactoryDependenciesBaseKey()))
+	else if (OriginalKeyName.StartsWith(UE::Interchange::FFactoryBaseNodeStaticData::FactoryDependenciesBaseKey()))
 	{
 		KeyDisplayName = TEXT("Factory Dependencies Index ");
 		const FString IndexKey = UE::Interchange::TArrayAttributeHelper<FString>::IndexKey();
-		int32 IndexPosition = NodeAttributeKey.Key.Find(IndexKey) + IndexKey.Len();
-		if (IndexPosition < NodeAttributeKey.Key.Len())
+		int32 IndexPosition = OriginalKeyName.Find(IndexKey) + IndexKey.Len();
+		if (IndexPosition < OriginalKeyName.Len())
 		{
-			KeyDisplayName += NodeAttributeKey.Key.RightChop(IndexPosition);
+			KeyDisplayName += OriginalKeyName.RightChop(IndexPosition);
 		}
 	}
 	else
@@ -60,7 +61,7 @@ FString UInterchangeFactoryBaseNode::GetKeyDisplayName(const UE::Interchange::FA
 
 FString UInterchangeFactoryBaseNode::GetAttributeCategory(const UE::Interchange::FAttributeKey& NodeAttributeKey) const
 {
-	if (NodeAttributeKey.Key.StartsWith(UE::Interchange::FFactoryBaseNodeStaticData::FactoryDependenciesBaseKey()))
+	if (NodeAttributeKey.ToString().StartsWith(UE::Interchange::FFactoryBaseNodeStaticData::FactoryDependenciesBaseKey()))
 	{
 		return TEXT("FactoryDependencies");
 	}
diff --git a/Engine/Source/Runtime/Experimental/Interchange/Core/Private/Nodes/InterchangeUserDefinedAttribute.cpp b/Engine/Source/Runtime/Experimental/Interchange/Core/Private/Nodes/InterchangeUserDefinedAttribute.cpp
--- a/Engine/Source/Runtime/Experimental/Interchange/Core/Private/Nodes/InterchangeUserDefinedAttribute.cpp
+++ b/Engine/Source/Runtime/Experimental/Interchange/Core/Private/Nodes/InterchangeUserDefinedAttribute.cpp
@@ -149,9 +149,10 @@ TArray<FInterchangeUserDefinedAttributeInfo> UInterchangeUserDefinedAttributesAP
 	int32 LeftChopIndex = UserDefinedAttributeValuePostKey.Len();
 	for (UE::Interchange::FAttributeKey& AttributeKey : AttributeKeys)
 	{
-		if (AttributeKey.Key.StartsWith(UserDefinedAttributeBaseKey) && AttributeKey.Key.EndsWith(UserDefinedAttributeValuePostKey))
+		const FString AttributeKeyString = AttributeKey.ToString();
+		if (AttributeKeyString.StartsWith(UserDefinedAttributeBaseKey) && AttributeKeyString.EndsWith(UserDefinedAttributeValuePostKey))
 		{
-			FString UserDefinedAttributeName = AttributeKey.Key.RightChop(RightChopIndex).LeftChop(LeftChopIndex);
+			FString UserDefinedAttributeName = AttributeKeyString.RightChop(RightChopIndex).LeftChop(LeftChopIndex);
 			FInterchangeUserDefinedAttributeInfo& UserDefinedAttributeInfo = UserDefinedAttributeInfos.AddDefaulted_GetRef();
 			UserDefinedAttributeInfo.Type = InterchangeNode->GetAttributeType(AttributeKey);
 			UserDefinedAttributeInfo.Name = UserDefinedAttributeName;
@@ -186,9 +187,10 @@ void UInterchangeUserDefinedAttributesAPI::DuplicateAllUserDefinedAttribute(cons
 	int32 LeftChopIndex = UserDefinedAttributeValuePostKey.Len();
 	for (UE::Interchange::FAttributeKey& AttributeKey : AttributeKeys)
 	{
-		if (AttributeKey.Key.StartsWith(UserDefinedAttributeBaseKey) && AttributeKey.Key.EndsWith(UserDefinedAttributeValuePostKey))
+		const FString AttributeKeyString = AttributeKey.ToString();
+		if (AttributeKeyString.StartsWith(UserDefinedAttributeBaseKey) && AttributeKeyString.EndsWith(UserDefinedAttributeValuePostKey))
 		{
-			FString UserDefinedAttributeName = AttributeKey.Key.RightChop(RightChopIndex).LeftChop(LeftChopIndex);
+			FString UserDefinedAttributeName = AttributeKeyString.RightChop(RightChopIndex).LeftChop(LeftChopIndex);
 			//Get the optional payload key
 			const FString StorageBaseKey = UserDefinedAttributeBaseKey + UserDefinedAttributeName;
 
diff --git a/Engine/Source/Runtime/Experimental/Interchange/Core/Private/Types/AttributeStorage.cpp b/Engine/Source/Runtime/Experimental/Interchange/Core/Private/Types/AttributeStorage.cpp
--- a/Engine/Source/Runtime/Experimental/Interchange/Core/Private/Types/AttributeStorage.cpp
+++ b/Engine/Source/Runtime/Experimental/Interchange/Core/Private/Types/AttributeStorage.cpp
@@ -91,44 +91,44 @@ namespace UE
 			//Errors
 			if (HasAttributeStorageResult(Result, EAttributeStorageResult::Operation_Error_AttributeAllocationCorrupted))
 			{
-				UE_LOG(LogInterchangeCore, Error, TEXT("Attribute storage operation [%s] Key[%s]: Storage is corrupted."), *OperationName, *(AttributeKey.Key));
+				UE_LOG(LogInterchangeCore, Error, TEXT("Attribute storage operation [%s] Key[%s]: Storage is corrupted."), *OperationName, *(AttributeKey.ToString()));
 			}
 
 			//////////////////////////////////////////////////////////////////////////
 			//Warning
 			if (HasAttributeStorageResult(Result, EAttributeStorageResult::Operation_Error_CannotFoundKey))
 			{
-				UE_LOG(LogInterchangeCore, Warning, TEXT("Attribute storage operation [%s] Key[%s]: Cannot find attribute key."), *OperationName, *(AttributeKey.Key));
+				UE_LOG(LogInterchangeCore, Warning, TEXT("Attribute storage operation [%s] Key[%s]: Cannot find attribute key."), *OperationName, *(AttributeKey.ToString()));
 			}
 
 			if (HasAttributeStorageResult(Result, EAttributeStorageResult::Operation_Error_CannotOverrideAttribute))
 			{
-				UE_LOG(LogInterchangeCore, Warning, TEXT("Attribute storage operation [%s] Key[%s]: Cannot override attribute."), *OperationName, *(AttributeKey.Key));
+				UE_LOG(LogInterchangeCore, Warning, TEXT("Attribute storage operation [%s] Key[%s]: Cannot override attribute."), *OperationName, *(AttributeKey.ToString()));
 			}
 
 			if (HasAttributeStorageResult(Result, EAttributeStorageResult::Operation_Error_CannotRemoveAttribute))
 			{
-				UE_LOG(LogInterchangeCore, Warning, TEXT("Attribute storage operation [%s] Key[%s]: Cannot remove an attribute."), *OperationName, *(AttributeKey.Key));
+				UE_LOG(LogInterchangeCore, Warning, TEXT("Attribute storage operation [%s] Key[%s]: Cannot remove an attribute."), *OperationName, *(AttributeKey.ToString()));
 			}
 
 			if (HasAttributeStorageResult(Result, EAttributeStorageResult::Operation_Error_WrongSize))
 			{
-				UE_LOG(LogInterchangeCore, Warning, TEXT("Attribute storage operation [%s] Key[%s]: Stored attribute value size do not match parameter value size."), *OperationName, *(AttributeKey.Key));
+				UE_LOG(LogInterchangeCore, Warning, TEXT("Attribute storage operation [%s] Key[%s]: Stored attribute value size do not match parameter value size."), *OperationName, *(AttributeKey.ToString()));
 			}
 
 			if (HasAttributeStorageResult(Result, EAttributeStorageResult::Operation_Error_WrongType))
 			{
-				UE_LOG(LogInterchangeCore, Warning, TEXT("Attribute storage operation [%s] Key[%s]: Stored attribute value type do not match parameter value type."), *OperationName, *(AttributeKey.Key));
+				UE_LOG(LogInterchangeCore, Warning, TEXT("Attribute storage operation [%s] Key[%s]: Stored attribute value type do not match parameter value type."), *OperationName, *(AttributeKey.ToString()));
 			}
 
 			if (HasAttributeStorageResult(Result, EAttributeStorageResult::Operation_Error_InvalidStorage))
 			{
-				UE_LOG(LogInterchangeCore, Warning, TEXT("Attribute storage operation [%s] Key[%s]: The storage is invalid (NULL)."), *OperationName, *(AttributeKey.Key));
+				UE_LOG(LogInterchangeCore, Warning, TEXT("Attribute storage operation [%s] Key[%s]: The storage is invalid (NULL)."), *OperationName, *(AttributeKey.ToString()));
 			}
 
 			if (HasAttributeStorageResult(Result, EAttributeStorageResult::Operation_Error_InvalidMultiSizeValueData))
 			{
-				UE_LOG(LogInterchangeCore, Warning, TEXT("Attribute storage operation [%s] Key[%s]: Cannot retrieve the multisize value data pointer."), *OperationName, *(AttributeKey.Key));
+				UE_LOG(LogInterchangeCore, Warning, TEXT("Attribute storage operation [%s] Key[%s]: Cannot retrieve the multisize value data pointer."), *OperationName, *(AttributeKey.ToString()));
 			}
 		}
 
diff --git a/Engine/Source/Runtime/Experimental/Interchange/Core/Public/Nodes/InterchangeBaseNode.h b/Engine/Source/Runtime/Experimental/Interchange/Core/Public/Nodes/InterchangeBaseNode.h
--- a/Engine/Source/Runtime/Experimental/Interchange/Core/Public/Nodes/InterchangeBaseNode.h
+++ b/Engine/Source/Runtime/Experimental/Interchange/Core/Public/Nodes/InterchangeBaseNode.h
@@ -100,12 +100,12 @@ const UE::Interchange::FAttributeKey Macro_Custom##AttributeName##Key = UE::Inte
 #define IMPLEMENT_NODE_ATTRIBUTE_DELEGATE_BY_PROPERTYNAME(AttributeName, AttributeType, ObjectType, PropertyName)		\
 bool ApplyCustom##AttributeName##ToAsset(UObject* Asset) const														\
 {																													\
-	return ApplyAttributeToObject<AttributeType>(Macro_Custom##AttributeName##Key.ToString(), Asset, PropertyName);	\
+	return ApplyAttributeToObject<AttributeType>(Macro_Custom##AttributeName##Key.Key, Asset, PropertyName);	\
 }																													\
 																													\
 bool FillCustom##AttributeName##FromAsset(UObject* Asset)															\
 {																													\
-	return FillAttributeFromObject<AttributeType>(Macro_Custom##AttributeName##Key.ToString(), Asset, PropertyName);\
+	return FillAttributeFromObject<AttributeType>(Macro_Custom##AttributeName##Key.Key, Asset, PropertyName);\
 }
 
 #else //#if WITH_ENGINE
@@ -130,7 +130,8 @@ bool FillCustom##AttributeName##FromAsset(UObject* Asset)															\
 	{																																				\
 		if(bAddApplyDelegate)																														\
 		{																																			\
-			AddApplyAndFillDelegates<AttributeType>(Macro_Custom##AttributeName##Key.ToString(), AssetType::StaticClass(), *Macro_Custom##AttributeName##Key.ToString()); \
+			FName AttributeNameImpl = TEXT(""#AttributeName);																						\
+			AddApplyAndFillDelegates<AttributeType>(Macro_Custom##AttributeName##Key.Key, AssetType::StaticClass(), AttributeNameImpl);				\
 		}																																			\
 		return true;																																\
 	}																																				\
@@ -178,33 +179,30 @@ if(HasAttribute(UE::Interchange::FAttributeKey(NodeAttributeKey)))
 return false;
 
 //Interchange namespace
-namespace UE
+namespace UE::Interchange
 {
-	namespace Interchange
+	DECLARE_DELEGATE_RetVal_OneParam(bool, FApplyAttributeToAsset, UObject*);
+	DECLARE_DELEGATE_RetVal_OneParam(bool, FFillAttributeToAsset, UObject*);
+
+	/**
+		* Helper struct use to declare static const data we use in the UInterchangeBaseNode
+		* Node that derive from UInterchangeBaseNode can also add a struct that derive from this one to add there static data
+		* @note: The static data are mainly for holding Attribute keys. All attributes that are always available for a node should be in this class or a derived class.
+		*/
+	struct INTERCHANGECORE_API FBaseNodeStaticData
 	{
-		DECLARE_DELEGATE_RetVal_OneParam(bool, FApplyAttributeToAsset, UObject*);
-		DECLARE_DELEGATE_RetVal_OneParam(bool, FFillAttributeToAsset, UObject*);
-
-		/**
-		 * Helper struct use to declare static const data we use in the UInterchangeBaseNode
-		 * Node that derive from UInterchangeBaseNode can also add a struct that derive from this one to add there static data
-		 * @note: The static data are mainly for holding Attribute keys. All attributes that are always available for a node should be in this class or a derived class.
-		 */
-		struct INTERCHANGECORE_API FBaseNodeStaticData
-		{
-			static const FAttributeKey& UniqueIDKey();
-			static const FAttributeKey& DisplayLabelKey();
-			static const FAttributeKey& ParentIDKey();
-			static const FAttributeKey& IsEnabledKey();
-			static const FString& TargetAssetIDsKey();
-			static const FAttributeKey& ClassTypeAttributeKey();
-			static const FAttributeKey& AssetNameKey();
-			static const FAttributeKey& NodeContainerTypeKey();
-			static const FAttributeKey& ReimportStrategyFlagsKey();
-		};
-
-	} //ns Interchange
-} //ns UE
+		static const FAttributeKey& UniqueIDKey();
+		static const FAttributeKey& DisplayLabelKey();
+		static const FAttributeKey& ParentIDKey();
+		static const FAttributeKey& IsEnabledKey();
+		static const FAttributeKey& TargetAssetIDsKey();
+		static const FAttributeKey& ClassTypeAttributeKey();
+		static const FAttributeKey& AssetNameKey();
+		static const FAttributeKey& NodeContainerTypeKey();
+		static const FAttributeKey& ReimportStrategyFlagsKey();
+	};
+
+} //ns UE::Interchange
 
 UENUM(BlueprintType)
 enum class EInterchangeNodeContainerType : uint8
@@ -327,70 +325,82 @@ public:
 	 * Remove any attribute from this node. Return false if we cannot remove it. If the attribute do not exist it will return true.
 	 */
 	UFUNCTION(BlueprintCallable, Category = "Interchange | Node")
-	bool RemoveAttribute(const FString& NodeAttributeKey);
+	bool RemoveAttribute(const FName& NodeAttributeKey);
 
 	/**
 	 * Add a boolean attribute to this node. Return false if the attribute do not exist or if we cannot add it
 	 */
 	UFUNCTION(BlueprintCallable, Category = "Interchange | Node")
-	bool AddBooleanAttribute(const FString& NodeAttributeKey, const bool& Value);
+	bool AddBooleanAttribute(const FName& NodeAttributeKey, const bool& Value);
 	
 	/**
 	 * Get a boolean attribute from this node. Return false if the attribute do not exist
 	 */
 	UFUNCTION(BlueprintCallable, Category = "Interchange | Node")
-	bool GetBooleanAttribute(const FString& NodeAttributeKey, bool& OutValue) const;
+	bool GetBooleanAttribute(const FName& NodeAttributeKey, bool& OutValue) const;
 
 	/**
 	 * Add a int32 attribute to this node. Return false if the attribute do not exist or if we cannot add it
 	 */
 	UFUNCTION(BlueprintCallable, Category = "Interchange | Node")
-	bool AddInt32Attribute(const FString& NodeAttributeKey, const int32& Value);
+	bool AddInt32Attribute(const FName& NodeAttributeKey, const int32& Value);
 	
 	/**
 	 * Get a int32 attribute from this node. Return false if the attribute do not exist
 	 */
 	UFUNCTION(BlueprintCallable, Category = "Interchange | Node")
-	bool GetInt32Attribute(const FString& NodeAttributeKey, int32& OutValue) const;
+	bool GetInt32Attribute(const FName& NodeAttributeKey, int32& OutValue) const;
 
 	/**
 	 * Add a float attribute to this node. Return false if the attribute do not exist or if we cannot add it
 	 */
 	UFUNCTION(BlueprintCallable, Category = "Interchange | Node")
-	bool AddFloatAttribute(const FString& NodeAttributeKey, const float& Value);
+	bool AddFloatAttribute(const FName& NodeAttributeKey, const float& Value);
 	
 	/**
 	 * Get a float attribute from this node. Return false if the attribute do not exist
 	 */
 	UFUNCTION(BlueprintCallable, Category = "Interchange | Node")
-	bool GetFloatAttribute(const FString& NodeAttributeKey, float& OutValue) const;
+	bool GetFloatAttribute(const FName& NodeAttributeKey, float& OutValue) const;
 
 	/**
 	 * Add a string attribute to this node. Return false if the attribute do not exist or if we cannot add it
 	 */
 	UFUNCTION(BlueprintCallable, Category = "Interchange | Node")
-	bool AddStringAttribute(const FString& NodeAttributeKey, const FString& Value);
+	bool AddStringAttribute(const FName& NodeAttributeKey, const FString& Value);
 	
 	/**
 	 * Get a string attribute from this node. Return false if the attribute do not exist
 	 */
 	UFUNCTION(BlueprintCallable, Category = "Interchange | Node")
-	bool GetStringAttribute(const FString& NodeAttributeKey, FString& OutValue) const;
+	bool GetStringAttribute(const FName& NodeAttributeKey, FString& OutValue) const;
+
+	/**
+	 * Add a guid attribute to this node. Return false if the attribute do not exist or if we cannot add it
+	 */
+	UFUNCTION(BlueprintCallable, Category = "Interchange | Node")
+	bool AddGuidAttribute(const FName& NodeAttributeKey, const FGuid& Value);
+
+	/**
+	 * Get a guid attribute from this node. Return false if the attribute do not exist
+	 */
+	UFUNCTION(BlueprintCallable, Category = "Interchange | Node")
+	bool GetGuidAttribute(const FName& NodeAttributeKey, FGuid& OutValue) const;
 
 	/**
 	 * Add a FLinearColor attribute to this node. Return false if the attribute do not exist or if we cannot add it
 	 */
 	UFUNCTION(BlueprintCallable, Category = "Interchange | Node")
-	bool AddLinearColorAttribute(const FString& NodeAttributeKey, const FLinearColor& Value);
+	bool AddLinearColorAttribute(const FName& NodeAttributeKey, const FLinearColor& Value);
 	
 	/**
 	 * Get a FLinearColor attribute from this node. Return false if the attribute do not exist
 	 */
 	UFUNCTION(BlueprintCallable, Category = "Interchange | Node")
-	bool GetLinearColorAttribute(const FString& NodeAttributeKey, FLinearColor& OutValue) const;
+	bool GetLinearColorAttribute(const FName& NodeAttributeKey, FLinearColor& OutValue) const;
 
 	template<typename AttributeType>
-	AttributeType GetAttributeChecked(const FString& NodeAttributeKey) const
+	AttributeType GetAttributeChecked(const FName& NodeAttributeKey) const
 	{
 		AttributeType Value = AttributeType();
 		check(HasAttribute(UE::Interchange::FAttributeKey(NodeAttributeKey)));
@@ -401,12 +411,12 @@ public:
 	}
 
 	template<typename AttributeType>
-	bool GetAttribute(const FString& NodeAttributeKey, AttributeType& OutValue) const
+	bool GetAttribute(const FName& NodeAttributeKey, AttributeType& OutValue) const
 	{
 		INTERCHANGE_BASE_NODE_GET_ATTRIBUTE(AttributeType);
 	}
 	template<typename AttributeType>
-	bool SetAttribute(const FString& NodeAttributeKey, const AttributeType& Value)
+	bool SetAttribute(const FName& NodeAttributeKey, const AttributeType& Value)
 	{
 		INTERCHANGE_BASE_NODE_ADD_ATTRIBUTE(AttributeType);
 	}
@@ -418,7 +428,7 @@ public:
 	 * @param PropertyName		The name of the property of the ObjectClass that we will read and write to.
 	 */
 	template<typename AttributeType>
-	inline void AddApplyAndFillDelegates(const FString& NodeAttributeKey, UClass* ObjectClass, const FName PropertyName);
+	inline void AddApplyAndFillDelegates(const FName& NodeAttributeKey, UClass* ObjectClass, const FName PropertyName);
 
 	/**
 	 * Writes an attribute value to a UObject.
@@ -428,7 +438,7 @@ public:
 	 * @result					True if we succeeded.
 	 */
 	template<typename AttributeType>
-	inline bool ApplyAttributeToObject(const FString& NodeAttributeKey, UObject* Object, const FName PropertyName) const;
+	inline bool ApplyAttributeToObject(const FName& NodeAttributeKey, UObject* Object, const FName PropertyName) const;
 
 	/**
 	 * Reads an attribute value from a UObject.
@@ -438,7 +448,7 @@ public:
 	 * @result					True if we succeeded.
 	 */
 	template<typename AttributeType>
-	inline bool FillAttributeFromObject(const FString& NodeAttributeKey, UObject* Object, const FName PropertyName);
+	inline bool FillAttributeFromObject(const FName& NodeAttributeKey, UObject* Object, const FName PropertyName);
 
 	/**
 	 * Return the unique id pass in the constructor.
@@ -481,6 +491,7 @@ public:
 	/**
 	 * Set the parent unique id.
 	 */
+	UE_DEPRECATED(5.1, "This function should not be use - Use UInterchangeBaseNodeContainer::SetNodeParentUid function to create a parent hierarchy. This will ensure the container cache is always in sync")
 	UFUNCTION(BlueprintCallable, Category = "Interchange | Node")
 	bool SetParentUid(const FString& ParentUid);
 
@@ -597,7 +608,7 @@ protected:
 };
 
 template<typename AttributeType>
-inline void UInterchangeBaseNode::AddApplyAndFillDelegates(const FString& NodeAttributeKey, UClass* ObjectClass, const FName PropertyName)
+inline void UInterchangeBaseNode::AddApplyAndFillDelegates(const FName& NodeAttributeKey, UClass* ObjectClass, const FName PropertyName)
 {
 	TArray<UE::Interchange::FApplyAttributeToAsset>& ApplyDelegates = ApplyCustomAttributeDelegates.FindOrAdd(ObjectClass);
 	ApplyDelegates.Add(
@@ -619,7 +630,7 @@ inline void UInterchangeBaseNode::AddApplyAndFillDelegates(const FString& NodeAt
 }
 
 template<typename AttributeType>
-inline bool UInterchangeBaseNode::ApplyAttributeToObject(const FString& NodeAttributeKey, UObject* Object, const FName PropertyName) const
+inline bool UInterchangeBaseNode::ApplyAttributeToObject(const FName& NodeAttributeKey, UObject* Object, const FName PropertyName) const
 {
 	if (!Object)
 	{
@@ -655,7 +666,7 @@ inline bool UInterchangeBaseNode::ApplyAttributeToObject(const FString& NodeAttr
  * If the target property is a FObjectPropertyBase, treat the string as an object path.
  */
 template<>
-inline bool UInterchangeBaseNode::ApplyAttributeToObject<FString>(const FString& NodeAttributeKey, UObject* Object, const FName PropertyName) const
+inline bool UInterchangeBaseNode::ApplyAttributeToObject<FString>(const FName& NodeAttributeKey, UObject* Object, const FName PropertyName) const
 {
 	if (!Object)
 	{
@@ -698,7 +709,7 @@ inline bool UInterchangeBaseNode::ApplyAttributeToObject<FString>(const FString&
  * If the target property is a FBoolProperty, treat the propertyg as a bitfield.
  */
 template<>
-inline bool UInterchangeBaseNode::ApplyAttributeToObject<bool>(const FString& NodeAttributeKey, UObject* Object, const FName PropertyName) const
+inline bool UInterchangeBaseNode::ApplyAttributeToObject<bool>(const FName& NodeAttributeKey, UObject* Object, const FName PropertyName) const
 {
 	if (!Object)
 	{
@@ -738,7 +749,7 @@ inline bool UInterchangeBaseNode::ApplyAttributeToObject<bool>(const FString& No
 }
 
 template<typename AttributeType>
-inline bool UInterchangeBaseNode::FillAttributeFromObject(const FString& NodeAttributeKey, UObject* Object, const FName PropertyName)
+inline bool UInterchangeBaseNode::FillAttributeFromObject(const FName& NodeAttributeKey, UObject* Object, const FName PropertyName)
 {
 	TVariant<UObject*, uint8*> Container;
 	Container.Set<UObject*>(Object);
@@ -774,7 +785,7 @@ inline bool UInterchangeBaseNode::FillAttributeFromObject(const FString& NodeAtt
  * If the target property is a FObjectPropertyBase, treat the string as an object path.
  */
 template<>
-inline bool UInterchangeBaseNode::FillAttributeFromObject<FString>(const FString& NodeAttributeKey, UObject* Object, const FName PropertyName)
+inline bool UInterchangeBaseNode::FillAttributeFromObject<FString>(const FName& NodeAttributeKey, UObject* Object, const FName PropertyName)
 {
 	TVariant<UObject*, uint8*> Container;
 	Container.Set<UObject*>(Object);
@@ -809,7 +820,7 @@ inline bool UInterchangeBaseNode::FillAttributeFromObject<FString>(const FString
  * If the target property is a FBoolProperty, treat the property as a bitfield.
  */
 template<>
-inline bool UInterchangeBaseNode::FillAttributeFromObject<bool>(const FString& NodeAttributeKey, UObject* Object, const FName PropertyName)
+inline bool UInterchangeBaseNode::FillAttributeFromObject<bool>(const FName& NodeAttributeKey, UObject* Object, const FName PropertyName)
 {
 	TVariant<UObject*, uint8*> Container;
 	Container.Set<UObject*>(Object);
diff --git a/Engine/Source/Runtime/Experimental/Interchange/Core/Public/Nodes/InterchangeBaseNodeContainer.h b/Engine/Source/Runtime/Experimental/Interchange/Core/Public/Nodes/InterchangeBaseNodeContainer.h
--- a/Engine/Source/Runtime/Experimental/Interchange/Core/Public/Nodes/InterchangeBaseNodeContainer.h
+++ b/Engine/Source/Runtime/Experimental/Interchange/Core/Public/Nodes/InterchangeBaseNodeContainer.h
@@ -123,6 +123,17 @@ public:
 	UFUNCTION(BlueprintCallable, Category = "Interchange | Node Container")
 	void LoadFromFile(const FString& Filename);
 
+	/* Fill the children uids cache to optimize the GetNodeChildrenUids call.*/
+	UFUNCTION(BlueprintCallable, Category = "Interchange | Node Container")
+	void ComputeChildrenCache();
+
+	/* Reset the children uids cache. */
+	UFUNCTION(BlueprintCallable, Category = "Interchange | Node Container")
+	void ResetChildrenCache()
+	{
+		ChildrenCache.Reset();
+	}
+
 private:
 
 	UInterchangeBaseNode* GetNodeChildrenInternal(const FString& NodeUniqueID, int32 ChildIndex);
@@ -130,4 +141,6 @@ private:
 	/** Flat List of the nodes. Since the nodes are variable size, we store a pointer. */
 	UPROPERTY()
 	TMap<FString, TObjectPtr<UInterchangeBaseNode> > Nodes;
+
+	mutable TMap<FString, TArray<FString> > ChildrenCache;
 };
diff --git a/Engine/Source/Runtime/Experimental/Interchange/Core/Public/Nodes/InterchangeBaseNodeUtilities.h b/Engine/Source/Runtime/Experimental/Interchange/Core/Public/Nodes/InterchangeBaseNodeUtilities.h
--- a/Engine/Source/Runtime/Experimental/Interchange/Core/Public/Nodes/InterchangeBaseNodeUtilities.h
+++ b/Engine/Source/Runtime/Experimental/Interchange/Core/Public/Nodes/InterchangeBaseNodeUtilities.h
@@ -273,7 +273,7 @@ private:
 	FAttributeKey KeyCount;
 	FAttributeKey GetKeyCount() const
 	{
-		ensure(!KeyCount.Key.IsEmpty()); return KeyCount;
+		ensure(!KeyCount.Key.IsNone()); return KeyCount;
 	}
 
 	FAttributeKey GetIndexKey(int32 Index) const
diff --git a/Engine/Source/Runtime/Experimental/Interchange/Core/Public/Types/AttributeStorage.h b/Engine/Source/Runtime/Experimental/Interchange/Core/Public/Types/AttributeStorage.h
--- a/Engine/Source/Runtime/Experimental/Interchange/Core/Public/Types/AttributeStorage.h
+++ b/Engine/Source/Runtime/Experimental/Interchange/Core/Public/Types/AttributeStorage.h
@@ -17,7 +17,7 @@ namespace UE
 	{
 		struct FAttributeKey
 		{
-			FString Key;
+			FName Key;
 
 			FAttributeKey() = default;
 		
@@ -28,48 +28,47 @@ namespace UE
 
 			explicit FAttributeKey(const FName& Other)
 			{
-				Key = Other.ToString();
+				Key = Other;
 			}
 
 			explicit FAttributeKey(const FString& Other)
 			{
-				Key = Other;
+				Key = FName(*Other);
 			}
 
-			explicit FAttributeKey(FString&& Other)
+			explicit FAttributeKey(FName&& Other)
 			{
 				Key = MoveTemp(Other);
 			}
 
 			explicit FAttributeKey(const FText& Other)
 			{
-				Key = Other.ToString();
+				Key = FName(*Other.ToString());
 			}
 
 			explicit FAttributeKey(const TCHAR* Other)
 			{
-				const FString Helper(Other);
-				Key = Helper;
+				Key = FName(Other);
 			}
 
 			FORCEINLINE const TCHAR* operator*() const
 			{
-				return *(Key);
+				return *(Key.ToString());
 			}
 
 			FORCEINLINE FAttributeKey& operator=(const FName& Other)
 			{
-				Key = Other.ToString();
+				Key = Other;
 				return *this;
 			}
 
 			FORCEINLINE FAttributeKey& operator=(const FString& Other)
 			{
-				Key = Other;
+				Key = FName(*Other);
 				return *this;
 			}
 
-			FORCEINLINE FAttributeKey& operator=(FString&& Other)
+			FORCEINLINE FAttributeKey& operator=(FName&& Other)
 			{
 				Key = MoveTemp(Other);
 				return *this;
@@ -77,14 +76,13 @@ namespace UE
 
 			FORCEINLINE FAttributeKey& operator=(const FText& Other)
 			{
-				Key = Other.ToString();
+				Key = *(Other.ToString());
 				return *this;
 			}
 
 			FORCEINLINE FAttributeKey& operator=(const TCHAR* Other)
 			{
-				const FString Helper(Other);
-				Key = Helper;
+				Key = FName(Other);
 				return *this;
 			}
 		
@@ -100,7 +98,7 @@ namespace UE
 
 			FORCEINLINE bool operator<(const FAttributeKey& Other) const
 			{
-				return Key < Other.Key;
+				return Key.LexicalLess(Other.Key);
 			}
 
 			FORCEINLINE bool operator<=(const FAttributeKey& Other) const
@@ -110,21 +108,36 @@ namespace UE
 
 			FORCEINLINE bool operator>(const FAttributeKey& Other) const
 			{
-				return Key > Other.Key;
+				return Other.Key.LexicalLess(Key);
 			}
 
 			FORCEINLINE bool operator>=(const FAttributeKey& Other) const
 			{
-				return Key.Compare(Other.Key) >= 0;
+				return Other.Key.LexicalLess(Key) || Key == Other.Key;
 			}
 
 			friend FArchive& operator<<(FArchive& Ar, FAttributeKey& AttributeKey)
 			{
-				Ar << AttributeKey.Key;
+				FString SerializeString;
+				if(Ar.IsSaving())
+				{
+					SerializeString = AttributeKey.Key.ToString();
+					Ar << SerializeString;
+				}
+				else
+				{
+					Ar << SerializeString;
+					AttributeKey.Key = FName(*SerializeString);
+				}
 				return Ar;
 			}
 
 			FORCEINLINE FString ToString() const
+			{
+				return Key.ToString();
+			}
+
+			FORCEINLINE FName ToName() const
 			{
 				return Key;
 			}
diff --git a/Engine/Source/Runtime/Experimental/Interchange/Engine/Private/Tasks/InterchangeTaskParsing.cpp b/Engine/Source/Runtime/Experimental/Interchange/Engine/Private/Tasks/InterchangeTaskParsing.cpp
--- a/Engine/Source/Runtime/Experimental/Interchange/Engine/Private/Tasks/InterchangeTaskParsing.cpp
+++ b/Engine/Source/Runtime/Experimental/Interchange/Engine/Private/Tasks/InterchangeTaskParsing.cpp
@@ -116,6 +116,9 @@ void UE::Interchange::FTaskParsing::DoTask(ENamedThreads::Type CurrentThread, co
 			{
 				continue;
 			}
+			//Translation and pipelines are not executed, compute the children cache for translated and factory nodes
+			BaseNodeContainer->ComputeChildrenCache();
+
 			const bool bCanImportSceneNode = AsyncHelper->TaskData.ImportType == EImportType::ImportType_Scene;
 			BaseNodeContainer->IterateNodesOfType<UInterchangeFactoryBaseNode>([&](const FString& NodeUID, UInterchangeFactoryBaseNode* FactoryNode)
 			{
diff --git a/Engine/Source/Runtime/Experimental/Interchange/Engine/Private/Tasks/InterchangeTaskTranslator.cpp b/Engine/Source/Runtime/Experimental/Interchange/Engine/Private/Tasks/InterchangeTaskTranslator.cpp
--- a/Engine/Source/Runtime/Experimental/Interchange/Engine/Private/Tasks/InterchangeTaskTranslator.cpp
+++ b/Engine/Source/Runtime/Experimental/Interchange/Engine/Private/Tasks/InterchangeTaskTranslator.cpp
@@ -47,4 +47,6 @@ void UE::Interchange::FTaskTranslator::DoTask(ENamedThreads::Type CurrentThread,
 	//Translate the source data
 	UInterchangeBaseNodeContainer& BaseNodeContainer = *(AsyncHelper->BaseNodeContainers[SourceIndex].Get());
 	Translator->Translate(BaseNodeContainer);
+	//Make sure the base node container cache is computed for all translated node
+	BaseNodeContainer.ComputeChildrenCache();
 }
