Fix CIS static analysis
#rnx

diff --git a/Engine/Plugins/Experimental/Interchange/Runtime/Source/Pipelines/Private/InterchangeGenericAnimationPipeline.cpp b/Engine/Plugins/Experimental/Interchange/Runtime/Source/Pipelines/Private/InterchangeGenericAnimationPipeline.cpp
--- a/Engine/Plugins/Experimental/Interchange/Runtime/Source/Pipelines/Private/InterchangeGenericAnimationPipeline.cpp
+++ b/Engine/Plugins/Experimental/Interchange/Runtime/Source/Pipelines/Private/InterchangeGenericAnimationPipeline.cpp
@@ -45,7 +45,7 @@ void UInterchangeGenericAnimationPipeline::AdjustSettingsForReimportType(EInterc
 {
 	check(!CommonSkeletalMeshesAndAnimationsProperties.IsNull());
 	if (ImportType == EInterchangeReimportType::AssetCustomLODImport
-		|| ImportType == EInterchangeReimportType::AssetCustomLODImport
+		|| ImportType == EInterchangeReimportType::AssetCustomLODReimport
 		|| ImportType == EInterchangeReimportType::AssetAlternateSkinningImport
 		|| ImportType == EInterchangeReimportType::AssetAlternateSkinningReimport)
 	{
diff --git a/Engine/Plugins/Experimental/Interchange/Runtime/Source/Pipelines/Private/InterchangeGenericMeshPipeline.cpp b/Engine/Plugins/Experimental/Interchange/Runtime/Source/Pipelines/Private/InterchangeGenericMeshPipeline.cpp
--- a/Engine/Plugins/Experimental/Interchange/Runtime/Source/Pipelines/Private/InterchangeGenericMeshPipeline.cpp
+++ b/Engine/Plugins/Experimental/Interchange/Runtime/Source/Pipelines/Private/InterchangeGenericMeshPipeline.cpp
@@ -17,7 +17,7 @@ void UInterchangeGenericMeshPipeline::AdjustSettingsForReimportType(EInterchange
 {
 	check(!CommonSkeletalMeshesAndAnimationsProperties.IsNull());
 	if(ImportType == EInterchangeReimportType::AssetCustomLODImport
-		|| ImportType == EInterchangeReimportType::AssetCustomLODImport
+		|| ImportType == EInterchangeReimportType::AssetCustomLODReimport
 		|| ImportType == EInterchangeReimportType::AssetAlternateSkinningImport
 		|| ImportType == EInterchangeReimportType::AssetAlternateSkinningReimport)
 	{
