Changed how the interchange texture factory test the non-power of two texture to reduce the code duplication and remove the config reads from ouside of the game thread.
Added the option to allow the import of non-power of two textures in the interchange generic texture pipeline.
Added some static assert to the interchange texture factory to help the long term maintenance.
Added some functions to the interchange texture 2d factory node to help the manipulation of udims source blocks.
Removed some texture translators since they can share their implementation. That also fixed the difference between interchange and the legacy factory.
Removed a copy of the raw source data for most of the image that are imported by the image wrappers (Tiff was already able to avoid the copy before that). This will greatly reduce the peak memory usage during the imports.
Fixed some other minor issues.

diff --git a/Engine/Plugins/Experimental/Interchange/Runtime/Source/FactoryNodes/InterchangeFactoryNodes.Build.cs b/Engine/Plugins/Experimental/Interchange/Runtime/Source/FactoryNodes/InterchangeFactoryNodes.Build.cs
--- a/Engine/Plugins/Experimental/Interchange/Runtime/Source/FactoryNodes/InterchangeFactoryNodes.Build.cs
+++ b/Engine/Plugins/Experimental/Interchange/Runtime/Source/FactoryNodes/InterchangeFactoryNodes.Build.cs
@@ -9,18 +9,13 @@ namespace UnrealBuildTool.Rules
 			PublicDependencyModuleNames.AddRange(
 				new string[]
 				{
+					"CinematicCamera",
 					"Core",
 					"CoreUObject",
+					"Engine",
 					"InterchangeCore",
 					"InterchangeNodes",
-				}
-			);
-
-			PublicDependencyModuleNames.AddRange(
-				new string[]
-				{
-					"CinematicCamera",
-					"Engine"
+					"TextureUtilitiesCommon",
 				}
 			);
 		}
diff --git a/Engine/Plugins/Experimental/Interchange/Runtime/Source/FactoryNodes/Public/InterchangeTexture2DFactoryNode.h b/Engine/Plugins/Experimental/Interchange/Runtime/Source/FactoryNodes/Public/InterchangeTexture2DFactoryNode.h
--- a/Engine/Plugins/Experimental/Interchange/Runtime/Source/FactoryNodes/Public/InterchangeTexture2DFactoryNode.h
+++ b/Engine/Plugins/Experimental/Interchange/Runtime/Source/FactoryNodes/Public/InterchangeTexture2DFactoryNode.h
@@ -2,10 +2,11 @@
 
 #pragma once
 
-#include "CoreMinimal.h"
 #include "InterchangeTextureFactoryNode.h"
 
+#include "CoreMinimal.h"
 #include "Engine/Texture2D.h"
+#include "UDIMUtilities.h"
 
 #include "InterchangeTexture2DFactoryNode.generated.h"
 
@@ -95,6 +96,66 @@ public:
 		SourceBlocks = InSourceBlocks;
 	}
 
+
+	/**
+	 * Get a source block from the texture
+	 *
+	 * @param X The X coordinate of the block
+	 * @param Y The Y coordinate of the block
+	 * @param OutSourceFile The source file for that block if found
+	 * @return True if the source file for the block was found
+	 */
+	UFUNCTION(BlueprintCallable, Category = "Interchange | Node | Texture | UDIMs")
+	bool GetSourceBlockByCoordinates(int32 X, int32 Y, FString& OutSourceFile) const
+	{
+		return GetSourceBlock(UE::TextureUtilitiesCommon::GetUDIMIndex(X, Y), OutSourceFile);
+	}
+
+	/**
+	 * Get a source block from the texture
+	 *
+	 * @param BlockIndex The UDIM Index of the block
+	 * @param OutSourceFile The source file for that block if found
+	 * @return True if the source file for the block was found
+	 */
+	UFUNCTION(BlueprintCallable, Category = "Interchange | Node | Texture | UDIMs")
+	bool GetSourceBlock(int32 BlockIndex, FString& OutSourceFile) const
+	{
+		return SourceBlocks.GetValue(BlockIndex, OutSourceFile);
+	}
+
+
+	/**
+	 * Set a source block for the texture
+	 *
+	 * @param X The X coordinate of the block
+	 * @param Y The Y coordinate of the block
+	 * @param InSourceFile The source file for that block
+	 * 
+	 * The textures must be of the same format and use the same pixel format
+	 * The first block in the map is used to determine the accepted texture format and pixel format
+	 */
+	UFUNCTION(BlueprintCallable, Category = "Interchange | Node | Texture | UDIMs")
+	void SetSourceBlockByCoordinates(int32 X, int32 Y,  const FString& InSourceFile)
+	{
+		SetSourceBlock(UE::TextureUtilitiesCommon::GetUDIMIndex(X, Y), InSourceFile);
+	}
+
+	/**
+	 * Set a source block for the texture
+	 *
+	 * @param BlockIndex The UDIM Index of the block
+	 * @param InSourceFile The source file for that block
+	 * 
+	 * The textures must be of the same format and use the same pixel format
+	 * The first block in the map is used to determine the accepted texture format and pixel format
+	 */
+	UFUNCTION(BlueprintCallable, Category = "Interchange | Node | Texture | UDIMs")
+	void SetSourceBlock(int32 BlockIndex, const FString& InSourceFile)
+	{
+		SourceBlocks.SetKeyValue(BlockIndex, InSourceFile);
+	}
+
 	// UDIMs ends here
 	//////////////////////////////////////////////////////////////////////////
 
diff --git a/Engine/Plugins/Experimental/Interchange/Runtime/Source/FactoryNodes/Public/InterchangeTextureFactoryNode.h b/Engine/Plugins/Experimental/Interchange/Runtime/Source/FactoryNodes/Public/InterchangeTextureFactoryNode.h
--- a/Engine/Plugins/Experimental/Interchange/Runtime/Source/FactoryNodes/Public/InterchangeTextureFactoryNode.h
+++ b/Engine/Plugins/Experimental/Interchange/Runtime/Source/FactoryNodes/Public/InterchangeTextureFactoryNode.h
@@ -658,6 +658,23 @@ public:
 	//////////////////////////////////////////////////////////////////////////
 	//Factory Import settings
 
+	/** Return false if the Attribute was not set previously.*/
+	UFUNCTION(BlueprintCallable, Category = "Interchange | Node | Texture | Import Settings")
+	bool GetCustomAllowNonPowerOfTwo(bool& AttributeValue) const
+	{
+		IMPLEMENT_NODE_ATTRIBUTE_GETTER(AllowNonPowerOfTwo, bool);
+	}
+
+	/*
+	 * Should the factory allow the import of texture that would have a resolution that is not a power of two
+	 * By default this is not allowed
+	 */
+	UFUNCTION(BlueprintCallable, Category = "Interchange | Node | Texture | Import Settings")
+	bool SetCustomAllowNonPowerOfTwo(const bool& AttributeValue)
+	{
+		IMPLEMENT_NODE_ATTRIBUTE_SETTER_NODELEGATE(AllowNonPowerOfTwo, bool)
+	}
+
 	/** Return false if the Attribute was not set previously.*/
 	UFUNCTION(BlueprintCallable, Category = "Interchange | Node | Texture | Import Settings")
 	bool GetCustomPreferCompressedSourceData(bool& AttributeValue) const
@@ -680,59 +697,61 @@ private:
 	const UE::Interchange::FAttributeKey ClassNameAttributeKey = UE::Interchange::FBaseNodeStaticData::ClassTypeAttributeKey();
 	const UE::Interchange::FAttributeKey AssetNameKey = UE::Interchange::FBaseNodeStaticData::AssetNameKey();
 
-	//Translated texture node unique id
-	const UE::Interchange::FAttributeKey Macro_CustomTranslatedTextureNodeUidKey = UE::Interchange::FAttributeKey(TEXT("TranslatedTextureNodeUid"));
+	// Translated texture node unique id
+	IMPLEMENT_NODE_ATTRIBUTE_KEY(TranslatedTextureNodeUid)
 	
-	//Texture Adjustments
-	const UE::Interchange::FAttributeKey Macro_CustomAdjustBrightnessKey = UE::Interchange::FAttributeKey(TEXT("AdjustBrightness"));
-	const UE::Interchange::FAttributeKey Macro_CustomAdjustBrightnessCurveKey = UE::Interchange::FAttributeKey(TEXT("AdjustBrightnessCurve"));
-	const UE::Interchange::FAttributeKey Macro_CustomAdjustVibranceKey = UE::Interchange::FAttributeKey(TEXT("AdjustVibrance"));
-	const UE::Interchange::FAttributeKey Macro_CustomAdjustSaturationKey = UE::Interchange::FAttributeKey(TEXT("AdjustSaturation"));
-	const UE::Interchange::FAttributeKey Macro_CustomAdjustRGBCurveKey = UE::Interchange::FAttributeKey(TEXT("AdjustRGBCurve"));
-	const UE::Interchange::FAttributeKey Macro_CustomAdjustHueKey = UE::Interchange::FAttributeKey(TEXT("AdjustHue"));
-	const UE::Interchange::FAttributeKey Macro_CustomAdjustMinAlphaKey = UE::Interchange::FAttributeKey(TEXT("AdjustMinAlpha"));
-	const UE::Interchange::FAttributeKey Macro_CustomAdjustMaxAlphaKey = UE::Interchange::FAttributeKey(TEXT("AdjustMaxAlpha"));
-	const UE::Interchange::FAttributeKey Macro_CustombChromaKeyTextureKey = UE::Interchange::FAttributeKey(TEXT("bChromaKeyTexture"));
-	const UE::Interchange::FAttributeKey Macro_CustomChromaKeyThresholdKey = UE::Interchange::FAttributeKey(TEXT("ChromaKeyThreshold"));
-	const UE::Interchange::FAttributeKey Macro_CustomChromaKeyColorKey = UE::Interchange::FAttributeKey(TEXT("ChromaKeyColor"));
-
-	//Texture Compression
-	const UE::Interchange::FAttributeKey Macro_CustomCompressionNoAlphaKey = UE::Interchange::FAttributeKey(TEXT("CompressionNoAlpha"));
-	const UE::Interchange::FAttributeKey Macro_CustomDeferCompressionKey = UE::Interchange::FAttributeKey(TEXT("DeferCompression"));
-	const UE::Interchange::FAttributeKey Macro_CustomLossyCompressionAmountKey = UE::Interchange::FAttributeKey(TEXT("LossyCompressionAmount"));
-	const UE::Interchange::FAttributeKey Macro_CustomMaxTextureSizeKey = UE::Interchange::FAttributeKey(TEXT("MaxTextureSize"));
-	const UE::Interchange::FAttributeKey Macro_CustomCompressionQualityKey = UE::Interchange::FAttributeKey(TEXT("CompressionQuality"));
-	const UE::Interchange::FAttributeKey Macro_CustomCompressionSettingsKey = UE::Interchange::FAttributeKey(TEXT("CompressionSettings"));
-
-	//Texture general
-	const UE::Interchange::FAttributeKey Macro_CustombDoScaleMipsForAlphaCoverageKey = UE::Interchange::FAttributeKey(TEXT("bDoScaleMipsForAlphaCoverage"));
-	const UE::Interchange::FAttributeKey Macro_CustomAlphaCoverageThresholdsKey = UE::Interchange::FAttributeKey(TEXT("AlphaCoverageThresholds"));
-	const UE::Interchange::FAttributeKey Macro_CustombFlipGreenChannelKey = UE::Interchange::FAttributeKey(TEXT("bFlipGreenChannel"));
-	const UE::Interchange::FAttributeKey Macro_CustombForcePVRTC4Key = UE::Interchange::FAttributeKey(TEXT("bForcePVRTC4"));
-	const UE::Interchange::FAttributeKey Macro_CustomPowerOfTwoModeKey = UE::Interchange::FAttributeKey(TEXT("PowerOfTwoMode"));
-	const UE::Interchange::FAttributeKey Macro_CustomPaddingColorKey = UE::Interchange::FAttributeKey(TEXT("PaddingColor"));
-	const UE::Interchange::FAttributeKey Macro_CustomFilterKey = UE::Interchange::FAttributeKey(TEXT("Filter"));
-	const UE::Interchange::FAttributeKey Macro_CustomMipLoadOptionsKey = UE::Interchange::FAttributeKey(TEXT("MipLoadOptions"));
-	const UE::Interchange::FAttributeKey Macro_CustomSRGBKey = UE::Interchange::FAttributeKey(TEXT("SRGB"));
-	const UE::Interchange::FAttributeKey Macro_CustombUseLegacyGammaKey = UE::Interchange::FAttributeKey(TEXT("bUseLegacyGamma"));
-	const UE::Interchange::FAttributeKey Macro_CustomVirtualTextureStreamingKey = UE::Interchange::FAttributeKey(TEXT("VirtualTextureStreaming"));
-
-	//Level of Detail
-	const UE::Interchange::FAttributeKey Macro_CustombPreserveBorderKey = UE::Interchange::FAttributeKey(TEXT("bPreserveBorder"));
-	const UE::Interchange::FAttributeKey Macro_CustomMipGenSettingsKey = UE::Interchange::FAttributeKey(TEXT("MipGenSettings"));
-	const UE::Interchange::FAttributeKey Macro_CustomLODBiasKey = UE::Interchange::FAttributeKey(TEXT("LODBias"));
-	const UE::Interchange::FAttributeKey Macro_CustomLODGroupKey = UE::Interchange::FAttributeKey(TEXT("LODGroup"));
+	// Texture Adjustments
+	IMPLEMENT_NODE_ATTRIBUTE_KEY(AdjustBrightness)
+	IMPLEMENT_NODE_ATTRIBUTE_KEY(AdjustBrightnessCurve)
+	IMPLEMENT_NODE_ATTRIBUTE_KEY(AdjustVibrance)
+	IMPLEMENT_NODE_ATTRIBUTE_KEY(AdjustSaturation)
+	IMPLEMENT_NODE_ATTRIBUTE_KEY(AdjustRGBCurve)
+	IMPLEMENT_NODE_ATTRIBUTE_KEY(AdjustHue)
+	IMPLEMENT_NODE_ATTRIBUTE_KEY(AdjustMinAlpha)
+	IMPLEMENT_NODE_ATTRIBUTE_KEY(AdjustMaxAlpha)
+	IMPLEMENT_NODE_ATTRIBUTE_KEY(bChromaKeyTexture)
+	IMPLEMENT_NODE_ATTRIBUTE_KEY(ChromaKeyThreshold)
+	IMPLEMENT_NODE_ATTRIBUTE_KEY(ChromaKeyColor)
+
+	// Texture Compression
+	IMPLEMENT_NODE_ATTRIBUTE_KEY(CompressionNoAlpha)
+	IMPLEMENT_NODE_ATTRIBUTE_KEY(DeferCompression)
+	IMPLEMENT_NODE_ATTRIBUTE_KEY(LossyCompressionAmount)
+	IMPLEMENT_NODE_ATTRIBUTE_KEY(MaxTextureSize)
+	IMPLEMENT_NODE_ATTRIBUTE_KEY(CompressionQuality)
+	IMPLEMENT_NODE_ATTRIBUTE_KEY(CompressionSettings)
+
+	// Texture general
+	IMPLEMENT_NODE_ATTRIBUTE_KEY(bDoScaleMipsForAlphaCoverage)
+	IMPLEMENT_NODE_ATTRIBUTE_KEY(AlphaCoverageThresholds)
+	IMPLEMENT_NODE_ATTRIBUTE_KEY(bFlipGreenChannel)
+	IMPLEMENT_NODE_ATTRIBUTE_KEY(bForcePVRTC4)
+	IMPLEMENT_NODE_ATTRIBUTE_KEY(PowerOfTwoMode)
+	IMPLEMENT_NODE_ATTRIBUTE_KEY(PaddingColor)
+	IMPLEMENT_NODE_ATTRIBUTE_KEY(Filter)
+	IMPLEMENT_NODE_ATTRIBUTE_KEY(MipLoadOptions)
+	IMPLEMENT_NODE_ATTRIBUTE_KEY(SRGB)
+	IMPLEMENT_NODE_ATTRIBUTE_KEY(bUseLegacyGamma)
+	IMPLEMENT_NODE_ATTRIBUTE_KEY(VirtualTextureStreaming)
+
+	// Level of Detail
+	IMPLEMENT_NODE_ATTRIBUTE_KEY(bPreserveBorder)
+	IMPLEMENT_NODE_ATTRIBUTE_KEY(MipGenSettings)
+	IMPLEMENT_NODE_ATTRIBUTE_KEY(LODBias)
+	IMPLEMENT_NODE_ATTRIBUTE_KEY(LODGroup)
 	//TODO support per platform data in the FAttributeStorage, so we can set different value per platform at the pipeline stage, We set only the default value for now
-	const UE::Interchange::FAttributeKey Macro_CustomDownscaleKey = UE::Interchange::FAttributeKey(TEXT("Downscale"));
-	const UE::Interchange::FAttributeKey Macro_CustomDownscaleOptionsKey = UE::Interchange::FAttributeKey(TEXT("DownscaleOptions"));
+	IMPLEMENT_NODE_ATTRIBUTE_KEY(Downscale)
+	IMPLEMENT_NODE_ATTRIBUTE_KEY(DownscaleOptions)
 
-	//Compositing
-	const UE::Interchange::FAttributeKey Macro_CustomCompositeTextureModeKey = UE::Interchange::FAttributeKey(TEXT("CompositeTextureMode"));
-	const UE::Interchange::FAttributeKey Macro_CustomCompositePowerKey = UE::Interchange::FAttributeKey(TEXT("CompositePower"));
+	// Compositing
+	IMPLEMENT_NODE_ATTRIBUTE_KEY(CompositeTextureMode)
+	IMPLEMENT_NODE_ATTRIBUTE_KEY(CompositePower)
 
 	// Texture Factory Settings
+	IMPLEMENT_NODE_ATTRIBUTE_KEY(AllowNonPowerOfTwo)
 	IMPLEMENT_NODE_ATTRIBUTE_KEY(PreferCompressedSourceData)
 
+
 	// no IMPLEMENT_NODE_ATTRIBUTE_APPLY_UOBJECT here ?
 
 	//TODO support per platform data in the FAttributeStorage, so we can set different value per platform at the pipeline stage, We set only the default value for now
diff --git a/Engine/Plugins/Experimental/Interchange/Runtime/Source/Import/InterchangeImport.Build.cs b/Engine/Plugins/Experimental/Interchange/Runtime/Source/Import/InterchangeImport.Build.cs
--- a/Engine/Plugins/Experimental/Interchange/Runtime/Source/Import/InterchangeImport.Build.cs
+++ b/Engine/Plugins/Experimental/Interchange/Runtime/Source/Import/InterchangeImport.Build.cs
@@ -29,6 +29,7 @@ namespace UnrealBuildTool.Rules
 					"CinematicCamera",
 					"GLTFCore",
 					"IESFile",
+					"ImageCore",
 					"ImageWrapper",
 					"InterchangeCommonParser",
 					"InterchangeDispatcher",
diff --git a/Engine/Plugins/Experimental/Interchange/Runtime/Source/Import/Private/InterchangeImportModule.cpp b/Engine/Plugins/Experimental/Interchange/Runtime/Source/Import/Private/InterchangeImportModule.cpp
--- a/Engine/Plugins/Experimental/Interchange/Runtime/Source/Import/Private/InterchangeImportModule.cpp
+++ b/Engine/Plugins/Experimental/Interchange/Runtime/Source/Import/Private/InterchangeImportModule.cpp
@@ -20,18 +20,13 @@
 #include "Scene/InterchangeCineCameraActorFactory.h"
 #include "Scene/InterchangeStaticMeshActorFactory.h"
 #include "Scene/InterchangeSkeletalMeshActorFactory.h"
-#include "Texture/InterchangeBMPTranslator.h"
+#include "Texture/InterchangeImageWrapperTranslator.h"
 #include "Texture/InterchangeDDSTranslator.h"
-#include "Texture/InterchangeEXRTranslator.h"
-#include "Texture/InterchangeHDRTranslator.h"
 #include "Texture/InterchangeIESTranslator.h"
 #include "Texture/InterchangeJPGTranslator.h"
 #include "Texture/InterchangePCXTranslator.h"
-#include "Texture/InterchangePNGTranslator.h"
 #include "Texture/InterchangePSDTranslator.h"
 #include "Texture/InterchangeTextureFactory.h"
-#include "Texture/InterchangeTGATranslator.h"
-#include "Texture/InterchangeTIFFTranslator.h"
 
 DEFINE_LOG_CATEGORY(LogInterchangeImport);
 
@@ -59,30 +54,18 @@ void FInterchangeImportModule::StartupModule()
 		InterchangeManager.RegisterTranslator(UInterchangeOBJTranslator::StaticClass());
 
 		//Textures
-		InterchangeManager.RegisterTranslator(UInterchangeBMPTranslator::StaticClass());
-		InterchangeManager.RegisterTextureOnlyTranslatorClass(UInterchangeBMPTranslator::StaticClass());
+		InterchangeManager.RegisterTranslator(UInterchangeImageWrapperTranslator::StaticClass());
+		InterchangeManager.RegisterTextureOnlyTranslatorClass(UInterchangeImageWrapperTranslator::StaticClass());
 		InterchangeManager.RegisterTranslator(UInterchangeDDSTranslator::StaticClass());
 		InterchangeManager.RegisterTextureOnlyTranslatorClass(UInterchangeDDSTranslator::StaticClass());
-		InterchangeManager.RegisterTranslator(UInterchangeEXRTranslator::StaticClass());
-		InterchangeManager.RegisterTextureOnlyTranslatorClass(UInterchangeEXRTranslator::StaticClass());
 		InterchangeManager.RegisterTranslator(UInterchangeJPGTranslator::StaticClass());
 		InterchangeManager.RegisterTextureOnlyTranslatorClass(UInterchangeJPGTranslator::StaticClass());
 		InterchangeManager.RegisterTranslator(UInterchangePCXTranslator::StaticClass());
 		InterchangeManager.RegisterTextureOnlyTranslatorClass(UInterchangePCXTranslator::StaticClass());
-		InterchangeManager.RegisterTranslator(UInterchangePNGTranslator::StaticClass());
-		InterchangeManager.RegisterTextureOnlyTranslatorClass(UInterchangePNGTranslator::StaticClass());
 		InterchangeManager.RegisterTranslator(UInterchangePSDTranslator::StaticClass());
 		InterchangeManager.RegisterTextureOnlyTranslatorClass(UInterchangePSDTranslator::StaticClass());
-		InterchangeManager.RegisterTranslator(UInterchangeTGATranslator::StaticClass());
-		InterchangeManager.RegisterTextureOnlyTranslatorClass(UInterchangeTGATranslator::StaticClass());
-		InterchangeManager.RegisterTranslator(UInterchangeHDRTranslator::StaticClass());
-		InterchangeManager.RegisterTextureOnlyTranslatorClass(UInterchangeHDRTranslator::StaticClass());
 		InterchangeManager.RegisterTranslator(UInterchangeIESTranslator::StaticClass());
 		InterchangeManager.RegisterTextureOnlyTranslatorClass(UInterchangeIESTranslator::StaticClass());
-#if WITH_LIBTIFF
-		InterchangeManager.RegisterTranslator(UInterchangeTIFFTranslator::StaticClass());
-		InterchangeManager.RegisterTextureOnlyTranslatorClass(UInterchangeTIFFTranslator::StaticClass());
-#endif // WITH_LIBTIFF
 
 		//Register the factories
 		InterchangeManager.RegisterFactory(UInterchangeTextureFactory::StaticClass());
diff --git a/Engine/Plugins/Experimental/Interchange/Runtime/Source/Import/Private/Texture/InterchangeBMPTranslator.cpp b/Engine/Plugins/Experimental/Interchange/Runtime/Source/Import/Private/Texture/InterchangeBMPTranslator.cpp
deleted file mode 100644
diff --git a/Engine/Plugins/Experimental/Interchange/Runtime/Source/Import/Private/Texture/InterchangeDDSTranslator.cpp b/Engine/Plugins/Experimental/Interchange/Runtime/Source/Import/Private/Texture/InterchangeDDSTranslator.cpp
--- a/Engine/Plugins/Experimental/Interchange/Runtime/Source/Import/Private/Texture/InterchangeDDSTranslator.cpp
+++ b/Engine/Plugins/Experimental/Interchange/Runtime/Source/Import/Private/Texture/InterchangeDDSTranslator.cpp
@@ -1,24 +1,23 @@
-// Copyright Epic Games, Inc. All Rights Reserved. 
+// Copyright Epic Games, Inc. All Rights Reserved.
+
 #include "Texture/InterchangeDDSTranslator.h"
 
 #include "DDSLoader.h"
 #include "Engine/Texture.h"
 #include "Engine/Texture2D.h"
+#include "HAL/FileManager.h"
 #include "IImageWrapper.h"
 #include "IImageWrapperModule.h"
 #include "InterchangeImportLog.h"
 #include "InterchangeTextureCubeNode.h"
 #include "InterchangeTextureNode.h"
-#include "Misc/ConfigCacheIni.h"
 #include "Misc/FileHelper.h"
 #include "Misc/Paths.h"
 #include "Misc/ScopedSlowTask.h"
 #include "Modules/ModuleManager.h"
 #include "Nodes/InterchangeBaseNodeContainer.h"
-#include "Texture/TextureTranslatorUtilities.h"
-
-#include "HAL/FileManager.h"
 #include "Serialization/Archive.h"
+#include "Texture/TextureTranslatorUtilities.h"
 
 namespace UE::Interchange::Private::InterchangeDDSTranslator
 { 
@@ -164,16 +163,6 @@ TOptional<UE::Interchange::FImportImage> UInterchangeDDSTranslator::GetTexturePa
 		return TOptional<UE::Interchange::FImportImage>();
 	}
 
-	bool bAllowNonPowerOfTwo = false;
-	GConfig->GetBool(TEXT("TextureImporter"), TEXT("AllowNonPowerOfTwoTextures"), bAllowNonPowerOfTwo, GEditorIni);
-
-	// DDS 2d texture
-	if (!UE::Interchange::FImportImageHelper::IsImportResolutionValid(DDSLoadHelper.DDSHeader->dwWidth, DDSLoadHelper.DDSHeader->dwHeight, bAllowNonPowerOfTwo))
-	{
-		UE_LOG(LogInterchangeImport, Error, TEXT("Failed to import PCX, invalid resolution. Resolution[%d, %d], AllowPowerOfTwo[%s], [%s]"), DDSLoadHelper.DDSHeader->dwWidth, DDSLoadHelper.DDSHeader->dwHeight, bAllowNonPowerOfTwo ? TEXT("True") : TEXT("false"), *Filename);
-		return TOptional<UE::Interchange::FImportImage>();
-	}
-
 	ETextureSourceFormat SourceFormat = DDSLoadHelper.ComputeSourceFormat();
 
 	// Invalid DDS format
@@ -255,15 +244,6 @@ TOptional<UE::Interchange::FImportSlicedImage> UInterchangeDDSTranslator::GetSli
 		return TOptional<UE::Interchange::FImportSlicedImage>();
 	}
 
-	bool bAllowNonPowerOfTwo = false;
-	GConfig->GetBool(TEXT("TextureImporter"), TEXT("AllowNonPowerOfTwoTextures"), bAllowNonPowerOfTwo, GEditorIni);
-
-	if (!UE::Interchange::FImportImageHelper::IsImportResolutionValid(DDSLoadHelper.DDSHeader->dwWidth, DDSLoadHelper.DDSHeader->dwHeight, bAllowNonPowerOfTwo))
-	{
-		UE_LOG(LogInterchangeImport, Error, TEXT("Failed to import PCX, invalid resolution. Resolution[%d, %d], AllowPowerOfTwo[%s], [%s]"), DDSLoadHelper.DDSHeader->dwWidth, DDSLoadHelper.DDSHeader->dwHeight, bAllowNonPowerOfTwo ? TEXT("True") : TEXT("false"), *Filename);
-		return TOptional<UE::Interchange::FImportSlicedImage>();
-	}
-
 	ETextureSourceFormat SourceFormat = DDSLoadHelper.ComputeSourceFormat();
 
 	// Invalid DDS format
diff --git a/Engine/Plugins/Experimental/Interchange/Runtime/Source/Import/Private/Texture/InterchangeEXRTranslator.cpp b/Engine/Plugins/Experimental/Interchange/Runtime/Source/Import/Private/Texture/InterchangeEXRTranslator.cpp
deleted file mode 100644
diff --git a/Engine/Plugins/Experimental/Interchange/Runtime/Source/Import/Private/Texture/InterchangeHDRTranslator.cpp b/Engine/Plugins/Experimental/Interchange/Runtime/Source/Import/Private/Texture/InterchangeHDRTranslator.cpp
deleted file mode 100644
diff --git a/Engine/Plugins/Experimental/Interchange/Runtime/Source/Import/Private/Texture/InterchangeImageWrapperTranslator.cpp b/Engine/Plugins/Experimental/Interchange/Runtime/Source/Import/Private/Texture/InterchangeImageWrapperTranslator.cpp
new file mode 100644
--- /dev/null
+++ b/Engine/Plugins/Experimental/Interchange/Runtime/Source/Import/Private/Texture/InterchangeImageWrapperTranslator.cpp
@@ -0,0 +1,154 @@
+// Copyright Epic Games, Inc. All Rights Reserved. 
+
+#include "Texture/InterchangeImageWrapperTranslator.h"
+
+#include "Engine/Texture.h"
+#include "Engine/Texture2D.h"
+#include "IImageWrapper.h"
+#include "IImageWrapperModule.h"
+#include "ImageCoreUtils.h"
+#include "InterchangeImportLog.h"
+#include "InterchangeTextureNode.h"
+#include "Memory/SharedBuffer.h"
+#include "Misc/FileHelper.h"
+#include "Misc/Paths.h"
+#include "Misc/ScopedSlowTask.h"
+#include "Modules/ModuleManager.h"
+#include "Nodes/InterchangeBaseNodeContainer.h"
+#include "Texture/TextureTranslatorUtilities.h"
+#include "TextureImportUtils.h"
+#include "TgaImageSupport.h"
+
+TArray<FString> UInterchangeImageWrapperTranslator::GetSupportedFormats() const
+{
+	return {
+			TEXT("png;Portable Network Graphic"),
+			TEXT("bmp;Bitmap image"),
+			TEXT("exr;OpenEXR image"),
+			TEXT("hdr;High Dynamic Range image"),
+#if WITH_LIBTIFF
+			TEXT("tif;Tag Image File Format"),
+			TEXT("tiff;Tag Image File Format"),
+#endif //WITH_LIBTIFF
+			TEXT("tga;Targa image")
+		};
+}
+
+bool UInterchangeImageWrapperTranslator::Translate(UInterchangeBaseNodeContainer& BaseNodeContainer) const
+{
+	return UE::Interchange::FTextureTranslatorUtilities::Generic2DTextureTranslate(GetSourceData(), BaseNodeContainer);
+}
+
+TOptional<UE::Interchange::FImportImage> UInterchangeImageWrapperTranslator::GetTexturePayloadData(const UInterchangeSourceData* PayloadSourceData, const FString& PayLoadKey) const
+{
+	check(PayloadSourceData == GetSourceData());
+
+	if (!GetSourceData())
+	{
+		UE_LOG(LogInterchangeImport, Error, TEXT("Failed to import PNG, bad source data."));
+		return TOptional<UE::Interchange::FImportImage>();
+	}
+
+	TArray64<uint8> SourceDataBuffer;
+	FString Filename = GetSourceData()->GetFilename();
+
+	//Make sure the key fit the filename, The key should always be valid
+	if (!Filename.Equals(PayLoadKey))
+	{
+		UE_LOG(LogInterchangeImport, Error, TEXT("Failed to import PNG, wrong payload key. [%s]"), *Filename);
+		return TOptional<UE::Interchange::FImportImage>();
+	}
+
+	if (!FPaths::FileExists(Filename))
+	{
+		UE_LOG(LogInterchangeImport, Error, TEXT("Failed to import PNG, cannot open file. [%s]"), *Filename);
+		return TOptional<UE::Interchange::FImportImage>();
+	}
+
+	if (!FFileHelper::LoadFileToArray(SourceDataBuffer, *Filename))
+	{
+		UE_LOG(LogInterchangeImport, Error, TEXT("Failed to import PNG, cannot load file content into an array. [%s]"), *Filename);
+		return TOptional<UE::Interchange::FImportImage>();
+	}
+
+	const uint8* Buffer = SourceDataBuffer.GetData();
+	const uint8* BufferEnd = Buffer + SourceDataBuffer.Num();
+
+	const int32 Length = BufferEnd - Buffer;
+
+	IImageWrapperModule& ImageWrapperModule = FModuleManager::LoadModuleChecked<IImageWrapperModule>(FName("ImageWrapper"));
+
+	EImageFormat ImageFormat = ImageWrapperModule.DetectImageFormat(Buffer, Length);
+
+
+	UE::Interchange::FImportImage PayloadData;
+
+	if (ImageFormat != EImageFormat::Invalid)
+	{
+		// Generic ImageWrapper loader :
+		// for PNG,EXR,BMP,TGA :
+		FImage LoadedImage;
+		if (ImageWrapperModule.DecompressImage(Buffer, Length, LoadedImage))
+		{
+			// Todo interchange: should these payload modification be part of the pipeline, factory or stay there?
+			UE::TextureUtilitiesCommon::AutoDetectAndChangeGrayScale(LoadedImage);
+
+			ETextureSourceFormat TextureFormat = FImageCoreUtils::ConvertToTextureSourceFormat(LoadedImage.Format);
+			bool bSRGB = LoadedImage.GammaSpace != EGammaSpace::Linear;
+
+			PayloadData.Init2DWithParams(
+				LoadedImage.SizeX,
+				LoadedImage.SizeY,
+				TextureFormat,
+				bSRGB,
+				false
+			);
+
+			PayloadData.RawData = MakeUniqueBufferFromArray(MoveTemp(LoadedImage.RawData));
+
+			if (ERawImageFormat::IsHDR(LoadedImage.Format))
+			{
+				PayloadData.CompressionSettings = TC_HDR;
+				check(bSRGB == false);
+			}
+
+			// do per-format processing to match legacy behavior :
+			if (ImageFormat == EImageFormat::TGA)
+			{
+				const FTGAFileHeader* TGA = (FTGAFileHeader*)Buffer;
+
+				if (TGA->ColorMapType == 1 && TGA->ImageTypeCode == 1 && TGA->BitsPerPixel == 8)
+				{
+					// Notes: The Scaleform GFx exporter (dll) strips all font glyphs into a single 8-bit texture.
+					// The targa format uses this for a palette index; GFx uses a palette of (i,i,i,i) so the index
+					// is also the alpha value.
+					//
+					// We store the image as PF_G8, where it will be used as alpha in the Glyph shader.
+
+					// ?? check or convert? or neither?
+					//check( TextureFormat == TSF_G8 );
+
+					PayloadData.CompressionSettings = TC_Grayscale;
+				}
+				else if (TGA->ColorMapType == 0 && TGA->ImageTypeCode == 3 && TGA->BitsPerPixel == 8)
+				{
+					// standard grayscale images
+
+					// ?? check or convert? or neither?
+					//check( TextureFormat == TSF_G8 );
+
+					PayloadData.CompressionSettings = TC_Grayscale;
+				}
+
+				if (PayloadData.CompressionSettings == TC_Grayscale && TGA->ImageTypeCode == 3)
+				{
+					// default grayscales to linear as they wont get compression otherwise and are commonly used as masks
+					PayloadData.bSRGB = false;
+				}
+			}
+		}
+	}
+
+	return PayloadData;
+}
+
diff --git a/Engine/Plugins/Experimental/Interchange/Runtime/Source/Import/Private/Texture/InterchangeJPGTranslator.cpp b/Engine/Plugins/Experimental/Interchange/Runtime/Source/Import/Private/Texture/InterchangeJPGTranslator.cpp
--- a/Engine/Plugins/Experimental/Interchange/Runtime/Source/Import/Private/Texture/InterchangeJPGTranslator.cpp
+++ b/Engine/Plugins/Experimental/Interchange/Runtime/Source/Import/Private/Texture/InterchangeJPGTranslator.cpp
@@ -10,7 +10,6 @@
 #include "InterchangeImportLog.h"
 #include "InterchangeTextureNode.h"
 #include "Memory/SharedBuffer.h"
-#include "Misc/ConfigCacheIni.h"
 #include "Misc/FileHelper.h"
 #include "Misc/Paths.h"
 #include "Misc/ScopedSlowTask.h"
@@ -50,10 +49,7 @@ TOptional<UE::Interchange::FImportImage> UInterchangeJPGTranslator::GetTexturePa
 
 bool UInterchangeJPGTranslator::SupportCompressedTexturePayloadData() const
 {
-	// For now this option is opt in via the config files once there is no technical risk this will become the default path.
-	bool bRetainJpegFormat = false;
-	GConfig->GetBool(TEXT("TextureImporter"), TEXT("RetainJpegFormat"), bRetainJpegFormat, GEditorIni);
-	return bRetainJpegFormat;
+	return true;
 }
 
 TOptional<UE::Interchange::FImportImage> UInterchangeJPGTranslator::GetCompressedTexturePayloadData(const UInterchangeSourceData* PayloadSourceData, const FString& PayLoadKey) const
@@ -110,10 +106,6 @@ TOptional<UE::Interchange::FImportImage> UInterchangeJPGTranslator::GetTexturePa
 	const uint8* Buffer = SourceDataBuffer.GetData();
 	const uint8* BufferEnd = Buffer + SourceDataBuffer.Num();
 
-	bool bAllowNonPowerOfTwo = false;
-	GConfig->GetBool(TEXT("TextureImporter"), TEXT("AllowNonPowerOfTwoTextures"), bAllowNonPowerOfTwo, GEditorIni);
-
-	// Validate it.
 	const int32 Length = BufferEnd - Buffer;
 
 	IImageWrapperModule& ImageWrapperModule = FModuleManager::LoadModuleChecked<IImageWrapperModule>(FName("ImageWrapper"));
@@ -127,11 +119,6 @@ TOptional<UE::Interchange::FImportImage> UInterchangeJPGTranslator::GetTexturePa
 		UE_LOG(LogInterchangeImport, Error, TEXT("Failed to decode JPEG. [%s]"), *Filename);
 		return TOptional<UE::Interchange::FImportImage>();
 	}
-	if (!UE::Interchange::FImportImageHelper::IsImportResolutionValid(JpegImageWrapper->GetWidth(), JpegImageWrapper->GetHeight(), bAllowNonPowerOfTwo))
-	{
-		UE_LOG(LogInterchangeImport, Error, TEXT("Failed to import JPEG, invalid resolution. Resolution[%d, %d], AllowPowerOfTwo[%s], [%s]"), JpegImageWrapper->GetWidth(), JpegImageWrapper->GetHeight(), bAllowNonPowerOfTwo ? TEXT("True") : TEXT("false"), *Filename);
-		return TOptional<UE::Interchange::FImportImage>();
-	}
 
 	// Select the texture's source format
 	ETextureSourceFormat TextureFormat = TSF_Invalid;
@@ -166,7 +153,7 @@ TOptional<UE::Interchange::FImportImage> UInterchangeJPGTranslator::GetTexturePa
 	UE::Interchange::FImportImage PayloadData;
 
 
-	const bool bShouldAllocateRawDataBuffer = !bShouldImportRaw;
+	const bool bShouldAllocateRawDataBuffer = false;
 
 	PayloadData.Init2DWithParams(
 		JpegImageWrapper->GetWidth(),
@@ -176,12 +163,17 @@ TOptional<UE::Interchange::FImportImage> UInterchangeJPGTranslator::GetTexturePa
 		bShouldAllocateRawDataBuffer
 	);
 
+	TArray64<uint8> RawData;
 	if (bShouldImportRaw)
 	{
 		PayloadData.RawData = MakeUniqueBufferFromArray(MoveTemp(SourceDataBuffer));
 		PayloadData.RawDataCompressionFormat = ETextureSourceCompressionFormat::TSCF_JPEG;
 	}
-	else if (!JpegImageWrapper->GetRaw(Format, BitDepth, PayloadData.GetArrayViewOfRawData()))
+	else if (JpegImageWrapper->GetRaw(Format, BitDepth, RawData))
+	{
+		PayloadData.RawData = MakeUniqueBufferFromArray(MoveTemp(RawData));
+	}
+	else
 	{
 		UE_LOG(LogInterchangeImport, Error, TEXT("Failed to decode JPEG. [%s]"), *Filename);
 		return TOptional<UE::Interchange::FImportImage>();
diff --git a/Engine/Plugins/Experimental/Interchange/Runtime/Source/Import/Private/Texture/InterchangePCXTranslator.cpp b/Engine/Plugins/Experimental/Interchange/Runtime/Source/Import/Private/Texture/InterchangePCXTranslator.cpp
--- a/Engine/Plugins/Experimental/Interchange/Runtime/Source/Import/Private/Texture/InterchangePCXTranslator.cpp
+++ b/Engine/Plugins/Experimental/Interchange/Runtime/Source/Import/Private/Texture/InterchangePCXTranslator.cpp
@@ -1,4 +1,5 @@
-// Copyright Epic Games, Inc. All Rights Reserved. 
+// Copyright Epic Games, Inc. All Rights Reserved.
+
 #include "Texture/InterchangePCXTranslator.h"
 
 #include "Engine/Texture.h"
@@ -7,7 +8,6 @@
 #include "IImageWrapperModule.h"
 #include "InterchangeImportLog.h"
 #include "InterchangeTextureNode.h"
-#include "Misc/ConfigCacheIni.h"
 #include "Misc/FileHelper.h"
 #include "Misc/Paths.h"
 #include "Misc/ScopedSlowTask.h"
@@ -113,10 +113,6 @@ TOptional<UE::Interchange::FImportImage> UInterchangePCXTranslator::GetTexturePa
 	const uint8* Buffer = SourceDataBuffer.GetData();
 	const uint8* BufferEnd = Buffer + SourceDataBuffer.Num();
 
-	bool bAllowNonPowerOfTwo = false;
-	GConfig->GetBool(TEXT("TextureImporter"), TEXT("AllowNonPowerOfTwoTextures"), bAllowNonPowerOfTwo, GEditorIni);
-
-	// Validate it.
 	const int32 Length = BufferEnd - Buffer;
 
 	FModuleManager::LoadModuleChecked<IImageWrapperModule>(FName("ImageWrapper"));
@@ -132,13 +128,7 @@ TOptional<UE::Interchange::FImportImage> UInterchangePCXTranslator::GetTexturePa
 		int32 NewU = PCX->XMax + 1 - PCX->XMin;
 		int32 NewV = PCX->YMax + 1 - PCX->YMin;
 
-		// Check the resolution of the imported texture to ensure validity
-		if (!UE::Interchange::FImportImageHelper::IsImportResolutionValid(NewU, NewV, bAllowNonPowerOfTwo))
-		{
-			UE_LOG(LogInterchangeImport, Error, TEXT("Failed to import PCX, invalid resolution. Resolution[%d, %d], AllowPowerOfTwo[%s], [%s]"), NewU, NewV, bAllowNonPowerOfTwo ? TEXT("True") : TEXT("false"), *Filename);
-			return TOptional<UE::Interchange::FImportImage>();
-		}
-		else if (PCX->NumPlanes == 1 && PCX->BitsPerPixel == 8)
+		if (PCX->NumPlanes == 1 && PCX->BitsPerPixel == 8)
 		{
 
 			// Set texture properties.
diff --git a/Engine/Plugins/Experimental/Interchange/Runtime/Source/Import/Private/Texture/InterchangePNGTranslator.cpp b/Engine/Plugins/Experimental/Interchange/Runtime/Source/Import/Private/Texture/InterchangePNGTranslator.cpp
deleted file mode 100644
diff --git a/Engine/Plugins/Experimental/Interchange/Runtime/Source/Import/Private/Texture/InterchangePSDTranslator.cpp b/Engine/Plugins/Experimental/Interchange/Runtime/Source/Import/Private/Texture/InterchangePSDTranslator.cpp
--- a/Engine/Plugins/Experimental/Interchange/Runtime/Source/Import/Private/Texture/InterchangePSDTranslator.cpp
+++ b/Engine/Plugins/Experimental/Interchange/Runtime/Source/Import/Private/Texture/InterchangePSDTranslator.cpp
@@ -1,4 +1,5 @@
-// Copyright Epic Games, Inc. All Rights Reserved. 
+// Copyright Epic Games, Inc. All Rights Reserved.
+
 #include "Texture/InterchangePSDTranslator.h"
 
 #include "Engine/Texture.h"
@@ -7,7 +8,6 @@
 #include "IImageWrapperModule.h"
 #include "InterchangeImportLog.h"
 #include "InterchangeTextureNode.h"
-#include "Misc/ConfigCacheIni.h"
 #include "Misc/FileHelper.h"
 #include "Misc/Paths.h"
 #include "Misc/ScopedSlowTask.h"
@@ -414,10 +414,6 @@ TOptional<UE::Interchange::FImportImage> UInterchangePSDTranslator::GetTexturePa
 	const uint8* Buffer = SourceDataBuffer.GetData();
 	const uint8* BufferEnd = Buffer + SourceDataBuffer.Num();
 
-	bool bAllowNonPowerOfTwo = false;
-	GConfig->GetBool(TEXT("TextureImporter"), TEXT("AllowNonPowerOfTwoTextures"), bAllowNonPowerOfTwo, GEditorIni);
-
-	// Validate it.
 	const int32 Length = BufferEnd - Buffer;
 
 	//
@@ -436,12 +432,6 @@ TOptional<UE::Interchange::FImportImage> UInterchangePSDTranslator::GetTexturePa
 		return TOptional<UE::Interchange::FImportImage>();
 	}
 
-	// Check the resolution of the imported texture to ensure validity
-	if (!UE::Interchange::FImportImageHelper::IsImportResolutionValid(PSDHeader.Width, PSDHeader.Height, bAllowNonPowerOfTwo))
-	{
-		return TOptional<UE::Interchange::FImportImage>();
-	}
-
 	if (!PSDHeader.IsSupported())
 	{
 		UE_LOG(LogInterchangeImport, Error, TEXT("Format of this PSD is not supported. Only Grayscale and RGBColor PSD images are currently supported, in 8-bit or 16-bit."));
diff --git a/Engine/Plugins/Experimental/Interchange/Runtime/Source/Import/Private/Texture/InterchangeTGATranslator.cpp b/Engine/Plugins/Experimental/Interchange/Runtime/Source/Import/Private/Texture/InterchangeTGATranslator.cpp
deleted file mode 100644
diff --git a/Engine/Plugins/Experimental/Interchange/Runtime/Source/Import/Private/Texture/InterchangeTIFFTranslator.cpp b/Engine/Plugins/Experimental/Interchange/Runtime/Source/Import/Private/Texture/InterchangeTIFFTranslator.cpp
deleted file mode 100644
diff --git a/Engine/Plugins/Experimental/Interchange/Runtime/Source/Import/Private/Texture/InterchangeTextureFactory.cpp b/Engine/Plugins/Experimental/Interchange/Runtime/Source/Import/Private/Texture/InterchangeTextureFactory.cpp
--- a/Engine/Plugins/Experimental/Interchange/Runtime/Source/Import/Private/Texture/InterchangeTextureFactory.cpp
+++ b/Engine/Plugins/Experimental/Interchange/Runtime/Source/Import/Private/Texture/InterchangeTextureFactory.cpp
@@ -15,6 +15,7 @@
 #include "InterchangeAssetImportData.h"
 #include "InterchangeImportCommon.h"
 #include "InterchangeImportLog.h"
+#include "InterchangeResult.h"
 #include "InterchangeTexture2DArrayFactoryNode.h"
 #include "InterchangeTexture2DArrayNode.h"
 #include "InterchangeTexture2DFactoryNode.h"
@@ -43,6 +44,7 @@
 
 #endif //WITH_EDITORONLY_DATA
 
+#define LOCTEXT_NAMESPACE "InterchangeTextureFactory"
 
 namespace UE::Interchange::Private::InterchangeTextureFactory
 {
@@ -118,7 +120,6 @@ namespace UE::Interchange::Private::InterchangeTextureFactory
 		UObject* Object;
 	};
 
-
 	/**
 	 * Return the supported class if the node is one otherwise return nullptr
 	 */
@@ -168,10 +169,12 @@ namespace UE::Interchange::Private::InterchangeTextureFactory
 		, UInterchangeTexture2DFactoryNode*
 		, UInterchangeTextureCubeFactoryNode*
 		, UInterchangeTexture2DArrayFactoryNode*
-		,UInterchangeTextureLightProfileFactoryNode* >;
+		, UInterchangeTextureLightProfileFactoryNode*>;
 
 	FTextureFactoryNodeVariant GetAsTextureFactoryNodeVariant(UInterchangeFactoryBaseNode* AssetNode, UClass* SupportedFactoryNodeClass)
 	{
+		static_assert(TVariantSize<FTextureFactoryNodeVariant>::Value == 5, "Please update the code below and this assert to reflect the change to the variant type.");
+
 		if (AssetNode)
 		{
 			if (!SupportedFactoryNodeClass)
@@ -253,10 +256,14 @@ namespace UE::Interchange::Private::InterchangeTextureFactory
 		, const UInterchangeTexture2DNode*
 		, const UInterchangeTextureCubeNode*
 		, const UInterchangeTexture2DArrayNode*
-		, const UInterchangeTextureLightProfileNode* >;
+		, const UInterchangeTextureLightProfileNode*>;
 
 	FTextureNodeVariant GetTextureNodeVariantFromFactoryVariant(const FTextureFactoryNodeVariant& FactoryVariant, const UInterchangeBaseNodeContainer* NodeContainer)
 	{
+		static_assert(TVariantSize<FTextureFactoryNodeVariant>::Value == 5, "Please update the code below and this assert to reflect the change to the variant type.");
+
+		static_assert(TVariantSize<FTextureNodeVariant>::Value == 5, "Please update the code below and this assert to reflect the change to the variant type.");
+
 		FString TextureNodeUniqueID;
 
 		if (UInterchangeTexture2DFactoryNode* const* Texture2DFactoryNode = FactoryVariant.TryGet<UInterchangeTexture2DFactoryNode*>())
@@ -304,6 +311,8 @@ namespace UE::Interchange::Private::InterchangeTextureFactory
 
 	bool HasPayloadKey(const FTextureNodeVariant& TextureNodeVariant)
 	{
+		static_assert(TVariantSize<FTextureNodeVariant>::Value == 5, "Please update the code below and this assert to reflect the change to the variant type.");
+
 		if (const UInterchangeTexture2DNode* const* TextureNode =  TextureNodeVariant.TryGet<const UInterchangeTexture2DNode*>())
 		{
 			return (*TextureNode)->GetPayLoadKey().IsSet();
@@ -329,6 +338,8 @@ namespace UE::Interchange::Private::InterchangeTextureFactory
 
 	TOptional<FString> GetPayloadKey(const FTextureNodeVariant& TextureNodeVariant)
 	{
+		static_assert(TVariantSize<FTextureNodeVariant>::Value == 5, "Please update the code below and this assert to reflect the change to the variant type.");
+
 		if (const UInterchangeTexture2DNode* const* TextureNode =  TextureNodeVariant.TryGet<const UInterchangeTexture2DNode*>())
 		{
 			return (*TextureNode)->GetPayLoadKey();
@@ -429,8 +440,9 @@ namespace UE::Interchange::Private::InterchangeTextureFactory
 			for (int32 Index = 0; Index < Images.Num(); ++Index)
 			{
 				int32 UDIMIndex = UDIMsAndSourcesFileArray[Index]->Key;
-				const int32 BlockX = (UDIMIndex - 1001) % 10;
-				const int32 BlockY = (UDIMIndex - 1001) / 10;
+				int32 BlockX;
+				int32 BlockY;
+				UE::TextureUtilitiesCommon::ExtractUDIMCoordinates(UDIMIndex, BlockX, BlockY);
 				BlockedImage.InitBlockFromImage(BlockX, BlockY, Images[Index]);
 			}
 
@@ -442,6 +454,12 @@ namespace UE::Interchange::Private::InterchangeTextureFactory
 
 	FTexturePayloadVariant GetTexturePayload(const UInterchangeSourceData* SourceData, const FString& PayloadKey, const FTextureNodeVariant& TextureNodeVariant, const FTextureFactoryNodeVariant& FactoryNodeVariant, const UInterchangeTranslatorBase* Translator)
 	{
+		static_assert(TVariantSize<FTextureFactoryNodeVariant>::Value == 5, "Please update the code below and this assert to reflect the change to the variant type.");
+
+		static_assert(TVariantSize<FTexturePayloadVariant>::Value == 5, "Please update the code below and this assert to reflect the change to the variant type.");
+
+		static_assert(TVariantSize<FTextureNodeVariant>::Value == 5, "Please update the code below and this assert to reflect the change to the variant type.");
+
 		// Standard texture 2D payload
 		if (const UInterchangeTexture2DNode* const* TextureNode =  TextureNodeVariant.TryGet<const UInterchangeTexture2DNode*>())
 		{
@@ -599,6 +617,8 @@ namespace UE::Interchange::Private::InterchangeTextureFactory
 
 	bool CanSetupTexture2DSourceData(FTexturePayloadVariant& TexturePayload)
 	{
+		static_assert(TVariantSize<FTexturePayloadVariant>::Value == 5, "Please update the code below and this assert to reflect the change to the variant type.");
+
 		if (TOptional<FImportBlockedImage>* BlockedImage = TexturePayload.TryGet<TOptional<FImportBlockedImage>>())
 		{
 			if (BlockedImage->IsSet())
@@ -690,6 +710,8 @@ namespace UE::Interchange::Private::InterchangeTextureFactory
 
 	bool CanSetupTextureCubeSourceData(FTexturePayloadVariant& TexturePayload)
 	{
+		static_assert(TVariantSize<FTexturePayloadVariant>::Value == 5, "Please update the code below and this assert to reflect the change to the variant type.");
+
 		if (TOptional<FImportSlicedImage>* SlicedImage = TexturePayload.TryGet<TOptional<FImportSlicedImage>>())
 		{
 			if (SlicedImage->IsSet())
@@ -755,6 +777,8 @@ namespace UE::Interchange::Private::InterchangeTextureFactory
 
 	bool CanSetupTexture2DArraySourceData(FTexturePayloadVariant& TexturePayload)
 	{
+		static_assert(TVariantSize<FTexturePayloadVariant>::Value == 5, "Please update the code below and this assert to reflect the change to the variant type.");
+
 		if (TOptional<FImportSlicedImage>* SlicedImage = TexturePayload.TryGet<TOptional<FImportSlicedImage>>())
 		{
 			if (SlicedImage->IsSet())
@@ -821,6 +845,8 @@ namespace UE::Interchange::Private::InterchangeTextureFactory
 
 	FSharedBuffer MoveRawDataToSharedBuffer(FTexturePayloadVariant& TexturePayload)
 	{
+		static_assert(TVariantSize<FTexturePayloadVariant>::Value == 5, "Please update the code below and this assert to reflect the change to the variant type.");
+
 		if (TOptional<FImportBlockedImage>* BlockedImage = TexturePayload.TryGet<TOptional<FImportBlockedImage>>())
 		{
 			return (*BlockedImage)->RawData.MoveToShared();
@@ -1163,6 +1189,8 @@ UObject* UInterchangeTextureFactory::CreateAsset(const FCreateAssetParams& Argum
 		return nullptr;
 	}
 
+	// Check if the imported texture(s) has a valid resolution
+	CheckForInvalidResolutions(TexturePayload, Arguments.SourceData, CastChecked<UInterchangeTextureFactoryNode>(Arguments.AssetNode));
 
 	bool bCanSetup = false;
 	// Check if the payload is valid for the Texture
@@ -1294,6 +1322,111 @@ void UInterchangeTextureFactory::PreImportPreCompletedCallback(const FImportPreC
 #endif
 }
 
+void UInterchangeTextureFactory::CheckForInvalidResolutions(UE::Interchange::Private::InterchangeTextureFactory::FTexturePayloadVariant& InPayloadVariant, const UInterchangeSourceData* SourceData, const UInterchangeTextureFactoryNode* TextureFactoryNode)
+{
+	using namespace UE::Interchange;
+
+	auto AddErrorMessage = [this, &TextureFactoryNode](FString&& InSourceAssetName, const FText& ErrorMessage) -> UInterchangeResultError_Generic*
+	{
+		UInterchangeResultError_Generic* Message = AddMessage<UInterchangeResultError_Generic>();
+		Message->SourceAssetName = MoveTemp(InSourceAssetName);
+		Message->DestinationAssetName = TextureFactoryNode->GetAssetName();
+		Message->AssetType = TextureFactoryNode->GetObjectClass();
+		Message->Text = ErrorMessage;
+		return Message;
+	};
+
+
+	bool bAllowNonPowerOfTwo = false;
+	TextureFactoryNode->GetCustomAllowNonPowerOfTwo(bAllowNonPowerOfTwo);
+
+	static_assert(TVariantSize<Private::InterchangeTextureFactory::FTexturePayloadVariant>::Value == 5, "Please update the code below and this assert to reflect the change to the variant type.");
+
+	FText ErrorMessage;
+	if (TOptional<FImportBlockedImage>* BlockedImagePtr = InPayloadVariant.TryGet<TOptional<FImportBlockedImage>>())
+	{
+		if (BlockedImagePtr->IsSet())
+		{
+			FImportBlockedImage& BlockedImage = BlockedImagePtr->GetValue();
+
+			for (int32 Index = 0; Index < BlockedImage.BlocksData.Num(); ++Index)
+			{
+				const FTextureSourceBlock& Block = BlockedImage.BlocksData[Index];
+				if (!FImportImageHelper::IsImportResolutionValid(Block.SizeX, Block.SizeY, bAllowNonPowerOfTwo, &ErrorMessage))
+				{
+					FString SourceFile;
+					if (const UInterchangeTexture2DFactoryNode* Texture2DFactoryNode = Cast<UInterchangeTexture2DFactoryNode>(TextureFactoryNode))
+					{
+						Texture2DFactoryNode->GetSourceBlockByCoordinates(Block.BlockX, Block.BlockY, SourceFile);
+					}
+					if (SourceFile.IsEmpty())
+					{
+						SourceFile = SourceData->GetFilename();
+					}
+					UInterchangeResultError_Generic* Message = AddErrorMessage(MoveTemp(SourceFile), ErrorMessage);
+
+					FString BlockMessage = FText::Format(LOCTEXT("BlockIndexInvalidResolutionError", "Invalid block (X : {0}, Y : {1}) \n"), Block.BlockX, Block.BlockY).ToString();
+					Message->Text = FText::FromString(BlockMessage.Append(Message->Text.ToString()));
+
+					BlockedImage.BlocksData.RemoveAtSwap(Index);
+					--Index;
+				}
+			}
+
+			if (BlockedImage.BlocksData.IsEmpty())
+			{
+				AddErrorMessage(FString(), LOCTEXT("UDIMHasNoValidBlock", "All blocks where invalid. The texture won't be imported."));
+
+				// Remove the payload
+				InPayloadVariant.Set<FEmptyVariantState>(FEmptyVariantState());
+			}
+		}
+	}
+	else if (TOptional<FImportImage>* ImagePtr = InPayloadVariant.TryGet<TOptional<FImportImage>>())
+	{
+		if (ImagePtr->IsSet())
+		{
+			const FImportImage& Image = ImagePtr->GetValue();
+			if (UTextureCube::StaticClass() != TextureFactoryNode->GetObjectClass() && !FImportImageHelper::IsImportResolutionValid(Image.SizeX, Image.SizeY, bAllowNonPowerOfTwo, &ErrorMessage))
+			{
+				AddErrorMessage(SourceData->GetFilename(), ErrorMessage);
+
+				// Remove the payload
+				InPayloadVariant.Set<FEmptyVariantState>(FEmptyVariantState());
+			}
+		}
+	}
+	else if (TOptional<FImportLightProfile>* LightProfilePtr = InPayloadVariant.TryGet<TOptional<FImportLightProfile>>())
+	{
+		if (LightProfilePtr->IsSet())
+		{
+			const FImportLightProfile& LightProfile = LightProfilePtr->GetValue();
+			if (!FImportImageHelper::IsImportResolutionValid(LightProfile.SizeX, LightProfile.SizeY, bAllowNonPowerOfTwo, &ErrorMessage))
+			{
+				AddErrorMessage(SourceData->GetFilename(), ErrorMessage);
+
+				// Remove the payload
+				InPayloadVariant.Set<FEmptyVariantState>(FEmptyVariantState());
+			}
+		}
+	}
+	else if (TOptional<FImportSlicedImage>* SlicedImagePtr = InPayloadVariant.TryGet<TOptional<FImportSlicedImage>>())
+	{
+		if (SlicedImagePtr->IsSet())
+		{
+			const FImportSlicedImage& SlicedImage = SlicedImagePtr->GetValue();
+
+			if (!FImportImageHelper::IsImportResolutionValid(SlicedImage.SizeX, SlicedImage.SizeY, bAllowNonPowerOfTwo, &ErrorMessage))
+			{
+				AddErrorMessage(SourceData->GetFilename(), ErrorMessage);
+
+				// Remove the payload
+				InPayloadVariant.Set<FEmptyVariantState>(FEmptyVariantState());
+			}
+		}
+	}
+}
+
 bool UInterchangeTextureFactory::GetSourceFilenames(const UObject* Object, TArray<FString>& OutSourceFilenames) const
 {
 #if WITH_EDITORONLY_DATA
@@ -1316,4 +1449,6 @@ bool UInterchangeTextureFactory::SetSourceFilename(const UObject* Object, const
 #endif
 
 	return false;
-}
\ No newline at end of file
+}
+
+#undef LOCTEXT_NAMESPACE
\ No newline at end of file
diff --git a/Engine/Plugins/Experimental/Interchange/Runtime/Source/Import/Private/Texture/InterchangeTexturePayloadData.cpp b/Engine/Plugins/Experimental/Interchange/Runtime/Source/Import/Private/Texture/InterchangeTexturePayloadData.cpp
--- a/Engine/Plugins/Experimental/Interchange/Runtime/Source/Import/Private/Texture/InterchangeTexturePayloadData.cpp
+++ b/Engine/Plugins/Experimental/Interchange/Runtime/Source/Import/Private/Texture/InterchangeTexturePayloadData.cpp
@@ -8,7 +8,7 @@
 #include "Misc/MessageDialog.h"
 #include "RHI.h"
 
-bool UE::Interchange::FImportImageHelper::IsImportResolutionValid(int32 Width, int32 Height, bool bAllowNonPowerOfTwo)
+bool UE::Interchange::FImportImageHelper::IsImportResolutionValid(int32 Width, int32 Height, bool bAllowNonPowerOfTwo, FText* OutErrorMessage)
 {
 	static const auto CVarVirtualTexturesEnabled = IConsoleManager::Get().FindTConsoleVariableDataInt(TEXT("r.VirtualTextures")); check(CVarVirtualTexturesEnabled);
 
@@ -19,34 +19,32 @@ bool UE::Interchange::FImportImageHelper::IsImportResolutionValid(int32 Width, i
 	// (Note, have to subtract 1 because 1x1 is a valid mip-size; this means a GMaxTextureMipCount of 4 means a max resolution of 8x8, not 2^4 = 16x16)
 	const int32 MaximumSupportedResolution = CVarVirtualTexturesEnabled->GetValueOnAnyThread() ? MaximumSupportedVirtualTextureResolution : (1 << (GMaxTextureMipCount - 1));
 
-	bool bValid = true;
-
 	// Check if the texture is above the supported resolution and prompt the user if they wish to continue if it is
 	if (Width > MaximumSupportedResolution || Height > MaximumSupportedResolution)
 	{
-		if (EAppReturnType::Yes != FMessageDialog::Open(EAppMsgType::YesNo, FText::Format(
-			NSLOCTEXT("Interchange", "Warning_LargeTextureImport", "Attempting to import {0} x {1} texture, proceed?\nLargest supported texture size: {2} x {3}"),
-			FText::AsNumber(Width), FText::AsNumber(Height), FText::AsNumber(MaximumSupportedResolution), FText::AsNumber(MaximumSupportedResolution))))
-		{
-			bValid = false;
-		}
-
-		if (bValid && (Width * Height) > FMath::Square(MaximumSupportedVirtualTextureResolution))
+		if ((Width * Height) > FMath::Square(MaximumSupportedVirtualTextureResolution))
 		{
-			//Warn->Log(ELogVerbosity::Error, *NSLOCTEXT("UnrealEd", "Warning_TextureSizeTooLarge", "Texture is too large to import").ToString());
-			bValid = false;
+			if (OutErrorMessage)
+			{
+				*OutErrorMessage = NSLOCTEXT("Interchange", "Warning_TextureSizeTooLarge", "Texture is too large to import");
+			}
+			
+			return false;
 		}
 	}
 
-	const bool bIsPowerOfTwo = FMath::IsPowerOfTwo(Width) && FMath::IsPowerOfTwo(Height);
 	// Check if the texture dimensions are powers of two
-	if (!bAllowNonPowerOfTwo && !bIsPowerOfTwo)
+	if (!bAllowNonPowerOfTwo)
 	{
-		//Warn->Log(ELogVerbosity::Error, *NSLOCTEXT("UnrealEd", "Warning_TextureNotAPowerOfTwo", "Cannot import texture with non-power of two dimensions").ToString());
-		bValid = false;
+		const bool bIsPowerOfTwo = FMath::IsPowerOfTwo(Width) && FMath::IsPowerOfTwo(Height);
+		if (!bIsPowerOfTwo)
+		{
+			*OutErrorMessage = NSLOCTEXT("Interchange", "Warning_TextureNotAPowerOfTwo", "Cannot import texture with non-power of two dimensions");
+			return false;
+		}
 	}
 
-	return bValid;
+	return true;
 }
 
 void UE::Interchange::FImportImage::Init2DWithParams(int32 InSizeX, int32 InSizeY, ETextureSourceFormat InFormat, bool bInSRGB, bool bShouldAllocateRawData)
diff --git a/Engine/Plugins/Experimental/Interchange/Runtime/Source/Import/Public/Texture/InterchangeBMPTranslator.h b/Engine/Plugins/Experimental/Interchange/Runtime/Source/Import/Public/Texture/InterchangeBMPTranslator.h
deleted file mode 100644
diff --git a/Engine/Plugins/Experimental/Interchange/Runtime/Source/Import/Public/Texture/InterchangeEXRTranslator.h b/Engine/Plugins/Experimental/Interchange/Runtime/Source/Import/Public/Texture/InterchangeEXRTranslator.h
deleted file mode 100644
diff --git a/Engine/Plugins/Experimental/Interchange/Runtime/Source/Import/Public/Texture/InterchangeHDRTranslator.h b/Engine/Plugins/Experimental/Interchange/Runtime/Source/Import/Public/Texture/InterchangeHDRTranslator.h
deleted file mode 100644
diff --git a/Engine/Plugins/Experimental/Interchange/Runtime/Source/Import/Public/Texture/InterchangeImageWrapperTranslator.h b/Engine/Plugins/Experimental/Interchange/Runtime/Source/Import/Public/Texture/InterchangeImageWrapperTranslator.h
new file mode 100644
--- /dev/null
+++ b/Engine/Plugins/Experimental/Interchange/Runtime/Source/Import/Public/Texture/InterchangeImageWrapperTranslator.h
@@ -0,0 +1,49 @@
+// Copyright Epic Games, Inc. All Rights Reserved.
+
+#pragma once
+
+#include "CoreMinimal.h"
+#include "UObject/Object.h"
+#include "UObject/ObjectMacros.h"
+#include "InterchangeTranslatorBase.h"
+#include "Nodes/InterchangeBaseNodeContainer.h"
+#include "Texture/InterchangeTexturePayloadData.h"
+#include "Texture/InterchangeTexturePayloadInterface.h"
+
+#include "InterchangeImageWrapperTranslator.generated.h"
+
+/**
+ * A translator for most of the image wrapper supported formats
+ */
+UCLASS(BlueprintType, Experimental)
+class INTERCHANGEIMPORT_API UInterchangeImageWrapperTranslator : public UInterchangeTranslatorBase, public IInterchangeTexturePayloadInterface
+{
+	GENERATED_BODY()
+public:
+
+	virtual TArray<FString> GetSupportedFormats() const override;
+
+	virtual bool DoesSupportAssetType(EInterchangeTranslatorAssetType AssetType) const override { return AssetType == EInterchangeTranslatorAssetType::Textures; }
+
+	/**
+	 * Translate the associated source data into a node hold by the specified nodes container.
+	 *
+	 * @param BaseNodeContainer - The unreal objects descriptions container where to put the translated source data.
+	 * @return true if the translator can translate the source data, false otherwise.
+	 */
+	virtual bool Translate(UInterchangeBaseNodeContainer& BaseNodeContainer) const override;
+
+	/* IInterchangeTexturePayloadInterface Begin */
+
+	/**
+	 * Once the translation is done, the import process need a way to retrieve payload data.
+	 * This payload will be use by the factories to create the asset.
+	 *
+	 * @param PayloadSourceData - The source data containing the data to translate
+	 * @param PayloadKey - The key to retrieve the a particular payload contain into the specified source data.
+	 * @return a PayloadData containing the import image data. The TOptional will not be set if there is an error.
+	 */
+	virtual TOptional<UE::Interchange::FImportImage> GetTexturePayloadData(const UInterchangeSourceData* PayloadSourceData, const FString& PayLoadKey) const override;
+
+	/* IInterchangeTexturePayloadInterface End */
+};
\ No newline at end of file
diff --git a/Engine/Plugins/Experimental/Interchange/Runtime/Source/Import/Public/Texture/InterchangePNGTranslator.h b/Engine/Plugins/Experimental/Interchange/Runtime/Source/Import/Public/Texture/InterchangePNGTranslator.h
deleted file mode 100644
diff --git a/Engine/Plugins/Experimental/Interchange/Runtime/Source/Import/Public/Texture/InterchangeTGATranslator.h b/Engine/Plugins/Experimental/Interchange/Runtime/Source/Import/Public/Texture/InterchangeTGATranslator.h
deleted file mode 100644
diff --git a/Engine/Plugins/Experimental/Interchange/Runtime/Source/Import/Public/Texture/InterchangeTIFFTranslator.h b/Engine/Plugins/Experimental/Interchange/Runtime/Source/Import/Public/Texture/InterchangeTIFFTranslator.h
deleted file mode 100644
diff --git a/Engine/Plugins/Experimental/Interchange/Runtime/Source/Import/Public/Texture/InterchangeTextureFactory.h b/Engine/Plugins/Experimental/Interchange/Runtime/Source/Import/Public/Texture/InterchangeTextureFactory.h
--- a/Engine/Plugins/Experimental/Interchange/Runtime/Source/Import/Public/Texture/InterchangeTextureFactory.h
+++ b/Engine/Plugins/Experimental/Interchange/Runtime/Source/Import/Public/Texture/InterchangeTextureFactory.h
@@ -17,6 +17,8 @@
 
 #include "InterchangeTextureFactory.generated.h"
 
+class UInterchangeTextureFactoryNode;
+
 namespace UE::Interchange::Private::InterchangeTextureFactory
 {
 	using FTexturePayloadVariant = TVariant<FEmptyVariantState
@@ -63,10 +65,15 @@ public:
 	//////////////////////////////////////////////////////////////////////////
 
 private:
+	/**
+	 * Check for any invalid resolutions and remove those from the payload
+	 */
+	void CheckForInvalidResolutions(UE::Interchange::Private::InterchangeTextureFactory::FTexturePayloadVariant& InPayloadVariant, const UInterchangeSourceData* SourceData, const UInterchangeTextureFactoryNode* TextureFactoryNode);
+	
 	UE::Interchange::Private::InterchangeTextureFactory::FProcessedPayload ProcessedPayload;
 
 #if WITH_EDITORONLY_DATA
-	// When importing a UDIM the data for the source files will be stored here
+	//  The data for the source files will be stored here during the import
 	TArray<FAssetImportInfo::FSourceFile> SourceFiles;
 #endif // WITH_EDITORONLY_DATA
 };
diff --git a/Engine/Plugins/Experimental/Interchange/Runtime/Source/Import/Public/Texture/InterchangeTexturePayloadData.h b/Engine/Plugins/Experimental/Interchange/Runtime/Source/Import/Public/Texture/InterchangeTexturePayloadData.h
--- a/Engine/Plugins/Experimental/Interchange/Runtime/Source/Import/Public/Texture/InterchangeTexturePayloadData.h
+++ b/Engine/Plugins/Experimental/Interchange/Runtime/Source/Import/Public/Texture/InterchangeTexturePayloadData.h
@@ -51,13 +51,14 @@ namespace UE
 			/**
 			 * Tests if the given height and width specify a supported texture resolution to import; Can optionally check if the height/width are powers of two
 			 *
-			 * @param Width - The width of an imported texture whose validity should be checked
-			 * @param Height - The height of an imported texture whose validity should be checked
-			 * @param bAllowNonPowerOfTwo - Whether or not non-power-of-two textures are allowed
+			 * @param Width The width of an imported texture whose validity should be checked
+			 * @param Height The height of an imported texture whose validity should be checked
+			 * @param bAllowNonPowerOfTwo Whether or not non-power-of-two textures are allowed
+			 * @param OutErrorMessage Optional output for an error message
 			 *
 			 * @return bool true if the given height/width represent a supported texture resolution, false if not
 			 */
-			static bool IsImportResolutionValid(int32 Width, int32 Height, bool bAllowNonPowerOfTwo);
+			static bool IsImportResolutionValid(int32 Width, int32 Height, bool bAllowNonPowerOfTwo, FText* OutErrorMessage = nullptr);
 		};
 	}//ns Interchange
 }//ns UE
diff --git a/Engine/Plugins/Experimental/Interchange/Runtime/Source/Pipelines/Private/InterchangeGenericTexturePipeline.cpp b/Engine/Plugins/Experimental/Interchange/Runtime/Source/Pipelines/Private/InterchangeGenericTexturePipeline.cpp
--- a/Engine/Plugins/Experimental/Interchange/Runtime/Source/Pipelines/Private/InterchangeGenericTexturePipeline.cpp
+++ b/Engine/Plugins/Experimental/Interchange/Runtime/Source/Pipelines/Private/InterchangeGenericTexturePipeline.cpp
@@ -235,6 +235,11 @@ UInterchangeTextureFactoryNode* UInterchangeGenericTexturePipeline::CreateTextur
 
 		TextureFactoryNode->AddTargetNodeUid(TextureNode->GetUniqueID());
 		TextureNode->AddTargetNodeUid(TextureFactoryNode->GetUniqueID());
+
+		if (bAllowNonPowerOfTwo)
+		{
+			TextureFactoryNode->SetCustomAllowNonPowerOfTwo(bAllowNonPowerOfTwo);
+		}
 	}
 
 #if WITH_EDITORONLY_DATA
diff --git a/Engine/Plugins/Experimental/Interchange/Runtime/Source/Pipelines/Public/InterchangeGenericTexturePipeline.h b/Engine/Plugins/Experimental/Interchange/Runtime/Source/Pipelines/Public/InterchangeGenericTexturePipeline.h
--- a/Engine/Plugins/Experimental/Interchange/Runtime/Source/Pipelines/Public/InterchangeGenericTexturePipeline.h
+++ b/Engine/Plugins/Experimental/Interchange/Runtime/Source/Pipelines/Public/InterchangeGenericTexturePipeline.h
@@ -32,15 +32,15 @@ public:
 	bool bDetectNormalMapTexture = true;
 
 	/** If enabled, the texture's green channel will be inverted for normal maps. */
-	UPROPERTY(EditAnywhere, Category = "Textures")
+	UPROPERTY(EditAnywhere, BlueprintReadWrite, Category = "Textures")
 	bool bFlipNormalMapGreenChannel = false;
 
 	/** If enabled detect if a texture use a UDIM pattern and if so import it as UIDMs. */
-	UPROPERTY(EditAnywhere, Category = "Textures")
+	UPROPERTY(EditAnywhere, BlueprintReadWrite, Category = "Textures")
 	bool bImportUDIMs = true;
 
 	/** Specify the files type that should be imported as long/lat cubemap */
-	UPROPERTY(EditAnywhere, Category = "Textures")
+	UPROPERTY(EditAnywhere, BlueprintReadWrite, Category = "Textures")
 	TSet<FString> FileExtensionsToImportAsLongLatCubemap = {"hdr"};
 
 	/** 
@@ -55,6 +55,10 @@ public:
 
 #endif
 
+	/** Should the textures that have a non-power of two resolution be imported */
+	UPROPERTY(EditAnywhere, BlueprintReadWrite, AdvancedDisplay, Category = "Textures")
+	bool bAllowNonPowerOfTwo = false;
+
 public:
 	virtual void AdjustSettingsForReimportType(EInterchangeReimportType ImportType, TObjectPtr<UObject> ReimportAsset) override;
 	virtual void ExecutePreImportPipeline(UInterchangeBaseNodeContainer* InBaseNodeContainer, const TArray<UInterchangeSourceData*>& InSourceDatas) override;
diff --git a/Engine/Source/Editor/UnrealEd/Private/Factories/EditorFactories.cpp b/Engine/Source/Editor/UnrealEd/Private/Factories/EditorFactories.cpp
--- a/Engine/Source/Editor/UnrealEd/Private/Factories/EditorFactories.cpp
+++ b/Engine/Source/Editor/UnrealEd/Private/Factories/EditorFactories.cpp
@@ -272,6 +272,7 @@
 #include "ObjectTools.h"
 
 #include "SkinWeightsUtilities.h"
+#include "TextureImportUtils.h"
 #include "UDIMUtilities.h"
 
 DEFINE_LOG_CATEGORY(LogEditorFactories);
@@ -3019,32 +3020,6 @@ void* FImportImage::GetMipData(int32 InMipIndex)
 	return &RawData[Offset];
 }
 
-static void AutoDetectAndChangeGrayScale(FImage & Image)
-{
-	if ( Image.Format != ERawImageFormat::BGRA8 )
-	{
-		return;
-	}
-
-	// auto-detect gray BGRA8 and change to G8
-
-	const FColor * Colors = (const FColor *)Image.RawData.GetData();
-	int64 NumPixels = Image.GetNumPixels();
-
-	for(int64 i=0;i<NumPixels;i++)
-	{
-		if ( Colors[i].A != 255 ||
-			Colors[i].R != Colors[i].B ||
-			Colors[i].G != Colors[i].B )
-		{
-			return;
-		}
-	}
-
-	// yes, it's gray, do it :
-	Image.ChangeFormat(ERawImageFormat::G8,Image.GammaSpace);
-}
-
 bool UTextureFactory::ImportImage(const uint8* Buffer, int64 Length, FFeedbackContext* Warn, EImageImportFlags Flags, FImportImage& OutImage)
 {
 	TRACE_CPUPROFILER_EVENT_SCOPE(UTextureFactory::ImportImage)
@@ -3114,7 +3089,7 @@ bool UTextureFactory::ImportImage(const uint8* Buffer, int64 Length, FFeedbackCo
 				return false;
 			}
 		
-			AutoDetectAndChangeGrayScale(LoadedImage);
+			UE::TextureUtilitiesCommon::AutoDetectAndChangeGrayScale(LoadedImage);
 
 			ETextureSourceFormat TextureFormat = FImageCoreUtils::ConvertToTextureSourceFormat(LoadedImage.Format);
 			bool bSRGB = LoadedImage.GammaSpace != EGammaSpace::Linear;
diff --git a/Engine/Source/Runtime/Experimental/Interchange/Core/Public/Nodes/InterchangeBaseNodeUtilities.h b/Engine/Source/Runtime/Experimental/Interchange/Core/Public/Nodes/InterchangeBaseNodeUtilities.h
--- a/Engine/Source/Runtime/Experimental/Interchange/Core/Public/Nodes/InterchangeBaseNodeUtilities.h
+++ b/Engine/Source/Runtime/Experimental/Interchange/Core/Public/Nodes/InterchangeBaseNodeUtilities.h
@@ -405,6 +405,8 @@ public:
 		{
 			return Pair->Value.Get(OutValue) == EAttributeStorageResult::Operation_Success;
 		}
+
+		return false;
 	}
 
 
diff --git a/Engine/Source/Runtime/TextureUtilitiesCommon/Private/TextureImportUtils.cpp b/Engine/Source/Runtime/TextureUtilitiesCommon/Private/TextureImportUtils.cpp
new file mode 100644
--- /dev/null
+++ b/Engine/Source/Runtime/TextureUtilitiesCommon/Private/TextureImportUtils.cpp
@@ -0,0 +1,44 @@
+// Copyright Epic Games, Inc. All Rights Reserved.
+
+#include "TextureImportUtils.h"
+
+#include "ImageCore.h"
+
+namespace UE
+{
+	namespace TextureUtilitiesCommon
+	{
+		/**
+		 * Detect the existence of gray scale image in some formats and convert those to a gray scale equivalent image
+		 * 
+		 * @return true if the image was converted
+		 */
+		bool AutoDetectAndChangeGrayScale(FImage& Image)
+		{
+			if (Image.Format != ERawImageFormat::BGRA8)
+			{
+				return false;
+			}
+
+			// auto-detect gray BGRA8 and change to G8
+
+			const FColor* Colors = (const FColor*)Image.RawData.GetData();
+			int64 NumPixels = Image.GetNumPixels();
+
+			for (int64 i = 0; i < NumPixels; i++)
+			{
+				if (Colors[i].A != 255 ||
+					Colors[i].R != Colors[i].B ||
+					Colors[i].G != Colors[i].B)
+				{
+					return false ;
+				}
+			}
+
+			// yes, it's gray, do it :
+			Image.ChangeFormat(ERawImageFormat::G8, Image.GammaSpace);
+
+			return true;
+		}
+	}
+}
\ No newline at end of file
diff --git a/Engine/Source/Runtime/TextureUtilitiesCommon/Private/UDIMUtilities.cpp b/Engine/Source/Runtime/TextureUtilitiesCommon/Private/UDIMUtilities.cpp
--- a/Engine/Source/Runtime/TextureUtilitiesCommon/Private/UDIMUtilities.cpp
+++ b/Engine/Source/Runtime/TextureUtilitiesCommon/Private/UDIMUtilities.cpp
@@ -54,7 +54,7 @@ namespace UE
 			return BlockY * 10 + BlockX + 1001;
 		}
 
-		TEXTUREUTILITIESCOMMON_API TMap<int32, FString> GetUDIMBlocksFromSourceFile(const FString& File, const FString& UdimRegexPattern, FString* OutFilenameWithoutUdimPatternAndExtension)
+		TMap<int32, FString> GetUDIMBlocksFromSourceFile(const FString& File, const FString& UdimRegexPattern, FString* OutFilenameWithoutUdimPatternAndExtension)
 		{
 			FString PreUDIMName;
 			FString PostUDIMName;
diff --git a/Engine/Source/Runtime/TextureUtilitiesCommon/Public/TextureImportUtils.h b/Engine/Source/Runtime/TextureUtilitiesCommon/Public/TextureImportUtils.h
new file mode 100644
--- /dev/null
+++ b/Engine/Source/Runtime/TextureUtilitiesCommon/Public/TextureImportUtils.h
@@ -0,0 +1,17 @@
+// Copyright Epic Games, Inc. All Rights Reserved.
+
+#pragma once
+
+#include "CoreTypes.h"
+
+struct FImage;
+
+namespace UE::TextureUtilitiesCommon
+{
+	/**
+		* Detect the existence of gray scale image in some formats and convert those to a gray scale equivalent image
+		* 
+		* @return true if the image was converted
+		*/
+	TEXTUREUTILITIESCOMMON_API bool AutoDetectAndChangeGrayScale(FImage& Image);
+}
\ No newline at end of file
diff --git a/Engine/Source/Runtime/TextureUtilitiesCommon/TextureUtilitiesCommon.build.cs b/Engine/Source/Runtime/TextureUtilitiesCommon/TextureUtilitiesCommon.build.cs
--- a/Engine/Source/Runtime/TextureUtilitiesCommon/TextureUtilitiesCommon.build.cs
+++ b/Engine/Source/Runtime/TextureUtilitiesCommon/TextureUtilitiesCommon.build.cs
@@ -16,6 +16,7 @@ public class TextureUtilitiesCommon : ModuleRules
 			new string[] {
 				"CoreUObject",
 				"Engine",
+				"ImageCore",
 				"Slate",
 				"SlateCore",
 			});
