Fixed logic in UInterchangeGenericLevelPipeline::ExecuteSceneNodePreImport
Issue: The ActorFactoryNode is calling AddFactoryDependencyUid with its own unique id. This creates aninfinite loop in FNodeDependencyCache::GetAccumulatedDependencies

diff --git a/Engine/Plugins/Experimental/Interchange/Runtime/Source/Pipelines/Private/InterchangeGenericScenesPipeline.cpp b/Engine/Plugins/Experimental/Interchange/Runtime/Source/Pipelines/Private/InterchangeGenericScenesPipeline.cpp
--- a/Engine/Plugins/Experimental/Interchange/Runtime/Source/Pipelines/Private/InterchangeGenericScenesPipeline.cpp
+++ b/Engine/Plugins/Experimental/Interchange/Runtime/Source/Pipelines/Private/InterchangeGenericScenesPipeline.cpp
@@ -6,6 +6,7 @@
 #include "InterchangeCameraNode.h"
 #include "InterchangeCineCameraFactoryNode.h"
 #include "InterchangeLightNode.h"
+#include "InterchangeMaterialFactoryNode.h"
 #include "InterchangeMeshNode.h"
 #include "InterchangePipelineLog.h"
 #include "InterchangeSceneNode.h"
@@ -89,8 +90,9 @@ void UInterchangeGenericLevelPipeline::ExecuteSceneNodePreImport(UInterchangeBas
 	const FString ActorFactoryNodeUid = FactoryNodeContainer->AddNode(ActorFactoryNode);
 	if (!SceneNode->GetParentUid().IsEmpty())
 	{
-		FactoryNodeContainer->SetNodeParentUid(ActorFactoryNodeUid, TEXT("Factory_") + SceneNode->GetParentUid());
-		ActorFactoryNode->AddFactoryDependencyUid(ActorFactoryNodeUid);
+		const FString ParentFactoryNodeUid = TEXT("Factory_") + SceneNode->GetParentUid();
+		FactoryNodeContainer->SetNodeParentUid(ActorFactoryNodeUid, ParentFactoryNodeUid);
+		ActorFactoryNode->AddFactoryDependencyUid(ParentFactoryNodeUid);
 	}
 
 	ActorFactoryNode->AddTargetNodeUid(SceneNode->GetUniqueID());
@@ -106,10 +108,8 @@ void UInterchangeGenericLevelPipeline::ExecuteSceneNodePreImport(UInterchangeBas
 
 	if (TranslatedAssetNode)
 	{
-		SetUpFactoryNode(ActorFactoryNode, TranslatedAssetNode, FactoryNodeContainer);
+		SetUpFactoryNode(ActorFactoryNode, SceneNode, TranslatedAssetNode, FactoryNodeContainer);
 	}
-
-	FactoryNodeContainer->AddNode(ActorFactoryNode);
 }
 
 UInterchangeActorFactoryNode* UInterchangeGenericLevelPipeline::CreateActorFactoryNode(const UInterchangeSceneNode* SceneNode, const UInterchangeBaseNode* TranslatedAssetNode, UInterchangeBaseNodeContainer* FactoryNodeContainer) const
@@ -124,7 +124,7 @@ UInterchangeActorFactoryNode* UInterchangeGenericLevelPipeline::CreateActorFacto
 	}
 }
 
-void UInterchangeGenericLevelPipeline::SetUpFactoryNode(UInterchangeActorFactoryNode* ActorFactoryNode, const UInterchangeBaseNode* TranslatedAssetNode, UInterchangeBaseNodeContainer* FactoryNodeContainer) const
+void UInterchangeGenericLevelPipeline::SetUpFactoryNode(UInterchangeActorFactoryNode* ActorFactoryNode, const UInterchangeSceneNode* SceneNode, const UInterchangeBaseNode* TranslatedAssetNode, UInterchangeBaseNodeContainer* FactoryNodeContainer) const
 {
 	if (const UInterchangeMeshNode* MeshNode = Cast<UInterchangeMeshNode>(TranslatedAssetNode))
 	{
@@ -137,6 +137,30 @@ void UInterchangeGenericLevelPipeline::SetUpFactoryNode(UInterchangeActorFactory
 		{
 			ActorFactoryNode->SetCustomActorClassName(AStaticMeshActor::StaticClass()->GetPathName());
 		}
+
+		if (SceneNode)
+		{
+			TMap<FString, FString> SlotMaterialDependencies;
+			SceneNode->GetSlotMaterialDependencies(SlotMaterialDependencies);
+
+			if (SlotMaterialDependencies.Num() > 0)
+			{
+				for (TPair<FString, FString>& SlotMaterialDependency : SlotMaterialDependencies)
+				{
+					const FString MaterialFactoryNodeUid = UInterchangeMaterialFactoryNode::GetMaterialFactoryNodeUidFromMaterialNodeUid(SlotMaterialDependency.Value);
+					if (FactoryNodeContainer->IsNodeUidValid(MaterialFactoryNodeUid))
+					{
+						// Create a factory dependency so Material asset are imported before the static mesh asset
+						TArray<FString> FactoryDependencies;
+						ActorFactoryNode->GetFactoryDependencies(FactoryDependencies);
+						if (!FactoryDependencies.Contains(MaterialFactoryNodeUid))
+						{
+							ActorFactoryNode->AddFactoryDependencyUid(MaterialFactoryNodeUid);
+						}
+					}
+				}
+			}
+		}
 	}
 	else if (const UInterchangeLightNode* LightNode = Cast<UInterchangeLightNode>(TranslatedAssetNode))
 	{
