Issue 166026: Compile errors in InterchangeGenericScenesPipeline.cpp

Solution: Removed unwanted code from previous submission

diff --git a/Engine/Plugins/Experimental/Interchange/Runtime/Source/Pipelines/Private/InterchangeGenericScenesPipeline.cpp b/Engine/Plugins/Experimental/Interchange/Runtime/Source/Pipelines/Private/InterchangeGenericScenesPipeline.cpp
--- a/Engine/Plugins/Experimental/Interchange/Runtime/Source/Pipelines/Private/InterchangeGenericScenesPipeline.cpp
+++ b/Engine/Plugins/Experimental/Interchange/Runtime/Source/Pipelines/Private/InterchangeGenericScenesPipeline.cpp
@@ -6,7 +6,6 @@
 #include "InterchangeCameraNode.h"
 #include "InterchangeCineCameraFactoryNode.h"
 #include "InterchangeLightNode.h"
-#include "InterchangeMaterialFactoryNode.h"
 #include "InterchangeMeshNode.h"
 #include "InterchangePipelineLog.h"
 #include "InterchangeSceneNode.h"
@@ -108,7 +107,7 @@ void UInterchangeGenericLevelPipeline::ExecuteSceneNodePreImport(UInterchangeBas
 
 	if (TranslatedAssetNode)
 	{
-		SetUpFactoryNode(ActorFactoryNode, SceneNode, TranslatedAssetNode, FactoryNodeContainer);
+		SetUpFactoryNode(ActorFactoryNode, TranslatedAssetNode, FactoryNodeContainer);
 	}
 }
 
@@ -124,7 +123,7 @@ UInterchangeActorFactoryNode* UInterchangeGenericLevelPipeline::CreateActorFacto
 	}
 }
 
-void UInterchangeGenericLevelPipeline::SetUpFactoryNode(UInterchangeActorFactoryNode* ActorFactoryNode, const UInterchangeSceneNode* SceneNode, const UInterchangeBaseNode* TranslatedAssetNode, UInterchangeBaseNodeContainer* FactoryNodeContainer) const
+void UInterchangeGenericLevelPipeline::SetUpFactoryNode(UInterchangeActorFactoryNode* ActorFactoryNode, const UInterchangeBaseNode* TranslatedAssetNode, UInterchangeBaseNodeContainer* FactoryNodeContainer) const
 {
 	if (const UInterchangeMeshNode* MeshNode = Cast<UInterchangeMeshNode>(TranslatedAssetNode))
 	{
@@ -137,30 +136,6 @@ void UInterchangeGenericLevelPipeline::SetUpFactoryNode(UInterchangeActorFactory
 		{
 			ActorFactoryNode->SetCustomActorClassName(AStaticMeshActor::StaticClass()->GetPathName());
 		}
-
-		if (SceneNode)
-		{
-			TMap<FString, FString> SlotMaterialDependencies;
-			SceneNode->GetSlotMaterialDependencies(SlotMaterialDependencies);
-
-			if (SlotMaterialDependencies.Num() > 0)
-			{
-				for (TPair<FString, FString>& SlotMaterialDependency : SlotMaterialDependencies)
-				{
-					const FString MaterialFactoryNodeUid = UInterchangeMaterialFactoryNode::GetMaterialFactoryNodeUidFromMaterialNodeUid(SlotMaterialDependency.Value);
-					if (FactoryNodeContainer->IsNodeUidValid(MaterialFactoryNodeUid))
-					{
-						// Create a factory dependency so Material asset are imported before the static mesh asset
-						TArray<FString> FactoryDependencies;
-						ActorFactoryNode->GetFactoryDependencies(FactoryDependencies);
-						if (!FactoryDependencies.Contains(MaterialFactoryNodeUid))
-						{
-							ActorFactoryNode->AddFactoryDependencyUid(MaterialFactoryNodeUid);
-						}
-					}
-				}
-			}
-		}
 	}
 	else if (const UInterchangeLightNode* LightNode = Cast<UInterchangeLightNode>(TranslatedAssetNode))
 	{
