Add support of Anisotropy input for materials
Add an 'Anisotropy' name parameter in the PBR namespace of the Material Definitions

diff --git a/Engine/Plugins/Experimental/Interchange/Runtime/Source/FactoryNodes/Private/InterchangeMaterialFactoryNode.cpp b/Engine/Plugins/Experimental/Interchange/Runtime/Source/FactoryNodes/Private/InterchangeMaterialFactoryNode.cpp
--- a/Engine/Plugins/Experimental/Interchange/Runtime/Source/FactoryNodes/Private/InterchangeMaterialFactoryNode.cpp
+++ b/Engine/Plugins/Experimental/Interchange/Runtime/Source/FactoryNodes/Private/InterchangeMaterialFactoryNode.cpp
@@ -92,6 +92,21 @@ bool UInterchangeMaterialFactoryNode::ConnectOutputToRoughness(const FString& Ex
 	return UInterchangeShaderPortsAPI::ConnectOuputToInput(this, UE::Interchange::Materials::PBR::Parameters::Roughness.ToString(), ExpressionNodeUid, OutputName);
 }
 
+bool UInterchangeMaterialFactoryNode::GetAnisotropyConnection(FString& ExpressionNodeUid, FString& OutputName) const
+{
+	return UInterchangeShaderPortsAPI::GetInputConnection(this, UE::Interchange::Materials::PBR::Parameters::Anisotropy.ToString(), ExpressionNodeUid, OutputName);
+}
+
+bool UInterchangeMaterialFactoryNode::ConnectToAnisotropy(const FString& ExpressionNodeUid)
+{
+	return UInterchangeShaderPortsAPI::ConnectDefaultOuputToInput(this, UE::Interchange::Materials::PBR::Parameters::Anisotropy.ToString(), ExpressionNodeUid);
+}
+
+bool UInterchangeMaterialFactoryNode::ConnectOutputToAnisotropy(const FString& ExpressionNodeUid, const FString& OutputName)
+{
+	return UInterchangeShaderPortsAPI::ConnectOuputToInput(this, UE::Interchange::Materials::PBR::Parameters::Anisotropy.ToString(), ExpressionNodeUid, OutputName);
+}
+
 bool UInterchangeMaterialFactoryNode::GetEmissiveColorConnection(FString& ExpressionNodeUid, FString& OutputName) const
 {
 	return UInterchangeShaderPortsAPI::GetInputConnection(this, UE::Interchange::Materials::Common::Parameters::EmissiveColor.ToString(), ExpressionNodeUid, OutputName);
diff --git a/Engine/Plugins/Experimental/Interchange/Runtime/Source/FactoryNodes/Public/InterchangeMaterialFactoryNode.h b/Engine/Plugins/Experimental/Interchange/Runtime/Source/FactoryNodes/Public/InterchangeMaterialFactoryNode.h
--- a/Engine/Plugins/Experimental/Interchange/Runtime/Source/FactoryNodes/Public/InterchangeMaterialFactoryNode.h
+++ b/Engine/Plugins/Experimental/Interchange/Runtime/Source/FactoryNodes/Public/InterchangeMaterialFactoryNode.h
@@ -72,6 +72,15 @@ public:
 	UFUNCTION(BlueprintCallable, Category = "Interchange | Node | MaterialFactory")
 	bool ConnectOutputToRoughness(const FString& ExpressionNodeUid, const FString& OutputName);
 
+	UFUNCTION(BlueprintCallable, Category = "Interchange | Node | MaterialFactory")
+	bool GetAnisotropyConnection(FString& ExpressionNodeUid, FString& OutputName) const;
+
+	UFUNCTION(BlueprintCallable, Category = "Interchange | Node | MaterialFactory")
+	bool ConnectToAnisotropy(const FString& ExpressionNodeUid);
+
+	UFUNCTION(BlueprintCallable, Category = "Interchange | Node | MaterialFactory")
+	bool ConnectOutputToAnisotropy(const FString& ExpressionNodeUid, const FString& OutputName);
+
 	UFUNCTION(BlueprintCallable, Category = "Interchange | Node | MaterialFactory")
 	bool GetEmissiveColorConnection(FString& ExpressionNodeUid, FString& OutputName) const;
 
diff --git a/Engine/Plugins/Experimental/Interchange/Runtime/Source/Import/Private/Material/InterchangeMaterialFactory.cpp b/Engine/Plugins/Experimental/Interchange/Runtime/Source/Import/Private/Material/InterchangeMaterialFactory.cpp
--- a/Engine/Plugins/Experimental/Interchange/Runtime/Source/Import/Private/Material/InterchangeMaterialFactory.cpp
+++ b/Engine/Plugins/Experimental/Interchange/Runtime/Source/Import/Private/Material/InterchangeMaterialFactory.cpp
@@ -475,6 +475,28 @@ void UInterchangeMaterialFactory::SetupMaterial(UMaterial* Material, const FCrea
 		}
 	}
 
+	// Anisotropy
+	{
+		FString AnisotropyUid;
+		FString OutputName;
+
+		if (MaterialFactoryNode->GetAnisotropyConnection(AnisotropyUid, OutputName))
+		{
+			const UInterchangeMaterialExpressionFactoryNode* Anisotropy = Cast<UInterchangeMaterialExpressionFactoryNode>(Arguments.NodeContainer->GetNode(AnisotropyUid));
+
+			if (Anisotropy)
+			{
+				if (UMaterialExpression* AnisotropyExpression = CreateExpressionsForNode(Material, Arguments, Anisotropy, Expressions))
+				{
+					if (FExpressionInput* AnisotropyInput = Material->GetExpressionInputForProperty(MP_Anisotropy))
+					{
+						AnisotropyExpression->ConnectExpression(AnisotropyInput, GetOutputIndex(*AnisotropyExpression, OutputName));
+					}
+				}
+			}
+		}
+	}
+
 	// Emissive
 	{
 		FString EmissiveColorUid;
diff --git a/Engine/Plugins/Experimental/Interchange/Runtime/Source/Nodes/Public/InterchangeMaterialDefinitions.h b/Engine/Plugins/Experimental/Interchange/Runtime/Source/Nodes/Public/InterchangeMaterialDefinitions.h
--- a/Engine/Plugins/Experimental/Interchange/Runtime/Source/Nodes/Public/InterchangeMaterialDefinitions.h
+++ b/Engine/Plugins/Experimental/Interchange/Runtime/Source/Nodes/Public/InterchangeMaterialDefinitions.h
@@ -167,6 +167,7 @@ namespace Interchange
 				const FName Metallic = TEXT("Metallic"); // Type: float
 				const FName Specular = TEXT("Specular"); // Type: float
 				const FName Roughness = TEXT("Roughness"); // Type: float
+				const FName Anisotropy = TEXT("Anisotropy"); // Type: float
 			}
 		}
 
