Interchange: Fixed importing into level multiple times.
- Moved setting the static mesh on the static mesh component to the post import callback as the component wants the static mesh render data to be initialized if it's not null.

diff --git a/Engine/Plugins/Experimental/Interchange/Runtime/Source/Import/Private/Scene/InterchangeStaticMeshActorFactory.cpp b/Engine/Plugins/Experimental/Interchange/Runtime/Source/Import/Private/Scene/InterchangeStaticMeshActorFactory.cpp
--- a/Engine/Plugins/Experimental/Interchange/Runtime/Source/Import/Private/Scene/InterchangeStaticMeshActorFactory.cpp
+++ b/Engine/Plugins/Experimental/Interchange/Runtime/Source/Import/Private/Scene/InterchangeStaticMeshActorFactory.cpp
@@ -36,14 +36,33 @@ UObject* UInterchangeStaticMeshActorFactory::CreateSceneObject(const UInterchang
 
 void UInterchangeStaticMeshActorFactory::SetupStaticMeshActor(const UInterchangeBaseNodeContainer* NodeContainer, const UInterchangeFactoryBaseNode* ActorFactoryNode, AStaticMeshActor* StaticMeshActor)
 {
-	UStaticMeshComponent* StaticMeshComponent = StaticMeshActor->GetStaticMeshComponent();
-	StaticMeshComponent->UnregisterComponent();
+	if (!StaticMeshActor)
+	{
+		return;
+	}
+
+	if (UStaticMeshComponent* StaticMeshComponent = StaticMeshActor->GetStaticMeshComponent())
+	{
+		StaticMeshComponent->UnregisterComponent();
+	}
+}
+
+void UInterchangeStaticMeshActorFactory::PostImportPreCompletedCallback(const FImportPreCompletedCallbackParams& Arguments)
+{
+	// Set the static mesh on the component in the post import callback, once the static mesh has been fully imported.
+	// The component doesn't like being set a static mesh with uninitialized render data.
 
-	if (const UInterchangeFactoryBaseNode* MeshNode = UE::Interchange::ActorHelper::FindAssetInstanceFactoryNode(NodeContainer, ActorFactoryNode))
+	if (AStaticMeshActor* StaticMeshActor = Cast<AStaticMeshActor>(Arguments.ImportedObject))
 	{
-		if (UStaticMesh* StaticMesh = Cast<UStaticMesh>(MeshNode->ReferenceObject.TryLoad()))
+		if (UStaticMeshComponent* StaticMeshComponent = StaticMeshActor->GetStaticMeshComponent())
 		{
-			StaticMeshComponent->SetStaticMesh(StaticMesh);
+			if (const UInterchangeFactoryBaseNode* MeshNode = UE::Interchange::ActorHelper::FindAssetInstanceFactoryNode(Arguments.NodeContainer, Arguments.FactoryNode))
+			{
+				if (UStaticMesh* StaticMesh = Cast<UStaticMesh>(MeshNode->ReferenceObject.TryLoad()))
+				{
+					StaticMeshComponent->SetStaticMesh(StaticMesh);
+				}
+			}
 		}
 	}
 }
\ No newline at end of file
diff --git a/Engine/Plugins/Experimental/Interchange/Runtime/Source/Import/Public/Scene/InterchangeStaticMeshActorFactory.h b/Engine/Plugins/Experimental/Interchange/Runtime/Source/Import/Public/Scene/InterchangeStaticMeshActorFactory.h
--- a/Engine/Plugins/Experimental/Interchange/Runtime/Source/Import/Public/Scene/InterchangeStaticMeshActorFactory.h
+++ b/Engine/Plugins/Experimental/Interchange/Runtime/Source/Import/Public/Scene/InterchangeStaticMeshActorFactory.h
@@ -31,6 +31,8 @@ public:
 		return false; 
 	}
 
+	virtual void PostImportPreCompletedCallback(const FImportPreCompletedCallbackParams& Arguments) override;
+
 	// Interchange factory base interface end
 	//////////////////////////////////////////////////////////////////////////
 
