Interchange Textures
- Added address mode support for 2D texture nodes.

Interchange GLTF
- Translated GLTF wrap mode.
- Renamed texture nodes to use the GLTFTexture name instead of the file URI since the same file can be sampled differently and result in multiple textures.

diff --git a/Engine/Plugins/Experimental/Interchange/Runtime/Source/FactoryNodes/Public/InterchangeTexture2DFactoryNode.h b/Engine/Plugins/Experimental/Interchange/Runtime/Source/FactoryNodes/Public/InterchangeTexture2DFactoryNode.h
--- a/Engine/Plugins/Experimental/Interchange/Runtime/Source/FactoryNodes/Public/InterchangeTexture2DFactoryNode.h
+++ b/Engine/Plugins/Experimental/Interchange/Runtime/Source/FactoryNodes/Public/InterchangeTexture2DFactoryNode.h
@@ -32,28 +32,28 @@ public:
 
 	/** Return false if the Attribute was not set previously.*/
 	UFUNCTION(BlueprintCallable, Category = "Interchange | Node | Texture")
-	bool GetCustomAddressX(uint8& AttributeValue) const
+	bool GetCustomAddressX(TEnumAsByte<TextureAddress>& AttributeValue) const
 	{
-		IMPLEMENT_NODE_ATTRIBUTE_GETTER(AddressX, uint8);
+		IMPLEMENT_NODE_ATTRIBUTE_GETTER(AddressX, TEnumAsByte<TextureAddress>);
 	}
 
 	UFUNCTION(BlueprintCallable, Category = "Interchange | Node | Texture")
-	bool SetCustomAddressX(const uint8 AttributeValue, bool bAddApplyDelegate = true)
+	bool SetCustomAddressX(const TEnumAsByte<TextureAddress> AttributeValue, bool bAddApplyDelegate = true)
 	{
-		IMPLEMENT_NODE_ATTRIBUTE_SETTER(UInterchangeTexture2DFactoryNode, AddressX, uint8, UTexture2D)
+		IMPLEMENT_NODE_ATTRIBUTE_SETTER(UInterchangeTexture2DFactoryNode, AddressX, TEnumAsByte<TextureAddress>, UTexture2D)
 	}
 
 	/** Return false if the Attribute was not set previously.*/
 	UFUNCTION(BlueprintCallable, Category = "Interchange | Node | Texture")
-	bool GetCustomAddressY(uint8& AttributeValue) const
+	bool GetCustomAddressY(TEnumAsByte<TextureAddress>& AttributeValue) const
 	{
-		IMPLEMENT_NODE_ATTRIBUTE_GETTER(AddressY, uint8);
+		IMPLEMENT_NODE_ATTRIBUTE_GETTER(AddressY, TEnumAsByte<TextureAddress>);
 	}
 
 	UFUNCTION(BlueprintCallable, Category = "Interchange | Node | Texture")
-	bool SetCustomAddressY(const uint8 AttributeValue, bool bAddApplyDelegate = true)
+	bool SetCustomAddressY(const TEnumAsByte<TextureAddress> AttributeValue, bool bAddApplyDelegate = true)
 	{
-		IMPLEMENT_NODE_ATTRIBUTE_SETTER(UInterchangeTexture2DFactoryNode, AddressY, uint8, UTexture2D)
+		IMPLEMENT_NODE_ATTRIBUTE_SETTER(UInterchangeTexture2DFactoryNode, AddressY, TEnumAsByte<TextureAddress>, UTexture2D)
 	}
 
 	//////////////////////////////////////////////////////////////////////////
diff --git a/Engine/Plugins/Experimental/Interchange/Runtime/Source/Import/Private/Gltf/InterchangeGltfTranslator.cpp b/Engine/Plugins/Experimental/Interchange/Runtime/Source/Import/Private/Gltf/InterchangeGltfTranslator.cpp
--- a/Engine/Plugins/Experimental/Interchange/Runtime/Source/Import/Private/Gltf/InterchangeGltfTranslator.cpp
+++ b/Engine/Plugins/Experimental/Interchange/Runtime/Source/Import/Private/Gltf/InterchangeGltfTranslator.cpp
@@ -32,7 +32,24 @@ namespace UE::Interchange::Gltf::Private
 			return NodeName + TEXT("_") + LexToString( NodeIndex );
 		}
 	}
+
+	EInterchangeTextureWrapMode ConvertWrap(GLTF::FSampler::EWrap Wrap)
+	{
+		switch (Wrap)
+		{
+		case GLTF::FSampler::EWrap::Repeat:
+			return EInterchangeTextureWrapMode::Wrap;
+		case GLTF::FSampler::EWrap::MirroredRepeat:
+			return EInterchangeTextureWrapMode::Mirror;
+		case GLTF::FSampler::EWrap::ClampToEdge:
+			return EInterchangeTextureWrapMode::Clamp;
+
+		default:
+			return EInterchangeTextureWrapMode::Wrap;
+		}
+	}
 }
+
 void UInterchangeGltfTranslator::HandleGltfNode( UInterchangeBaseNodeContainer& NodeContainer, const GLTF::FNode& GltfNode, const FString& ParentNodeUid, const int32 NodeIndex ) const
 {
 	using namespace UE::Interchange::Gltf::Private;
@@ -157,7 +174,7 @@ void UInterchangeGltfTranslator::HandleGltfMaterialParameter( UInterchangeBaseNo
 		NodeContainer.SetNodeParentUid( NodeUid, ShaderNode.GetUniqueID() );
 
 		ColorNode->SetCustomShaderType( Standard::Nodes::TextureSample::Name.ToString() );
-		const FString TextureUid = TEXT("\\Texture\\") + GltfAsset.Textures[ TextureMap.TextureIndex ].Source.URI;
+		const FString TextureUid = TEXT("\\Texture\\") + GltfAsset.Textures[ TextureMap.TextureIndex ].Name;
 		ColorNode->AddStringAttribute( UInterchangeShaderPortsAPI::MakeInputValueKey( Standard::Nodes::TextureSample::Inputs::Texture.ToString() ), TextureUid );
 
 		if ( TextureMap.bHasTextureTransform )
@@ -737,9 +754,13 @@ bool UInterchangeGltfTranslator::Translate( UInterchangeBaseNodeContainer& NodeC
 		for ( const GLTF::FTexture& GltfTexture : GltfAsset.Textures )
 		{
 			UInterchangeTexture2DNode* TextureNode = NewObject< UInterchangeTexture2DNode >( &NodeContainer );
-			FString TextureNodeUid = TEXT("\\Texture\\") + GltfTexture.Source.URI;
-			TextureNode->InitializeNode( TextureNodeUid, GltfTexture.Source.URI, EInterchangeNodeContainerType::TranslatedAsset );
+			FString TextureNodeUid = TEXT("\\Texture\\") + GltfTexture.Name;
+			TextureNode->InitializeNode( TextureNodeUid, GltfTexture.Name, EInterchangeNodeContainerType::TranslatedAsset );
 			TextureNode->SetPayLoadKey( LexToString( TextureIndex++ ) );
+
+			TextureNode->SetCustomWrapU( UE::Interchange::Gltf::Private::ConvertWrap( GltfTexture.Sampler.WrapS ) );
+			TextureNode->SetCustomWrapV( UE::Interchange::Gltf::Private::ConvertWrap( GltfTexture.Sampler.WrapT ) );
+
 			NodeContainer.AddNode( TextureNode );
 		}
 	}
diff --git a/Engine/Plugins/Experimental/Interchange/Runtime/Source/Nodes/Public/InterchangeTexture2DNode.h b/Engine/Plugins/Experimental/Interchange/Runtime/Source/Nodes/Public/InterchangeTexture2DNode.h
--- a/Engine/Plugins/Experimental/Interchange/Runtime/Source/Nodes/Public/InterchangeTexture2DNode.h
+++ b/Engine/Plugins/Experimental/Interchange/Runtime/Source/Nodes/Public/InterchangeTexture2DNode.h
@@ -87,4 +87,34 @@ public:
 protected:
 
 	UE::Interchange::TMapAttributeHelper<int32, FString> SourceBlocks;
+
+// Custom attributes
+public:
+	UFUNCTION(BlueprintCallable, Category = "Interchange | Node | Texture")
+	bool SetCustomWrapU(const EInterchangeTextureWrapMode& AttributeValue)
+	{
+		IMPLEMENT_NODE_ATTRIBUTE_SETTER_NODELEGATE(WrapU, EInterchangeTextureWrapMode);
+	}
+
+	UFUNCTION(BlueprintCallable, Category = "Interchange | Node | Texture")
+	bool GetCustomWrapU(EInterchangeTextureWrapMode& AttributeValue) const
+	{
+		IMPLEMENT_NODE_ATTRIBUTE_GETTER(WrapU, EInterchangeTextureWrapMode);
+	}
+
+	UFUNCTION(BlueprintCallable, Category = "Interchange | Node | Texture")
+	bool SetCustomWrapV(const EInterchangeTextureWrapMode& AttributeValue)
+	{
+		IMPLEMENT_NODE_ATTRIBUTE_SETTER_NODELEGATE(WrapV, EInterchangeTextureWrapMode);
+	}
+
+	UFUNCTION(BlueprintCallable, Category = "Interchange | Node | Texture")
+	bool GetCustomWrapV(EInterchangeTextureWrapMode& AttributeValue) const
+	{
+		IMPLEMENT_NODE_ATTRIBUTE_GETTER(WrapV, EInterchangeTextureWrapMode);
+	}
+
+private:
+	IMPLEMENT_NODE_ATTRIBUTE_KEY(WrapU);
+	IMPLEMENT_NODE_ATTRIBUTE_KEY(WrapV);
 };
diff --git a/Engine/Plugins/Experimental/Interchange/Runtime/Source/Nodes/Public/InterchangeTextureNode.h b/Engine/Plugins/Experimental/Interchange/Runtime/Source/Nodes/Public/InterchangeTextureNode.h
--- a/Engine/Plugins/Experimental/Interchange/Runtime/Source/Nodes/Public/InterchangeTextureNode.h
+++ b/Engine/Plugins/Experimental/Interchange/Runtime/Source/Nodes/Public/InterchangeTextureNode.h
@@ -23,6 +23,14 @@ namespace UE
 	}//ns Interchange
 }//ns UE
 
+UENUM(BlueprintType)
+enum class EInterchangeTextureWrapMode : uint8
+{
+	Wrap,
+	Clamp,
+	Mirror
+};
+
 UCLASS(BlueprintType, Abstract, Experimental)
 class INTERCHANGENODES_API UInterchangeTextureNode : public UInterchangeBaseNode
 {
diff --git a/Engine/Plugins/Experimental/Interchange/Runtime/Source/Pipelines/Private/InterchangeGenericTexturePipeline.cpp b/Engine/Plugins/Experimental/Interchange/Runtime/Source/Pipelines/Private/InterchangeGenericTexturePipeline.cpp
--- a/Engine/Plugins/Experimental/Interchange/Runtime/Source/Pipelines/Private/InterchangeGenericTexturePipeline.cpp
+++ b/Engine/Plugins/Experimental/Interchange/Runtime/Source/Pipelines/Private/InterchangeGenericTexturePipeline.cpp
@@ -75,6 +75,23 @@ namespace UE::Interchange::Private
 		}
 	}
 #endif
+
+	TextureAddress ConvertWrap(const EInterchangeTextureWrapMode WrapMode)
+	{
+		switch (WrapMode)
+		{
+		case EInterchangeTextureWrapMode::Wrap:
+			return TA_Wrap;
+		case EInterchangeTextureWrapMode::Clamp:
+			return TA_Clamp;
+		case EInterchangeTextureWrapMode::Mirror:
+			return TA_Mirror;
+
+		default:
+			ensureMsgf(false, TEXT("Unkown Interchange Texture Wrap Mode"));
+			return TA_Wrap;
+		}
+	}
 }
 
 void UInterchangeGenericTexturePipeline::AdjustSettingsForReimportType(EInterchangeReimportType ImportType, TObjectPtr<UObject> ReimportAsset)
@@ -177,6 +194,18 @@ UInterchangeTextureFactoryNode* UInterchangeGenericTexturePipeline::HandleCreati
 		if (const UInterchangeTexture2DNode* Texture2DNode = Cast<UInterchangeTexture2DNode>(TextureNode))
 		{
 			SourceBlocks = Texture2DNode->GetSourceBlocks();
+
+			EInterchangeTextureWrapMode WrapU;
+			if (Texture2DNode->GetCustomWrapU(WrapU))
+			{
+				Texture2DFactoryNode->SetCustomAddressX(UE::Interchange::Private::ConvertWrap(WrapU));
+			}
+
+			EInterchangeTextureWrapMode WrapV;
+			if (Texture2DNode->GetCustomWrapV(WrapV))
+			{
+				Texture2DFactoryNode->SetCustomAddressY(UE::Interchange::Private::ConvertWrap(WrapV));
+			}
 		}
 
 #if WITH_EDITOR
