Interchange GLTF: Fixed texture names when it's missing from the source data. Use the URI in that case.

diff --git a/Engine/Plugins/Experimental/Interchange/Runtime/Source/Import/Private/Gltf/InterchangeGltfTranslator.cpp b/Engine/Plugins/Experimental/Interchange/Runtime/Source/Import/Private/Gltf/InterchangeGltfTranslator.cpp
--- a/Engine/Plugins/Experimental/Interchange/Runtime/Source/Import/Private/Gltf/InterchangeGltfTranslator.cpp
+++ b/Engine/Plugins/Experimental/Interchange/Runtime/Source/Import/Private/Gltf/InterchangeGltfTranslator.cpp
@@ -33,6 +33,18 @@ namespace UE::Interchange::Gltf::Private
 		}
 	}
 
+	FString GenerateTextureName( const GLTF::FTexture& Texture, int32 NodeIndex )
+	{
+		if ( !Texture.Name.IsEmpty() )
+		{
+			return GenerateUniqueIdForGltfNode( Texture.Name, NodeIndex );
+		}
+		else
+		{
+			return GenerateUniqueIdForGltfNode( FPaths::GetBaseFilename( Texture.Source.URI ), NodeIndex );
+		}
+	}
+
 	EInterchangeTextureWrapMode ConvertWrap(GLTF::FSampler::EWrap Wrap)
 	{
 		switch (Wrap)
@@ -131,6 +143,7 @@ void UInterchangeGltfTranslator::HandleGltfMaterialParameter( UInterchangeBaseNo
 		const FString& MapName, const TVariant< FLinearColor, float >& MapFactor, const FString& OutputChannel, const bool bInverse, const bool bIsNormal ) const
 {
 	using namespace UE::Interchange::Materials;
+	using namespace UE::Interchange::Gltf::Private;
 
 	UInterchangeShaderNode* NodeToConnectTo = &ShaderNode;
 	FString InputToConnectTo = MapName;
@@ -174,7 +187,10 @@ void UInterchangeGltfTranslator::HandleGltfMaterialParameter( UInterchangeBaseNo
 		NodeContainer.SetNodeParentUid( NodeUid, ShaderNode.GetUniqueID() );
 
 		ColorNode->SetCustomShaderType( Standard::Nodes::TextureSample::Name.ToString() );
-		const FString TextureUid = TEXT("\\Texture\\") + GltfAsset.Textures[ TextureMap.TextureIndex ].Name;
+
+		const FString TextureName = GenerateTextureName( GltfAsset.Textures[ TextureMap.TextureIndex ], TextureMap.TextureIndex );
+		const FString TextureUid = TEXT("\\Texture\\") + TextureName;
+
 		ColorNode->AddStringAttribute( UInterchangeShaderPortsAPI::MakeInputValueKey( Standard::Nodes::TextureSample::Inputs::Texture.ToString() ), TextureUid );
 
 		if ( TextureMap.bHasTextureTransform )
@@ -753,9 +769,12 @@ bool UInterchangeGltfTranslator::Translate( UInterchangeBaseNodeContainer& NodeC
 		int32 TextureIndex = 0;
 		for ( const GLTF::FTexture& GltfTexture : GltfAsset.Textures )
 		{
+			const FString TextureName = GenerateTextureName( GltfTexture, TextureIndex );
+
 			UInterchangeTexture2DNode* TextureNode = NewObject< UInterchangeTexture2DNode >( &NodeContainer );
-			FString TextureNodeUid = TEXT("\\Texture\\") + GltfTexture.Name;
-			TextureNode->InitializeNode( TextureNodeUid, GltfTexture.Name, EInterchangeNodeContainerType::TranslatedAsset );
+			FString TextureNodeUid = TEXT("\\Texture\\") + TextureName;
+
+			TextureNode->InitializeNode( TextureNodeUid, TextureName, EInterchangeNodeContainerType::TranslatedAsset );
 			TextureNode->SetPayLoadKey( LexToString( TextureIndex++ ) );
 
 			TextureNode->SetCustomWrapU( UE::Interchange::Gltf::Private::ConvertWrap( GltfTexture.Sampler.WrapS ) );
