Implement ImportOffset generic options
#rnx

diff --git a/Engine/Plugins/Experimental/Interchange/Runtime/Source/FactoryNodes/Private/InterchangeCommonPipelineDataFactoryNode.cpp b/Engine/Plugins/Experimental/Interchange/Runtime/Source/FactoryNodes/Private/InterchangeCommonPipelineDataFactoryNode.cpp
new file mode 100644
--- /dev/null
+++ b/Engine/Plugins/Experimental/Interchange/Runtime/Source/FactoryNodes/Private/InterchangeCommonPipelineDataFactoryNode.cpp
@@ -0,0 +1,46 @@
+// Copyright Epic Games, Inc. All Rights Reserved.
+
+#include "InterchangeCommonPipelineDataFactoryNode.h"
+
+#include "InterchangeSceneNode.h"
+#include "Nodes/InterchangeBaseNodeContainer.h"
+
+namespace UE::Interchange::CommonPipelineData
+{
+	FString GetCommonPipelineDataUniqueID()
+	{
+		static FString StaticUid = TEXT("CommonPipelineDataFactoryNode");
+		return StaticUid;
+	}
+}
+
+UInterchangeCommonPipelineDataFactoryNode* UInterchangeCommonPipelineDataFactoryNode::FindOrCreateUniqueInstance(UInterchangeBaseNodeContainer* NodeContainer)
+{
+	const FString StaticUid = UE::Interchange::CommonPipelineData::GetCommonPipelineDataUniqueID();
+	UInterchangeCommonPipelineDataFactoryNode* CommonPipelineDataFactoryNode = Cast<UInterchangeCommonPipelineDataFactoryNode>(NodeContainer->GetFactoryNode(StaticUid));
+	if (!CommonPipelineDataFactoryNode)
+	{
+		CommonPipelineDataFactoryNode = NewObject<UInterchangeCommonPipelineDataFactoryNode>(NodeContainer, NAME_None);
+		CommonPipelineDataFactoryNode->InitializeNode(StaticUid, StaticUid, EInterchangeNodeContainerType::FactoryData);
+		NodeContainer->AddNode(CommonPipelineDataFactoryNode);
+	}
+	return CommonPipelineDataFactoryNode;
+}
+UInterchangeCommonPipelineDataFactoryNode* UInterchangeCommonPipelineDataFactoryNode::GetUniqueInstance(const UInterchangeBaseNodeContainer* NodeContainer)
+{
+	static FString StaticUid = UE::Interchange::CommonPipelineData::GetCommonPipelineDataUniqueID();
+	return Cast<UInterchangeCommonPipelineDataFactoryNode>(NodeContainer->GetFactoryNode(StaticUid));
+}
+
+bool UInterchangeCommonPipelineDataFactoryNode::GetCustomGlobalOffsetTransform(FTransform& AttributeValue) const
+{
+	IMPLEMENT_NODE_ATTRIBUTE_GETTER(GlobalOffsetTransform, FTransform);
+}
+
+UFUNCTION(BlueprintCallable, Category = "Interchange | Node | Common Pipeline Data")
+bool UInterchangeCommonPipelineDataFactoryNode::SetCustomGlobalOffsetTransform(const UInterchangeBaseNodeContainer* NodeContainer, const FTransform& AttributeValue)
+{
+	IMPLEMENT_NODE_ATTRIBUTE_SETTER_NODELEGATE(GlobalOffsetTransform, FTransform);
+	//Reset all scene node container cache
+	UInterchangeSceneNode::ResetAllGlobalTransformCaches(NodeContainer);
+}
diff --git a/Engine/Plugins/Experimental/Interchange/Runtime/Source/FactoryNodes/Public/InterchangeCommonPipelineDataFactoryNode.h b/Engine/Plugins/Experimental/Interchange/Runtime/Source/FactoryNodes/Public/InterchangeCommonPipelineDataFactoryNode.h
new file mode 100644
--- /dev/null
+++ b/Engine/Plugins/Experimental/Interchange/Runtime/Source/FactoryNodes/Public/InterchangeCommonPipelineDataFactoryNode.h
@@ -0,0 +1,37 @@
+// Copyright Epic Games, Inc. All Rights Reserved.
+
+#pragma once
+
+#include "Nodes/InterchangeFactoryBaseNode.h"
+
+#include "InterchangeCommonPipelineDataFactoryNode.generated.h"
+
+class UInterchangeBaseNodeContainer;
+
+
+/* This factory node is the place where pipeline can set global data that can be use by factories. */
+UCLASS(BlueprintType, Experimental)
+class INTERCHANGEFACTORYNODES_API UInterchangeCommonPipelineDataFactoryNode : public UInterchangeFactoryBaseNode
+{
+	GENERATED_BODY()
+
+public:
+	/* The pipelines that want to modify the common data shoiuld ensure they create the unique common pipeline node. */
+	static UInterchangeCommonPipelineDataFactoryNode* FindOrCreateUniqueInstance(UInterchangeBaseNodeContainer* NodeContainer);
+	
+	/* If the unique instance doesn't exist it will return nullptr. This function should be use by the factories, to avoid creating a node. */
+	static UInterchangeCommonPipelineDataFactoryNode* GetUniqueInstance(const UInterchangeBaseNodeContainer* NodeContainer);
+
+	/** Return the global offset transform set by the pipelines. */
+	UFUNCTION(BlueprintCallable, Category = "Interchange | Node | Common Pipeline Data")
+	bool GetCustomGlobalOffsetTransform(FTransform& AttributeValue) const;
+
+	/** Pipeline can set a global transform, factories will use this global offset when importing asset. */
+	UFUNCTION(BlueprintCallable, Category = "Interchange | Node | Common Pipeline Data")
+	bool SetCustomGlobalOffsetTransform(const UInterchangeBaseNodeContainer* NodeContainer, const FTransform& AttributeValue);
+
+private:
+	UInterchangeCommonPipelineDataFactoryNode() {};
+
+	const UE::Interchange::FAttributeKey Macro_CustomGlobalOffsetTransformKey = UE::Interchange::FAttributeKey(TEXT("GlobalOffsetTransform"));
+};
\ No newline at end of file
diff --git a/Engine/Plugins/Experimental/Interchange/Runtime/Source/Import/Private/Animation/InterchangeAnimSequenceFactory.cpp b/Engine/Plugins/Experimental/Interchange/Runtime/Source/Import/Private/Animation/InterchangeAnimSequenceFactory.cpp
--- a/Engine/Plugins/Experimental/Interchange/Runtime/Source/Import/Private/Animation/InterchangeAnimSequenceFactory.cpp
+++ b/Engine/Plugins/Experimental/Interchange/Runtime/Source/Import/Private/Animation/InterchangeAnimSequenceFactory.cpp
@@ -4,6 +4,7 @@
 #include "Animation/AnimSequence.h"
 #include "Animation/InterchangeAnimationPayload.h"
 #include "Animation/InterchangeAnimationPayloadInterface.h"
+#include "InterchangeCommonPipelineDataFactoryNode.h"
 #include "InterchangeAssetImportData.h"
 #include "InterchangeCommonAnimationPayload.h"
 #include "InterchangeImportCommon.h"
@@ -88,6 +89,16 @@ namespace UE::Interchange::Private
 			Controller.RemoveAllBoneTracks();
 			Controller.SetPlayLength(FGenericPlatformMath::Max<float>(SequenceLength, MINIMUM_ANIMATION_LENGTH));
 
+			FTransform3f GlobalOffsetTransform;
+			{
+				FTransform TempTransform = FTransform::Identity;
+				if (UInterchangeCommonPipelineDataFactoryNode* CommonPipelineDataFactoryNode = UInterchangeCommonPipelineDataFactoryNode::GetUniqueInstance(NodeContainer))
+				{
+					CommonPipelineDataFactoryNode->GetCustomGlobalOffsetTransform(TempTransform);
+				}
+				GlobalOffsetTransform = FTransform3f(TempTransform);
+			}
+
 			for (const FString& NodeUid : SkeletonNodes)
 			{
 				if (const UInterchangeSceneNode* SkeletonSceneNode = Cast<UInterchangeSceneNode>(NodeContainer->GetNode(NodeUid)))
@@ -99,6 +110,9 @@ namespace UE::Interchange::Private
 						//Skip this bone, we did not found it in the skeleton
 						continue;
 					}
+					//If we are getting the root 
+					bool bApplyGlobalOffset = NodeUid.Equals(SkeletonRootUid);
+
 					FString PayloadKey;
 					if (SkeletonSceneNode->GetCustomTransformCurvePayloadKey(PayloadKey))
 					{
@@ -164,6 +178,15 @@ namespace UE::Interchange::Private
 								}
 							}
 							FQuat4f Quaternion = FQuat4f::MakeFromEuler(Euler);
+							
+							if (bApplyGlobalOffset)
+							{
+								FTransform3f AnimKeyTransform(Quaternion, Position, Scale);
+								AnimKeyTransform = AnimKeyTransform * GlobalOffsetTransform;
+								Position = AnimKeyTransform.GetLocation();
+								Quaternion = AnimKeyTransform.GetRotation();
+								Scale = AnimKeyTransform.GetScale3D();
+							}
 							RawTrack.ScaleKeys.Add(Scale);
 							RawTrack.PosKeys.Add(Position);
 							RawTrack.RotKeys.Add(Quaternion);
diff --git a/Engine/Plugins/Experimental/Interchange/Runtime/Source/Import/Private/Mesh/InterchangeSkeletalMeshFactory.cpp b/Engine/Plugins/Experimental/Interchange/Runtime/Source/Import/Private/Mesh/InterchangeSkeletalMeshFactory.cpp
--- a/Engine/Plugins/Experimental/Interchange/Runtime/Source/Import/Private/Mesh/InterchangeSkeletalMeshFactory.cpp
+++ b/Engine/Plugins/Experimental/Interchange/Runtime/Source/Import/Private/Mesh/InterchangeSkeletalMeshFactory.cpp
@@ -6,6 +6,7 @@
 #include "Engine/SkeletalMesh.h"
 #include "GPUSkinPublicDefs.h"
 #include "InterchangeAssetImportData.h"
+#include "InterchangeCommonPipelineDataFactoryNode.h"
 #include "InterchangeImportCommon.h"
 #include "InterchangeImportLog.h"
 #include "InterchangeMaterialFactoryNode.h"
@@ -288,6 +289,12 @@ namespace UE
 				TArray<double> SkinWeights;
 				SkinWeights.AddZeroed(VertexCount);
 
+				FTransform GlobalOffsetTransform = FTransform::Identity;
+				if (UInterchangeCommonPipelineDataFactoryNode* CommonPipelineDataFactoryNode = UInterchangeCommonPipelineDataFactoryNode::GetUniqueInstance(NodeContainer))
+				{
+					CommonPipelineDataFactoryNode->GetCustomGlobalOffsetTransform(GlobalOffsetTransform);
+				}
+
 				//We assume normalize weight method in this bind pose conversion
 
 				const FTransform MeshGlobalTransformInverse = MeshGlobalTransform.Inverse();
@@ -303,11 +310,11 @@ namespace UE
 					}
 					
 					FTransform JointBindPoseGlobalTransform;
-					if (!JointNode->GetCustomBindPoseGlobalTransform(NodeContainer, JointBindPoseGlobalTransform))
+					if (!JointNode->GetCustomBindPoseGlobalTransform(NodeContainer, GlobalOffsetTransform, JointBindPoseGlobalTransform))
 					{
 						//If there is no bind pose we will fall back on the CustomGlobalTransform of the link.
 						//We ensure here because any scenenode should have a valid CustomGlobalTransform.
-						if (!ensure(JointNode->GetCustomGlobalTransform(NodeContainer, JointBindPoseGlobalTransform)))
+						if (!ensure(JointNode->GetCustomGlobalTransform(NodeContainer, GlobalOffsetTransform, JointBindPoseGlobalTransform)))
 						{
 							//No value to convert from, skip this joint.
 							continue;
@@ -315,7 +322,7 @@ namespace UE
 					}
 
 					FTransform JointTimeZeroGlobalTransform;
-					if (!JointNode->GetCustomTimeZeroGlobalTransform(NodeContainer, JointTimeZeroGlobalTransform))
+					if (!JointNode->GetCustomTimeZeroGlobalTransform(NodeContainer, GlobalOffsetTransform, JointTimeZeroGlobalTransform))
 					{
 						//If there is no time zero global transform we cannot set the bind pose to time zero.
 						//We must skip this joint.
@@ -997,6 +1004,12 @@ UObject* UInterchangeSkeletalMeshFactory::CreateAsset(const FCreateAssetParams&
 		return nullptr;
 	}
 
+	FTransform GlobalOffsetTransform = FTransform::Identity;
+	if (UInterchangeCommonPipelineDataFactoryNode* CommonPipelineDataFactoryNode = UInterchangeCommonPipelineDataFactoryNode::GetUniqueInstance(Arguments.NodeContainer))
+	{
+		CommonPipelineDataFactoryNode->GetCustomGlobalOffsetTransform(GlobalOffsetTransform);
+	}
+
 	//Dirty the DDC Key for any imported Skeletal Mesh
 	SkeletalMesh->InvalidateDeriveDataCacheGUID();
 	USkeleton* SkeletonReference = nullptr;
@@ -1142,9 +1155,9 @@ UObject* UInterchangeSkeletalMeshFactory::CreateAsset(const FCreateAssetParams&
 					MeshReference.MeshNode = Cast<UInterchangeMeshNode>(Arguments.NodeContainer->GetNode(MeshDependencyUid));
 					//Cache the scene node global matrix, we will use this matrix to bake the vertices, add the node geometric mesh offset to this matrix to bake it properly
 					FTransform SceneNodeGlobalTransform;
-					if (!bUseTimeZeroAsBindPose || !MeshReference.SceneNode->GetCustomTimeZeroGlobalTransform(Arguments.NodeContainer, SceneNodeGlobalTransform))
+					if (!bUseTimeZeroAsBindPose || !MeshReference.SceneNode->GetCustomTimeZeroGlobalTransform(Arguments.NodeContainer, GlobalOffsetTransform, SceneNodeGlobalTransform))
 					{
-						ensure(MeshReference.SceneNode->GetCustomGlobalTransform(Arguments.NodeContainer, SceneNodeGlobalTransform));
+						ensure(MeshReference.SceneNode->GetCustomGlobalTransform(Arguments.NodeContainer, GlobalOffsetTransform, SceneNodeGlobalTransform));
 					}
 					FTransform SceneNodeGeometricTransform;
 					if(MeshReference.SceneNode->GetCustomGeometricTransform(SceneNodeGeometricTransform))
@@ -1153,6 +1166,11 @@ UObject* UInterchangeSkeletalMeshFactory::CreateAsset(const FCreateAssetParams&
 					}
 					MeshReference.SceneGlobalTransform = SceneNodeGlobalTransform;
 				}
+				else
+				{
+					MeshReference.SceneGlobalTransform = GlobalOffsetTransform;
+				}
+
 				if (!ensure(MeshReference.MeshNode != nullptr))
 				{
 					UE_LOG(LogInterchangeImport, Warning, TEXT("Invalid LOD mesh reference when importing SkeletalMesh asset %s"), *Arguments.AssetName);
diff --git a/Engine/Plugins/Experimental/Interchange/Runtime/Source/Import/Private/Mesh/InterchangeSkeletonHelper.cpp b/Engine/Plugins/Experimental/Interchange/Runtime/Source/Import/Private/Mesh/InterchangeSkeletonHelper.cpp
--- a/Engine/Plugins/Experimental/Interchange/Runtime/Source/Import/Private/Mesh/InterchangeSkeletonHelper.cpp
+++ b/Engine/Plugins/Experimental/Interchange/Runtime/Source/Import/Private/Mesh/InterchangeSkeletonHelper.cpp
@@ -3,6 +3,7 @@
 
 #if WITH_EDITOR
 #include "Animation/Skeleton.h"
+#include "InterchangeCommonPipelineDataFactoryNode.h"
 #include "InterchangeImportLog.h"
 #include "InterchangeSceneNode.h"
 #include "ReferenceSkeleton.h"
@@ -187,12 +188,27 @@ namespace UE::Interchange::Private
 		Info.Name = JointNode->GetDisplayLabel();
 
 		FTransform LocalTransform;
-		ensure(JointNode->GetCustomLocalTransform(LocalTransform));
-
 		FTransform TimeZeroLocalTransform;
-		const bool bHasTimeZeroTransform = JointNode->GetCustomTimeZeroLocalTransform(TimeZeroLocalTransform);
+		bool bHasTimeZeroTransform = false;
 		FTransform BindPoseLocalTransform;
-		const bool bHasBindPoseTransform = JointNode->GetCustomBindPoseLocalTransform(BindPoseLocalTransform);
+		bool bHasBindPoseTransform = false;
+		if (ParentIndex == INDEX_NONE)
+		{
+			FTransform GlobalOffsetTransform = FTransform::Identity;
+			if (UInterchangeCommonPipelineDataFactoryNode* CommonPipelineDataFactoryNode = UInterchangeCommonPipelineDataFactoryNode::GetUniqueInstance(NodeContainer))
+			{
+				CommonPipelineDataFactoryNode->GetCustomGlobalOffsetTransform(GlobalOffsetTransform);
+			}
+			ensure(JointNode->GetCustomGlobalTransform(NodeContainer, GlobalOffsetTransform, LocalTransform));
+			bHasTimeZeroTransform = JointNode->GetCustomTimeZeroGlobalTransform(NodeContainer, GlobalOffsetTransform, TimeZeroLocalTransform);
+			bHasBindPoseTransform = JointNode->GetCustomBindPoseGlobalTransform(NodeContainer, GlobalOffsetTransform, BindPoseLocalTransform);
+		}
+		else
+		{
+			ensure(JointNode->GetCustomLocalTransform(LocalTransform));
+			bHasTimeZeroTransform = JointNode->GetCustomTimeZeroLocalTransform(TimeZeroLocalTransform);
+			bHasBindPoseTransform = JointNode->GetCustomBindPoseLocalTransform(BindPoseLocalTransform);
+		}
 
 		Info.LocalTransform = bHasBindPoseTransform ? BindPoseLocalTransform : LocalTransform;
 		//If user want to bind the mesh at time zero try to get the time zero transform
diff --git a/Engine/Plugins/Experimental/Interchange/Runtime/Source/Import/Private/Mesh/InterchangeStaticMeshFactory.cpp b/Engine/Plugins/Experimental/Interchange/Runtime/Source/Import/Private/Mesh/InterchangeStaticMeshFactory.cpp
--- a/Engine/Plugins/Experimental/Interchange/Runtime/Source/Import/Private/Mesh/InterchangeStaticMeshFactory.cpp
+++ b/Engine/Plugins/Experimental/Interchange/Runtime/Source/Import/Private/Mesh/InterchangeStaticMeshFactory.cpp
@@ -10,6 +10,7 @@
 #include "Engine/StaticMesh.h"
 #include "Engine/StaticMeshSocket.h"
 #include "InterchangeAssetImportData.h"
+#include "InterchangeCommonPipelineDataFactoryNode.h"
 #include "InterchangeImportCommon.h"
 #include "InterchangeImportLog.h"
 #include "InterchangeMaterialFactoryNode.h"
@@ -433,6 +434,12 @@ TArray<UInterchangeStaticMeshFactory::FMeshPayload> UInterchangeStaticMeshFactor
 		return Payloads;
 	}
 
+	FTransform GlobalOffsetTransform = FTransform::Identity;
+	if (UInterchangeCommonPipelineDataFactoryNode* CommonPipelineDataFactoryNode = UInterchangeCommonPipelineDataFactoryNode::GetUniqueInstance(Arguments.NodeContainer))
+	{
+		CommonPipelineDataFactoryNode->GetCustomGlobalOffsetTransform(GlobalOffsetTransform);
+	}
+
 	for (const FString& MeshUid : MeshUids)
 	{
 		FMeshPayload Payload;
@@ -451,7 +458,7 @@ TArray<UInterchangeStaticMeshFactory::FMeshPayload> UInterchangeStaticMeshFactor
 
 			// Get the transform from the scene node
 			FTransform SceneNodeGlobalTransform;
-			if (SceneNode->GetCustomGlobalTransform(Arguments.NodeContainer, SceneNodeGlobalTransform))
+			if (SceneNode->GetCustomGlobalTransform(Arguments.NodeContainer, GlobalOffsetTransform, SceneNodeGlobalTransform))
 			{
 				Payload.Transform = SceneNodeGlobalTransform;
 			}
@@ -461,6 +468,11 @@ TArray<UInterchangeStaticMeshFactory::FMeshPayload> UInterchangeStaticMeshFactor
 			SceneNode->GetCustomAssetInstanceUid(MeshDependencyUid);
 			MeshNode = Cast<UInterchangeMeshNode>(Arguments.NodeContainer->GetNode(MeshDependencyUid));
 		}
+		else
+		{
+			//If we have a mesh that is not reference by a scene node, we must apply the global offset.
+			Payload.Transform = GlobalOffsetTransform;
+		}
 
 		if (!ensure(MeshNode))
 		{
@@ -1288,6 +1300,12 @@ bool UInterchangeStaticMeshFactory::ImportSockets(const FCreateAssetParams& Argu
 
 	TSet<FName> ImportedSocketNames;
 
+	FTransform GlobalOffsetTransform = FTransform::Identity;
+	if (UInterchangeCommonPipelineDataFactoryNode* CommonPipelineDataFactoryNode = UInterchangeCommonPipelineDataFactoryNode::GetUniqueInstance(Arguments.NodeContainer))
+	{
+		CommonPipelineDataFactoryNode->GetCustomGlobalOffsetTransform(GlobalOffsetTransform);
+	}
+
 	for (const FString& SocketUid : SocketUids)
 	{
 		if (const UInterchangeSceneNode* SceneNode = Cast<UInterchangeSceneNode>(Arguments.NodeContainer->GetNode(SocketUid)))
@@ -1302,7 +1320,7 @@ bool UInterchangeStaticMeshFactory::ImportSockets(const FCreateAssetParams& Argu
 			ImportedSocketNames.Add(SocketName);
 
 			FTransform GlobalTransform;
-			SceneNode->GetCustomGlobalTransform(Arguments.NodeContainer, GlobalTransform);
+			SceneNode->GetCustomGlobalTransform(Arguments.NodeContainer, GlobalOffsetTransform, GlobalTransform);
 
 			UStaticMeshSocket* Socket = StaticMesh->FindSocket(SocketName);
 			if (!Socket)
diff --git a/Engine/Plugins/Experimental/Interchange/Runtime/Source/Nodes/Private/InterchangeSceneNode.cpp b/Engine/Plugins/Experimental/Interchange/Runtime/Source/Nodes/Private/InterchangeSceneNode.cpp
--- a/Engine/Plugins/Experimental/Interchange/Runtime/Source/Nodes/Private/InterchangeSceneNode.cpp
+++ b/Engine/Plugins/Experimental/Interchange/Runtime/Source/Nodes/Private/InterchangeSceneNode.cpp
@@ -216,9 +216,9 @@ bool UInterchangeSceneNode::SetCustomLocalTransform(const UInterchangeBaseNodeCo
 	IMPLEMENT_NODE_ATTRIBUTE_SETTER_NODELEGATE(LocalTransform, FTransform);
 }
 
-bool UInterchangeSceneNode::GetCustomGlobalTransform(const UInterchangeBaseNodeContainer* BaseNodeContainer, FTransform& AttributeValue, bool bForceRecache /*= false*/) const
+bool UInterchangeSceneNode::GetCustomGlobalTransform(const UInterchangeBaseNodeContainer* BaseNodeContainer, const FTransform& GlobalOffsetTransform, FTransform& AttributeValue, bool bForceRecache /*= false*/) const
 {
-	return GetGlobalTransformInternal(Macro_CustomLocalTransformKey, CacheGlobalTransform, BaseNodeContainer, AttributeValue, bForceRecache);
+	return GetGlobalTransformInternal(Macro_CustomLocalTransformKey, CacheGlobalTransform, BaseNodeContainer, GlobalOffsetTransform, AttributeValue, bForceRecache);
 }
 
 bool UInterchangeSceneNode::GetCustomBindPoseLocalTransform(FTransform& AttributeValue) const
@@ -235,9 +235,9 @@ bool UInterchangeSceneNode::SetCustomBindPoseLocalTransform(const UInterchangeBa
 	IMPLEMENT_NODE_ATTRIBUTE_SETTER_NODELEGATE(BindPoseLocalTransform, FTransform);
 }
 
-bool UInterchangeSceneNode::GetCustomBindPoseGlobalTransform(const UInterchangeBaseNodeContainer* BaseNodeContainer, FTransform& AttributeValue, bool bForceRecache /*= false*/) const
+bool UInterchangeSceneNode::GetCustomBindPoseGlobalTransform(const UInterchangeBaseNodeContainer* BaseNodeContainer, const FTransform& GlobalOffsetTransform, FTransform& AttributeValue, bool bForceRecache /*= false*/) const
 {
-	return GetGlobalTransformInternal(Macro_CustomBindPoseLocalTransformKey, CacheBindPoseGlobalTransform, BaseNodeContainer, AttributeValue, bForceRecache);
+	return GetGlobalTransformInternal(Macro_CustomBindPoseLocalTransformKey, CacheBindPoseGlobalTransform, BaseNodeContainer, GlobalOffsetTransform, AttributeValue, bForceRecache);
 }
 
 bool UInterchangeSceneNode::GetCustomTimeZeroLocalTransform(FTransform& AttributeValue) const
@@ -254,9 +254,9 @@ bool UInterchangeSceneNode::SetCustomTimeZeroLocalTransform(const UInterchangeBa
 	IMPLEMENT_NODE_ATTRIBUTE_SETTER_NODELEGATE(TimeZeroLocalTransform, FTransform);
 }
 
-bool UInterchangeSceneNode::GetCustomTimeZeroGlobalTransform(const UInterchangeBaseNodeContainer* BaseNodeContainer, FTransform& AttributeValue, bool bForceRecache /*= false*/) const
+bool UInterchangeSceneNode::GetCustomTimeZeroGlobalTransform(const UInterchangeBaseNodeContainer* BaseNodeContainer, const FTransform& GlobalOffsetTransform, FTransform& AttributeValue, bool bForceRecache /*= false*/) const
 {
-	return GetGlobalTransformInternal(Macro_CustomTimeZeroLocalTransformKey, CacheTimeZeroGlobalTransform, BaseNodeContainer, AttributeValue, bForceRecache);
+	return GetGlobalTransformInternal(Macro_CustomTimeZeroLocalTransformKey, CacheTimeZeroGlobalTransform, BaseNodeContainer, GlobalOffsetTransform, AttributeValue, bForceRecache);
 }
 
 bool UInterchangeSceneNode::GetCustomGeometricTransform(FTransform& AttributeValue) const
@@ -361,7 +361,12 @@ void UInterchangeSceneNode::ResetGlobalTransformCachesOfNodeAndAllChildren(const
 	}
 }
 
-bool UInterchangeSceneNode::GetGlobalTransformInternal(const UE::Interchange::FAttributeKey LocalTransformKey, TOptional<FTransform>& CacheTransform, const UInterchangeBaseNodeContainer* BaseNodeContainer, FTransform& AttributeValue, bool bForceRecache) const
+bool UInterchangeSceneNode::GetGlobalTransformInternal(const UE::Interchange::FAttributeKey LocalTransformKey
+	, TOptional<FTransform>& CacheTransform
+	, const UInterchangeBaseNodeContainer* BaseNodeContainer
+	, const FTransform& GlobalOffsetTransform
+	, FTransform& AttributeValue
+	, bool bForceRecache) const
 {
 	if (!Attributes->ContainAttribute(LocalTransformKey))
 	{
@@ -385,22 +390,31 @@ bool UInterchangeSceneNode::GetGlobalTransformInternal(const UE::Interchange::FA
 				{
 					if (LocalTransformKey == Macro_CustomLocalTransformKey)
 					{
-						ParentSceneNode->GetCustomGlobalTransform(BaseNodeContainer, GlobalParent, bForceRecache);
+						ParentSceneNode->GetCustomGlobalTransform(BaseNodeContainer, GlobalOffsetTransform, GlobalParent, bForceRecache);
 					}
 					else if (LocalTransformKey == Macro_CustomBindPoseLocalTransformKey)
 					{
-						ParentSceneNode->GetCustomBindPoseGlobalTransform(BaseNodeContainer, GlobalParent, bForceRecache);
+						//Its possible the root skeleton id have a parent that is not a joint, in that case we need to fall back on the normal transform
+						if (!ParentSceneNode->GetCustomBindPoseGlobalTransform(BaseNodeContainer, GlobalOffsetTransform, GlobalParent, bForceRecache))
+						{
+							ParentSceneNode->GetCustomGlobalTransform(BaseNodeContainer, GlobalOffsetTransform, GlobalParent, bForceRecache);
+						}
 					}
 					else if (LocalTransformKey == Macro_CustomTimeZeroLocalTransformKey)
 					{
-						ParentSceneNode->GetCustomTimeZeroGlobalTransform(BaseNodeContainer, GlobalParent, bForceRecache);
+						//Its possible the root skeleton id have a parent that is not a joint, in that case we need to fall back on the normal transform
+						if(!ParentSceneNode->GetCustomTimeZeroGlobalTransform(BaseNodeContainer, GlobalOffsetTransform, GlobalParent, bForceRecache))
+						{
+							ParentSceneNode->GetCustomGlobalTransform(BaseNodeContainer, GlobalOffsetTransform, GlobalParent, bForceRecache);
+						}
 					}
 				}
 				CacheTransform = LocalTransform * GlobalParent;
 			}
 			else
 			{
-				CacheTransform = LocalTransform;
+				//Scene Node without parent will need the global offset to be apply
+				CacheTransform = LocalTransform * GlobalOffsetTransform;
 			}
 		}
 		else
diff --git a/Engine/Plugins/Experimental/Interchange/Runtime/Source/Nodes/Public/InterchangeSceneNode.h b/Engine/Plugins/Experimental/Interchange/Runtime/Source/Nodes/Public/InterchangeSceneNode.h
--- a/Engine/Plugins/Experimental/Interchange/Runtime/Source/Nodes/Public/InterchangeSceneNode.h
+++ b/Engine/Plugins/Experimental/Interchange/Runtime/Source/Nodes/Public/InterchangeSceneNode.h
@@ -96,7 +96,7 @@ public:
 
 	/** Return the default scene node global transform. This value is computed with all parent local transform. */
 	UFUNCTION(BlueprintCallable, Category = "Interchange | Node | Scene")
-	bool GetCustomGlobalTransform(const UInterchangeBaseNodeContainer* BaseNodeContainer, FTransform& AttributeValue, bool bForceRecache = false) const;
+	bool GetCustomGlobalTransform(const UInterchangeBaseNodeContainer* BaseNodeContainer, const FTransform& GlobalOffsetTransform, FTransform& AttributeValue, bool bForceRecache = false) const;
 
 	//Bind pose transform is the transform of the joint when the binding with the mesh was done.
 	//This attribute should be set only if we have a joint.
@@ -111,7 +111,7 @@ public:
 
 	/** Return the bind pose scene node global transform. This value is computed with all parent bind pose local transform. */
 	UFUNCTION(BlueprintCallable, Category = "Interchange | Node | Joint")
-	bool GetCustomBindPoseGlobalTransform(const UInterchangeBaseNodeContainer* BaseNodeContainer, FTransform& AttributeValue, bool bForceRecache = false) const;
+	bool GetCustomBindPoseGlobalTransform(const UInterchangeBaseNodeContainer* BaseNodeContainer, const FTransform& GlobalOffsetTransform, FTransform& AttributeValue, bool bForceRecache = false) const;
 
 	//Time zero transform is the transform of the node at time zero.
 	//This is useful when there is no bind pose or when we import rigid mesh.
@@ -126,7 +126,7 @@ public:
 
 	/** Return the time zero scene node global transform. This value is computed with all parent timezero local transform. */
 	UFUNCTION(BlueprintCallable, Category = "Interchange | Node | Joint")
-	bool GetCustomTimeZeroGlobalTransform(const UInterchangeBaseNodeContainer* BaseNodeContainer, FTransform& AttributeValue, bool bForceRecache = false) const;
+	bool GetCustomTimeZeroGlobalTransform(const UInterchangeBaseNodeContainer* BaseNodeContainer, const FTransform& GlobalOffsetTransform, FTransform& AttributeValue, bool bForceRecache = false) const;
 
 	/** Return the geometric offset. Any mesh attach to this scene node will be offset using this transform. */
 	UFUNCTION(BlueprintCallable, Category = "Interchange | Node | Joint")
@@ -192,7 +192,12 @@ public:
 
 private:
 
-	bool GetGlobalTransformInternal(const UE::Interchange::FAttributeKey LocalTransformKey, TOptional<FTransform>& CacheTransform, const UInterchangeBaseNodeContainer* BaseNodeContainer, FTransform& AttributeValue, bool bForceRecache) const;
+	bool GetGlobalTransformInternal(const UE::Interchange::FAttributeKey LocalTransformKey
+		, TOptional<FTransform>& CacheTransform
+		, const UInterchangeBaseNodeContainer* BaseNodeContainer
+		, const FTransform& GlobalOffsetTransform
+		, FTransform& AttributeValue
+		, bool bForceRecache) const;
 
 	//Scene Attribute Keys
 	const UE::Interchange::FAttributeKey Macro_CustomLocalTransformKey = UE::Interchange::FAttributeKey(TEXT("LocalTransform"));
diff --git a/Engine/Plugins/Experimental/Interchange/Runtime/Source/Pipelines/Private/InterchangeGenericAssetsPipeline.cpp b/Engine/Plugins/Experimental/Interchange/Runtime/Source/Pipelines/Private/InterchangeGenericAssetsPipeline.cpp
--- a/Engine/Plugins/Experimental/Interchange/Runtime/Source/Pipelines/Private/InterchangeGenericAssetsPipeline.cpp
+++ b/Engine/Plugins/Experimental/Interchange/Runtime/Source/Pipelines/Private/InterchangeGenericAssetsPipeline.cpp
@@ -5,6 +5,7 @@
 #include "Animation/Skeleton.h"
 #include "CoreMinimal.h"
 #include "InterchangeAnimSequenceFactoryNode.h"
+#include "InterchangeCommonPipelineDataFactoryNode.h"
 #include "InterchangeGenericAnimationPipeline.h"
 #include "InterchangeGenericMaterialPipeline.h"
 #include "InterchangeGenericMeshPipeline.h"
@@ -156,6 +157,20 @@ void UInterchangeGenericAssetsPipeline::ExecutePreImportPipeline(UInterchangeBas
 
 
 	BaseNodeContainer = InBaseNodeContainer;
+
+	//Setup the Global import offset
+
+	{
+		FTransform ImportOffsetTransform;
+		ImportOffsetTransform.SetTranslation(ImportOffsetTranslation);
+		ImportOffsetTransform.SetRotation(FQuat(ImportOffsetRotation));
+		ImportOffsetTransform.SetScale3D(FVector(ImportOffsetUniformScale));
+
+		UInterchangeCommonPipelineDataFactoryNode* CommonPipelineDataFactoryNode = UInterchangeCommonPipelineDataFactoryNode::FindOrCreateUniqueInstance(BaseNodeContainer);
+		CommonPipelineDataFactoryNode->SetCustomGlobalOffsetTransform(BaseNodeContainer, ImportOffsetTransform);
+		
+	}
+
 	SourceDatas.Empty(InSourceDatas.Num());
 	for (const UInterchangeSourceData* SourceData : InSourceDatas)
 	{
diff --git a/Engine/Plugins/Experimental/Interchange/Runtime/Source/Pipelines/Private/InterchangeGenericScenesPipeline.cpp b/Engine/Plugins/Experimental/Interchange/Runtime/Source/Pipelines/Private/InterchangeGenericScenesPipeline.cpp
--- a/Engine/Plugins/Experimental/Interchange/Runtime/Source/Pipelines/Private/InterchangeGenericScenesPipeline.cpp
+++ b/Engine/Plugins/Experimental/Interchange/Runtime/Source/Pipelines/Private/InterchangeGenericScenesPipeline.cpp
@@ -5,6 +5,7 @@
 #include "InterchangeActorFactoryNode.h"
 #include "InterchangeCameraNode.h"
 #include "InterchangeCineCameraFactoryNode.h"
+#include "InterchangeCommonPipelineDataFactoryNode.h"
 #include "InterchangeLightNode.h"
 #include "InterchangeMeshNode.h"
 #include "InterchangePipelineLog.h"
@@ -26,6 +27,12 @@ void UInterchangeGenericLevelPipeline::ExecutePreImportPipeline(UInterchangeBase
 		return;
 	}
 
+	FTransform GlobalOffsetTransform = FTransform::Identity;
+	if (UInterchangeCommonPipelineDataFactoryNode* CommonPipelineDataFactoryNode = UInterchangeCommonPipelineDataFactoryNode::GetUniqueInstance(InBaseNodeContainer))
+	{
+		CommonPipelineDataFactoryNode->GetCustomGlobalOffsetTransform(GlobalOffsetTransform);
+	}
+
 	TArray<UInterchangeSceneNode*> SceneNodes;
 
 	//Find all translated node we need for this pipeline
@@ -58,12 +65,12 @@ void UInterchangeGenericLevelPipeline::ExecutePreImportPipeline(UInterchangeBase
 					continue;
 				}
 			}
-			ExecuteSceneNodePreImport(InBaseNodeContainer, SceneNode, InBaseNodeContainer);
+			ExecuteSceneNodePreImport(InBaseNodeContainer, GlobalOffsetTransform, SceneNode, InBaseNodeContainer);
 		}
 	}
 }
 
-void UInterchangeGenericLevelPipeline::ExecuteSceneNodePreImport(UInterchangeBaseNodeContainer* InBaseNodeContainer, const UInterchangeSceneNode* SceneNode, UInterchangeBaseNodeContainer* FactoryNodeContainer)
+void UInterchangeGenericLevelPipeline::ExecuteSceneNodePreImport(UInterchangeBaseNodeContainer* InBaseNodeContainer, const FTransform& GlobalOffsetTransform, const UInterchangeSceneNode* SceneNode, UInterchangeBaseNodeContainer* FactoryNodeContainer)
 {
 	if (!SceneNode)
 	{
@@ -97,8 +104,9 @@ void UInterchangeGenericLevelPipeline::ExecuteSceneNodePreImport(UInterchangeBas
 	ActorFactoryNode->AddTargetNodeUid(SceneNode->GetUniqueID());
 	SceneNode->AddTargetNodeUid(ActorFactoryNode->GetUniqueID());
 
+	//TODO move this code to the factory, a stack over pipeline can change the global offset transform which will affect this value.
 	FTransform GlobalTransform;
-	if (SceneNode->GetCustomGlobalTransform(InBaseNodeContainer, GlobalTransform))
+	if (SceneNode->GetCustomGlobalTransform(InBaseNodeContainer, GlobalOffsetTransform, GlobalTransform))
 	{
 		ActorFactoryNode->SetCustomGlobalTransform(GlobalTransform);
 	}
diff --git a/Engine/Plugins/Experimental/Interchange/Runtime/Source/Pipelines/Public/InterchangeGenericScenesPipeline.h b/Engine/Plugins/Experimental/Interchange/Runtime/Source/Pipelines/Public/InterchangeGenericScenesPipeline.h
--- a/Engine/Plugins/Experimental/Interchange/Runtime/Source/Pipelines/Public/InterchangeGenericScenesPipeline.h
+++ b/Engine/Plugins/Experimental/Interchange/Runtime/Source/Pipelines/Public/InterchangeGenericScenesPipeline.h
@@ -32,7 +32,7 @@ protected:
 	/**
 	 * PreImport step called for each translated SceneNode.
 	 */
-	virtual void ExecuteSceneNodePreImport(UInterchangeBaseNodeContainer* InBaseNodeContainer, const UInterchangeSceneNode* SceneNode, UInterchangeBaseNodeContainer* FactoryNodeContainer);
+	virtual void ExecuteSceneNodePreImport(UInterchangeBaseNodeContainer* InBaseNodeContainer, const FTransform& GlobalOffsetTransform, const UInterchangeSceneNode* SceneNode, UInterchangeBaseNodeContainer* FactoryNodeContainer);
 
 	/**
 	 * Return a new Actor Factory Node to be used for the given SceneNode.
