Change the default value for the import uniform scale to 1.0

diff --git a/Engine/Plugins/Experimental/Interchange/Runtime/Source/Pipelines/Public/InterchangeGenericAssetsPipeline.h b/Engine/Plugins/Experimental/Interchange/Runtime/Source/Pipelines/Public/InterchangeGenericAssetsPipeline.h
--- a/Engine/Plugins/Experimental/Interchange/Runtime/Source/Pipelines/Public/InterchangeGenericAssetsPipeline.h
+++ b/Engine/Plugins/Experimental/Interchange/Runtime/Source/Pipelines/Public/InterchangeGenericAssetsPipeline.h
@@ -50,7 +50,7 @@ public:
 
 	/** Uniform scale offset applied to meshes and animations. */
 	UPROPERTY(EditAnywhere, BlueprintReadWrite, Category = "Common", meta = (DisplayName = "Offset Uniform Scale"))
-	float ImportOffsetUniformScale;
+	float ImportOffsetUniformScale = 1.0f;
 
 	//////	COMMON_MESHES_CATEGORY Properties //////
 	UPROPERTY(VisibleAnywhere, Instanced, Category = "Common Meshes")
