Make sure sub generic asset pipelines are not showing in the pipeline stack dropdown list.

diff --git a/Engine/Plugins/Experimental/Interchange/Runtime/Source/Pipelines/Public/InterchangeGenericAnimationPipeline.h b/Engine/Plugins/Experimental/Interchange/Runtime/Source/Pipelines/Public/InterchangeGenericAnimationPipeline.h
--- a/Engine/Plugins/Experimental/Interchange/Runtime/Source/Pipelines/Public/InterchangeGenericAnimationPipeline.h
+++ b/Engine/Plugins/Experimental/Interchange/Runtime/Source/Pipelines/Public/InterchangeGenericAnimationPipeline.h
@@ -28,7 +28,7 @@ enum class EInterchangeAnimationRange : uint8
 };
 
 
-UCLASS(BlueprintType, Experimental)
+UCLASS(BlueprintType, Experimental, hidedropdown)
 class INTERCHANGEPIPELINES_API UInterchangeGenericAnimationPipeline : public UInterchangePipelineBase
 {
 	GENERATED_BODY()
diff --git a/Engine/Plugins/Experimental/Interchange/Runtime/Source/Pipelines/Public/InterchangeGenericAssetsPipelineSharedSettings.h b/Engine/Plugins/Experimental/Interchange/Runtime/Source/Pipelines/Public/InterchangeGenericAssetsPipelineSharedSettings.h
--- a/Engine/Plugins/Experimental/Interchange/Runtime/Source/Pipelines/Public/InterchangeGenericAssetsPipelineSharedSettings.h
+++ b/Engine/Plugins/Experimental/Interchange/Runtime/Source/Pipelines/Public/InterchangeGenericAssetsPipelineSharedSettings.h
@@ -44,7 +44,7 @@ enum class EInterchangeVertexColorImportOption : uint8
 };
 
 
-UCLASS(BlueprintType, Experimental)
+UCLASS(BlueprintType, Experimental, hidedropdown)
 class INTERCHANGEPIPELINES_API UInterchangeGenericCommonMeshesProperties : public UInterchangePipelineBase
 {
 	GENERATED_BODY()
@@ -72,7 +72,7 @@ public:
 	FColor VertexOverrideColor;
 };
 
-UCLASS(BlueprintType, Experimental)
+UCLASS(BlueprintType, Experimental, hidedropdown)
 class INTERCHANGEPIPELINES_API UInterchangeGenericCommonSkeletalMeshesAndAnimationsProperties : public UInterchangePipelineBase
 {
 	GENERATED_BODY()
diff --git a/Engine/Plugins/Experimental/Interchange/Runtime/Source/Pipelines/Public/InterchangeGenericMaterialPipeline.h b/Engine/Plugins/Experimental/Interchange/Runtime/Source/Pipelines/Public/InterchangeGenericMaterialPipeline.h
--- a/Engine/Plugins/Experimental/Interchange/Runtime/Source/Pipelines/Public/InterchangeGenericMaterialPipeline.h
+++ b/Engine/Plugins/Experimental/Interchange/Runtime/Source/Pipelines/Public/InterchangeGenericMaterialPipeline.h
@@ -27,7 +27,7 @@ enum class EInterchangeMaterialImportOption : uint8
 	ImportAsMaterialInstances,
 };
 
-UCLASS(BlueprintType, Experimental)
+UCLASS(BlueprintType, Experimental, hidedropdown)
 class INTERCHANGEPIPELINES_API UInterchangeGenericMaterialPipeline : public UInterchangePipelineBase
 {
 	GENERATED_BODY()
diff --git a/Engine/Plugins/Experimental/Interchange/Runtime/Source/Pipelines/Public/InterchangeGenericMeshPipeline.h b/Engine/Plugins/Experimental/Interchange/Runtime/Source/Pipelines/Public/InterchangeGenericMeshPipeline.h
--- a/Engine/Plugins/Experimental/Interchange/Runtime/Source/Pipelines/Public/InterchangeGenericMeshPipeline.h
+++ b/Engine/Plugins/Experimental/Interchange/Runtime/Source/Pipelines/Public/InterchangeGenericMeshPipeline.h
@@ -23,8 +23,8 @@ class UInterchangeStaticMeshFactoryNode;
 class UInterchangeStaticMeshLodDataNode;
 class UPhysicsAsset;
 
-
-UCLASS(BlueprintType, Experimental)
+/* Hide drop down will make sure the class is not showing in the class picker */
+UCLASS(BlueprintType, Experimental, hidedropdown)
 class INTERCHANGEPIPELINES_API UInterchangeGenericMeshPipeline : public UInterchangePipelineBase
 {
 	GENERATED_BODY()
diff --git a/Engine/Plugins/Experimental/Interchange/Runtime/Source/Pipelines/Public/InterchangeGenericTexturePipeline.h b/Engine/Plugins/Experimental/Interchange/Runtime/Source/Pipelines/Public/InterchangeGenericTexturePipeline.h
--- a/Engine/Plugins/Experimental/Interchange/Runtime/Source/Pipelines/Public/InterchangeGenericTexturePipeline.h
+++ b/Engine/Plugins/Experimental/Interchange/Runtime/Source/Pipelines/Public/InterchangeGenericTexturePipeline.h
@@ -13,7 +13,7 @@
 class UInterchangeTextureFactoryNode;
 class UInterchangeTextureNode;
 
-UCLASS(BlueprintType, Experimental)
+UCLASS(BlueprintType, Experimental, hidedropdown)
 class INTERCHANGEPIPELINES_API UInterchangeGenericTexturePipeline : public UInterchangePipelineBase
 {
 	GENERATED_BODY()
