Fix crash when using interchange to import an already existing skeletalmesh. The interchange skeletalmesh factory now find the situation and do a re-import.

#rnx

diff --git a/Engine/Plugins/Experimental/Interchange/Runtime/Source/Import/Private/Mesh/InterchangeSkeletalMeshFactory.cpp b/Engine/Plugins/Experimental/Interchange/Runtime/Source/Import/Private/Mesh/InterchangeSkeletalMeshFactory.cpp
--- a/Engine/Plugins/Experimental/Interchange/Runtime/Source/Import/Private/Mesh/InterchangeSkeletalMeshFactory.cpp
+++ b/Engine/Plugins/Experimental/Interchange/Runtime/Source/Import/Private/Mesh/InterchangeSkeletalMeshFactory.cpp
@@ -988,12 +988,9 @@ UObject* UInterchangeSkeletalMeshFactory::CreateAsset(const FCreateAssetParams&
 	}
 
 	USkeletalMesh* SkeletalMesh = Cast<USkeletalMesh>(SkeletalMeshObject);
-
-	const bool bIsReImport = (Arguments.ReimportObject != nullptr);
-
 	if (!ensure(SkeletalMesh))
 	{
-		if (!bIsReImport)
+		if (Arguments.ReimportObject == nullptr)
 		{
 			UE_LOG(LogInterchangeImport, Error, TEXT("Could not create skeletalMesh asset %s"), *Arguments.AssetName);
 		}
@@ -1004,6 +1001,9 @@ UObject* UInterchangeSkeletalMeshFactory::CreateAsset(const FCreateAssetParams&
 		return nullptr;
 	}
 
+	//This is consider has a re-import if we have a reimport object or if the object exist and have some valid LOD
+	const bool bIsReImport = (Arguments.ReimportObject != nullptr) || (SkeletalMesh->GetLODNum() > 0);
+
 	FTransform GlobalOffsetTransform = FTransform::Identity;
 	if (UInterchangeCommonPipelineDataFactoryNode* CommonPipelineDataFactoryNode = UInterchangeCommonPipelineDataFactoryNode::GetUniqueInstance(Arguments.NodeContainer))
 	{
