Interchange: First pass on an API to create nodes and standardize unique ids.
- Added a static Create function on UInterchangeTexture2DNode, UInterchangeShaderGraphNode and UInterchangeShaderNode.
- Added a static MakeNodeUid function on UInterchangeTextureNode, UInterchangeShaderGraphNode and UInterchangeShaderNode.
- Updated the FBX, GLTF and OBJ translators to leverage the new APIs.

diff --git a/Engine/Plugins/Experimental/Interchange/Runtime/Source/Import/Private/Gltf/InterchangeGltfTranslator.cpp b/Engine/Plugins/Experimental/Interchange/Runtime/Source/Import/Private/Gltf/InterchangeGltfTranslator.cpp
--- a/Engine/Plugins/Experimental/Interchange/Runtime/Source/Import/Private/Gltf/InterchangeGltfTranslator.cpp
+++ b/Engine/Plugins/Experimental/Interchange/Runtime/Source/Import/Private/Gltf/InterchangeGltfTranslator.cpp
@@ -21,7 +21,7 @@
 
 namespace UE::Interchange::Gltf::Private
 {
-	FString GenerateUniqueIdForGltfNode( const FString& NodeName, int32 NodeIndex )
+	FString GenerateNameForGltfNode( const FString& NodeName, int32 NodeIndex )
 	{
 		if ( NodeIndex == 0 )
 		{
@@ -37,11 +37,11 @@ namespace UE::Interchange::Gltf::Private
 	{
 		if ( !Texture.Name.IsEmpty() )
 		{
-			return GenerateUniqueIdForGltfNode( Texture.Name, NodeIndex );
+			return GenerateNameForGltfNode( Texture.Name, NodeIndex );
 		}
 		else
 		{
-			return GenerateUniqueIdForGltfNode( FPaths::GetBaseFilename( Texture.Source.URI ), NodeIndex );
+			return GenerateNameForGltfNode( FPaths::GetBaseFilename( Texture.Source.URI ), NodeIndex );
 		}
 	}
 
@@ -66,7 +66,7 @@ void UInterchangeGltfTranslator::HandleGltfNode( UInterchangeBaseNodeContainer&
 {
 	using namespace UE::Interchange::Gltf::Private;
 
-	const FString NodeUid = ParentNodeUid + TEXT("\\") + GenerateUniqueIdForGltfNode( GltfNode.Name, NodeIndex );
+	const FString NodeUid = ParentNodeUid + TEXT("\\") + GenerateNameForGltfNode( GltfNode.Name, NodeIndex );
 
 	const UInterchangeSceneNode* ParentSceneNode = Cast< UInterchangeSceneNode >( NodeContainer.GetNode( ParentNodeUid ) );
 
@@ -85,7 +85,7 @@ void UInterchangeGltfTranslator::HandleGltfNode( UInterchangeBaseNodeContainer&
 		{
 			if ( GltfAsset.Meshes.IsValidIndex( GltfNode.MeshIndex ) )
 			{
-				const FString MeshNodeUid = TEXT("\\Mesh\\") + GenerateUniqueIdForGltfNode( GltfAsset.Meshes[ GltfNode.MeshIndex ].Name, GltfNode.MeshIndex );
+				const FString MeshNodeUid = TEXT("\\Mesh\\") + GenerateNameForGltfNode( GltfAsset.Meshes[ GltfNode.MeshIndex ].Name, GltfNode.MeshIndex );
 				InterchangeSceneNode->SetCustomAssetInstanceUid( MeshNodeUid );
 			}
 			break;
@@ -97,7 +97,7 @@ void UInterchangeGltfTranslator::HandleGltfNode( UInterchangeBaseNodeContainer&
 
 			if ( GltfAsset.Cameras.IsValidIndex( GltfNode.CameraIndex ) )
 			{
-				const FString CameraNodeUid = TEXT("\\Camera\\") + GenerateUniqueIdForGltfNode( GltfAsset.Cameras[ GltfNode.CameraIndex ].Name, GltfNode.CameraIndex );
+				const FString CameraNodeUid = TEXT("\\Camera\\") + GenerateNameForGltfNode( GltfAsset.Cameras[ GltfNode.CameraIndex ].Name, GltfNode.CameraIndex );
 				InterchangeSceneNode->SetCustomAssetInstanceUid( CameraNodeUid );
 			}
 			break;
@@ -109,7 +109,7 @@ void UInterchangeGltfTranslator::HandleGltfNode( UInterchangeBaseNodeContainer&
 
 			if ( GltfAsset.Lights.IsValidIndex( GltfNode.LightIndex ) )
 			{
-				const FString LightNodeUid = TEXT("\\Light\\") + GenerateUniqueIdForGltfNode( GltfAsset.Lights[ GltfNode.LightIndex ].Name, GltfNode.LightIndex );
+				const FString LightNodeUid = TEXT("\\Light\\") + GenerateNameForGltfNode( GltfAsset.Lights[ GltfNode.LightIndex ].Name, GltfNode.LightIndex );
 				InterchangeSceneNode->SetCustomAssetInstanceUid( LightNodeUid );
 			}
 		}
@@ -151,11 +151,7 @@ void UInterchangeGltfTranslator::HandleGltfMaterialParameter( UInterchangeBaseNo
 	if (bInverse)
 	{
 		const FString OneMinusNodeName = MapName + TEXT("OneMinus");
-		const FString OneMinusNodeUid = ShaderNode.GetUniqueID() + TEXT("_") + OneMinusNodeName;
-		UInterchangeShaderNode* OneMinusNode = NewObject< UInterchangeShaderNode >( &NodeContainer );
-		OneMinusNode->InitializeNode( OneMinusNodeUid, OneMinusNodeName, EInterchangeNodeContainerType::TranslatedAsset );
-		NodeContainer.AddNode( OneMinusNode );
-		NodeContainer.SetNodeParentUid( OneMinusNodeUid, ShaderNode.GetUniqueID() );
+		UInterchangeShaderNode* OneMinusNode = UInterchangeShaderNode::Create( &NodeContainer, OneMinusNodeName, ShaderNode.GetUniqueID() );
 
 		OneMinusNode->SetCustomShaderType(Standard::Nodes::OneMinus::Name.ToString());
 
@@ -178,18 +174,12 @@ void UInterchangeGltfTranslator::HandleGltfMaterialParameter( UInterchangeBaseNo
 
 	if ( bTextureHasImportance && GltfAsset.Textures.IsValidIndex( TextureMap.TextureIndex ) )
 	{
-		const FString NodeName = MapName;
-		const FString NodeUid = ShaderNode.GetUniqueID() + TEXT("_") + NodeName;
-
-		UInterchangeShaderNode* ColorNode = NewObject< UInterchangeShaderNode >( &NodeContainer );
-		ColorNode->InitializeNode( NodeUid, NodeName, EInterchangeNodeContainerType::TranslatedAsset );
-		NodeContainer.AddNode( ColorNode );
-		NodeContainer.SetNodeParentUid( NodeUid, ShaderNode.GetUniqueID() );
-
+		const FString ColorNodeName = MapName;
+		UInterchangeShaderNode* ColorNode = UInterchangeShaderNode::Create( &NodeContainer, ColorNodeName, ShaderNode.GetUniqueID() );
 		ColorNode->SetCustomShaderType( Standard::Nodes::TextureSample::Name.ToString() );
 
 		const FString TextureName = GenerateTextureName( GltfAsset.Textures[ TextureMap.TextureIndex ], TextureMap.TextureIndex );
-		const FString TextureUid = TEXT("\\Texture\\") + TextureName;
+		const FString TextureUid = UInterchangeTextureNode::MakeNodeUid( TextureName );
 
 		ColorNode->AddStringAttribute( UInterchangeShaderPortsAPI::MakeInputValueKey( Standard::Nodes::TextureSample::Inputs::Texture.ToString() ), TextureUid );
 
@@ -211,30 +201,22 @@ void UInterchangeGltfTranslator::HandleGltfMaterialParameter( UInterchangeBaseNo
 
 		if ( bNeedsFactorNode )
 		{
-			const FString FactorNodeUid = NodeUid + TEXT("_Factor");
-			UInterchangeShaderNode* FactorNode = NewObject< UInterchangeShaderNode >( &NodeContainer );
-			FactorNode->InitializeNode( FactorNodeUid, NodeName + TEXT("_Factor"), EInterchangeNodeContainerType::TranslatedAsset );
-			NodeContainer.AddNode( FactorNode );
-			NodeContainer.SetNodeParentUid( FactorNodeUid, ShaderNode.GetUniqueID() );
+			UInterchangeShaderNode* FactorNode = UInterchangeShaderNode::Create( &NodeContainer, ColorNodeName + TEXT("_Factor"), ShaderNode.GetUniqueID() );
 
 			if ( bIsNormal )
 			{
 				FactorNode->SetCustomShaderType( Standard::Nodes::FlattenNormal::Name.ToString() );
 
-				const FString FactorOneMinusNodeUid = FactorNodeUid + TEXT("_OneMinus");
-				UInterchangeShaderNode* FactorOneMinusNode = NewObject< UInterchangeShaderNode >( &NodeContainer );
-				FactorOneMinusNode->InitializeNode( FactorOneMinusNodeUid, NodeName + TEXT("_Factor_OneMinus"), EInterchangeNodeContainerType::TranslatedAsset );
+				UInterchangeShaderNode* FactorOneMinusNode = UInterchangeShaderNode::Create( &NodeContainer, ColorNodeName + TEXT("_Factor_OneMinus"), ShaderNode.GetUniqueID() );
 				FactorOneMinusNode->SetCustomShaderType(Standard::Nodes::OneMinus::Name.ToString());
-				NodeContainer.AddNode( FactorOneMinusNode );
-				NodeContainer.SetNodeParentUid( FactorOneMinusNodeUid, ShaderNode.GetUniqueID() );
 
 				if ( MapFactor.IsType< float >() )
 				{
 					FactorOneMinusNode->AddFloatAttribute( UInterchangeShaderPortsAPI::MakeInputValueKey( Standard::Nodes::OneMinus::Inputs::Input.ToString() ), MapFactor.Get< float >() );
 				}
 
-				UInterchangeShaderPortsAPI::ConnectOuputToInput( FactorNode, Standard::Nodes::FlattenNormal::Inputs::Normal.ToString(), NodeUid, OutputChannel );
-				UInterchangeShaderPortsAPI::ConnectDefaultOuputToInput( FactorNode, Standard::Nodes::FlattenNormal::Inputs::Flatness.ToString(), FactorOneMinusNodeUid );
+				UInterchangeShaderPortsAPI::ConnectOuputToInput( FactorNode, Standard::Nodes::FlattenNormal::Inputs::Normal.ToString(), ColorNode->GetUniqueID(), OutputChannel );
+				UInterchangeShaderPortsAPI::ConnectDefaultOuputToInput( FactorNode, Standard::Nodes::FlattenNormal::Inputs::Flatness.ToString(), FactorOneMinusNode->GetUniqueID() );
 			}
 			else
 			{
@@ -249,14 +231,14 @@ void UInterchangeGltfTranslator::HandleGltfMaterialParameter( UInterchangeBaseNo
 					FactorNode->AddLinearColorAttribute( UInterchangeShaderPortsAPI::MakeInputValueKey( Standard::Nodes::Multiply::Inputs::B.ToString() ), MapFactor.Get< FLinearColor >() );
 				}
 
-				UInterchangeShaderPortsAPI::ConnectOuputToInput( FactorNode, Standard::Nodes::Multiply::Inputs::A.ToString(), NodeUid, OutputChannel );
+				UInterchangeShaderPortsAPI::ConnectOuputToInput( FactorNode, Standard::Nodes::Multiply::Inputs::A.ToString(), ColorNode->GetUniqueID(), OutputChannel );
 			}
 
-			UInterchangeShaderPortsAPI::ConnectDefaultOuputToInput( NodeToConnectTo, InputToConnectTo, FactorNodeUid );
+			UInterchangeShaderPortsAPI::ConnectDefaultOuputToInput( NodeToConnectTo, InputToConnectTo, FactorNode->GetUniqueID() );
 		}
 		else
 		{
-			UInterchangeShaderPortsAPI::ConnectOuputToInput( NodeToConnectTo, InputToConnectTo, NodeUid, OutputChannel );
+			UInterchangeShaderPortsAPI::ConnectOuputToInput( NodeToConnectTo, InputToConnectTo, ColorNode->GetUniqueID(), OutputChannel );
 		}
 	}
 	else
@@ -518,11 +500,7 @@ void UInterchangeGltfTranslator::HandleGltfTransmission( UInterchangeBaseNodeCon
 		 * So we're setting the inverse of the red channel as the opacity.
 		 */
 		const FString OneMinusNodeName =  TEXT("OpacityOneMinus");
-		const FString OneMinusNodeUid = ShaderGraphNode.GetUniqueID() + TEXT("_") + OneMinusNodeName;
-		UInterchangeShaderNode* OneMinusNode = NewObject< UInterchangeShaderNode >( &NodeContainer );
-		OneMinusNode->InitializeNode( OneMinusNodeUid, OneMinusNodeName, EInterchangeNodeContainerType::TranslatedAsset );
-		NodeContainer.AddNode( OneMinusNode );
-		NodeContainer.SetNodeParentUid( OneMinusNodeUid, ShaderGraphNode.GetUniqueID() );
+		UInterchangeShaderNode* OneMinusNode = UInterchangeShaderNode::Create( &NodeContainer, OneMinusNodeName, ShaderGraphNode.GetUniqueID() );
 
 		OneMinusNode->SetCustomShaderType(Standard::Nodes::OneMinus::Name.ToString());
 
@@ -543,14 +521,9 @@ void UInterchangeGltfTranslator::HandleGltfTransmission( UInterchangeBaseNodeCon
 			if ( UInterchangeShaderPortsAPI::GetInputConnection( &ShaderGraphNode, PBR::Parameters::Metallic.ToString(), MetallicNodeUid, MetallicNodeOutput ) )
 			{
 				const FString MetallicLerpNodeName =  TEXT("OpacityMetallicLerp");
-				const FString MetallicLerpNodeUid = ShaderGraphNode.GetUniqueID() + TEXT("_") + MetallicLerpNodeName;
-
-				UInterchangeShaderNode* LerpMetallicNode = NewObject< UInterchangeShaderNode >( &NodeContainer );
-				LerpMetallicNode->InitializeNode( MetallicLerpNodeUid, MetallicLerpNodeName, EInterchangeNodeContainerType::TranslatedAsset );
+				UInterchangeShaderNode* LerpMetallicNode = UInterchangeShaderNode::Create( &NodeContainer, MetallicLerpNodeName, ShaderGraphNode.GetUniqueID() );
 				LerpMetallicNode->SetCustomShaderType( Standard::Nodes::Lerp::Name.ToString() );
 
-				NodeContainer.AddNode( LerpMetallicNode );
-				NodeContainer.SetNodeParentUid( MetallicLerpNodeUid, ShaderGraphNode.GetUniqueID() );
 
 				LerpMetallicNode->AddFloatAttribute( UInterchangeShaderPortsAPI::MakeInputValueKey( Standard::Nodes::Lerp::Inputs::B.ToString() ), 1.f );
 				UInterchangeShaderPortsAPI::ConnectDefaultOuputToInput( LerpMetallicNode, Standard::Nodes::Lerp::Inputs::A.ToString(), CurrentNode->GetUniqueID() );
@@ -566,15 +539,9 @@ void UInterchangeGltfTranslator::HandleGltfTransmission( UInterchangeBaseNodeCon
 			if ( UInterchangeShaderPortsAPI::GetInputConnection( &ShaderGraphNode, PBR::Parameters::Opacity.ToString(), OpacityNodeUid, OpacityNodeOutput ) )
 			{
 				const FString OpacityLerpNodeName =  TEXT("OpacityLerp");
-				const FString OpaciotyLerpNodeUid = ShaderGraphNode.GetUniqueID() + TEXT("_") + OpacityLerpNodeName;
-
-				UInterchangeShaderNode* OpacityLerpNode = NewObject< UInterchangeShaderNode >( &NodeContainer );
-				OpacityLerpNode->InitializeNode( OpaciotyLerpNodeUid, OpacityLerpNodeName, EInterchangeNodeContainerType::TranslatedAsset );
+				UInterchangeShaderNode* OpacityLerpNode = UInterchangeShaderNode::Create( &NodeContainer, OpacityLerpNodeName, ShaderGraphNode.GetUniqueID() );
 				OpacityLerpNode->SetCustomShaderType( Standard::Nodes::Lerp::Name.ToString() );
 
-				NodeContainer.AddNode( OpacityLerpNode );
-				NodeContainer.SetNodeParentUid( OpaciotyLerpNodeUid, ShaderGraphNode.GetUniqueID() );
-
 				OpacityLerpNode->AddFloatAttribute( UInterchangeShaderPortsAPI::MakeInputValueKey( Standard::Nodes::Lerp::Inputs::A.ToString() ), 0.f );
 				UInterchangeShaderPortsAPI::ConnectDefaultOuputToInput( OpacityLerpNode, Standard::Nodes::Lerp::Inputs::B.ToString(), CurrentNode->GetUniqueID() );
 				UInterchangeShaderPortsAPI::ConnectOuputToInput( OpacityLerpNode, Standard::Nodes::Lerp::Inputs::Factor.ToString(), OpacityNodeUid, OpacityNodeOutput );
@@ -618,15 +585,9 @@ void UInterchangeGltfTranslator::HandleGltfTransmission( UInterchangeBaseNodeCon
 			if ( !OpacityNodeUid.IsEmpty() )
 			{
 				const FString OpacityLerpNodeName =  TEXT("OpacityTransmissionLerp");
-				const FString OpaciotyLerpNodeUid = ShaderGraphNode.GetUniqueID() + TEXT("_") + OpacityLerpNodeName;
-
-				UInterchangeShaderNode* OpacityLerpNode = NewObject< UInterchangeShaderNode >( &NodeContainer );
-				OpacityLerpNode->InitializeNode( OpaciotyLerpNodeUid, OpacityLerpNodeName, EInterchangeNodeContainerType::TranslatedAsset );
+				UInterchangeShaderNode* OpacityLerpNode = UInterchangeShaderNode::Create( &NodeContainer, OpacityLerpNodeName, ShaderGraphNode.GetUniqueID() );
 				OpacityLerpNode->SetCustomShaderType( Standard::Nodes::Lerp::Name.ToString() );
 
-				NodeContainer.AddNode( OpacityLerpNode );
-				NodeContainer.SetNodeParentUid( OpaciotyLerpNodeUid, ShaderGraphNode.GetUniqueID() );
-
 				OpacityLerpNode->AddLinearColorAttribute( UInterchangeShaderPortsAPI::MakeInputValueKey( Standard::Nodes::Lerp::Inputs::A.ToString() ), FLinearColor::White );
 				UInterchangeShaderPortsAPI::ConnectOuputToInput( OpacityLerpNode, Standard::Nodes::Lerp::Inputs::Factor.ToString(), OpacityNodeUid, OpacityNodeOutput );
 
@@ -663,12 +624,7 @@ void UInterchangeGltfTranslator::HandleGltfTextureTransform( UInterchangeBaseNod
 
 	// Texture Coordinates
 	{
-		const FString TexCoordNodeUid = ShaderNode.GetUniqueID() + TEXT("_TexCoord");
-		UInterchangeShaderNode* TexCoordNode = NewObject< UInterchangeShaderNode >( &NodeContainer );
-		TexCoordNode->InitializeNode( TexCoordNodeUid, ShaderNode.GetDisplayLabel() + TEXT("_TexCoord"), EInterchangeNodeContainerType::TranslatedAsset );
-		NodeContainer.AddNode( TexCoordNode );
-		NodeContainer.SetNodeParentUid( TexCoordNodeUid, ShaderNode.GetUniqueID() );
-
+		UInterchangeShaderNode* TexCoordNode = UInterchangeShaderNode::Create( &NodeContainer, TEXT("TexCoord"), ShaderNode.GetUniqueID() );
 		TexCoordNode->SetCustomShaderType( Standard::Nodes::TextureCoordinate::Name.ToString() );
 
 		TexCoordNode->AddInt32Attribute( UInterchangeShaderPortsAPI::MakeInputValueKey( Standard::Nodes::TextureCoordinate::Inputs::Index.ToString() ), TexCoordIndex );
@@ -771,16 +727,11 @@ bool UInterchangeGltfTranslator::Translate( UInterchangeBaseNodeContainer& NodeC
 		{
 			const FString TextureName = GenerateTextureName( GltfTexture, TextureIndex );
 
-			UInterchangeTexture2DNode* TextureNode = NewObject< UInterchangeTexture2DNode >( &NodeContainer );
-			FString TextureNodeUid = TEXT("\\Texture\\") + TextureName;
-
-			TextureNode->InitializeNode( TextureNodeUid, TextureName, EInterchangeNodeContainerType::TranslatedAsset );
+			UInterchangeTexture2DNode* TextureNode = UInterchangeTexture2DNode::Create( &NodeContainer, TextureName );
 			TextureNode->SetPayLoadKey( LexToString( TextureIndex++ ) );
 
 			TextureNode->SetCustomWrapU( UE::Interchange::Gltf::Private::ConvertWrap( GltfTexture.Sampler.WrapS ) );
 			TextureNode->SetCustomWrapV( UE::Interchange::Gltf::Private::ConvertWrap( GltfTexture.Sampler.WrapT ) );
-
-			NodeContainer.AddNode( TextureNode );
 		}
 	}
 
@@ -789,10 +740,7 @@ bool UInterchangeGltfTranslator::Translate( UInterchangeBaseNodeContainer& NodeC
 		int32 MaterialIndex = 0;
 		for ( const GLTF::FMaterial& GltfMaterial : GltfAsset.Materials )
 		{
-			UInterchangeShaderGraphNode* ShaderGraphNode = NewObject< UInterchangeShaderGraphNode >( &NodeContainer );
-			FString ShaderGraphNodeUid = TEXT("\\Material\\") + GenerateUniqueIdForGltfNode( GltfMaterial.Name, MaterialIndex );
-			ShaderGraphNode->InitializeNode( ShaderGraphNodeUid, TEXT("Material_") + GenerateUniqueIdForGltfNode( GltfMaterial.Name, MaterialIndex ), EInterchangeNodeContainerType::TranslatedAsset );
-			NodeContainer.AddNode( ShaderGraphNode );
+			UInterchangeShaderGraphNode* ShaderGraphNode = UInterchangeShaderGraphNode::Create( &NodeContainer, TEXT("Material_") + GenerateNameForGltfNode( GltfMaterial.Name, MaterialIndex ) );
 
 			HandleGltfMaterial( NodeContainer, GltfMaterial, *ShaderGraphNode );
 			++MaterialIndex;
@@ -805,9 +753,9 @@ bool UInterchangeGltfTranslator::Translate( UInterchangeBaseNodeContainer& NodeC
 		for ( const GLTF::FMesh& GltfMesh : GltfAsset.Meshes )
 		{
 			UInterchangeMeshNode* MeshNode = NewObject< UInterchangeMeshNode >( &NodeContainer );
-			FString MeshNodeUid = TEXT("\\Mesh\\") + GenerateUniqueIdForGltfNode( GltfMesh.Name, MeshIndex );
+			FString MeshNodeUid = TEXT("\\Mesh\\") + GenerateNameForGltfNode( GltfMesh.Name, MeshIndex );
 
-			MeshNode->InitializeNode( MeshNodeUid, GltfMesh.Name, EInterchangeNodeContainerType::TranslatedAsset );
+			MeshNode->InitializeNode( MeshNodeUid, GenerateNameForGltfNode( GltfMesh.Name, MeshIndex ), EInterchangeNodeContainerType::TranslatedAsset );
 			MeshNode->SetPayLoadKey( LexToString( MeshIndex++ ) );
 			NodeContainer.AddNode( MeshNode );
 
@@ -816,7 +764,7 @@ bool UInterchangeGltfTranslator::Translate( UInterchangeBaseNodeContainer& NodeC
 			{
 				if ( GltfAsset.Materials.IsValidIndex( Primitive.MaterialIndex ) )
 				{
-					const FString ShaderGraphNodeUid = TEXT("\\Material\\") + GenerateUniqueIdForGltfNode( GltfAsset.Materials[ Primitive.MaterialIndex ].Name, Primitive.MaterialIndex );
+					const FString ShaderGraphNodeUid = UInterchangeShaderGraphNode::MakeNodeUid( TEXT("Material_") + GenerateNameForGltfNode( GltfAsset.Materials[ Primitive.MaterialIndex ].Name, Primitive.MaterialIndex ) );
 					MeshNode->SetMaterialDependencyUid( ShaderGraphNodeUid );
 				}
 			}
@@ -829,7 +777,7 @@ bool UInterchangeGltfTranslator::Translate( UInterchangeBaseNodeContainer& NodeC
 		for ( const GLTF::FCamera& GltfCamera : GltfAsset.Cameras )
 		{
 			UInterchangeCameraNode* CameraNode = NewObject< UInterchangeCameraNode >( &NodeContainer );
-			FString CameraNodeUid = TEXT("\\Camera\\") + GenerateUniqueIdForGltfNode( GltfCamera.Name, CameraIndex );
+			FString CameraNodeUid = TEXT("\\Camera\\") + GenerateNameForGltfNode( GltfCamera.Name, CameraIndex );
 			CameraNode->InitializeNode( CameraNodeUid, GltfCamera.Name, EInterchangeNodeContainerType::TranslatedAsset );
 			NodeContainer.AddNode( CameraNode );
 			++CameraIndex;
@@ -842,7 +790,7 @@ bool UInterchangeGltfTranslator::Translate( UInterchangeBaseNodeContainer& NodeC
 		for ( const GLTF::FLight& GltfLight : GltfAsset.Lights )
 		{
 			UInterchangeLightNode* LightNode = NewObject< UInterchangeLightNode >( &NodeContainer );
-			FString LightNodeUid = TEXT("\\Light\\") + GenerateUniqueIdForGltfNode( GltfLight.Name, LightIndex );
+			FString LightNodeUid = TEXT("\\Light\\") + GenerateNameForGltfNode( GltfLight.Name, LightIndex );
 			LightNode->InitializeNode( LightNodeUid, GltfLight.Name, EInterchangeNodeContainerType::TranslatedAsset );
 			NodeContainer.AddNode( LightNode );
 			++LightIndex;
@@ -862,7 +810,7 @@ bool UInterchangeGltfTranslator::Translate( UInterchangeBaseNodeContainer& NodeC
 				SceneName = TEXT("Scene");
 			}
 
-			SceneName = GenerateUniqueIdForGltfNode( SceneName, SceneIndex );
+			SceneName = GenerateNameForGltfNode( SceneName, SceneIndex );
 
 			FString SceneNodeUid = TEXT("\\Scene\\") + SceneName;
 			SceneNode->InitializeNode( SceneNodeUid, SceneName, EInterchangeNodeContainerType::TranslatedScene );
@@ -916,7 +864,7 @@ TFuture< TOptional< UE::Interchange::FStaticMeshPayloadData > > UInterchangeGltf
 
 			if ( GltfAsset.Materials.IsValidIndex( MaterialIndex ) )
 			{
-				StaticMeshAttributes.GetPolygonGroupMaterialSlotNames()[ MaterialSlotIndex ] = *(TEXT("Material_") + GenerateUniqueIdForGltfNode( GltfAsset.Materials[ MaterialIndex ].Name, MaterialIndex ));
+				StaticMeshAttributes.GetPolygonGroupMaterialSlotNames()[ MaterialSlotIndex ] = *(TEXT("Material_") + GenerateNameForGltfNode( GltfAsset.Materials[ MaterialIndex ].Name, MaterialIndex ));
 			}
 		}
 	}
diff --git a/Engine/Plugins/Experimental/Interchange/Runtime/Source/Import/Private/Mesh/InterchangeOBJTranslator.cpp b/Engine/Plugins/Experimental/Interchange/Runtime/Source/Import/Private/Mesh/InterchangeOBJTranslator.cpp
--- a/Engine/Plugins/Experimental/Interchange/Runtime/Source/Import/Private/Mesh/InterchangeOBJTranslator.cpp
+++ b/Engine/Plugins/Experimental/Interchange/Runtime/Source/Import/Private/Mesh/InterchangeOBJTranslator.cpp
@@ -782,7 +782,7 @@ namespace ObjParser
  */
 static FString MakeMeshNodeUid(const FString& Name)
 {
-	return FString(TEXT("\\Mesh\\")) + (Name.IsEmpty() ? FString(TEXT("Null")) : Name);
+	return TEXT("\\Mesh\\") + (Name.IsEmpty() ? FString(TEXT("Null")) : Name);
 }
 
 
@@ -790,9 +790,9 @@ static FString MakeMeshNodeUid(const FString& Name)
  * Build and return a UID name for a shader graph node.
  * @todo: move this to be a static method on the ShaderGraphNode class itself?
  */
-static FString MakeShaderGraphNodeUid(const FString& Name)
+static FString MakeShaderGraphNodeName(const FString& Name)
 {
-	return FString(TEXT("\\Material\\MAT_")) + (Name.IsEmpty() ? FString(TEXT("Null")) : Name);
+	return TEXT("MAT_") + (Name.IsEmpty() ? FString(TEXT("Null")) : Name);
 }
 
 
@@ -800,9 +800,9 @@ static FString MakeShaderGraphNodeUid(const FString& Name)
  * Build and return a UID name for a texture node.
  * @todo: move this to be a static method on the TextureNode class itself?
  */
-static FString MakeTextureNodeUid(const FString& Name)
+static FString MakeTextureNodeName(const FString& Name)
 {
-	return FString(TEXT("\\Texture\\")) + (Name.IsEmpty() ? FString(TEXT("Null")) : Name);
+	return TEXT("TEX_") + (Name.IsEmpty() ? FString(TEXT("Null")) : Name);
 }
 
 
@@ -903,7 +903,7 @@ bool UInterchangeOBJTranslator::Translate(UInterchangeBaseNodeContainer& BaseNod
 
 			if (!Group.Value.MaterialName.IsEmpty())
 			{
-				MeshNode->SetMaterialDependencyUid(MakeShaderGraphNodeUid(Group.Value.MaterialName));
+				MeshNode->SetMaterialDependencyUid(UInterchangeShaderGraphNode::MakeNodeUid(MakeShaderGraphNodeName(Group.Value.MaterialName)));
 			}
 		}
 
@@ -913,9 +913,9 @@ bool UInterchangeOBJTranslator::Translate(UInterchangeBaseNodeContainer& BaseNod
 		{
 			const FString& MaterialName = Material.Key;
 			const FObjData::FMaterialData& MaterialData = Material.Value;
-			FString NodeUid = MakeShaderGraphNodeUid(MaterialName);
+			FString NodeUid = UInterchangeShaderGraphNode::MakeNodeUid(MaterialName);
 
-			UInterchangeShaderGraphNode* ShaderGraphNode = NewObject<UInterchangeShaderGraphNode>(&BaseNodeContainer);
+			UInterchangeShaderGraphNode* ShaderGraphNode = UInterchangeShaderGraphNode::Create(&BaseNodeContainer, MaterialName);
 			ShaderGraphNode->InitializeNode(NodeUid, MaterialName, EInterchangeNodeContainerType::TranslatedAsset);
 			BaseNodeContainer.AddNode(ShaderGraphNode);
 
@@ -990,24 +990,16 @@ void UInterchangeOBJTranslator::AddMaterialNodes(UInterchangeBaseNodeContainer&
 
 		using namespace UE::Interchange::Materials::Standard::Nodes;
 
-		FString TextureName = FString(TEXT("TEX_")) + FPaths::GetBaseFilename(TexturePath);
+		const FString TextureName = MakeTextureNodeName(FPaths::GetBaseFilename(TexturePath));
 
-		UInterchangeShaderNode* TextureSampleShader = NewObject<UInterchangeShaderNode>(&BaseNodeContainer);
+		UInterchangeShaderNode* TextureSampleShader = UInterchangeShaderNode::Create(&BaseNodeContainer, TextureName, ShaderGraphNode->GetUniqueID());
 		TextureSampleShader->SetCustomShaderType(TextureSample::Name.ToString());
 
-		FString TextureSampleShaderUid = ShaderGraphNode->GetUniqueID() + TEXT("\\Textures\\") + TextureName;
-		TextureSampleShader->InitializeNode(TextureSampleShaderUid, TextureName, EInterchangeNodeContainerType::TranslatedAsset);
-
-		BaseNodeContainer.AddNode(TextureSampleShader);
-		BaseNodeContainer.SetNodeParentUid(TextureSampleShaderUid, ShaderGraphNode->GetUniqueID());
-
-		FString TextureNodeUid = MakeTextureNodeUid(TextureName);
+		FString TextureNodeUid = UInterchangeTextureNode::MakeNodeUid(TextureName);
 		const UInterchangeTexture2DNode* TextureNode = Cast<const UInterchangeTexture2DNode>(BaseNodeContainer.GetNode(TextureNodeUid));
 		if (!TextureNode)
 		{
-			UInterchangeTexture2DNode* NewTextureNode = NewObject<UInterchangeTexture2DNode>(&BaseNodeContainer);
-			NewTextureNode->InitializeNode(TextureNodeUid, TextureName, EInterchangeNodeContainerType::TranslatedAsset);
-			BaseNodeContainer.AddNode(NewTextureNode);
+			UInterchangeTexture2DNode* NewTextureNode = UInterchangeTexture2DNode::Create(&BaseNodeContainer, TextureName);
 
 			FString NormalizedTexturePath(TexturePath);
 			FPaths::NormalizeFilename(NormalizedTexturePath);
@@ -1016,7 +1008,7 @@ void UInterchangeOBJTranslator::AddMaterialNodes(UInterchangeBaseNodeContainer&
 
 		TextureSampleShader->AddStringAttribute(UInterchangeShaderPortsAPI::MakeInputValueKey(TextureSample::Inputs::Texture.ToString()), TextureNodeUid);
 
-		UInterchangeShaderPortsAPI::ConnectDefaultOuputToInput(ShaderGraphNode, InputType, TextureSampleShaderUid);
+		UInterchangeShaderPortsAPI::ConnectDefaultOuputToInput(ShaderGraphNode, InputType, TextureSampleShader->GetUniqueID());
 	}
 	else
 	{
diff --git a/Engine/Plugins/Experimental/Interchange/Runtime/Source/Nodes/Private/InterchangeShaderGraphNode.cpp b/Engine/Plugins/Experimental/Interchange/Runtime/Source/Nodes/Private/InterchangeShaderGraphNode.cpp
--- a/Engine/Plugins/Experimental/Interchange/Runtime/Source/Nodes/Private/InterchangeShaderGraphNode.cpp
+++ b/Engine/Plugins/Experimental/Interchange/Runtime/Source/Nodes/Private/InterchangeShaderGraphNode.cpp
@@ -1,6 +1,7 @@
 // Copyright Epic Games, Inc. All Rights Reserved.
 
 #include "InterchangeShaderGraphNode.h"
+#include "Nodes/InterchangeBaseNodeContainer.h"
 
 const TCHAR* UInterchangeShaderPortsAPI::InputPrefix = TEXT("Inputs");
 const TCHAR* UInterchangeShaderPortsAPI::InputSeparator = TEXT(":");
@@ -112,6 +113,37 @@ bool UInterchangeShaderPortsAPI::GetInputConnection(const UInterchangeBaseNode*
 	return false;
 }
 
+FString UInterchangeShaderNode::MakeNodeUid(const FStringView NodeName, const FStringView ParentNodeUid)
+{
+	FString ShaderNodeUid = FString(UInterchangeBaseNode::HierarchySeparator) + TEXT("Shaders") + FString(UInterchangeBaseNode::HierarchySeparator);
+
+	if (ParentNodeUid.IsEmpty())
+	{
+		ShaderNodeUid += NodeName;
+	}
+	else
+	{
+		ShaderNodeUid += ParentNodeUid + FString(UInterchangeBaseNode::HierarchySeparator) + NodeName;
+	}
+
+	return ShaderNodeUid;
+}
+
+UInterchangeShaderNode* UInterchangeShaderNode::Create(UInterchangeBaseNodeContainer* NodeContainer, const FStringView NodeName, const FStringView ParentNodeUid)
+{
+	check(NodeContainer);
+
+	const FString ShaderNodeUid = MakeNodeUid(NodeName, ParentNodeUid);
+
+	UInterchangeShaderNode* ShaderNode = NewObject< UInterchangeShaderNode >(NodeContainer);
+	ShaderNode->InitializeNode(ShaderNodeUid, FString(NodeName), EInterchangeNodeContainerType::TranslatedAsset);
+
+	NodeContainer->AddNode(ShaderNode);
+	NodeContainer->SetNodeParentUid(ShaderNodeUid, FString(ParentNodeUid));
+
+	return ShaderNode;
+}
+
 FString UInterchangeShaderNode::GetTypeName() const
 {
 	const FString TypeName = TEXT("ShaderNode");
@@ -128,6 +160,25 @@ bool UInterchangeShaderNode::SetCustomShaderType(const FString& AttributeValue)
 	IMPLEMENT_NODE_ATTRIBUTE_SETTER_NODELEGATE(ShaderType, FString);
 }
 
+FString UInterchangeShaderGraphNode::MakeNodeUid(const FStringView NodeName)
+{
+	return FString(UInterchangeBaseNode::HierarchySeparator) + TEXT("Shaders") + FString(UInterchangeBaseNode::HierarchySeparator) + NodeName;
+}
+
+UInterchangeShaderGraphNode* UInterchangeShaderGraphNode::Create(UInterchangeBaseNodeContainer* NodeContainer, const FStringView NodeName)
+{
+	check(NodeContainer);
+
+	const FString ShaderGraphNodeUid = MakeNodeUid(NodeName);
+
+	UInterchangeShaderGraphNode* ShaderGraphNode = NewObject< UInterchangeShaderGraphNode >(NodeContainer);
+	ShaderGraphNode->InitializeNode(ShaderGraphNodeUid, FString(NodeName), EInterchangeNodeContainerType::TranslatedAsset);
+
+	NodeContainer->AddNode(ShaderGraphNode);
+
+	return ShaderGraphNode;
+}
+
 FString UInterchangeShaderGraphNode::GetTypeName() const
 {
 	const FString TypeName = TEXT("ShaderGraphNode");
diff --git a/Engine/Plugins/Experimental/Interchange/Runtime/Source/Nodes/Private/InterchangeTexture2DNode.cpp b/Engine/Plugins/Experimental/Interchange/Runtime/Source/Nodes/Private/InterchangeTexture2DNode.cpp
new file mode 100644
--- /dev/null
+++ b/Engine/Plugins/Experimental/Interchange/Runtime/Source/Nodes/Private/InterchangeTexture2DNode.cpp
@@ -0,0 +1,20 @@
+// Copyright Epic Games, Inc. All Rights Reserved.
+
+#include "InterchangeTexture2DNode.h"
+#include "Nodes/InterchangeBaseNodeContainer.h"
+
+FString UInterchangeTexture2DNode::MakeNodeUid(const FStringView NodeName)
+{
+	return UInterchangeTextureNode::MakeNodeUid(NodeName);
+}
+
+ UInterchangeTexture2DNode* UInterchangeTexture2DNode::Create(UInterchangeBaseNodeContainer* NodeContainer, const FStringView TextureNodeName)
+{
+	UInterchangeTexture2DNode* TextureNode = NewObject< UInterchangeTexture2DNode >(NodeContainer);
+	const FString TextureNodeUid = UInterchangeTextureNode::MakeNodeUid(TextureNodeName);
+
+	TextureNode->InitializeNode(TextureNodeUid, FString(TextureNodeName), EInterchangeNodeContainerType::TranslatedAsset);
+	NodeContainer->AddNode(TextureNode);
+
+	return TextureNode;
+}
\ No newline at end of file
diff --git a/Engine/Plugins/Experimental/Interchange/Runtime/Source/Nodes/Public/InterchangeShaderGraphNode.h b/Engine/Plugins/Experimental/Interchange/Runtime/Source/Nodes/Public/InterchangeShaderGraphNode.h
--- a/Engine/Plugins/Experimental/Interchange/Runtime/Source/Nodes/Public/InterchangeShaderGraphNode.h
+++ b/Engine/Plugins/Experimental/Interchange/Runtime/Source/Nodes/Public/InterchangeShaderGraphNode.h
@@ -6,7 +6,9 @@
 #include "UObject/Object.h"
  
 #include "InterchangeShaderGraphNode.generated.h"
- 
+
+class UInterchangeBaseNodeContainer;
+
 /**
  * The Shader Ports API manages a set of inputs and outputs attributes.
  * This API can be used over any InterchangeBaseNode that wants to support shader ports as attributes.
@@ -89,7 +91,7 @@ private:
 	static const TCHAR* InputPrefix;
 	static const TCHAR* InputSeparator;
 };
- 
+
 /**
  * A shader node is a named set of inputs and outputs. It can be connected to other shader nodes and finally to a shader graph input.
  */
@@ -99,6 +101,16 @@ class INTERCHANGENODES_API UInterchangeShaderNode : public UInterchangeBaseNode
 	GENERATED_BODY()
  
 public:
+	/**
+	 * Build and return a UID name for a shader node.
+	 */
+	static FString MakeNodeUid(const FStringView NodeName, const FStringView ParentNodeUid);
+
+	/**
+	 * Creates a new UInterchangeShaderNode and adds it to NodeContainer as a translated node.
+	 */
+	static UInterchangeShaderNode* Create(UInterchangeBaseNodeContainer* NodeContainer, const FStringView NodeName, const FStringView ParentNodeUid);
+
 	virtual FString GetTypeName() const override;
  
 public:
@@ -125,6 +137,16 @@ class INTERCHANGENODES_API UInterchangeShaderGraphNode : public UInterchangeShad
 	GENERATED_BODY()
  
 public:
+	/**
+	 * Build and return a UID name for a shader graph node.
+	 */
+	static FString MakeNodeUid(const FStringView NodeName);
+
+	/**
+	 * Creates a new UInterchangeShaderGraphNode and adds it to NodeContainer as a translated node.
+	 */
+	static UInterchangeShaderGraphNode* Create(UInterchangeBaseNodeContainer* NodeContainer, const FStringView NodeName);
+
 	/**
 	 * Return the node type name of the class, we use this when reporting error
 	 */
diff --git a/Engine/Plugins/Experimental/Interchange/Runtime/Source/Nodes/Public/InterchangeTexture2DNode.h b/Engine/Plugins/Experimental/Interchange/Runtime/Source/Nodes/Public/InterchangeTexture2DNode.h
--- a/Engine/Plugins/Experimental/Interchange/Runtime/Source/Nodes/Public/InterchangeTexture2DNode.h
+++ b/Engine/Plugins/Experimental/Interchange/Runtime/Source/Nodes/Public/InterchangeTexture2DNode.h
@@ -7,6 +7,8 @@
 
 #include "InterchangeTexture2DNode.generated.h"
 
+class UInterchangeBaseNodeContainer;
+
 //Interchange namespace
 namespace UE::Interchange
 {
@@ -27,6 +29,16 @@ class INTERCHANGENODES_API UInterchangeTexture2DNode : public UInterchangeTextur
 
 public:
 
+	/**
+	 * Build and return a UID name for a texture 2D node.
+	 */
+	static FString MakeNodeUid(const FStringView NodeName);
+
+	/**
+	 * Creates a new UInterchangeTexture2DNode and adds it to NodeContainer as a translated node.
+	 */
+	static UInterchangeTexture2DNode* Create(UInterchangeBaseNodeContainer* NodeContainer, const FStringView TextureNodeName);
+
 	virtual void PostInitProperties()
 	{
 		Super::PostInitProperties();
diff --git a/Engine/Plugins/Experimental/Interchange/Runtime/Source/Nodes/Public/InterchangeTextureNode.h b/Engine/Plugins/Experimental/Interchange/Runtime/Source/Nodes/Public/InterchangeTextureNode.h
--- a/Engine/Plugins/Experimental/Interchange/Runtime/Source/Nodes/Public/InterchangeTextureNode.h
+++ b/Engine/Plugins/Experimental/Interchange/Runtime/Source/Nodes/Public/InterchangeTextureNode.h
@@ -37,6 +37,14 @@ class INTERCHANGENODES_API UInterchangeTextureNode : public UInterchangeBaseNode
 	GENERATED_BODY()
 
 public:
+	/**
+	 * Build and return a UID name for a texture node.
+	 */
+	static FString MakeNodeUid(const FStringView NodeName)
+	{
+		return FString(UInterchangeBaseNode::HierarchySeparator) +  TEXT("Textures") + FString(UInterchangeBaseNode::HierarchySeparator) + NodeName;
+	}
+
 	/**
 	 * Return the node type name of the class, we use this when reporting error
 	 */
diff --git a/Engine/Plugins/Experimental/Interchange/Runtime/Source/Parsers/Fbx/Private/FbxMaterial.cpp b/Engine/Plugins/Experimental/Interchange/Runtime/Source/Parsers/Fbx/Private/FbxMaterial.cpp
--- a/Engine/Plugins/Experimental/Interchange/Runtime/Source/Parsers/Fbx/Private/FbxMaterial.cpp
+++ b/Engine/Plugins/Experimental/Interchange/Runtime/Source/Parsers/Fbx/Private/FbxMaterial.cpp
@@ -35,29 +35,15 @@ namespace UE
 				return ShaderGraphNode;
 			}
 
-			UInterchangeShaderNode* FFbxMaterial::CreateShaderNode(UInterchangeBaseNodeContainer& NodeContainer, const FString& ParentUid, const FString& NodeName)
+			UInterchangeTexture2DNode* FFbxMaterial::CreateTexture2DNode(UInterchangeBaseNodeContainer& NodeContainer, const FString& TextureFilePath)
 			{
-				UInterchangeShaderNode* ShaderNode = NewObject<UInterchangeShaderNode>(&NodeContainer);
+				const FString TextureName = FPaths::GetBaseFilename(TextureFilePath);
+				UInterchangeTexture2DNode* TextureNode = UInterchangeTexture2DNode::Create(&NodeContainer, TextureName);
 
-				const FString NodeUid = ParentUid + TEXT("\\") + NodeName;
-				ShaderNode->InitializeNode(NodeUid, NodeName, EInterchangeNodeContainerType::TranslatedAsset);
-				NodeContainer.AddNode(ShaderNode);
-				NodeContainer.SetNodeParentUid(NodeUid, ParentUid);
-
-				return ShaderNode;
-			}
-
-			UInterchangeTexture2DNode* FFbxMaterial::CreateTexture2DNode(UInterchangeBaseNodeContainer& NodeContainer, const FString& NodeUid, const FString& TextureFilePath)
-			{
-				FString DisplayLabel = FPaths::GetBaseFilename(TextureFilePath);
-				UInterchangeTexture2DNode* TextureNode = NewObject<UInterchangeTexture2DNode>(&NodeContainer, NAME_None);
-				TextureNode->InitializeNode(NodeUid, DisplayLabel, EInterchangeNodeContainerType::TranslatedAsset);
-
-				//All texture translator expect a file has the payload key
+				//All texture translator expect a file as the payload key
 				FString NormalizeFilePath = TextureFilePath;
 				FPaths::NormalizeFilename(NormalizeFilePath);
 				TextureNode->SetPayLoadKey(NormalizeFilePath);
-				NodeContainer.AddNode(TextureNode);
 
 				return TextureNode;
 			}
@@ -76,7 +62,7 @@ namespace UE
 				if (bInverse)
 				{
 					const FString OneMinusNodeName = InputName.ToString() + TEXT("OneMinus");
-					UInterchangeShaderNode* OneMinusNode = CreateShaderNode(NodeContainer, ShaderGraphNode->GetUniqueID(), OneMinusNodeName);
+					UInterchangeShaderNode* OneMinusNode = UInterchangeShaderNode::Create(&NodeContainer, OneMinusNodeName, ShaderGraphNode->GetUniqueID());
 					OneMinusNode->SetCustomShaderType(OneMinus::Name.ToString());
 
 					UInterchangeShaderPortsAPI::ConnectDefaultOuputToInput(ShaderGraphNode, InputName.ToString(), OneMinusNode->GetUniqueID());
@@ -88,7 +74,7 @@ namespace UE
 				if (!FMath::IsNearlyEqual(Factor, 1.f))
 				{
 					FString LerpNodeName = InputName.ToString() + TEXT("Lerp");
-					UInterchangeShaderNode* LerpNode = CreateShaderNode(NodeContainer, ShaderGraphNode->GetUniqueID(), LerpNodeName);
+					UInterchangeShaderNode* LerpNode = UInterchangeShaderNode::Create(&NodeContainer, LerpNodeName, ShaderGraphNode->GetUniqueID());
 					LerpNode->SetCustomShaderType(Lerp::Name.ToString());
 
 					if (DefaultValue.IsType<float>())
@@ -118,43 +104,30 @@ namespace UE
 					if(!TextureFilename.IsEmpty() && FPaths::FileExists(TextureFilename))
 					{
 						FString TextureName = FPaths::GetBaseFilename(TextureFilename);
-						UInterchangeShaderNode* TextureSampleShader = NewObject<UInterchangeShaderNode>(&NodeContainer);
+						UInterchangeShaderNode* TextureSampleShader = UInterchangeShaderNode::Create(&NodeContainer, TextureName, ShaderGraphNode->GetUniqueID());
 						TextureSampleShader->SetCustomShaderType(TextureSample::Name.ToString());
 
-						FString TextureSampleShaderUid = ShaderGraphNode->GetUniqueID() + TEXT("\\Textures\\") + TextureName;
-						TextureSampleShader->InitializeNode(TextureSampleShaderUid, TextureName, EInterchangeNodeContainerType::TranslatedAsset);
-
-						NodeContainer.AddNode(TextureSampleShader);
-						NodeContainer.SetNodeParentUid(TextureSampleShaderUid, ShaderGraphNode->GetUniqueID());
-
-						FString TextureNodeUid = TEXT("\\Texture\\") + TextureFilename;
-						const UInterchangeTexture2DNode* TextureNode = Cast<const UInterchangeTexture2DNode>(NodeContainer.GetNode(TextureNodeUid));
+						const UInterchangeTexture2DNode* TextureNode = Cast<const UInterchangeTexture2DNode>(NodeContainer.GetNode(UInterchangeTextureNode::MakeNodeUid(TextureName)));
 
 						if (!TextureNode)
 						{
-							TextureNode = CreateTexture2DNode(NodeContainer, TextureNodeUid, TextureFilename);
+							TextureNode = CreateTexture2DNode(NodeContainer, TextureFilename);
 						}
 
 						TextureSampleShader->AddStringAttribute(UInterchangeShaderPortsAPI::MakeInputValueKey(TextureSample::Inputs::Texture.ToString()), TextureNode->GetUniqueID());
 
 						if (!FMath::IsNearlyEqual(FbxTexture->GetScaleU(), 1.0) || !FMath::IsNearlyEqual(FbxTexture->GetScaleV(), 1.0))
 						{
-							UInterchangeShaderNode* TextureCoordinateShader = NewObject<UInterchangeShaderNode>(&NodeContainer);
+							UInterchangeShaderNode* TextureCoordinateShader = UInterchangeShaderNode::Create(&NodeContainer, TextureName + TEXT("_Coordinate"), ShaderGraphNode->GetUniqueID());
 							TextureCoordinateShader->SetCustomShaderType(TextureCoordinate::Name.ToString());
 
-							FString TextureCoordinateShaderUid = TextureSampleShader->GetUniqueID() + TEXT("_Coordinate");
-							TextureCoordinateShader->InitializeNode(TextureCoordinateShaderUid, TextureName + TEXT("_Coordinate"), EInterchangeNodeContainerType::TranslatedAsset);
-
-							NodeContainer.AddNode(TextureCoordinateShader);
-							NodeContainer.SetNodeParentUid(TextureCoordinateShaderUid, ShaderGraphNode->GetUniqueID());
-
 							TextureCoordinateShader->AddFloatAttribute(UInterchangeShaderPortsAPI::MakeInputValueKey(TextureCoordinate::Inputs::UTiling.ToString()), (float)FbxTexture->GetScaleU());
 							TextureCoordinateShader->AddFloatAttribute(UInterchangeShaderPortsAPI::MakeInputValueKey(TextureCoordinate::Inputs::VTiling.ToString()), (float)FbxTexture->GetScaleV());
 
-							UInterchangeShaderPortsAPI::ConnectDefaultOuputToInput(TextureSampleShader, TextureSample::Inputs::Coordinates.ToString(), TextureCoordinateShaderUid);
+							UInterchangeShaderPortsAPI::ConnectDefaultOuputToInput(TextureSampleShader, TextureSample::Inputs::Coordinates.ToString(), TextureCoordinateShader->GetUniqueID());
 						}
 
-						UInterchangeShaderPortsAPI::ConnectDefaultOuputToInput(NodeToConnectTo, InputToConnectTo, TextureSampleShaderUid);
+						UInterchangeShaderPortsAPI::ConnectDefaultOuputToInput(NodeToConnectTo, InputToConnectTo, TextureSampleShader->GetUniqueID());
 					}
 					else
 					{
@@ -299,12 +272,11 @@ namespace UE
 						continue;
 					}
 					//Create a texture node and make it child of the material node
-					TArray<FString> JsonErrorMessage;
-					FString NodeUid = TEXT("\\Texture\\") + TextureFilename;
-					const UInterchangeTexture2DNode* TextureNode = Cast<UInterchangeTexture2DNode>(NodeContainer.GetNode(NodeUid));
+					const FString TextureName = FPaths::GetBaseFilename(TextureFilename);
+					const UInterchangeTexture2DNode* TextureNode = Cast<UInterchangeTexture2DNode>(NodeContainer.GetNode(UInterchangeTextureNode::MakeNodeUid(TextureName)));
 					if (!TextureNode)
 					{
-						TextureNode = CreateTexture2DNode(NodeContainer, NodeUid, TextureFilename);
+						CreateTexture2DNode(NodeContainer, TextureFilename);
 					}
 				}
 			}
diff --git a/Engine/Plugins/Experimental/Interchange/Runtime/Source/Parsers/Fbx/Private/FbxMaterial.h b/Engine/Plugins/Experimental/Interchange/Runtime/Source/Parsers/Fbx/Private/FbxMaterial.h
--- a/Engine/Plugins/Experimental/Interchange/Runtime/Source/Parsers/Fbx/Private/FbxMaterial.h
+++ b/Engine/Plugins/Experimental/Interchange/Runtime/Source/Parsers/Fbx/Private/FbxMaterial.h
@@ -53,8 +53,7 @@ namespace UE
 
 			protected:
 				UInterchangeShaderGraphNode* CreateShaderGraphNode(UInterchangeBaseNodeContainer& NodeContainer, const FString& NodeUID, const FString& NodeName);
-				UInterchangeShaderNode* CreateShaderNode(UInterchangeBaseNodeContainer& NodeContainer, const FString& NodeUid, const FString& NodeName);
-				UInterchangeTexture2DNode* CreateTexture2DNode(UInterchangeBaseNodeContainer& NodeContainer, const FString& NodeUID, const FString& TextureFilePath);
+				UInterchangeTexture2DNode* CreateTexture2DNode(UInterchangeBaseNodeContainer& NodeContainer, const FString& TextureFilePath);
 				void ConvertPropertyToShaderNode(UInterchangeBaseNodeContainer& NodeContainer, UInterchangeShaderGraphNode* ShaderGraphNode,
 					FbxProperty& Property, float Factor, FName PropertyName, TVariant<FLinearColor, float> DefaultValue, bool bInverse = false);
 
diff --git a/Engine/Source/Runtime/Experimental/Interchange/Core/Public/Nodes/InterchangeBaseNode.h b/Engine/Source/Runtime/Experimental/Interchange/Core/Public/Nodes/InterchangeBaseNode.h
--- a/Engine/Source/Runtime/Experimental/Interchange/Core/Public/Nodes/InterchangeBaseNode.h
+++ b/Engine/Source/Runtime/Experimental/Interchange/Core/Public/Nodes/InterchangeBaseNode.h
@@ -235,6 +235,8 @@ class INTERCHANGECORE_API UInterchangeBaseNode : public UObject
 	GENERATED_BODY()
 
 public:
+	static constexpr const TCHAR* HierarchySeparator = TEXT("\\");
+
 	UInterchangeBaseNode();
 
 	/**
