Fixed crash when obj file does not contain UV data
Fixed incorrect name of Interchange node created for material

diff --git a/Engine/Plugins/Experimental/Interchange/Runtime/Source/Import/Private/Mesh/InterchangeOBJTranslator.cpp b/Engine/Plugins/Experimental/Interchange/Runtime/Source/Import/Private/Mesh/InterchangeOBJTranslator.cpp
--- a/Engine/Plugins/Experimental/Interchange/Runtime/Source/Import/Private/Mesh/InterchangeOBJTranslator.cpp
+++ b/Engine/Plugins/Experimental/Interchange/Runtime/Source/Import/Private/Mesh/InterchangeOBJTranslator.cpp
@@ -20,7 +20,9 @@
 #include "Modules/ModuleManager.h"
 #include "Nodes/InterchangeBaseNodeContainer.h"
 #include "StaticMeshAttributes.h"
+#include "StaticMeshOperations.h"
 #include "UObject/GCObjectScopeGuard.h"
+#include "UVMapSettings.h"
 
 
 /**
@@ -163,7 +165,10 @@ TArray<int32> FObjData::GetUVIndicesUsedByGroup(const FGroupData& GroupData) con
 	{
 		for (const FVertexData& VertexData : FaceData.Vertices)
 		{
-			UVIndexSet.Add(VertexData.UVIndex);
+			if (VertexData.UVIndex != INDEX_NONE)
+			{
+				UVIndexSet.Add(VertexData.UVIndex);
+			}
 		}
 	}
 
@@ -263,12 +268,15 @@ FMeshDescription FObjData::MakeMeshDescriptionForGroup(const FString& GroupName)
 
 	// Create UVs and initialize values
 
-	const int32 UVChannel = 0;
-	MeshDescription.ReserveNewUVs(UVIndexMapping.Num());
-	for (int32 ObjUVIndex : UVIndexMapping)
+	if (UVIndexMapping.Num() > 0)
 	{
-		FUVID UVIndex = MeshDescription.CreateUV(UVChannel);
-		Attributes.GetUVCoordinates(UVChannel)[UVIndex] = UVs[ObjUVIndex];
+		const int32 UVChannel = 0;
+		MeshDescription.ReserveNewUVs(UVIndexMapping.Num());
+		for (int32 ObjUVIndex : UVIndexMapping)
+		{
+			FUVID UVIndex = MeshDescription.CreateUV(UVChannel);
+			Attributes.GetUVCoordinates(UVChannel)[UVIndex] = UVs[ObjUVIndex];
+		}
 	}
 
 	// Create polygon group
@@ -359,6 +367,34 @@ FMeshDescription FObjData::MakeMeshDescriptionForGroup(const FString& GroupName)
 		}
 	}
 
+	// Create default UVs if mesh does not have UVs
+
+	if (UVIndexMapping.Num() == 0)
+	{
+		FBox MeshBoundingBox = MeshDescription.ComputeBoundingBox();
+		FUVMapParameters UVParameters(MeshBoundingBox.GetCenter(), FQuat::Identity, MeshBoundingBox.GetSize(), FVector::OneVector, FVector2D::UnitVector);
+		TMap<FVertexInstanceID, FVector2D> TexCoords;
+		FStaticMeshOperations::GenerateBoxUV(MeshDescription, UVParameters, TexCoords);
+
+		TVertexInstanceAttributesRef<FVector2f> VertexInstanceUVs = Attributes.GetVertexInstanceUVs();
+		if (VertexInstanceUVs.GetNumChannels() == 0)
+		{
+			VertexInstanceUVs.SetNumChannels(1);
+		}
+
+		for (const FVertexInstanceID VertexInstanceID : MeshDescription.VertexInstances().GetElementIDs())
+		{
+			if (const FVector2D* UVCoord = TexCoords.Find(VertexInstanceID))
+			{
+				VertexInstanceUVs.Set(VertexInstanceID, 0, FVector2f(*UVCoord));
+			}
+			else
+			{
+				ensureMsgf(false, TEXT("Tried to apply UV data that did not match the MeshDescription."));
+			}
+		}
+	}
+
 	MeshDescription.ResumeVertexInstanceIndexing();
 	MeshDescription.ResumePolygonIndexing();
 	MeshDescription.ResumePolygonGroupIndexing();
@@ -913,7 +949,7 @@ bool UInterchangeOBJTranslator::Translate(UInterchangeBaseNodeContainer& BaseNod
 		{
 			const FString& MaterialName = Material.Key;
 			const FObjData::FMaterialData& MaterialData = Material.Value;
-			FString NodeUid = UInterchangeShaderGraphNode::MakeNodeUid(MaterialName);
+			FString NodeUid = UInterchangeShaderGraphNode::MakeNodeUid(MakeShaderGraphNodeName(MaterialName));
 
 			UInterchangeShaderGraphNode* ShaderGraphNode = UInterchangeShaderGraphNode::Create(&BaseNodeContainer, MaterialName);
 			ShaderGraphNode->InitializeNode(NodeUid, MaterialName, EInterchangeNodeContainerType::TranslatedAsset);
