Add MaterialX Translator to Interchange module
* Ignore MaterialX warning -Wmissing-braces on clang
* Only enable the Translator with the editor for the moment, because of the MaterialX exceptions during Runtime

diff --git a/Engine/Plugins/Experimental/Interchange/Runtime/Source/Import/InterchangeImport.Build.cs b/Engine/Plugins/Experimental/Interchange/Runtime/Source/Import/InterchangeImport.Build.cs
--- a/Engine/Plugins/Experimental/Interchange/Runtime/Source/Import/InterchangeImport.Build.cs
+++ b/Engine/Plugins/Experimental/Interchange/Runtime/Source/Import/InterchangeImport.Build.cs
@@ -51,6 +51,8 @@ namespace UnrealBuildTool.Rules
 					}
 				);
 			}
+
+			AddEngineThirdPartyPrivateStaticDependencies(Target, "MaterialX");
 		}
 	}
 }
diff --git a/Engine/Plugins/Experimental/Interchange/Runtime/Source/Import/Private/InterchangeImportModule.cpp b/Engine/Plugins/Experimental/Interchange/Runtime/Source/Import/Private/InterchangeImportModule.cpp
--- a/Engine/Plugins/Experimental/Interchange/Runtime/Source/Import/Private/InterchangeImportModule.cpp
+++ b/Engine/Plugins/Experimental/Interchange/Runtime/Source/Import/Private/InterchangeImportModule.cpp
@@ -9,6 +9,7 @@
 #include "InterchangeImportLog.h"
 #include "InterchangeManager.h"
 #include "Material/InterchangeMaterialFactory.h"
+#include "MaterialX/InterchangeMaterialXTranslator.h"
 #include "Mesh/InterchangeOBJTranslator.h"
 #include "Mesh/InterchangePhysicsAssetFactory.h"
 #include "Mesh/InterchangeSkeletalMeshFactory.h"
@@ -53,6 +54,9 @@ void FInterchangeImportModule::StartupModule()
 		InterchangeManager.RegisterTranslator(UInterchangeGltfTranslator::StaticClass());
 		InterchangeManager.RegisterTranslator(UInterchangeOBJTranslator::StaticClass());
 
+		//Materials
+		InterchangeManager.RegisterTranslator(UInterchangeMaterialXTranslator::StaticClass());
+
 		//Textures
 		InterchangeManager.RegisterTranslator(UInterchangeImageWrapperTranslator::StaticClass());
 		InterchangeManager.RegisterTextureOnlyTranslatorClass(UInterchangeImageWrapperTranslator::StaticClass());
@@ -80,7 +84,7 @@ void FInterchangeImportModule::StartupModule()
 		InterchangeManager.RegisterFactory(UInterchangeStaticMeshActorFactory::StaticClass());
 		InterchangeManager.RegisterFactory(UInterchangeSkeletalMeshActorFactory::StaticClass());
 	};
-
+	
 	if (GEngine)
 	{
 		RegisterItems();
diff --git a/Engine/Plugins/Experimental/Interchange/Runtime/Source/Import/Private/MaterialX/InterchangeMaterialXDefinitions.h b/Engine/Plugins/Experimental/Interchange/Runtime/Source/Import/Private/MaterialX/InterchangeMaterialXDefinitions.h
new file mode 100644
--- /dev/null
+++ b/Engine/Plugins/Experimental/Interchange/Runtime/Source/Import/Private/MaterialX/InterchangeMaterialXDefinitions.h
@@ -0,0 +1,95 @@
+#pragma once
+
+namespace MaterialX
+{
+	namespace StandardSurface
+	{
+		namespace Input
+		{
+			static constexpr const char* Base = "base";
+			static constexpr const char* BaseColor = "base_color";
+			static constexpr const char* DiffuseRoughness = "diffuse_roughness";
+			static constexpr const char* Metalness = "metalness";
+			static constexpr const char* Specular = "specular";
+			static constexpr const char* SpecularColor = "specular_color";
+			static constexpr const char* SpecularRoughness = "specular_roughness";
+			static constexpr const char* SpecularIOR = "specular_IOR";
+			static constexpr const char* SpecularAnisotropy = "specular_anisotropy";
+			static constexpr const char* SpecularRotation = "specular_rotation";
+			static constexpr const char* Transmission = "transmission";
+			static constexpr const char* TransmissionColor = "transmission_color";
+			static constexpr const char* TransmissionDepth = "transmission_depth";
+			static constexpr const char* TransmissionScatter = "transmission_scatter";
+			static constexpr const char* TransmissionScatterAnisotropy = "transmission_scatter_anisotropy";
+			static constexpr const char* TransmissionDispersion = "transmission_dispersion";
+			static constexpr const char* TransmissionExtraRoughness= "transmission_extra_roughness";
+			static constexpr const char* Subsurface = "subsurface";
+			static constexpr const char* SubsurfaceColor = "subsurface_color";
+			static constexpr const char* SubsurfaceRadius = "subsurface_radius";
+			static constexpr const char* SubsurfaceScale = "subsurface_scale";
+			static constexpr const char* SubsurfaceAnisotropy = "subsurface_anisotropy";
+			static constexpr const char* Sheen = "sheen";
+			static constexpr const char* SheenColor = "sheen_color";
+			static constexpr const char* SheenRoughness = "sheen_roughness";
+			static constexpr const char* Coat = "coat";
+			static constexpr const char* CoatColor = "coat_color";
+			static constexpr const char* CoatRoughness = "coat_roughness";
+			static constexpr const char* CoatAnisotropy = "coat_anisotropy";
+			static constexpr const char* CoatRotation = "coat_rotation";
+			static constexpr const char* CoatIOR = "coat_IOR";
+			static constexpr const char* CoatNormal = "coat_normal";
+			static constexpr const char* CoatAffectColor = "coat_affect_color";
+			static constexpr const char* CoatAffectRoughness = "coat_affect_roughness";
+			static constexpr const char* ThinFilmThickness = "thin_film_thickness";
+			static constexpr const char* ThinFilmIOR = "thin_film_IOR";
+			static constexpr const char* Emission = "emission";
+			static constexpr const char* EmissionColor = "emission_color";
+			static constexpr const char* Opacity = "opacity";
+			static constexpr const char* ThinWalled= "thin_walled";
+			static constexpr const char* Normal = "normal";
+			static constexpr const char* Tangent = "tangent";
+		}
+	}
+
+	namespace Attributes
+	{
+		static constexpr const char* IsVisited = "UE:IsVisited";
+		static constexpr const char* OldName = "UE:OldName";
+	}
+
+	namespace Category
+	{
+		static constexpr const char* Add = "add";
+		static constexpr const char* Sub = "subtract";
+		static constexpr const char* Multiply = "multiply";
+		static constexpr const char* Clamp = "clamp";
+		static constexpr const char* Max = "max";
+		static constexpr const char* Min = "min";
+		static constexpr const char* Mix = "mix";
+		static constexpr const char* DotProduct = "dotproduct";
+		static constexpr const char* Combine2 = "combine2";
+		static constexpr const char* Combine3 = "combine2";
+		static constexpr const char* Combine4 = "combine4";
+		static constexpr const char* Constant = "constant";
+		static constexpr const char* Cos = "cos";
+		static constexpr const char* Sin = "sin";
+		static constexpr const char* Dot = "dot";
+		static constexpr const char* Image = "image";
+		static constexpr const char* TiledImage = "tiledimage";
+		static constexpr const char* NormalMap = "normalmap";
+		static constexpr const char* Extract = "extract";
+		static constexpr const char* TexCoord = "texcoord";
+	}
+
+	namespace NodeDefinition
+	{
+		static constexpr const char* StandardSurface = "ND_standard_surface_surfaceshader";
+	}
+
+	namespace Library
+	{
+		static constexpr const char* Std = "stdlib";
+		static constexpr const char* Pbr= "pbrlib";
+		static constexpr const char* Bxdf = "bxdf";
+	}
+}
\ No newline at end of file
diff --git a/Engine/Plugins/Experimental/Interchange/Runtime/Source/Import/Private/MaterialX/InterchangeMaterialXTranslator.cpp b/Engine/Plugins/Experimental/Interchange/Runtime/Source/Import/Private/MaterialX/InterchangeMaterialXTranslator.cpp
new file mode 100644
--- /dev/null
+++ b/Engine/Plugins/Experimental/Interchange/Runtime/Source/Import/Private/MaterialX/InterchangeMaterialXTranslator.cpp
@@ -0,0 +1,601 @@
+#include "MaterialX/InterchangeMaterialXTranslator.h"
+#if WITH_EDITOR
+#include "MaterialXFormat/Util.h"
+#endif
+#include "InterchangeImportLog.h"
+#include "InterchangeManager.h"
+#include "InterchangeMaterialDefinitions.h"
+#include "InterchangeShaderGraphNode.h"
+#include "InterchangeMaterialXDefinitions.h"
+#include "InterchangeTexture2DNode.h"
+#include "UObject/GCObjectScopeGuard.h"
+
+#define LOCTEXT_NAMESPACE "InterchangeMaterialXTranslator"
+
+namespace mx = MaterialX;
+
+UInterchangeMaterialXTranslator::UInterchangeMaterialXTranslator()
+#if WITH_EDITOR
+	:
+	InputNamesMaterialX2UE{
+		{TEXT("in"),        TEXT("Input")},
+		{TEXT("in1"),       TEXT("A")},
+		{TEXT("in2"),       TEXT("B")},
+		{TEXT("fg"),        TEXT("A")},
+		{TEXT("bg"),        TEXT("B")},
+		{TEXT("mix"),       TEXT("Factor")},
+		{TEXT("low"),       TEXT("Min")},
+		{TEXT("high"),      TEXT("Max")},
+		{TEXT("texcoord"),  UE::Interchange::Materials::Standard::Nodes::TextureSample::Inputs::Coordinates.ToString()} },
+		NodeNamesMaterialX2UE
+{
+	{MaterialX::Category::Add,      TEXT("Add")},
+	{MaterialX::Category::Sub,      TEXT("Subtract")},
+	{MaterialX::Category::Multiply, TEXT("Multiply")},
+	{MaterialX::Category::Sin,      TEXT("Sine")},
+	{MaterialX::Category::Cos,      TEXT("Cosine")},
+	{MaterialX::Category::Clamp,    TEXT("Clamp")},
+	{MaterialX::Category::Mix,      TEXT("Lerp")},
+	{MaterialX::Category::Max,      TEXT("Max")},
+	{MaterialX::Category::Min,      TEXT("Min")},
+	{MaterialX::Category::Combine2, TEXT("AppendVector")},
+	{MaterialX::Category::Combine3, TEXT("AppendVector")},
+	{MaterialX::Category::Combine4, TEXT("AppendVector")}
+},
+UEInputs{ TEXT("A"), TEXT("B"), TEXT("Input"), TEXT("Factor"), TEXT("Min"), TEXT("Max") }
+#endif // WITH_EDITOR
+{
+}
+
+TArray<FString> UInterchangeMaterialXTranslator::GetSupportedFormats() const
+{
+	static const bool bStandardSurfacePackageLoaded = []() -> bool
+	{
+		if(FString FunctionPath = FPackageName::ExportTextPathToObjectPath(TEXT("MaterialFunction'/Interchange/Functions/MX_StandardSurface.MX_StandardSurface'"));
+			FPackageName::DoesPackageExist(FunctionPath))
+		{
+			if(FSoftObjectPath(FunctionPath).TryLoad())
+			{
+				return true;
+			}
+			else
+			{
+				UE_LOG(LogInterchangeImport, Warning, TEXT("Couldn't load %s"), *FunctionPath);
+			}
+		}
+		else
+		{
+			UE_LOG(LogInterchangeImport, Warning, TEXT("Couldn't find %s"), *FunctionPath);
+		}
+
+		return false;
+	}();
+
+	return bStandardSurfacePackageLoaded ? TArray<FString>{ TEXT("mtlx;MaterialX File Format") } : TArray<FString>{};
+}
+
+bool UInterchangeMaterialXTranslator::Translate(UInterchangeBaseNodeContainer& BaseNodeContainer) const
+{
+	bool bIsDocumentValid = false;
+	bool bIsReferencesValid = true;
+
+#if WITH_EDITOR
+	FString Filename = GetSourceData()->GetFilename();
+
+	if(!FPaths::FileExists(Filename))
+	{
+		return false;
+	}
+
+	mx::FileSearchPath MaterialXLibFolder{ TCHAR_TO_UTF8(*FPaths::Combine(
+		FPaths::EngineDir(),
+		TEXT("Binaries"),
+		TEXT("ThirdParty"),
+		TEXT("MaterialX"),
+		TEXT("libraries"))) };
+
+	mx::DocumentPtr MaterialXLibrary = mx::createDocument();
+
+	mx::StringSet LoadedLibs = mx::loadLibraries({ mx::Library::Std, mx::Library::Pbr, mx::Library::Bxdf }, MaterialXLibFolder, MaterialXLibrary);
+	if(LoadedLibs.empty())
+	{
+		UInterchangeResultError_Generic* Message = AddMessage<UInterchangeResultError_Generic>();
+		Message->Text = FText::Format(LOCTEXT("MaterialXLibrariesNotFound", "Couldn't load MaterialX libraries from {0}"),
+			FText::FromString(MaterialXLibFolder.asString().c_str()));
+		return false;
+	}
+
+	mx::DocumentPtr Document = mx::createDocument();
+	mx::readFromXmlFile(Document, TCHAR_TO_UTF8(*Filename));
+
+	Document->importLibrary(MaterialXLibrary);
+	bIsDocumentValid = Document->validate();
+
+	for(mx::ElementPtr Elem : Document->traverseTree())
+	{
+		//make sure to read only the current file otherwise we'll process the entire library
+		if(Elem->getActiveSourceUri() != Document->getActiveSourceUri())
+		{
+			continue;
+		}
+
+		mx::NodePtr Node = Elem->asA<mx::Node>();
+
+		//for the moment the entry point is only on surfacematerial node
+		if(Node && Node->getType() == mx::MATERIAL_TYPE_STRING && Node->getCategory() == mx::SURFACE_MATERIAL_NODE_STRING)
+		{
+			if(!Node->getTypeDef())
+			{
+				UInterchangeResultError_Generic* Message = AddMessage<UInterchangeResultError_Generic>();
+				Message->Text = FText::Format(LOCTEXT("TypeDefNotFound", "<{0}> has no matching TypeDef"),
+					FText::FromString(Node->getName().c_str()));
+
+				bIsReferencesValid = false;
+				break;
+			}
+
+			bool bHasStandardSurface = false;
+
+			for(mx::InputPtr Input : Node->getInputs())
+			{
+				//we only support standard_surface for now
+				mx::NodePtr ConnectedNode = Input->getConnectedNode();
+
+				if(ConnectedNode && ConnectedNode->getCategory() == "standard_surface")
+				{
+					ProcessStandardSurface(BaseNodeContainer, ConnectedNode, Document);
+					bHasStandardSurface = true;
+				}
+			}
+
+			if(!bHasStandardSurface)
+			{
+				UInterchangeResultWarning_Generic* Message = AddMessage<UInterchangeResultWarning_Generic>();
+				Message->Text = FText::Format(LOCTEXT("StandardSurfaceNotFound", "<{0}> has no standard_surface inputs"),
+					FText::FromString(Node->getName().c_str()));
+			}
+		}
+	}
+
+#endif // WITH_EDITOR
+
+	return bIsDocumentValid && bIsReferencesValid;
+}
+
+TOptional<UE::Interchange::FImportImage> UInterchangeMaterialXTranslator::GetTexturePayloadData(const UInterchangeSourceData* InSourceData, const FString& PayloadKey) const
+{
+	UInterchangeSourceData* PayloadSourceData = UInterchangeManager::GetInterchangeManager().CreateSourceData(PayloadKey);
+	FGCObjectScopeGuard ScopedSourceData(PayloadSourceData);
+
+	if(!PayloadSourceData)
+	{
+		return TOptional<UE::Interchange::FImportImage>();
+	}
+
+	UInterchangeTranslatorBase* SourceTranslator = UInterchangeManager::GetInterchangeManager().GetTranslatorForSourceData(PayloadSourceData);
+	FGCObjectScopeGuard ScopedSourceTranslator(SourceTranslator);
+	const IInterchangeTexturePayloadInterface* TextureTranslator = Cast<IInterchangeTexturePayloadInterface>(SourceTranslator);
+
+	if(!ensure(TextureTranslator))
+	{
+		return TOptional<UE::Interchange::FImportImage>();
+	}
+
+	return TextureTranslator->GetTexturePayloadData(PayloadSourceData, PayloadKey);
+}
+
+#if WITH_EDITOR
+void UInterchangeMaterialXTranslator::ProcessStandardSurface(UInterchangeBaseNodeContainer& NodeContainer, MaterialX::NodePtr StandardSurfaceNode, MaterialX::DocumentPtr Document) const
+{
+	auto GetInputFloatValue = [](mx::InputPtr Input)
+	{
+		if(Input->hasValueString())
+		{
+			return mx::fromValueString<float>(Input->getValueString());
+		}
+		else
+		{
+			return mx::fromValueString<float>(Input->getDefaultValue()->getValueString());
+		}
+	};
+
+	auto GetInputColor3Value = [](mx::InputPtr Input)
+	{
+		mx::Color3 Color;
+
+		if(Input->hasValueString())
+		{
+			Color = mx::fromValueString<mx::Color3>(Input->getValueString());
+		}
+		else
+		{
+			Color = mx::fromValueString<mx::Color3>(Input->getDefaultValue()->getValueString());
+		}
+
+		return FLinearColor(Color[0], Color[1], Color[2]);
+	};
+
+	using namespace UE::Interchange::Materials;
+
+	FString StandardSurfaceName{ FPaths::GetBaseFilename(StandardSurfaceNode->getActiveSourceUri().c_str()) };
+	UInterchangeShaderGraphNode* ShaderGraphNode = NewObject<UInterchangeShaderGraphNode>(&NodeContainer);
+	FString ShaderGraphNodeUID = TEXT("\\Material\\") + StandardSurfaceName;
+	ShaderGraphNode->InitializeNode(ShaderGraphNodeUID, StandardSurfaceName, EInterchangeNodeContainerType::TranslatedAsset);
+	NodeContainer.AddNode(ShaderGraphNode);
+	ShaderGraphNode->SetCustomShaderType(StandardSurface::Name.ToString());
+
+	TMap<FString, UInterchangeShaderNode*> NamesToShaderNodes;
+	NamesToShaderNodes.Add(StandardSurfaceName, ShaderGraphNode);
+
+	//Base
+	{
+		//Weight
+		mx::InputPtr InputBase = GetStandardSurfaceInput(StandardSurfaceNode, mx::StandardSurface::Input::Base, Document);
+
+		if(!ConnectNodeGraphOutputToInput(InputBase, ShaderGraphNode, StandardSurface::Parameters::Base.ToString(), NamesToShaderNodes, NodeContainer))
+		{
+			ShaderGraphNode->AddFloatAttribute(UInterchangeShaderPortsAPI::MakeInputValueKey(StandardSurface::Parameters::Base.ToString()), GetInputFloatValue(InputBase));
+		}
+
+		//Color
+		mx::InputPtr InputBaseColor = GetStandardSurfaceInput(StandardSurfaceNode, mx::StandardSurface::Input::BaseColor, Document);
+
+		if(!ConnectNodeGraphOutputToInput(InputBaseColor, ShaderGraphNode, StandardSurface::Parameters::BaseColor.ToString(), NamesToShaderNodes, NodeContainer))
+		{
+			ShaderGraphNode->AddLinearColorAttribute(UInterchangeShaderPortsAPI::MakeInputValueKey(StandardSurface::Parameters::BaseColor.ToString()), GetInputColor3Value(InputBaseColor));
+		}
+	}
+
+	//Specular
+	{
+		//Weight
+		mx::InputPtr InputSpecular = GetStandardSurfaceInput(StandardSurfaceNode, mx::StandardSurface::Input::Specular, Document);
+
+		if(!ConnectNodeGraphOutputToInput(InputSpecular, ShaderGraphNode, StandardSurface::Parameters::Specular.ToString(), NamesToShaderNodes, NodeContainer))
+		{
+			ShaderGraphNode->AddFloatAttribute(UInterchangeShaderPortsAPI::MakeInputValueKey(StandardSurface::Parameters::Specular.ToString()), GetInputFloatValue(InputSpecular));
+		}
+
+		//Roughness
+		{
+			mx::InputPtr InputSpecularRoughness = GetStandardSurfaceInput(StandardSurfaceNode, mx::StandardSurface::Input::SpecularRoughness, Document);
+
+			if(!ConnectNodeGraphOutputToInput(InputSpecularRoughness, ShaderGraphNode, StandardSurface::Parameters::SpecularRoughness.ToString(), NamesToShaderNodes, NodeContainer))
+			{
+				ShaderGraphNode->AddFloatAttribute(UInterchangeShaderPortsAPI::MakeInputValueKey(StandardSurface::Parameters::SpecularRoughness.ToString()), GetInputFloatValue(InputSpecularRoughness));
+			}
+		}
+	}
+
+	//Metallic
+	{
+		mx::InputPtr InputMetalness = GetStandardSurfaceInput(StandardSurfaceNode, mx::StandardSurface::Input::Metalness, Document);
+
+		if(!ConnectNodeGraphOutputToInput(InputMetalness, ShaderGraphNode, StandardSurface::Parameters::Metalness.ToString(), NamesToShaderNodes, NodeContainer))
+		{
+			ShaderGraphNode->AddFloatAttribute(UInterchangeShaderPortsAPI::MakeInputValueKey(StandardSurface::Parameters::Metalness.ToString()), GetInputFloatValue(InputMetalness));
+		}
+	}
+
+	//Normal
+	{
+		//No need to take the default input if there is no Normal input
+		mx::InputPtr InputNormal = StandardSurfaceNode->getInput(mx::StandardSurface::Input::Normal);
+
+		if(InputNormal)
+		{
+			ConnectNodeGraphOutputToInput(InputNormal, ShaderGraphNode, StandardSurface::Parameters::Normal.ToString(), NamesToShaderNodes, NodeContainer);
+		}
+	}
+}
+
+bool UInterchangeMaterialXTranslator::ConnectNodeGraphOutputToInput(MaterialX::InputPtr InputToNodeGraph, UInterchangeShaderNode* ShaderNode, const FString& ParentInputName, TMap<FString, UInterchangeShaderNode*>& NamesToShaderNodes, UInterchangeBaseNodeContainer& NodeContainer) const
+{
+	using namespace UE::Interchange::Materials::Standard::Nodes;
+
+	bool bHasNodeGraph = false;
+
+	if(InputToNodeGraph->hasNodeGraphString())
+	{
+		bHasNodeGraph = true;
+
+		mx::OutputPtr Output = InputToNodeGraph->getConnectedOutput();
+
+		if(!Output)
+		{
+			UInterchangeResultWarning_Generic* Message = AddMessage<UInterchangeResultWarning_Generic>();
+			Message->Text = FText::Format(LOCTEXT("OutputNotFound", "Couldn't find a connected output to ({0})"),
+				FText::FromString(InputToNodeGraph->getName().c_str()));
+
+			return false;
+		}
+
+		for(mx::Edge Edge : Output->traverseGraph())
+		{
+			if(mx::NodePtr UpstreamNode = Edge.getUpstreamElement()->asA<mx::Node>())
+			{
+				UInterchangeShaderNode* ParentShaderNode = ShaderNode;
+				FString InputChannelName = ParentInputName;
+
+				//Replace the input's name by the one used in UE
+				RenameInputsNames(UpstreamNode);
+
+				if(mx::NodePtr DownstreamNode = Edge.getDownstreamElement()->asA<mx::Node>())
+				{
+					if(UInterchangeShaderNode** FoundNode = NamesToShaderNodes.Find(DownstreamNode->getName().c_str()))
+					{
+						ParentShaderNode = *FoundNode;
+					}
+
+					if(mx::InputPtr ConnectedInput = Edge.getConnectingElement()->asA<mx::Input>())
+					{
+						InputChannelName = ConnectedInput->getName().c_str();
+					}
+				}
+
+				if(ConnectNodeOutputToInput(UpstreamNode, ParentShaderNode, InputChannelName, NamesToShaderNodes, NodeContainer))
+				{
+					continue;
+				}
+				else if(UpstreamNode->getCategory() == mx::Category::Constant)
+				{
+					if(mx::InputPtr InputConstant = UpstreamNode->getInput("value"); !AddAttribute(InputConstant, InputChannelName, ParentShaderNode))
+					{
+						UInterchangeResultWarning_Generic* Message = AddMessage<UInterchangeResultWarning_Generic>();
+						Message->Text = FText::Format(LOCTEXT("InputTypeNotSupported", "<{0}>: \"{1}\" is not supported yet"),
+							FText::FromString(InputConstant->getName().c_str()),
+							FText::FromString(InputConstant->getType().c_str()));
+					}
+
+				}
+				else if(UpstreamNode->getCategory() == mx::Category::Extract)
+				{
+					UInterchangeShaderNode* MaskShaderNode = CreateShaderNode(UpstreamNode->getName().c_str(), Mask::Name.ToString(), ParentShaderNode, NamesToShaderNodes, NodeContainer);
+
+					if(mx::InputPtr InputIndex = UpstreamNode->getInput("index"))
+					{
+						int32 Index = mx::fromValueString<int>(InputIndex->getValueString());
+						switch(Index)
+						{
+						case 0: MaskShaderNode->AddBooleanAttribute(Mask::Attributes::R, true); break;
+						case 1: MaskShaderNode->AddBooleanAttribute(Mask::Attributes::G, true); break;
+						case 2: MaskShaderNode->AddBooleanAttribute(Mask::Attributes::B, true); break;
+						case 3: MaskShaderNode->AddBooleanAttribute(Mask::Attributes::A, true); break;
+						default:
+							UInterchangeResultError_Generic* Message = AddMessage<UInterchangeResultError_Generic>();
+							Message->Text = FText::FromString(TEXT("Wrong index number for extract node, values are from [0-3]"));
+							break;
+						}
+					}
+
+					UInterchangeShaderPortsAPI::ConnectDefaultOuputToInput(ParentShaderNode, InputChannelName, MaskShaderNode->GetUniqueID());
+				}
+				//dot means identity, input == output
+				else if(UpstreamNode->getCategory() == mx::Category::Dot || UpstreamNode->getCategory() == mx::Category::NormalMap)
+				{
+					if(mx::InputPtr Input = GetInputFromOldName(UpstreamNode, "in"))
+					{
+						Input->setName(TCHAR_TO_UTF8(*InputChannelName)); //let's take the parent node's input name
+						NamesToShaderNodes.FindOrAdd(UpstreamNode->getName().c_str(), ParentShaderNode);
+					}
+				}
+				else if(UpstreamNode->getCategory() == mx::Category::Image || UpstreamNode->getCategory() == mx::Category::TiledImage)
+				{
+					if(UInterchangeTextureNode* TextureNode = CreateTextureNode(UpstreamNode, NodeContainer))
+					{
+						//By default set the output of a texture to RGB, if the type is float then it's either up to an extract node or the Material Input to handle it
+						FString OutputChannel{ TEXT("RGB") };
+
+						if(UpstreamNode->getType() == "vector4" || UpstreamNode->getType() == "color4")
+							OutputChannel = TEXT("RGBA");
+
+						UInterchangeShaderNode* TextureShaderNode = CreateShaderNode(UpstreamNode->getName().c_str(), TextureSample::Name.ToString(), ParentShaderNode, NamesToShaderNodes, NodeContainer);
+						TextureShaderNode->AddStringAttribute(UInterchangeShaderPortsAPI::MakeInputValueKey(TextureSample::Inputs::Texture.ToString()), TextureNode->GetUniqueID());
+						UInterchangeShaderPortsAPI::ConnectOuputToInput(ParentShaderNode, InputChannelName, TextureShaderNode->GetUniqueID(), OutputChannel);
+					}
+					else
+					{
+						AddAttribute(UpstreamNode->getInput("default"), InputChannelName, ParentShaderNode);
+					}
+				}
+				else if(UpstreamNode->getCategory() == mx::Category::TexCoord)
+				{
+					UInterchangeShaderNode* TextureCoordinateNode = CreateShaderNode(UpstreamNode->getName().c_str(), TextureCoordinate::Name.ToString(), ParentShaderNode, NamesToShaderNodes, NodeContainer);
+					UInterchangeShaderPortsAPI::ConnectDefaultOuputToInput(ParentShaderNode, InputChannelName, TextureCoordinateNode->GetUniqueID());
+				}
+				else
+				{
+					UInterchangeResultWarning_Generic* Message = AddMessage<UInterchangeResultWarning_Generic>();
+					Message->Text = FText::Format(LOCTEXT("NodeCategoryNotSupported", "<{0}> is not supported yet"),
+						FText::FromString(UpstreamNode->getCategory().c_str()));
+				}
+			}
+		}
+	}
+
+	return bHasNodeGraph;
+}
+
+bool UInterchangeMaterialXTranslator::ConnectNodeOutputToInput(MaterialX::NodePtr Node, UInterchangeShaderNode* ParentShaderNode, const FString& InputChannelName, TMap<FString, UInterchangeShaderNode*>& NamesToShaderNodes, UInterchangeBaseNodeContainer& NodeContainer) const
+{
+	bool bIsConnected = false;
+
+	if(const FString* ShaderType = NodeNamesMaterialX2UE.Find(Node->getCategory().c_str()))
+	{
+		UInterchangeShaderNode* OperatorNode = nullptr;
+
+		OperatorNode = CreateShaderNode(Node->getName().c_str(), *ShaderType, ParentShaderNode, NamesToShaderNodes, NodeContainer);
+
+		for(mx::InputPtr Input : Node->getInputs())
+		{
+			if(Input->hasValue())
+			{
+				if(const FString* InputNameFound = UEInputs.Find(Input->getName().c_str()))
+				{
+					AddAttribute(Input, *InputNameFound, OperatorNode);
+				}
+			}
+		}
+
+		bIsConnected = UInterchangeShaderPortsAPI::ConnectDefaultOuputToInput(ParentShaderNode, InputChannelName, OperatorNode->GetUniqueID());
+	}
+
+	return bIsConnected;
+}
+
+UInterchangeShaderNode* UInterchangeMaterialXTranslator::CreateShaderNode(const FString& NodeName, const FString& ShaderType, UInterchangeShaderNode* ParentNode, TMap<FString, UInterchangeShaderNode*>& NamesToShaderNodes, UInterchangeBaseNodeContainer& NodeContainer) const
+{
+	UInterchangeShaderNode* Node;
+
+	if(UInterchangeShaderNode** FoundNode = NamesToShaderNodes.Find(NodeName); !FoundNode)
+	{
+		const FString NodeUID = ParentNode->GetUniqueID() + TEXT("_") + NodeName;
+		Node = NewObject<UInterchangeShaderNode>(&NodeContainer);
+		Node->InitializeNode(NodeUID, NodeName, EInterchangeNodeContainerType::TranslatedAsset);
+		NodeContainer.AddNode(Node);
+		NodeContainer.SetNodeParentUid(NodeUID, ParentNode->GetUniqueID());
+		Node->SetCustomShaderType(ShaderType);
+
+		NamesToShaderNodes.Add(NodeName, Node);
+	}
+	else
+	{
+		Node = *FoundNode;
+	}
+
+	return Node;
+}
+
+UInterchangeTextureNode* UInterchangeMaterialXTranslator::CreateTextureNode(MaterialX::NodePtr Node, UInterchangeBaseNodeContainer& NodeContainer) const
+{
+	UInterchangeTextureNode* TextureNode = nullptr;
+
+	//A node image should have an input file otherwise the user should check its default value
+	if(Node)
+	{
+		if(mx::InputPtr InputFile = Node->getInput("file"); InputFile && InputFile->hasValue())
+		{
+			FString Filepath{ InputFile->getValueString().c_str() };
+			FString Filename = FPaths::GetCleanFilename(Filepath);
+			FString TextureNodeUID = TEXT("\\Texture\\") + Filename;
+
+			//Only add the TextureNode once
+			TextureNode = const_cast<UInterchangeTextureNode *>(Cast<UInterchangeTextureNode>(NodeContainer.GetNode(TextureNodeUID)));
+			if(TextureNode == nullptr)
+			{
+				TextureNode = NewObject<UInterchangeTexture2DNode>(&NodeContainer);
+				TextureNode->InitializeNode(TextureNodeUID, Filename, EInterchangeNodeContainerType::TranslatedAsset);
+				NodeContainer.AddNode(TextureNode);
+
+				if(FPaths::IsRelative(Filepath))
+				{
+					Filepath = FPaths::ConvertRelativePathToFull(FPaths::GetPath(Node->getActiveSourceUri().c_str()), Filepath);
+				}
+
+				TextureNode->SetPayLoadKey(Filepath);
+			}
+		}
+	}
+
+	return TextureNode;
+}
+
+const FString& UInterchangeMaterialXTranslator::GetMatchedInputName(MaterialX::InputPtr Input) const
+{
+	static FString Empty{ "" };
+
+	if(Input)
+	{
+		if(const FString* Result = InputNamesMaterialX2UE.Find(Input->getName().c_str()))
+		{
+			return *Result;
+		}
+	}
+
+	return Empty;
+}
+
+void UInterchangeMaterialXTranslator::RenameInputsNames(MaterialX::NodePtr Node) const
+{
+	if(Node)
+	{
+		if(const std::string& IsVisited = Node->getAttribute(mx::Attributes::IsVisited); IsVisited.empty())
+		{
+			Node->setAttribute(mx::Attributes::IsVisited, "true");
+
+			for(mx::InputPtr Input : Node->getInputs())
+			{
+				if(const FString& Name = GetMatchedInputName(Input); !Name.IsEmpty())
+				{
+					Input->setAttribute(mx::Attributes::OldName, Input->getName()); // keep the old name for further processing
+					Input->setName(TCHAR_TO_UTF8(*Name));
+				}
+			}
+		}
+	}
+}
+
+MaterialX::InputPtr UInterchangeMaterialXTranslator::GetInputFromOldName(MaterialX::NodePtr Node, const char* OldNameAttribute) const
+{
+	mx::InputPtr InputRes{ nullptr };
+
+	for(mx::InputPtr Input : Node->getInputs())
+	{
+		if(Input->getAttribute(mx::Attributes::OldName) == OldNameAttribute)
+		{
+			InputRes = Input;
+		}
+	}
+
+	return InputRes;
+}
+
+MaterialX::InputPtr UInterchangeMaterialXTranslator::GetStandardSurfaceInput(MaterialX::NodePtr StandardSurface, const char* InputName, MaterialX::DocumentPtr Document) const
+{
+	mx::InputPtr Input = StandardSurface->getInput(InputName);
+
+	if(!Input)
+	{
+		Input = Document->getNodeDef(mx::NodeDefinition::StandardSurface)->getInput(InputName);
+	}
+
+	return Input;
+}
+
+bool UInterchangeMaterialXTranslator::AddAttribute(MaterialX::InputPtr Input, const FString& InputChannelName, UInterchangeShaderNode* ShaderNode) const
+{
+	if(Input)
+	{
+		if(Input->getType() == "float")
+		{
+			return ShaderNode->AddFloatAttribute(UInterchangeShaderPortsAPI::MakeInputValueKey(InputChannelName), mx::fromValueString<float>(Input->getValueString()));
+		}
+
+		FLinearColor LinearColor;
+
+		bool bIsColor = false;
+
+		if(bool bIsColor3 = (Input->getType() == "color3" || Input->getType() == "vector3"))
+		{
+			mx::Color3 Color = mx::fromValueString<mx::Color3>(Input->getValueString());
+			LinearColor = FLinearColor{ Color[0], Color[1], Color[2] };
+			bIsColor = true;
+		}
+		else if(bool bIsColor4 = (Input->getType() == "color4" || Input->getType() == "vector4"))
+		{
+			mx::Color4 Color = mx::fromValueString<mx::Color4>(Input->getValueString());
+			LinearColor = FLinearColor{ Color[0], Color[1], Color[2], Color[3] };
+			bIsColor = true;
+		}
+
+		if(bIsColor)
+		{
+			return ShaderNode->AddLinearColorAttribute(UInterchangeShaderPortsAPI::MakeInputValueKey(InputChannelName), LinearColor);
+		}
+	}
+
+	return false;
+}
+#endif //WITH_EDITOR
+
+#undef LOCTEXT_NAMESPACE
\ No newline at end of file
diff --git a/Engine/Plugins/Experimental/Interchange/Runtime/Source/Import/Public/MaterialX/InterchangeMaterialXTranslator.h b/Engine/Plugins/Experimental/Interchange/Runtime/Source/Import/Public/MaterialX/InterchangeMaterialXTranslator.h
new file mode 100644
--- /dev/null
+++ b/Engine/Plugins/Experimental/Interchange/Runtime/Source/Import/Public/MaterialX/InterchangeMaterialXTranslator.h
@@ -0,0 +1,161 @@
+// Copyright Epic Games, Inc. All Rights Reserved.
+#pragma once
+
+#include "InterchangeTranslatorBase.h"
+#include "Texture/InterchangeTexturePayloadInterface.h"
+
+#if WITH_EDITOR
+#if defined(PRAGMA_DISABLE_MISSING_BRACES_WARNINGS)
+PRAGMA_DISABLE_MISSING_BRACES_WARNINGS
+#endif
+#include "MaterialXCore/Document.h"
+#if defined(PRAGMA_ENABLE_MISSING_BRACES_WARNINGS)
+PRAGMA_ENABLE_MISSING_BRACES_WARNINGS
+#endif
+#endif
+
+#include "InterchangeMaterialXTranslator.generated.h"
+
+class UInterchangeTextureNode;
+class UInterchangeShaderNode;
+
+UCLASS(BlueprintType, Experimental)
+class UInterchangeMaterialXTranslator : public UInterchangeTranslatorBase, public IInterchangeTexturePayloadInterface
+{
+	GENERATED_BODY()
+
+public:
+
+	UInterchangeMaterialXTranslator();
+
+	/** Begin UInterchangeTranslatorBase API*/
+	virtual TArray<FString> GetSupportedFormats() const override;
+
+	/**
+	 * Translate the associated source data into a node hold by the specified nodes container.
+	 *
+	 * @param BaseNodeContainer - The unreal objects descriptions container where to put the translated source data.
+	 * @return true if the translator can translate the source data, false otherwise.
+	 */
+	virtual bool Translate( UInterchangeBaseNodeContainer& BaseNodeContainer ) const override;
+	/** End UInterchangeTranslatorBase API*/
+
+	virtual TOptional<UE::Interchange::FImportImage> GetTexturePayloadData(const UInterchangeSourceData* PayloadSourceData, const FString& PayloadKey) const override;
+
+#if WITH_EDITOR
+protected:
+
+	/**
+	 * Process Autodesk's standard surface shader
+	 * 
+	 * @param NodeContainer - The unreal objects descriptions container where to put the translated source data.
+	 * @param StandardSurfaceNode - The <standard_surface> in the MaterialX file
+	 */
+	void ProcessStandardSurface(UInterchangeBaseNodeContainer & NodeContainer, MaterialX::NodePtr StandardSurfaceNode, MaterialX::DocumentPtr Document) const;
+
+	/**
+	 * Connect an ouput in the NodeGraph to the ShaderGraph
+	 * 
+	 * @param InputToNodeGraph - The input from the standard surface to retrieve the output in the NodeGraph
+	 * @param ShaderNode - The Interchange shader node to connect the MaterialX's node graph to
+	 * @param ParentInputName - The name of the input of the shader node to which we want the node graph to be connected to
+	 * @param NamesToShaderNodes - Map of the shader nodes already created
+	 * @param NodeContainer - The unreal objects descriptions container where to put the translated source data.
+	 * 
+	 * @return true if the given input is attached to one of the outputs of a node graph
+	 */
+	bool ConnectNodeGraphOutputToInput(MaterialX::InputPtr InputToNodeGraph, UInterchangeShaderNode * ShaderNode, const FString& ParentInputName, TMap<FString, UInterchangeShaderNode*> & NamesToShaderNodes, UInterchangeBaseNodeContainer & NodeContainer) const;
+
+	/**
+	 * Create and Connect the output of a node to a shader node
+	 * 
+	 * @param Node - The MaterialX node of a given type used to create the appropriate shader node
+	 * @param ParentShaderNode - The shader node to connect to
+	 * @param InputChannelName - The input of the ParentShaderNode to connect to
+	 * @param NamesToShaderNodes - Map of the shader nodes already created
+	 * @param NodeContainer - The unreal objects descriptions container where to put the translated source data.
+	 * 
+	 * @return true if a shader node has been successfully created and is connected to the given input
+	 */
+	bool ConnectNodeOutputToInput(MaterialX::NodePtr Node, UInterchangeShaderNode* ParentShaderNode, const FString& InputChannelName, TMap<FString, UInterchangeShaderNode*>& NamesToShaderNodes, UInterchangeBaseNodeContainer& NodeContainer) const;
+
+	/**
+	 * Helper function to create an InterchangeShaderNode
+	 * 
+	 * @param NodeName - The name of the shader node
+	 * @param ShaderType - The shader node's type we want to create
+	 * @param ParentNode - The parent node of the created node
+	 * @param NamesToShaderNodes - Map of the shader nodes already created
+	 * @param NodeContainer - The unreal objects descriptions container where to put the translated source data.
+	 * 
+	 * @return The shader node that was created
+	 */
+	UInterchangeShaderNode* CreateShaderNode(const FString & NodeName, const FString & ShaderType, UInterchangeShaderNode * ParentNode, TMap<FString, UInterchangeShaderNode*>& NamesToShaderNodes, UInterchangeBaseNodeContainer & NodeContainer) const;
+
+	/**
+	 * Helper function to create an InterchangeTextureNode
+	 *
+	 * @param Node - The MaterialX node, it should be of the category <image> no test is done on it
+	 * @param NodeContainer - The unreal objects descriptions container where to put the translated source data.
+	 * 
+	 * @return The texture node that was created
+	 */
+	UInterchangeTextureNode* CreateTextureNode(MaterialX::NodePtr Node, UInterchangeBaseNodeContainer& NodeContainer) const;
+
+	/**
+	 * Get the UE corresponding name of MaterialX input of a material
+	 * 
+	 * @param Input - MaterialX input
+	 * 
+	 * @return The matched name of the input else empty string
+	 */
+	const FString& GetMatchedInputName(MaterialX::InputPtr Input) const;
+
+	/**
+	 * Rename the inputs names of a node to correspond to the one used by UE, it will keep the old inputs names under the attribute "oldname"
+	 * 
+	 * @param Node - Look up to all inputs of Node and rename them to match UE names
+	 */
+	void RenameInputsNames(MaterialX::NodePtr Node) const;
+
+	/**
+	 * Helper function to retrieve an Input in a Node from its old name (after a renaming)
+	 * 
+	 * @param Node - The node with the Inputs to look up
+	 * @param OldNameAttribute - The previous name of an Input
+	 * 
+	 * @return The found Input, nullptr otherwise
+	 */
+	MaterialX::InputPtr GetInputFromOldName(MaterialX::NodePtr Node, const char * OldNameAttribute) const;
+
+	/**
+	 * Retrieve the input from a standard_surface node, or take the default input from the library, 
+	 * this function should only be called after testing the MaterialX libraries have been successfully imported, meaning the node definition of the standard_surface 
+	 * should always be valid
+	 * 
+	 * @param StandardSurface - the <standard_surface> node
+	 * @param InputName - the input name to retrieve
+	 * @param Document - the MaterialX library which has the node definition of the standard_surface
+	 * 
+	 * @return the input from the given name
+	 */
+	MaterialX::InputPtr GetStandardSurfaceInput(MaterialX::NodePtr StandardSurface, const char* InputName, MaterialX::DocumentPtr Document) const;
+
+	/**
+	 * Add an attribute to a shader node, only floats and linear colors are supported for the moment
+	 * 
+	 * @param Input - The MaterialX input to retrieve and add the value from, must be of type float/color/vector
+	 * @param InputChannelName - The name of the shader node's input to add the attribute
+	 * @param ShaderNode - The shader node to which we want to add the attribute
+	 * 
+	 * @return true if the attribute was successfully added
+	 */
+	bool AddAttribute(MaterialX::InputPtr Input, const FString& InputChannelName, UInterchangeShaderNode* ShaderNode) const;
+
+private:
+
+	TMap<FString, FString> InputNamesMaterialX2UE;
+	TMap<FString, FString> NodeNamesMaterialX2UE; //given a MaterialX node category, return the UE category
+	TSet<FString> UEInputs;
+#endif // WITH_EDITOR
+};
\ No newline at end of file
