Add ability for translators to force the creation of unused materials. To do so:
 - Indicate uniqueness of  UInterchangeSourceNode in the node container by providing two static functions to find/create and get.
 - Added methods  GetCustomImportUnusedMaterial and SetCustomImportUnusedMaterial on UInterchangeSourceNode to check ability or not to import unused materials.
 - Modified UInterchangeGenericMaterialPipeline to enable or not material factory nodes based on the value returned by GetCustomImportUnusedMaterial on the UInterchangeSourceNode when added to the node container.

diff --git a/Engine/Plugins/Experimental/Interchange/Runtime/Source/Import/Private/MaterialX/InterchangeMaterialXTranslator.cpp b/Engine/Plugins/Experimental/Interchange/Runtime/Source/Import/Private/MaterialX/InterchangeMaterialXTranslator.cpp
--- a/Engine/Plugins/Experimental/Interchange/Runtime/Source/Import/Private/MaterialX/InterchangeMaterialXTranslator.cpp
+++ b/Engine/Plugins/Experimental/Interchange/Runtime/Source/Import/Private/MaterialX/InterchangeMaterialXTranslator.cpp
@@ -10,12 +10,43 @@
 #include "InterchangeShaderGraphNode.h"
 #include "InterchangeMaterialXDefinitions.h"
 #include "InterchangeTexture2DNode.h"
+#include "Nodes/InterchangeSourceNode.h"
 #include "UObject/GCObjectScopeGuard.h"
 
 #define LOCTEXT_NAMESPACE "InterchangeMaterialXTranslator"
 
 namespace mx = MaterialX;
 
+namespace UE::Interchange::MaterialX
+{
+	bool IsStandardSurfacePackageLoaded()
+	{
+		static const bool bStandardSurfacePackageLoaded = []() -> bool
+		{
+			if (FString FunctionPath = FPackageName::ExportTextPathToObjectPath(TEXT("MaterialFunction'/Interchange/Functions/MX_StandardSurface.MX_StandardSurface'"));
+				FPackageName::DoesPackageExist(FunctionPath))
+			{
+				if (FSoftObjectPath(FunctionPath).TryLoad())
+				{
+					return true;
+				}
+				else
+				{
+					UE_LOG(LogInterchangeImport, Warning, TEXT("Couldn't load %s"), *FunctionPath);
+				}
+			}
+			else
+			{
+				UE_LOG(LogInterchangeImport, Warning, TEXT("Couldn't find %s"), *FunctionPath);
+			}
+
+			return false;
+		}();
+
+		return bStandardSurfacePackageLoaded;
+	}
+}
+
 UInterchangeMaterialXTranslator::UInterchangeMaterialXTranslator()
 #if WITH_EDITOR
 	:
@@ -51,29 +82,14 @@ UEInputs{ TEXT("A"), TEXT("B"), TEXT("Input"), TEXT("Factor"), TEXT("Min"), TEXT
 
 TArray<FString> UInterchangeMaterialXTranslator::GetSupportedFormats() const
 {
-	static const bool bStandardSurfacePackageLoaded = []() -> bool
-	{
-		if(FString FunctionPath = FPackageName::ExportTextPathToObjectPath(TEXT("MaterialFunction'/Interchange/Functions/MX_StandardSurface.MX_StandardSurface'"));
-			FPackageName::DoesPackageExist(FunctionPath))
+	// Call to UInterchangeMaterialXTranslator::GetSupportedFormats is not supported out of game thread
+	// A more global solution must be found for translators which require some initialization
+	if (!IsInGameThread())
 		{
-			if(FSoftObjectPath(FunctionPath).TryLoad())
-			{
-				return true;
+		return TArray<FString>{};
 			}
-			else
-			{
-				UE_LOG(LogInterchangeImport, Warning, TEXT("Couldn't load %s"), *FunctionPath);
-			}
-		}
-		else
-		{
-			UE_LOG(LogInterchangeImport, Warning, TEXT("Couldn't find %s"), *FunctionPath);
-		}
-
-		return false;
-	}();
 
-	return bStandardSurfacePackageLoaded ? TArray<FString>{ TEXT("mtlx;MaterialX File Format") } : TArray<FString>{};
+	return UE::Interchange::MaterialX::IsStandardSurfacePackageLoaded() ? TArray<FString>{ TEXT("mtlx;MaterialX File Format") } : TArray<FString>{};
 }
 
 bool UInterchangeMaterialXTranslator::Translate(UInterchangeBaseNodeContainer& BaseNodeContainer) const
@@ -161,6 +177,12 @@ bool UInterchangeMaterialXTranslator::Translate(UInterchangeBaseNodeContainer& B
 
 #endif // WITH_EDITOR
 
+	if (bIsDocumentValid && bIsReferencesValid)
+	{
+		UInterchangeSourceNode* SourceNode = UInterchangeSourceNode::FindOrCreateUniqueInstance(&BaseNodeContainer);
+		SourceNode->SetCustomImportUnusedMaterial(true);
+	}
+
 	return bIsDocumentValid && bIsReferencesValid;
 }
 
diff --git a/Engine/Plugins/Experimental/Interchange/Runtime/Source/Parsers/Fbx/Private/FbxScene.cpp b/Engine/Plugins/Experimental/Interchange/Runtime/Source/Parsers/Fbx/Private/FbxScene.cpp
--- a/Engine/Plugins/Experimental/Interchange/Runtime/Source/Parsers/Fbx/Private/FbxScene.cpp
+++ b/Engine/Plugins/Experimental/Interchange/Runtime/Source/Parsers/Fbx/Private/FbxScene.cpp
@@ -472,16 +472,13 @@ namespace UE
 				 FbxNode* RootNode = SDKScene->GetRootNode();
 				 //Create a source node where we can store any general file info
 				 {
-					 FString NodeName = FFbxHelper::GetFbxObjectName(RootNode) + TEXT("_SourceInfo");
-					 FString NodeUniqueID = TEXT("\\SourceInfo\\") + FFbxHelper::GetFbxNodeHierarchyName(RootNode);
-					 UInterchangeSourceNode* SourceNode = NewObject<UInterchangeSourceNode>(&NodeContainer, NAME_None);
+					 UInterchangeSourceNode* SourceNode = UInterchangeSourceNode::FindOrCreateUniqueInstance(&NodeContainer);
 					 if (!ensure(SourceNode))
 					 {
 						 UInterchangeResultError_Generic* Message = Parser.AddMessage<UInterchangeResultError_Generic>();
 						 Message->Text = LOCTEXT("NodeAllocationError", "Unable to allocate a node when importing FBX.");
 						 return;
 					 }
-					 SourceNode->InitializeNode(NodeUniqueID, NodeName, EInterchangeNodeContainerType::TranslatedAsset);
 
 					 //Store the fbx frame rate
 					 {
@@ -515,8 +512,6 @@ namespace UE
 						 SourceNode->SetCustomAnimatedTimeStart(AnimatedTimeSpan.GetStart().GetSecondDouble());
 						 SourceNode->SetCustomAnimatedTimeEnd(AnimatedTimeSpan.GetStop().GetSecondDouble());
 					 }
-
-					 NodeContainer.AddNode(SourceNode);
 				 }
 				 AddHierarchyRecursively(nullptr, RootNode, SDKScene, NodeContainer, PayloadContexts);
 			}
diff --git a/Engine/Plugins/Experimental/Interchange/Runtime/Source/Pipelines/Private/InterchangeGenericAnimationPipeline.cpp b/Engine/Plugins/Experimental/Interchange/Runtime/Source/Pipelines/Private/InterchangeGenericAnimationPipeline.cpp
--- a/Engine/Plugins/Experimental/Interchange/Runtime/Source/Pipelines/Private/InterchangeGenericAnimationPipeline.cpp
+++ b/Engine/Plugins/Experimental/Interchange/Runtime/Source/Pipelines/Private/InterchangeGenericAnimationPipeline.cpp
@@ -89,53 +89,49 @@ void UInterchangeGenericAnimationPipeline::ExecutePreImportPipeline(UInterchange
 	double RangeStart = 0;
 	double RangeStop = 0;
 	bool bRangeIsValid = false;
-	TArray<FString> SourceNodeUids;
-	BaseNodeContainer->GetNodes(UInterchangeSourceNode::StaticClass(), SourceNodeUids);
-	if (SourceNodeUids.Num() > 0)
+	const UInterchangeSourceNode* SourceNode = UInterchangeSourceNode::GetUniqueInstance(BaseNodeContainer);
+	if (SourceNode)
 	{
-		if (const UInterchangeSourceNode* SourceNode = Cast<UInterchangeSourceNode>(BaseNodeContainer->GetNode(SourceNodeUids[0])))
+		if (bImportBoneTracks)
 		{
-			if (bImportBoneTracks)
+			int32 Numerator, Denominator;
+			if (!bUse30HzToBakeBoneAnimation && CustomBoneAnimationSampleRate == 0 && SourceNode->GetCustomSourceFrameRateNumerator(Numerator))
 			{
-				int32 Numerator, Denominator;
-				if (!bUse30HzToBakeBoneAnimation && CustomBoneAnimationSampleRate == 0 && SourceNode->GetCustomSourceFrameRateNumerator(Numerator))
-				{
-					if (SourceNode->GetCustomSourceFrameRateDenominator(Denominator) && Denominator > 0 && Numerator > 0)
-					{
-						SampleRate = static_cast<double>(Numerator) / static_cast<double>(Denominator);
-					}
-				}
-				else if ((!bUse30HzToBakeBoneAnimation && CustomBoneAnimationSampleRate > 0))
+				if (SourceNode->GetCustomSourceFrameRateDenominator(Denominator) && Denominator > 0 && Numerator > 0)
 				{
-					SampleRate = static_cast<double>(CustomBoneAnimationSampleRate);
+					SampleRate = static_cast<double>(Numerator) / static_cast<double>(Denominator);
 				}
-			
-				if (AnimationRange == EInterchangeAnimationRange::Timeline)
+			}
+			else if ((!bUse30HzToBakeBoneAnimation && CustomBoneAnimationSampleRate > 0))
+			{
+				SampleRate = static_cast<double>(CustomBoneAnimationSampleRate);
+			}
+
+			if (AnimationRange == EInterchangeAnimationRange::Timeline)
+			{
+				if (SourceNode->GetCustomSourceTimelineStart(RangeStart))
 				{
-					if (SourceNode->GetCustomSourceTimelineStart(RangeStart))
+					if (SourceNode->GetCustomSourceTimelineEnd(RangeStop))
 					{
-						if (SourceNode->GetCustomSourceTimelineEnd(RangeStop))
-						{
-							bRangeIsValid = true;
-						}
+						bRangeIsValid = true;
 					}
 				}
-				else if (AnimationRange == EInterchangeAnimationRange::Animated)
+			}
+			else if (AnimationRange == EInterchangeAnimationRange::Animated)
+			{
+				if (SourceNode->GetCustomAnimatedTimeStart(RangeStart))
 				{
-					if (SourceNode->GetCustomAnimatedTimeStart(RangeStart))
+					if (SourceNode->GetCustomAnimatedTimeEnd(RangeStop))
 					{
-						if (SourceNode->GetCustomAnimatedTimeEnd(RangeStop))
-						{
-							bRangeIsValid = true;
-						}
+						bRangeIsValid = true;
 					}
 				}
-				else if (AnimationRange == EInterchangeAnimationRange::SetRange)
-				{
-					RangeStart = static_cast<double>(FrameImportRange.Min) / SampleRate;
-					RangeStop = static_cast<double>(FrameImportRange.Max) / SampleRate;
-					bRangeIsValid = true;
-				}
+			}
+			else if (AnimationRange == EInterchangeAnimationRange::SetRange)
+			{
+				RangeStart = static_cast<double>(FrameImportRange.Min) / SampleRate;
+				RangeStop = static_cast<double>(FrameImportRange.Max) / SampleRate;
+				bRangeIsValid = true;
 			}
 		}
 	}
diff --git a/Engine/Plugins/Experimental/Interchange/Runtime/Source/Pipelines/Private/InterchangeGenericMaterialPipeline.cpp b/Engine/Plugins/Experimental/Interchange/Runtime/Source/Pipelines/Private/InterchangeGenericMaterialPipeline.cpp
--- a/Engine/Plugins/Experimental/Interchange/Runtime/Source/Pipelines/Private/InterchangeGenericMaterialPipeline.cpp
+++ b/Engine/Plugins/Experimental/Interchange/Runtime/Source/Pipelines/Private/InterchangeGenericMaterialPipeline.cpp
@@ -36,6 +36,7 @@
 #include "Misc/Paths.h"
 #include "Nodes/InterchangeBaseNode.h"
 #include "Nodes/InterchangeBaseNodeContainer.h"
+#include "Nodes/InterchangeSourceNode.h"
 #include "Templates/Function.h"
 #include "Templates/SubclassOf.h"
 #include "UObject/Object.h"
@@ -85,14 +86,21 @@ void UInterchangeGenericMaterialPipeline::ExecutePreImportPipeline(UInterchangeB
 		}
 	});
 
+	// Check to see whether materials should be created even if unused
+	// By default we do not create the materials, every node with mesh attribute can enable them. So we won't create unused materials.
+	bool bImportUnusedMaterial = false;
+	if (const UInterchangeSourceNode* SourceNode = UInterchangeSourceNode::GetUniqueInstance(BaseNodeContainer))
+	{
+		SourceNode->GetCustomImportUnusedMaterial(bImportUnusedMaterial);
+	}
+
 	if (MaterialImport == EInterchangeMaterialImportOption::ImportAsMaterials)
 	{
 		for (const UInterchangeShaderGraphNode* ShaderGraphNode : MaterialNodes)
 		{
 			if (UInterchangeMaterialFactoryNode* MaterialFactoryNode = CreateMaterialFactoryNode(ShaderGraphNode))
 			{
-				//By default we do not create the materials, every node with mesh attribute can enable them. So we wont create unused materials.
-				MaterialFactoryNode->SetEnabled(false);
+				MaterialFactoryNode->SetEnabled(bImportUnusedMaterial);
 			}
 		}
 	}
@@ -102,8 +110,7 @@ void UInterchangeGenericMaterialPipeline::ExecutePreImportPipeline(UInterchangeB
 		{
 			if (UInterchangeMaterialInstanceFactoryNode* MaterialInstanceFactoryNode = CreateMaterialInstanceFactoryNode(ShaderGraphNode))
 			{
-				//By default we do not create the materials, every node with mesh attribute can enable them. So we wont create unused materials.
-				MaterialInstanceFactoryNode->SetEnabled(false);
+				MaterialInstanceFactoryNode->SetEnabled(bImportUnusedMaterial);
 			}
 		}
 	}
diff --git a/Engine/Source/Runtime/Experimental/Interchange/Core/Private/Nodes/InterchangeSourceNode.cpp b/Engine/Source/Runtime/Experimental/Interchange/Core/Private/Nodes/InterchangeSourceNode.cpp
--- a/Engine/Source/Runtime/Experimental/Interchange/Core/Private/Nodes/InterchangeSourceNode.cpp
+++ b/Engine/Source/Runtime/Experimental/Interchange/Core/Private/Nodes/InterchangeSourceNode.cpp
@@ -1,12 +1,24 @@
 // Copyright Epic Games, Inc. All Rights Reserved.
 #include "Nodes/InterchangeSourceNode.h"
 
+#include "Nodes/InterchangeBaseNodeContainer.h"
+
 #include "CoreMinimal.h"
 #include "Types/AttributeStorage.h"
 #include "UObject/Class.h"
 #include "UObject/Object.h"
 #include "UObject/ObjectMacros.h"
 
+
+namespace UE::Interchange::SourceNode
+{
+	FString GetSourceNodeUniqueID()
+	{
+		static FString StaticUid = TEXT("__SourceNode__");
+		return StaticUid;
+	}
+}
+
 UInterchangeSourceNode::UInterchangeSourceNode()
 {
 }
@@ -22,6 +34,27 @@ FString UInterchangeSourceNode::GetTypeName() const
 	return TypeName;
 }
 
+
+UInterchangeSourceNode* UInterchangeSourceNode::FindOrCreateUniqueInstance(UInterchangeBaseNodeContainer* NodeContainer)
+{
+	const FString StaticUid = UE::Interchange::SourceNode::GetSourceNodeUniqueID();
+	UInterchangeSourceNode* SourceNode = Cast<UInterchangeSourceNode>(const_cast<UInterchangeBaseNode*>(NodeContainer->GetNode(StaticUid)));
+	if (!SourceNode)
+	{
+		SourceNode = NewObject<UInterchangeSourceNode>(NodeContainer, NAME_None);
+		SourceNode->InitializeNode(StaticUid, StaticUid, EInterchangeNodeContainerType::FactoryData);
+		NodeContainer->AddNode(SourceNode);
+	}
+
+	return SourceNode;
+}
+
+const UInterchangeSourceNode* UInterchangeSourceNode::GetUniqueInstance(const UInterchangeBaseNodeContainer* NodeContainer)
+{
+	static FString StaticUid = UE::Interchange::SourceNode::GetSourceNodeUniqueID();
+	return Cast<const UInterchangeSourceNode>(NodeContainer->GetNode(StaticUid));
+}
+
 bool UInterchangeSourceNode::GetCustomSourceFrameRateNumerator(int32& AttributeValue) const
 {
 	IMPLEMENT_NODE_ATTRIBUTE_GETTER(SourceFrameRateNumerator, int32);
@@ -80,4 +113,14 @@ bool UInterchangeSourceNode::GetCustomAnimatedTimeEnd(double& AttributeValue) co
 bool UInterchangeSourceNode::SetCustomAnimatedTimeEnd(const double& AttributeValue)
 {
 	IMPLEMENT_NODE_ATTRIBUTE_SETTER_NODELEGATE(AnimatedTimeEnd, double);
-}
\ No newline at end of file
+}
+
+bool UInterchangeSourceNode::GetCustomImportUnusedMaterial(bool& AttributeValue) const
+{
+	IMPLEMENT_NODE_ATTRIBUTE_GETTER(ImportUnusedMaterial, bool);
+}
+
+bool UInterchangeSourceNode::SetCustomImportUnusedMaterial(const bool& AttributeValue)
+{
+	IMPLEMENT_NODE_ATTRIBUTE_SETTER_NODELEGATE(ImportUnusedMaterial, bool);
+}
diff --git a/Engine/Source/Runtime/Experimental/Interchange/Core/Public/Nodes/InterchangeSourceNode.h b/Engine/Source/Runtime/Experimental/Interchange/Core/Public/Nodes/InterchangeSourceNode.h
--- a/Engine/Source/Runtime/Experimental/Interchange/Core/Public/Nodes/InterchangeSourceNode.h
+++ b/Engine/Source/Runtime/Experimental/Interchange/Core/Public/Nodes/InterchangeSourceNode.h
@@ -12,7 +12,7 @@
 
 #include "InterchangeSourceNode.generated.h"
 
-
+class UInterchangeBaseNodeContainer;
 
 /**
  * This class allow a translator to add general source data that describe the whole source. Pipeline can use this information.
@@ -22,9 +22,9 @@ class INTERCHANGECORE_API UInterchangeSourceNode : public UInterchangeBaseNode
 {
 	GENERATED_BODY()
 
-public:
 	UInterchangeSourceNode();
 
+public:
 	/**
 	 * Initialize the base data of the node
 	 * @param UniqueID - The uniqueId for this node
@@ -39,6 +39,12 @@ public:
 	 */
 	virtual FString GetTypeName() const override;
 
+	/* The translators that want to modify the common data should ensure they create the unique common pipeline node. */
+	static UInterchangeSourceNode* FindOrCreateUniqueInstance(UInterchangeBaseNodeContainer* NodeContainer);
+
+	/* If the unique instance doesn't exist it will return nullptr. This function should be use by the pipelines, to avoid creating a node. */
+	static const UInterchangeSourceNode* GetUniqueInstance(const UInterchangeBaseNodeContainer* NodeContainer);
+
 	/** Query the source frame rate numerator. */
 	UFUNCTION(BlueprintCallable, Category = "Interchange | Node | Source")
 	bool GetCustomSourceFrameRateNumerator(int32& AttributeValue) const;
@@ -87,6 +93,14 @@ public:
 	UFUNCTION(BlueprintCallable, Category = "Interchange | Node | Source")
 	bool SetCustomAnimatedTimeEnd(const double& AttributeValue);
 
+	/** Query the source animated time end. */
+	UFUNCTION(BlueprintCallable, Category = "Interchange | Node | Source")
+	bool GetCustomImportUnusedMaterial(bool& AttributeValue) const;
+
+	/** Store the source animated time end. */
+	UFUNCTION(BlueprintCallable, Category = "Interchange | Node | Source")
+	bool SetCustomImportUnusedMaterial(const bool& AttributeValue);
+
 private:
 	const UE::Interchange::FAttributeKey Macro_CustomSourceFrameRateNumeratorKey = UE::Interchange::FAttributeKey(TEXT("SourceFrameRateNumerator"));
 	const UE::Interchange::FAttributeKey Macro_CustomSourceFrameRateDenominatorKey = UE::Interchange::FAttributeKey(TEXT("SourceFrameRateDenominator"));
@@ -94,5 +108,6 @@ private:
 	const UE::Interchange::FAttributeKey Macro_CustomSourceTimelineEndKey = UE::Interchange::FAttributeKey(TEXT("SourceTimelineEnd"));
 	const UE::Interchange::FAttributeKey Macro_CustomAnimatedTimeStartKey = UE::Interchange::FAttributeKey(TEXT("AnimatedTimeStart"));
 	const UE::Interchange::FAttributeKey Macro_CustomAnimatedTimeEndKey = UE::Interchange::FAttributeKey(TEXT("AnimatedTimeEnd"));
+	const UE::Interchange::FAttributeKey Macro_CustomImportUnusedMaterialKey = UE::Interchange::FAttributeKey(TEXT("ImportUnusedMaterial"));
 };
 
