Interchange: Add rigid mesh support, we need to translate animation on any scene node. Then the pipeline can decide to create or not a rigid skeletal mesh.
#rnx

diff --git a/Engine/Plugins/Experimental/Interchange/Runtime/Source/Dispatcher/Public/InterchangeCommands.h b/Engine/Plugins/Experimental/Interchange/Runtime/Source/Dispatcher/Public/InterchangeCommands.h
--- a/Engine/Plugins/Experimental/Interchange/Runtime/Source/Dispatcher/Public/InterchangeCommands.h
+++ b/Engine/Plugins/Experimental/Interchange/Runtime/Source/Dispatcher/Public/InterchangeCommands.h
@@ -17,7 +17,7 @@ namespace UE
 		{
 		public:
 			//Major version should be updated when there is an existing API that has been change.
-			static int32 GetMajor() { return 7; }
+			static int32 GetMajor() { return 8; }
 			//Minor version must be updated when there is an addition to the API.
 			static int32 GetMinor() { return 0; }
 			//Patch version should be update if there is some bug fixes in the private code.
diff --git a/Engine/Plugins/Experimental/Interchange/Runtime/Source/Import/Private/Animation/InterchangeAnimSequenceFactory.cpp b/Engine/Plugins/Experimental/Interchange/Runtime/Source/Import/Private/Animation/InterchangeAnimSequenceFactory.cpp
--- a/Engine/Plugins/Experimental/Interchange/Runtime/Source/Import/Private/Animation/InterchangeAnimSequenceFactory.cpp
+++ b/Engine/Plugins/Experimental/Interchange/Runtime/Source/Import/Private/Animation/InterchangeAnimSequenceFactory.cpp
@@ -129,7 +129,6 @@ namespace UE::Interchange::Private
 						RawTrack.ScaleKeys.Reserve(BakeKeyCount);
 						TArray<float> TimeKeys;
 						TimeKeys.Reserve(BakeKeyCount);
-
 						for (double CurrentTime = RangeStart; CurrentTime <= RangeStop + SMALL_NUMBER; CurrentTime += BakeInterval)
 						{
 							//Default value to identity
@@ -140,6 +139,12 @@ namespace UE::Interchange::Private
 							for (int32 CurveIndex = 0; CurveIndex < AnimationTransformPayload->TransformCurves.Num(); ++CurveIndex)
 							{
 								const UE::Interchange::FAnimationCurveTransformPayloadData& CurveData = AnimationTransformPayload->TransformCurves[CurveIndex];
+								if (CurveData.Curve.GetNumKeys() <= 0)
+								{
+									//skip empty curve, the default value will be applied
+									continue;
+								}
+
 								if (CurveData.TransformChannel == EInterchangeTransformCurveChannel::TranslationX)
 								{
 									Position.X = CurveData.Curve.Eval(CurrentTime);
diff --git a/Engine/Plugins/Experimental/Interchange/Runtime/Source/Import/Private/Mesh/InterchangeSkeletalMeshFactory.cpp b/Engine/Plugins/Experimental/Interchange/Runtime/Source/Import/Private/Mesh/InterchangeSkeletalMeshFactory.cpp
--- a/Engine/Plugins/Experimental/Interchange/Runtime/Source/Import/Private/Mesh/InterchangeSkeletalMeshFactory.cpp
+++ b/Engine/Plugins/Experimental/Interchange/Runtime/Source/Import/Private/Mesh/InterchangeSkeletalMeshFactory.cpp
@@ -443,7 +443,8 @@ namespace UE
 					FElementIDRemappings ElementIDRemappings;
 					LodMeshPayload->LodMeshDescription.Compact(ElementIDRemappings);
 
-					if (bSkinControlPointToTimeZero)
+					const bool bIsRigidMesh = LodMeshPayload->JointNames.Num() <= 0 && MeshNodeContext.SceneNode;
+					if (bSkinControlPointToTimeZero && !bIsRigidMesh)
 					{
 						//We need to rebind the mesh at time 0. Skeleton joint have the time zero transform, so we need to apply the skinning to the mesh
 						//With the skeleton transform at time zero
@@ -457,6 +458,7 @@ namespace UE
 					}
 
 					const int32 RefBoneCount = RefBonesBinary.Num();
+					
 					//Remap the influence vertex index to point on the correct index
 					if (LodMeshPayload->JointNames.Num() > 0)
 					{
@@ -478,6 +480,36 @@ namespace UE
 							}
 						}
 					}
+					else if(bIsRigidMesh)
+					{
+						// We have a rigid mesh instance (a scene node point on the mesh, the scene node will be the bone on which the rigid mesh is skin).
+						// We must add skinning to the mesh description on bone 0 and remap it to the correct RefBonesBinary in the append settings
+						const FString ToSkinBoneName = MeshNodeContext.SceneNode->GetDisplayLabel();
+						for (int32 RefBoneIndex = 0; RefBoneIndex < RefBoneCount; ++RefBoneIndex)
+						{
+							const SkeletalMeshImportData::FBone& Bone = RefBonesBinary[RefBoneIndex];
+							if (Bone.Name.Equals(ToSkinBoneName))
+							{
+								SkeletalMeshAppendSettings.SourceRemapBoneIndex.AddZeroed_GetRef() = RefBoneIndex;
+								break;
+							}
+						}
+						//Add the skinning in the mesh description
+						{
+							FSkeletalMeshAttributes PayloadSkeletalMeshAttributes(LodMeshPayload->LodMeshDescription);
+							PayloadSkeletalMeshAttributes.Register();
+							using namespace UE::AnimationCore;
+							TArray<FBoneWeight> BoneWeights;
+							FBoneWeight& BoneWeight = BoneWeights.AddDefaulted_GetRef();
+							BoneWeight.SetBoneIndex(0);
+							BoneWeight.SetWeight(1.0f);
+							FSkinWeightsVertexAttributesRef PayloadVertexSkinWeights = PayloadSkeletalMeshAttributes.GetVertexSkinWeights();
+							for (const FVertexID& PayloadVertexID : LodMeshPayload->LodMeshDescription.Vertices().GetElementIDs())
+							{
+								PayloadVertexSkinWeights.Set(PayloadVertexID, BoneWeights);
+							}
+						}
+					}
 					//Bake the payload, with the provide transform
 					if (MeshNodeContext.SceneGlobalTransform.IsSet())
 					{
@@ -488,7 +520,7 @@ namespace UE
 						AppendSettings.MeshTransform.Reset();
 					}
 					FStaticMeshOperations::AppendMeshDescription(LodMeshPayload->LodMeshDescription, LodMeshDescription, AppendSettings);
-					if (MeshNodeContext.MeshNode->IsSkinnedMesh())
+					if (MeshNodeContext.MeshNode->IsSkinnedMesh() || bIsRigidMesh)
 					{
 						FSkeletalMeshOperations::AppendSkinWeight(LodMeshPayload->LodMeshDescription, LodMeshDescription, SkeletalMeshAppendSettings);
 					}
diff --git a/Engine/Plugins/Experimental/Interchange/Runtime/Source/Nodes/Private/InterchangeSceneNode.cpp b/Engine/Plugins/Experimental/Interchange/Runtime/Source/Nodes/Private/InterchangeSceneNode.cpp
--- a/Engine/Plugins/Experimental/Interchange/Runtime/Source/Nodes/Private/InterchangeSceneNode.cpp
+++ b/Engine/Plugins/Experimental/Interchange/Runtime/Source/Nodes/Private/InterchangeSceneNode.cpp
@@ -109,11 +109,14 @@ FString UInterchangeSceneNode::GetAttributeCategory(const UE::Interchange::FAttr
 		return FString(TEXT("Scene"));
 	}
 	else if (NodeAttributeKey == Macro_CustomBindPoseLocalTransformKey
-		|| NodeAttributeKey == Macro_CustomTimeZeroLocalTransformKey
-		|| NodeAttributeKey == Macro_CustomTransformCurvePayloadKeyKey)
+		|| NodeAttributeKey == Macro_CustomTimeZeroLocalTransformKey)
 	{
 		return FString(TEXT("Joint"));
 	}
+	else if (NodeAttributeKey == Macro_CustomTransformCurvePayloadKeyKey)
+	{
+		return FString(TEXT("Animation"));
+	}
 	else if (NodeAttributeKeyString.StartsWith(UE::Interchange::FSceneNodeStaticData::GetNodeSpecializeTypeBaseKey().ToString()))
 	{
 		return FString(TEXT("SpecializeType"));
diff --git a/Engine/Plugins/Experimental/Interchange/Runtime/Source/Parsers/Fbx/Private/FbxAnimation.cpp b/Engine/Plugins/Experimental/Interchange/Runtime/Source/Parsers/Fbx/Private/FbxAnimation.cpp
--- a/Engine/Plugins/Experimental/Interchange/Runtime/Source/Parsers/Fbx/Private/FbxAnimation.cpp
+++ b/Engine/Plugins/Experimental/Interchange/Runtime/Source/Parsers/Fbx/Private/FbxAnimation.cpp
@@ -563,7 +563,7 @@ namespace UE::Interchange::Private
 		return true;
 	}
 
-	void FFbxAnimation::AddJointAnimation(FbxScene* SDKScene, FbxNode* JointNode, UInterchangeBaseNodeContainer& NodeContainer, UInterchangeSceneNode* SceneNode, TMap<FString, TSharedPtr<FPayloadContextBase>>& PayloadContexts)
+	void FFbxAnimation::AddNodeTransformAnimation(FbxScene* SDKScene, FbxNode* JointNode, UInterchangeBaseNodeContainer& NodeContainer, UInterchangeSceneNode* SceneNode, TMap<FString, TSharedPtr<FPayloadContextBase>>& PayloadContexts)
 	{
 		int32 NumAnimations = SDKScene->GetSrcObjectCount<FbxAnimStack>();
 		//Anim stack should be merge so we expect to have only one stack here
diff --git a/Engine/Plugins/Experimental/Interchange/Runtime/Source/Parsers/Fbx/Private/FbxAnimation.h b/Engine/Plugins/Experimental/Interchange/Runtime/Source/Parsers/Fbx/Private/FbxAnimation.h
--- a/Engine/Plugins/Experimental/Interchange/Runtime/Source/Parsers/Fbx/Private/FbxAnimation.h
+++ b/Engine/Plugins/Experimental/Interchange/Runtime/Source/Parsers/Fbx/Private/FbxAnimation.h
@@ -18,7 +18,7 @@ namespace UE::Interchange::Private
 	{
 	public:
 		virtual ~FAnimationPayloadContextTransform() {}
-		virtual FString GetPayloadType() const override { return TEXT("JointAnimation-PayloadContext"); }
+		virtual FString GetPayloadType() const override { return TEXT("TransformAnimation-PayloadContext"); }
 		virtual bool FetchPayloadToFile(FFbxParser& Parser, const FString& PayloadFilepath) override;
 		FbxNode* Node = nullptr;
 		FbxScene* SDKScene = nullptr;
@@ -27,6 +27,7 @@ namespace UE::Interchange::Private
 	class FFbxAnimation
 	{
 	public:
-		static void AddJointAnimation(FbxScene* SDKScene, FbxNode* JointNode, UInterchangeBaseNodeContainer& NodeContainer, UInterchangeSceneNode* UnrealNode, TMap<FString, TSharedPtr<FPayloadContextBase>>& PayloadContexts);
+		/** This function add the payload key if the scene node transform is animated. */
+		static void AddNodeTransformAnimation(FbxScene* SDKScene, FbxNode* JointNode, UInterchangeBaseNodeContainer& NodeContainer, UInterchangeSceneNode* UnrealNode, TMap<FString, TSharedPtr<FPayloadContextBase>>& PayloadContexts);
 	};
 }//ns UE::Interchange::Private
diff --git a/Engine/Plugins/Experimental/Interchange/Runtime/Source/Parsers/Fbx/Private/FbxScene.cpp b/Engine/Plugins/Experimental/Interchange/Runtime/Source/Parsers/Fbx/Private/FbxScene.cpp
--- a/Engine/Plugins/Experimental/Interchange/Runtime/Source/Parsers/Fbx/Private/FbxScene.cpp
+++ b/Engine/Plugins/Experimental/Interchange/Runtime/Source/Parsers/Fbx/Private/FbxScene.cpp
@@ -238,7 +238,6 @@ namespace UE
 									UnrealNode->SetCustomTimeZeroLocalTransform(&NodeContainer, GlobalTransform, bResetCache);
 								}
 							}
-							FFbxAnimation::AddJointAnimation(SDKScene, Node, NodeContainer, UnrealNode, PayloadContexts);
 							break;
 						}
 
@@ -280,6 +279,8 @@ namespace UE
 						}
 					}
 				}
+				//Scene node transform can be animated, add the animation payload key.
+				FFbxAnimation::AddNodeTransformAnimation(SDKScene, Node, NodeContainer, UnrealNode, PayloadContexts);
 
 				//Add all custom Attributes for the node
 				auto IsSupportedExtraAttributeType = [](int32 InDataType)
diff --git a/Engine/Plugins/Experimental/Interchange/Runtime/Source/Pipelines/Private/InterchangeGenericSkeletalMeshPipeline.cpp b/Engine/Plugins/Experimental/Interchange/Runtime/Source/Pipelines/Private/InterchangeGenericSkeletalMeshPipeline.cpp
--- a/Engine/Plugins/Experimental/Interchange/Runtime/Source/Pipelines/Private/InterchangeGenericSkeletalMeshPipeline.cpp
+++ b/Engine/Plugins/Experimental/Interchange/Runtime/Source/Pipelines/Private/InterchangeGenericSkeletalMeshPipeline.cpp
@@ -156,13 +156,14 @@ void UInterchangeGenericMeshPipeline::ExecutePreImportPipelineSkeletalMesh()
 		bool bFoundMeshes = false;
 		if (CommonMeshesProperties->bBakeMeshes)
 		{
-			PipelineMeshesUtilities->GetCombinedSkinnedMeshInstances(MeshUidsPerSkeletonRootUid, bConvertStaticMeshToSkeletalMesh);
+			PipelineMeshesUtilities->GetCombinedSkinnedMeshInstances(BaseNodeContainer, MeshUidsPerSkeletonRootUid, bConvertStaticMeshToSkeletalMesh);
 			const bool bUseMeshInstance = true;
 			bFoundMeshes = CreatePerSkeletonRootUidCombinedSkinnedMesh(bUseMeshInstance);
 		}
 
 		if (!bFoundMeshes)
 		{
+			//Fall back to support mesh not reference by any scene node
 			MeshUidsPerSkeletonRootUid.Empty();
 			PipelineMeshesUtilities->GetCombinedSkinnedMeshGeometries(MeshUidsPerSkeletonRootUid);
 			const bool bUseMeshInstance = false;
diff --git a/Engine/Plugins/Experimental/Interchange/Runtime/Source/Pipelines/Private/InterchangePipelineMeshesUtilities.cpp b/Engine/Plugins/Experimental/Interchange/Runtime/Source/Pipelines/Private/InterchangePipelineMeshesUtilities.cpp
--- a/Engine/Plugins/Experimental/Interchange/Runtime/Source/Pipelines/Private/InterchangePipelineMeshesUtilities.cpp
+++ b/Engine/Plugins/Experimental/Interchange/Runtime/Source/Pipelines/Private/InterchangePipelineMeshesUtilities.cpp
@@ -398,8 +398,10 @@ void UInterchangePipelineMeshesUtilities::IterateAllMeshInstanceUsingMeshGeometr
 	}
 }
 
-void UInterchangePipelineMeshesUtilities::GetCombinedSkinnedMeshInstances(TMap<FString, TArray<FString>>& OutMeshInstanceUidsPerSkeletonRootUid, const bool bConvertStaticMeshToSkeletalMesh) const
+void UInterchangePipelineMeshesUtilities::GetCombinedSkinnedMeshInstances(UInterchangeBaseNodeContainer* BaseNodeContainer, TMap<FString, TArray<FString>>& OutMeshInstanceUidsPerSkeletonRootUid, const bool bConvertStaticMeshToSkeletalMesh) const
 {
+	check(BaseNodeContainer);
+
 	for (const TPair<FString, FInterchangeMeshInstance>& MeshInstanceUidAndMeshInstance : MeshInstancesPerMeshInstanceUid)
 	{
 		const FInterchangeMeshInstance& MeshInstance = MeshInstanceUidAndMeshInstance.Value;
@@ -428,7 +430,20 @@ void UInterchangePipelineMeshesUtilities::GetCombinedSkinnedMeshInstances(TMap<F
 			else if (bIsStaticMesh)
 			{
 				//Create a joint from the instance node (the scene node pointing on the mesh).
+				//Since we are dealing with rigid mesh and combine we will use the outermost valid parent
 				SkeletonRootUid = MeshInstance.MeshInstanceUid;
+				if (const UInterchangeSceneNode* SceneNode = Cast<UInterchangeSceneNode>(BaseNodeContainer->GetNode(MeshInstance.MeshInstanceUid)))
+				{
+					FString ParentUid = SceneNode->GetParentUid();
+					while (!ParentUid.Equals(UInterchangeBaseNode::InvalidNodeUid()))
+					{
+						if (const UInterchangeSceneNode* ParentNode = Cast<UInterchangeSceneNode>(BaseNodeContainer->GetNode(ParentUid)))
+						{
+							SkeletonRootUid = ParentUid;
+							ParentUid = ParentNode->GetParentUid();
+						}
+					}
+				}
 			}
 			else
 			{
diff --git a/Engine/Plugins/Experimental/Interchange/Runtime/Source/Pipelines/Public/InterchangePipelineMeshesUtilities.h b/Engine/Plugins/Experimental/Interchange/Runtime/Source/Pipelines/Public/InterchangePipelineMeshesUtilities.h
--- a/Engine/Plugins/Experimental/Interchange/Runtime/Source/Pipelines/Public/InterchangePipelineMeshesUtilities.h
+++ b/Engine/Plugins/Experimental/Interchange/Runtime/Source/Pipelines/Public/InterchangePipelineMeshesUtilities.h
@@ -241,7 +241,7 @@ public:
 	* Return a list of skinned FInterchangeMeshInstance uid that can be combined together.
 	* We cannot create a skinned mesh with multiple skeleton root node, This function return combined MeshInstance per skeleton roots
 	*/
-	void GetCombinedSkinnedMeshInstances(TMap<FString, TArray<FString>>& OutMeshInstanceUidsPerSkeletonRootUid, const bool bConvertStaticMeshToSkeletalMesh) const;
+	void GetCombinedSkinnedMeshInstances(UInterchangeBaseNodeContainer* BaseNodeContainer, TMap<FString, TArray<FString>>& OutMeshInstanceUidsPerSkeletonRootUid, const bool bConvertStaticMeshToSkeletalMesh) const;
 	
 	/**
 	* Return a list of skinned FInterchangeMeshGeometry uid that can be combined together.
