Fix crash when adding error message. The sub generic asset pipelines result container was not set, and this was triggering assert when adding message to the unset container.
Add a error message in case we have a zero scale in the UI. Set back the scale to the default value 1.0f
#rnx

diff --git a/Engine/Plugins/Experimental/Interchange/Runtime/Source/Pipelines/Private/InterchangeGenericAssetsPipeline.cpp b/Engine/Plugins/Experimental/Interchange/Runtime/Source/Pipelines/Private/InterchangeGenericAssetsPipeline.cpp
--- a/Engine/Plugins/Experimental/Interchange/Runtime/Source/Pipelines/Private/InterchangeGenericAssetsPipeline.cpp
+++ b/Engine/Plugins/Experimental/Interchange/Runtime/Source/Pipelines/Private/InterchangeGenericAssetsPipeline.cpp
@@ -140,6 +140,28 @@ void UInterchangeGenericAssetsPipeline::ExecutePreImportPipeline(UInterchangeBas
 		return;
 	}
 
+	//Set the result container to allow error message
+	//The parent Results container should be set at this point
+	ensure(!Results.IsNull());
+	{
+		if (TexturePipeline)
+		{
+			TexturePipeline->SetResultsContainer(Results);
+		}
+		if (MaterialPipeline)
+		{
+			MaterialPipeline->SetResultsContainer(Results);
+		}
+		if (MeshPipeline)
+		{
+			MeshPipeline->SetResultsContainer(Results);
+		}
+		if (AnimationPipeline)
+		{
+			AnimationPipeline->SetResultsContainer(Results);
+		}
+	}
+
 	//////////////////////////////////////////////////////////////////////////
 	//Make sure all options go together
 	
@@ -159,8 +181,21 @@ void UInterchangeGenericAssetsPipeline::ExecutePreImportPipeline(UInterchangeBas
 	BaseNodeContainer = InBaseNodeContainer;
 
 	//Setup the Global import offset
-
 	{
+		//Make sure the scale value is greater than zero, warn the user in this case and set the scale to the default value 1.0f
+		if (ImportOffsetUniformScale < UE_SMALL_NUMBER)
+		{
+			FNumberFormattingOptions FormatingOptions;
+			FormatingOptions.SetMaximumFractionalDigits(6);
+			FormatingOptions.SetMinimumFractionalDigits(1);
+			float DefaultScaleValue = 1.0f;
+			UInterchangeResultError_Generic* Message = AddMessage<UInterchangeResultError_Generic>();
+			Message->Text = FText::Format(NSLOCTEXT("UInterchangeGenericAssetsPipeline", "BadImportOffsetUniformScale", "Value [{0}] for ImportOffsetUniformScale setting is too small, we will use the default value [{1}]."),
+				FText::AsNumber(ImportOffsetUniformScale, &FormatingOptions),
+				FText::AsNumber(DefaultScaleValue, &FormatingOptions));
+			ImportOffsetUniformScale = DefaultScaleValue;
+		}
+
 		FTransform ImportOffsetTransform;
 		ImportOffsetTransform.SetTranslation(ImportOffsetTranslation);
 		ImportOffsetTransform.SetRotation(FQuat(ImportOffsetRotation));
@@ -181,17 +216,14 @@ void UInterchangeGenericAssetsPipeline::ExecutePreImportPipeline(UInterchangeBas
 	{
 		TexturePipeline->ScriptedExecutePreImportPipeline(InBaseNodeContainer, InSourceDatas);
 	}
-
 	if (MaterialPipeline)
 	{
 		MaterialPipeline->ScriptedExecutePreImportPipeline(InBaseNodeContainer, InSourceDatas);
 	}
-
 	if (MeshPipeline)
 	{
 		MeshPipeline->ScriptedExecutePreImportPipeline(InBaseNodeContainer, InSourceDatas);
 	}
-
 	if (AnimationPipeline)
 	{
 		AnimationPipeline->ScriptedExecutePreImportPipeline(InBaseNodeContainer, InSourceDatas);
@@ -214,17 +246,14 @@ void UInterchangeGenericAssetsPipeline::ExecutePostImportPipeline(const UInterch
 	{
 		TexturePipeline->ScriptedExecutePostImportPipeline(InBaseNodeContainer, NodeKey, CreatedAsset, bIsAReimport);
 	}
-
 	if (MaterialPipeline)
 	{
 		MaterialPipeline->ScriptedExecutePostImportPipeline(InBaseNodeContainer, NodeKey, CreatedAsset, bIsAReimport);
 	}
-
 	if (MeshPipeline)
 	{
 		MeshPipeline->ScriptedExecutePostImportPipeline(InBaseNodeContainer, NodeKey, CreatedAsset, bIsAReimport);
 	}
-
 	if (AnimationPipeline)
 	{
 		AnimationPipeline->ScriptedExecutePostImportPipeline(InBaseNodeContainer, NodeKey, CreatedAsset, bIsAReimport);
diff --git a/Engine/Source/Runtime/Experimental/Interchange/Engine/Private/Tasks/InterchangeTaskPipeline.cpp b/Engine/Source/Runtime/Experimental/Interchange/Engine/Private/Tasks/InterchangeTaskPipeline.cpp
--- a/Engine/Source/Runtime/Experimental/Interchange/Engine/Private/Tasks/InterchangeTaskPipeline.cpp
+++ b/Engine/Source/Runtime/Experimental/Interchange/Engine/Private/Tasks/InterchangeTaskPipeline.cpp
@@ -112,7 +112,6 @@ void UE::Interchange::FTaskPipelinePostImport::DoTask(ENamedThreads::Type Curren
 		return;
 	}
 	UInterchangePipelineBase* Pipeline = AsyncHelper->Pipelines[PipelineIndex];
-	Pipeline->SetResultsContainer(AsyncHelper->AssetImportResult->GetResults());
 
 	//Call the pipeline outside of the lock, we do this in case the pipeline take a long time. We call it for each asset created by this import
 	for (int32 ObjectIndex = 0; ObjectIndex < ImportedObjects.Num(); ++ObjectIndex)
