Interchange: tweak some material processing in the generic pipeline
- Favor translucency over coats when creating materials, as we do when creating instances.
- Ignore twosidedness for thin translucency as the results look visually wrong.

diff --git a/Engine/Plugins/Experimental/Interchange/Runtime/Source/Pipelines/Private/InterchangeGenericMaterialPipeline.cpp b/Engine/Plugins/Experimental/Interchange/Runtime/Source/Pipelines/Private/InterchangeGenericMaterialPipeline.cpp
--- a/Engine/Plugins/Experimental/Interchange/Runtime/Source/Pipelines/Private/InterchangeGenericMaterialPipeline.cpp
+++ b/Engine/Plugins/Experimental/Interchange/Runtime/Source/Pipelines/Private/InterchangeGenericMaterialPipeline.cpp
@@ -672,7 +672,8 @@ void UInterchangeGenericMaterialPipeline::HandleCommonParameters(const UIntercha
 {
 	using namespace UE::Interchange::Materials::Common;
 
-	// Two sidedness
+	// Two sidedness (ignored for thin translucency as it looks wrong)
+	if (!IsThinTranslucentModel(ShaderGraphNode))
 	{
 		bool bTwoSided = false;
 		ShaderGraphNode->GetCustomTwoSided(bTwoSided);
@@ -1371,10 +1372,11 @@ UInterchangeMaterialFactoryNode* UInterchangeGenericMaterialPipeline::CreateMate
 
 	HandlePBRModel(ShaderGraphNode, MaterialFactoryNode); // Always process the PBR parameters. If they were already assigned from Phong or Lambert, they will be ignored.
 	
-	if (!HandleClearCoat(ShaderGraphNode, MaterialFactoryNode))
+	// Can't have different shading models
+	// Favor translucency over coats (clear coat, sheen, etc.) since it tends to have a bigger impact visually
+	if (!HandleThinTranslucent(ShaderGraphNode, MaterialFactoryNode))
 	{
-		// Can't have different shading models
-		if (!HandleThinTranslucent(ShaderGraphNode, MaterialFactoryNode))
+		if (!HandleClearCoat(ShaderGraphNode, MaterialFactoryNode))
 		{
 			HandleSheen(ShaderGraphNode, MaterialFactoryNode);
 		}
diff --git a/Engine/Plugins/Experimental/Interchange/Runtime/Source/Pipelines/Private/InterchangePipelineMeshesUtilities.cpp b/Engine/Plugins/Experimental/Interchange/Runtime/Source/Pipelines/Private/InterchangePipelineMeshesUtilities.cpp
--- a/Engine/Plugins/Experimental/Interchange/Runtime/Source/Pipelines/Private/InterchangePipelineMeshesUtilities.cpp
+++ b/Engine/Plugins/Experimental/Interchange/Runtime/Source/Pipelines/Private/InterchangePipelineMeshesUtilities.cpp
@@ -69,7 +69,9 @@ UInterchangePipelineMeshesUtilities* UInterchangePipelineMeshesUtilities::Create
 					FString MeshUid;
 					if (SceneNode->GetCustomAssetInstanceUid(MeshUid))
 					{
-						if (BaseNodeContainer->GetNode(MeshUid)->IsA<UInterchangeMeshNode>())
+						const UInterchangeBaseNode* MeshNode = BaseNodeContainer->GetNode(MeshUid);
+
+						if (MeshNode && MeshNode->IsA<UInterchangeMeshNode>())
 						{
 							const UInterchangeSceneNode* ParentMeshSceneNode = Cast<UInterchangeSceneNode>(BaseNodeContainer->GetNode(SceneNode->GetParentUid()));
 							if (ParentMeshSceneNode)
