Added support for material slot names and their associated materials for static meshes and static mesh actors
Added  storage of correspondence between slot names and material names to UInterchangeMeshNode
Added  storage of correspondence between slot names and material names to UInterchangeStaticMeshFactoryNode
Added  storage of correspondence between slot names and material names to UInterchangeSceneNode
Added  storage of correspondence between slot names and material names to UInterchangeActorFactoryNode
Modified UInterchangeGenericMeshPipeline to use  UInterchangeSceneNode's slot dependency
Modified UInterchangeGenericScenesPipeline to use  UInterchangeSceneNode's slot dependency and create UInterchangeMeshActorFactoryNode
Added logic in UInterchangeStaticMeshFactory to apply correspondence between slot names and material names stored in UInterchangeStaticMeshFactoryNode to the static mesh
Added logic in UInterchangeStaticMeshActorFactory to apply correspondence between slot names and material names stored in UInterchangeStaticMeshFactoryNode to the static mesh
Added logic in UInterchangeSkeletalMeshActorFactory to apply correspondence between slot names and material names stored in UInterchangeActorFactoryNode to the static mesh component
Added code to skip FStaticMeshOperations::AppendMeshDescription if there is only one mesh in the payload
Modified UE::Interchange::TMapAttributeHelper<FString, FString> by adding a method to rebuild the cache from the attributes. This was needed to rebuild cache after Interchange nodes using this helper were deserialized.
Fixed issue in InterchangeTexture2DFactoryNode and  InterchangeTexture2DNode with serialization of Interchange nodes using UE::Interchange::TMapAttributeHelper<FString, FString>. Serialize method needed to be override in order to restore the cache.
Added UInterchangeMeshActorFactoryNode class to manage slots on static mesh and skeletal mesh actors.
Added helper function to transfer slot dependencies from translated node to factory node
Added abstract class UInterchangeMeshFactoryNode from which UInterchangeSkeletalMeshFactoryNode and UInterchangeStaticMeshFactoryNode now derived
Modified UInterchangeSkeletalMeshFactoryNode and UInterchangeStaticMeshFactoryNode based on new UInterchangeMeshFactoryNode class
Added logic in UInterchangeSkeletalMeshFactory::CreateAsset to handle slot material dependencies
Added logic in UInterchangeGenericMeshPipeline::AddLodDataToSkeletalMesh to handle slot material dependencies

Note:
 - LODs' material assignment does not work for skeletal and static meshes for a scene import. This needs to be fixed separately

diff --git a/Engine/Plugins/Experimental/Interchange/Runtime/Source/FactoryNodes/Private/InterchangeMeshActorFactoryNode.cpp b/Engine/Plugins/Experimental/Interchange/Runtime/Source/FactoryNodes/Private/InterchangeMeshActorFactoryNode.cpp
new file mode 100644
--- /dev/null
+++ b/Engine/Plugins/Experimental/Interchange/Runtime/Source/FactoryNodes/Private/InterchangeMeshActorFactoryNode.cpp
@@ -0,0 +1,32 @@
+// Copyright Epic Games, Inc. All Rights Reserved.
+
+#include "InterchangeMeshActorFactoryNode.h"
+
+#if WITH_ENGINE
+	#include "GameFramework/Actor.h"
+#endif
+
+UInterchangeMeshActorFactoryNode::UInterchangeMeshActorFactoryNode()
+{
+	SlotMaterialDependencies.Initialize(Attributes.ToSharedRef(), TEXT("__SlotMaterialDependencies__"));
+}
+
+void UInterchangeMeshActorFactoryNode::GetSlotMaterialDependencies(TMap<FString, FString>& OutMaterialDependencies) const
+{
+	OutMaterialDependencies = SlotMaterialDependencies.ToMap();
+}
+
+bool UInterchangeMeshActorFactoryNode::GetSlotMaterialDependencyUid(const FString& SlotName, FString& OutMaterialDependency) const
+{
+	return SlotMaterialDependencies.GetValue(SlotName, OutMaterialDependency);
+}
+
+bool UInterchangeMeshActorFactoryNode::SetSlotMaterialDependencyUid(const FString& SlotName, const FString& MaterialDependencyUid)
+{
+	return SlotMaterialDependencies.SetKeyValue(SlotName, MaterialDependencyUid);
+}
+
+bool UInterchangeMeshActorFactoryNode::RemoveSlotMaterialDependencyUid(const FString& SlotName)
+{
+	return SlotMaterialDependencies.RemoveKey(SlotName);
+}
diff --git a/Engine/Plugins/Experimental/Interchange/Runtime/Source/FactoryNodes/Private/InterchangeMeshFactoryNode.cpp b/Engine/Plugins/Experimental/Interchange/Runtime/Source/FactoryNodes/Private/InterchangeMeshFactoryNode.cpp
new file mode 100644
--- /dev/null
+++ b/Engine/Plugins/Experimental/Interchange/Runtime/Source/FactoryNodes/Private/InterchangeMeshFactoryNode.cpp
@@ -0,0 +1,118 @@
+// Copyright Epic Games, Inc. All Rights Reserved.
+
+#include "InterchangeMeshFactoryNode.h"
+
+namespace UE::Interchange
+{
+	const FAttributeKey& FMeshFactoryNodeStaticData::GetLodDependenciesBaseKey()
+	{
+		static FAttributeKey LodDependencies_BaseKey = FAttributeKey(TEXT("Lod_Dependencies"));
+		return LodDependencies_BaseKey;
+	}
+
+	const FAttributeKey& FMeshFactoryNodeStaticData::GetSlotMaterialDependencyBaseKey()
+	{
+		static FAttributeKey SlotMaterialDependency_BaseKey(TEXT("__SlotMaterialDependency__"));
+		return SlotMaterialDependency_BaseKey;
+	}
+} // namespace UE::Interchange
+
+
+UInterchangeMeshFactoryNode::UInterchangeMeshFactoryNode()
+{
+	LodDependencies.Initialize(Attributes, UE::Interchange::FMeshFactoryNodeStaticData::GetLodDependenciesBaseKey().ToString());
+	SlotMaterialDependencies.Initialize(Attributes.ToSharedRef(), UE::Interchange::FMeshFactoryNodeStaticData::GetSlotMaterialDependencyBaseKey().ToString());
+}
+
+FString UInterchangeMeshFactoryNode::GetKeyDisplayName(const UE::Interchange::FAttributeKey& NodeAttributeKey) const
+{
+	FString KeyDisplayName = NodeAttributeKey.ToString();
+	FString NodeAttributeKeyString = NodeAttributeKey.ToString();
+	if (NodeAttributeKey == UE::Interchange::FMeshFactoryNodeStaticData::GetLodDependenciesBaseKey())
+	{
+		KeyDisplayName = TEXT("LOD Dependencies Count");
+		return KeyDisplayName;
+	}
+	else if (NodeAttributeKeyString.StartsWith(UE::Interchange::FMeshFactoryNodeStaticData::GetLodDependenciesBaseKey().ToString()))
+	{
+		KeyDisplayName = TEXT("LOD Dependencies Index ");
+		const FString IndexKey = UE::Interchange::TArrayAttributeHelper<FString>::IndexKey();
+		int32 IndexPosition = NodeAttributeKeyString.Find(IndexKey) + IndexKey.Len();
+		if (IndexPosition < NodeAttributeKeyString.Len())
+		{
+			KeyDisplayName += NodeAttributeKeyString.RightChop(IndexPosition);
+		}
+		return KeyDisplayName;
+	}
+	return Super::GetKeyDisplayName(NodeAttributeKey);
+}
+
+int32 UInterchangeMeshFactoryNode::GetLodDataCount() const
+{
+	return LodDependencies.GetCount();
+}
+
+void UInterchangeMeshFactoryNode::GetLodDataUniqueIds(TArray<FString>& OutLodDataUniqueIds) const
+{
+	LodDependencies.GetItems(OutLodDataUniqueIds);
+}
+
+bool UInterchangeMeshFactoryNode::AddLodDataUniqueId(const FString& LodDataUniqueId)
+{
+	return LodDependencies.AddItem(LodDataUniqueId);
+}
+
+bool UInterchangeMeshFactoryNode::RemoveLodDataUniqueId(const FString& LodDataUniqueId)
+{
+	return LodDependencies.RemoveItem(LodDataUniqueId);
+}
+
+bool UInterchangeMeshFactoryNode::GetCustomVertexColorReplace(bool& AttributeValue) const
+{
+	IMPLEMENT_NODE_ATTRIBUTE_GETTER(VertexColorReplace, bool)
+}
+
+bool UInterchangeMeshFactoryNode::SetCustomVertexColorReplace(const bool& AttributeValue)
+{
+	IMPLEMENT_NODE_ATTRIBUTE_SETTER_NODELEGATE(VertexColorReplace, bool)
+}
+
+bool UInterchangeMeshFactoryNode::GetCustomVertexColorIgnore(bool& AttributeValue) const
+{
+	IMPLEMENT_NODE_ATTRIBUTE_GETTER(VertexColorIgnore, bool)
+}
+
+bool UInterchangeMeshFactoryNode::SetCustomVertexColorIgnore(const bool& AttributeValue)
+{
+	IMPLEMENT_NODE_ATTRIBUTE_SETTER_NODELEGATE(VertexColorIgnore, bool)
+}
+
+bool UInterchangeMeshFactoryNode::GetCustomVertexColorOverride(FColor& AttributeValue) const
+{
+	IMPLEMENT_NODE_ATTRIBUTE_GETTER(VertexColorOverride, FColor)
+}
+
+bool UInterchangeMeshFactoryNode::SetCustomVertexColorOverride(const FColor& AttributeValue)
+{
+	IMPLEMENT_NODE_ATTRIBUTE_SETTER_NODELEGATE(VertexColorOverride, FColor)
+}
+
+void UInterchangeMeshFactoryNode::GetSlotMaterialDependencies(TMap<FString, FString>& OutMaterialDependencies) const
+{
+	OutMaterialDependencies = SlotMaterialDependencies.ToMap();
+}
+
+bool UInterchangeMeshFactoryNode::GetSlotMaterialDependencyUid(const FString& SlotName, FString& OutMaterialDependency) const
+{
+	return SlotMaterialDependencies.GetValue(SlotName, OutMaterialDependency);
+}
+
+bool UInterchangeMeshFactoryNode::SetSlotMaterialDependencyUid(const FString& SlotName, const FString& MaterialDependencyUid)
+{
+	return SlotMaterialDependencies.SetKeyValue(SlotName, MaterialDependencyUid);
+}
+
+bool UInterchangeMeshFactoryNode::RemoveSlotMaterialDependencyUid(const FString& SlotName)
+{
+	return SlotMaterialDependencies.RemoveKey(SlotName);
+}
diff --git a/Engine/Plugins/Experimental/Interchange/Runtime/Source/FactoryNodes/Private/InterchangeSkeletalMeshFactoryNode.cpp b/Engine/Plugins/Experimental/Interchange/Runtime/Source/FactoryNodes/Private/InterchangeSkeletalMeshFactoryNode.cpp
--- a/Engine/Plugins/Experimental/Interchange/Runtime/Source/FactoryNodes/Private/InterchangeSkeletalMeshFactoryNode.cpp
+++ b/Engine/Plugins/Experimental/Interchange/Runtime/Source/FactoryNodes/Private/InterchangeSkeletalMeshFactoryNode.cpp
@@ -6,22 +6,11 @@
 #include "Engine/SkeletalMesh.h"
 #endif
 
-//Interchange namespace
-namespace UE::Interchange
-{
-	const FAttributeKey& FSkeletalMeshNodeStaticData::GetLodDependenciesBaseKey()
-	{
-		static FAttributeKey LodDependencies_BaseKey(TEXT("Lod_Dependencies"));
-		return LodDependencies_BaseKey;
-	}
-}//ns UE::Interchange
-
 UInterchangeSkeletalMeshFactoryNode::UInterchangeSkeletalMeshFactoryNode()
 {
 #if WITH_ENGINE
 	AssetClass = nullptr;
 #endif
-	LodDependencies.Initialize(Attributes, UE::Interchange::FSkeletalMeshNodeStaticData::GetLodDependenciesBaseKey().ToString());
 }
 
 void UInterchangeSkeletalMeshFactoryNode::InitializeSkeletalMeshNode(const FString& UniqueID, const FString& DisplayLabel, const FString& InAssetClass)
@@ -34,48 +23,12 @@ void UInterchangeSkeletalMeshFactoryNode::InitializeSkeletalMeshNode(const FStri
 	FillAssetClassFromAttribute();
 }
 
-void UInterchangeSkeletalMeshFactoryNode::Serialize(FArchive& Ar)
-{
-	Super::Serialize(Ar);
-#if WITH_ENGINE
-	if (Ar.IsLoading())
-	{
-		//Make sure the class is properly set when we compile with engine, this will set the
-		//bIsNodeClassInitialized to true.
-		SetNodeClassFromClassAttribute();
-	}
-#endif //#if WITH_ENGINE
-}
-
 FString UInterchangeSkeletalMeshFactoryNode::GetTypeName() const
 {
 	const FString TypeName = TEXT("SkeletalMeshNode");
 	return TypeName;
 }
 
-FString UInterchangeSkeletalMeshFactoryNode::GetKeyDisplayName(const UE::Interchange::FAttributeKey& NodeAttributeKey) const
-{
-	FString KeyDisplayName = NodeAttributeKey.ToString();
-	const FString NodeAttributeKeyString = KeyDisplayName;
-	if (NodeAttributeKey == UE::Interchange::FSkeletalMeshNodeStaticData::GetLodDependenciesBaseKey())
-	{
-		KeyDisplayName = TEXT("LOD Dependencies Count");
-		return KeyDisplayName;
-	}
-	else if (NodeAttributeKeyString.StartsWith(UE::Interchange::FSkeletalMeshNodeStaticData::GetLodDependenciesBaseKey().ToString()))
-	{
-		KeyDisplayName = TEXT("LOD Dependencies Index ");
-		const FString IndexKey = UE::Interchange::TArrayAttributeHelper<FString>::IndexKey();
-		int32 IndexPosition = NodeAttributeKeyString.Find(IndexKey) + IndexKey.Len();
-		if (IndexPosition < NodeAttributeKeyString.Len())
-		{
-			KeyDisplayName += NodeAttributeKeyString.RightChop(IndexPosition);
-		}
-		return KeyDisplayName;
-	}
-	return Super::GetKeyDisplayName(NodeAttributeKey);
-}
-
 UClass* UInterchangeSkeletalMeshFactoryNode::GetObjectClass() const
 {
 	ensure(bIsNodeClassInitialized);
@@ -86,26 +39,6 @@ UClass* UInterchangeSkeletalMeshFactoryNode::GetObjectClass() const
 #endif
 }
 
-int32 UInterchangeSkeletalMeshFactoryNode::GetLodDataCount() const
-{
-	return LodDependencies.GetCount();
-}
-
-void UInterchangeSkeletalMeshFactoryNode::GetLodDataUniqueIds(TArray<FString>& OutLodDataUniqueIds) const
-{
-	LodDependencies.GetItems(OutLodDataUniqueIds);
-}
-
-bool UInterchangeSkeletalMeshFactoryNode::AddLodDataUniqueId(const FString& LodDataUniqueId)
-{
-	return LodDependencies.AddItem(LodDataUniqueId);
-}
-
-bool UInterchangeSkeletalMeshFactoryNode::RemoveLodDataUniqueId(const FString& LodDataUniqueId)
-{
-	return LodDependencies.RemoveItem(LodDataUniqueId);
-}
-
 bool UInterchangeSkeletalMeshFactoryNode::GetCustomSkeletonSoftObjectPath(FSoftObjectPath& AttributeValue) const
 {
 	IMPLEMENT_NODE_ATTRIBUTE_GETTER(SkeletonSoftObjectPath, FSoftObjectPath)
@@ -146,36 +79,6 @@ bool UInterchangeSkeletalMeshFactoryNode::SetCustomPhysicAssetSoftObjectPath(con
 	IMPLEMENT_NODE_ATTRIBUTE_SETTER_NODELEGATE(PhysicAssetSoftObjectPath, FSoftObjectPath)
 }
 
-bool UInterchangeSkeletalMeshFactoryNode::GetCustomVertexColorReplace(bool& AttributeValue) const
-{
-	IMPLEMENT_NODE_ATTRIBUTE_GETTER(VertexColorReplace, bool)
-}
-
-bool UInterchangeSkeletalMeshFactoryNode::SetCustomVertexColorReplace(const bool& AttributeValue)
-{
-	IMPLEMENT_NODE_ATTRIBUTE_SETTER_NODELEGATE(VertexColorReplace, bool)
-}
-
-bool UInterchangeSkeletalMeshFactoryNode::GetCustomVertexColorIgnore(bool& AttributeValue) const
-{
-	IMPLEMENT_NODE_ATTRIBUTE_GETTER(VertexColorIgnore, bool)
-}
-
-bool UInterchangeSkeletalMeshFactoryNode::SetCustomVertexColorIgnore(const bool& AttributeValue)
-{
-	IMPLEMENT_NODE_ATTRIBUTE_SETTER_NODELEGATE(VertexColorIgnore, bool)
-}
-
-bool UInterchangeSkeletalMeshFactoryNode::GetCustomVertexColorOverride(FColor& AttributeValue) const
-{
-	IMPLEMENT_NODE_ATTRIBUTE_GETTER(VertexColorOverride, FColor)
-}
-
-bool UInterchangeSkeletalMeshFactoryNode::SetCustomVertexColorOverride(const FColor& AttributeValue)
-{
-	IMPLEMENT_NODE_ATTRIBUTE_SETTER_NODELEGATE(VertexColorOverride, FColor)
-}
-
 bool UInterchangeSkeletalMeshFactoryNode::GetCustomImportContentType(EInterchangeSkeletalMeshContentType& AttributeValue) const
 {
 	FString OperationName = GetTypeName() + TEXT(".GetImportContentType");
@@ -244,12 +147,3 @@ bool UInterchangeSkeletalMeshFactoryNode::SetNodeClassFromClassAttribute()
 	}
 	return bIsNodeClassInitialized;
 }
-
-bool UInterchangeSkeletalMeshFactoryNode::IsEditorOnlyDataDefined()
-{
-#if WITH_EDITORONLY_DATA
-	return true;
-#else
-	return false;
-#endif
-}
diff --git a/Engine/Plugins/Experimental/Interchange/Runtime/Source/FactoryNodes/Private/InterchangeStaticMeshFactoryNode.cpp b/Engine/Plugins/Experimental/Interchange/Runtime/Source/FactoryNodes/Private/InterchangeStaticMeshFactoryNode.cpp
--- a/Engine/Plugins/Experimental/Interchange/Runtime/Source/FactoryNodes/Private/InterchangeStaticMeshFactoryNode.cpp
+++ b/Engine/Plugins/Experimental/Interchange/Runtime/Source/FactoryNodes/Private/InterchangeStaticMeshFactoryNode.cpp
@@ -10,12 +10,6 @@ namespace UE
 {
 	namespace Interchange
 	{
-		const FAttributeKey& FStaticMeshNodeStaticData::GetLodDependenciesBaseKey()
-		{
-			static FAttributeKey LodDependencies_BaseKey = FAttributeKey(TEXT("Lod_Dependencies"));
-			return LodDependencies_BaseKey;
-		}
-
 		const FAttributeKey& FStaticMeshNodeStaticData::GetSocketUidsBaseKey()
 		{
 			static FAttributeKey SocketUids_BaseKey = FAttributeKey(TEXT("SocketUids"));
@@ -30,7 +24,6 @@ UInterchangeStaticMeshFactoryNode::UInterchangeStaticMeshFactoryNode()
 #if WITH_ENGINE
 	AssetClass = nullptr;
 #endif
-	LodDependencies.Initialize(Attributes, UE::Interchange::FStaticMeshNodeStaticData::GetLodDependenciesBaseKey().ToString());
 	SocketUids.Initialize(Attributes, UE::Interchange::FStaticMeshNodeStaticData::GetSocketUidsBaseKey().ToString());
 }
 
@@ -44,47 +37,12 @@ void UInterchangeStaticMeshFactoryNode::InitializeStaticMeshNode(const FString&
 	FillAssetClassFromAttribute();
 }
 
-void UInterchangeStaticMeshFactoryNode::Serialize(FArchive& Ar)
-{
-	Super::Serialize(Ar);
-#if WITH_ENGINE
-	if (Ar.IsLoading())
-	{
-		// Make sure the class is properly set when we compile with engine, this will set the bIsNodeClassInitialized to true.
-		SetNodeClassFromClassAttribute();
-	}
-#endif //#if WITH_ENGINE
-}
-
 FString UInterchangeStaticMeshFactoryNode::GetTypeName() const
 {
 	const FString TypeName = TEXT("StaticMeshNode");
 	return TypeName;
 }
 
-FString UInterchangeStaticMeshFactoryNode::GetKeyDisplayName(const UE::Interchange::FAttributeKey& NodeAttributeKey) const
-{
-	FString KeyDisplayName = NodeAttributeKey.ToString();
-	FString NodeAttributeKeyString = NodeAttributeKey.ToString();
-	if (NodeAttributeKey == UE::Interchange::FStaticMeshNodeStaticData::GetLodDependenciesBaseKey())
-	{
-		KeyDisplayName = TEXT("LOD Dependencies Count");
-		return KeyDisplayName;
-	}
-	else if (NodeAttributeKeyString.StartsWith(UE::Interchange::FStaticMeshNodeStaticData::GetLodDependenciesBaseKey().ToString()))
-	{
-		KeyDisplayName = TEXT("LOD Dependencies Index ");
-		const FString IndexKey = UE::Interchange::TArrayAttributeHelper<FString>::IndexKey();
-		int32 IndexPosition = NodeAttributeKeyString.Find(IndexKey) + IndexKey.Len();
-		if (IndexPosition < NodeAttributeKeyString.Len())
-		{
-			KeyDisplayName += NodeAttributeKeyString.RightChop(IndexPosition);
-		}
-		return KeyDisplayName;
-	}
-	return Super::GetKeyDisplayName(NodeAttributeKey);
-}
-
 UClass* UInterchangeStaticMeshFactoryNode::GetObjectClass() const
 {
 	ensure(bIsNodeClassInitialized);
@@ -95,26 +53,6 @@ UClass* UInterchangeStaticMeshFactoryNode::GetObjectClass() const
 #endif
 }
 
-int32 UInterchangeStaticMeshFactoryNode::GetLodDataCount() const
-{
-	return LodDependencies.GetCount();
-}
-
-void UInterchangeStaticMeshFactoryNode::GetLodDataUniqueIds(TArray<FString>& OutLodDataUniqueIds) const
-{
-	LodDependencies.GetItems(OutLodDataUniqueIds);
-}
-
-bool UInterchangeStaticMeshFactoryNode::AddLodDataUniqueId(const FString& LodDataUniqueId)
-{
-	return LodDependencies.AddItem(LodDataUniqueId);
-}
-
-bool UInterchangeStaticMeshFactoryNode::RemoveLodDataUniqueId(const FString& LodDataUniqueId)
-{
-	return LodDependencies.RemoveItem(LodDataUniqueId);
-}
-
 int32 UInterchangeStaticMeshFactoryNode::GetSocketUidCount() const
 {
 	return SocketUids.GetCount();
@@ -148,36 +86,6 @@ bool UInterchangeStaticMeshFactoryNode::RemoveSocketUd(const FString& SocketUid)
 	return SocketUids.RemoveItem(SocketUid);
 }
 
-bool UInterchangeStaticMeshFactoryNode::GetCustomVertexColorReplace(bool& AttributeValue) const
-{
-	IMPLEMENT_NODE_ATTRIBUTE_GETTER(VertexColorReplace, bool)
-}
-
-bool UInterchangeStaticMeshFactoryNode::SetCustomVertexColorReplace(const bool& AttributeValue)
-{
-	IMPLEMENT_NODE_ATTRIBUTE_SETTER_NODELEGATE(VertexColorReplace, bool)
-}
-
-bool UInterchangeStaticMeshFactoryNode::GetCustomVertexColorIgnore(bool& AttributeValue) const
-{
-	IMPLEMENT_NODE_ATTRIBUTE_GETTER(VertexColorIgnore, bool)
-}
-
-bool UInterchangeStaticMeshFactoryNode::SetCustomVertexColorIgnore(const bool& AttributeValue)
-{
-	IMPLEMENT_NODE_ATTRIBUTE_SETTER_NODELEGATE(VertexColorIgnore, bool)
-}
-
-bool UInterchangeStaticMeshFactoryNode::GetCustomVertexColorOverride(FColor& AttributeValue) const
-{
-	IMPLEMENT_NODE_ATTRIBUTE_GETTER(VertexColorOverride, FColor)
-}
-
-bool UInterchangeStaticMeshFactoryNode::SetCustomVertexColorOverride(const FColor& AttributeValue)
-{
-	IMPLEMENT_NODE_ATTRIBUTE_SETTER_NODELEGATE(VertexColorOverride, FColor)
-}
-
 void UInterchangeStaticMeshFactoryNode::FillAssetClassFromAttribute()
 {
 #if WITH_ENGINE
@@ -200,12 +108,3 @@ bool UInterchangeStaticMeshFactoryNode::SetNodeClassFromClassAttribute()
 	}
 	return bIsNodeClassInitialized;
 }
-
-bool UInterchangeStaticMeshFactoryNode::IsEditorOnlyDataDefined()
-{
-#if WITH_EDITORONLY_DATA
-	return true;
-#else
-	return false;
-#endif
-}
diff --git a/Engine/Plugins/Experimental/Interchange/Runtime/Source/FactoryNodes/Public/InterchangeMeshActorFactoryNode.h b/Engine/Plugins/Experimental/Interchange/Runtime/Source/FactoryNodes/Public/InterchangeMeshActorFactoryNode.h
new file mode 100644
--- /dev/null
+++ b/Engine/Plugins/Experimental/Interchange/Runtime/Source/FactoryNodes/Public/InterchangeMeshActorFactoryNode.h
@@ -0,0 +1,56 @@
+// Copyright Epic Games, Inc. All Rights Reserved.
+
+#pragma once
+
+#include "InterchangeActorFactoryNode.h"
+
+#include "InterchangeMeshActorFactoryNode.generated.h"
+
+UCLASS(BlueprintType, Experimental)
+class INTERCHANGEFACTORYNODES_API UInterchangeMeshActorFactoryNode : public UInterchangeActorFactoryNode
+{
+	GENERATED_BODY()
+
+public:
+	UInterchangeMeshActorFactoryNode();
+
+	/**
+	 * Override serialize to restore SlotMaterialDependencies on load.
+	 */
+	virtual void Serialize(FArchive& Ar) override
+	{
+		Super::Serialize(Ar);
+
+		if (Ar.IsLoading() && bIsInitialized)
+		{
+			SlotMaterialDependencies.RebuildCache();
+		}
+	}
+
+	/**
+	 * Allow to retrieve the correspondence table between slot names and assigned materials for this object.
+	 */
+	UFUNCTION(BlueprintCallable, Category = "Interchange | Node | StaticMesh")
+	void GetSlotMaterialDependencies(TMap<FString, FString>& OutMaterialDependencies) const;
+
+	/**
+	 * Allow to retrieve one Material dependency for a given slot of this object.
+	 */
+	UFUNCTION(BlueprintCallable, Category = "Interchange | Node | StaticMesh")
+	bool GetSlotMaterialDependencyUid(const FString& SlotName, FString& OutMaterialDependency) const;
+
+	/**
+	 * Add one Material dependency to a specific slot name of this object.
+	 */
+	UFUNCTION(BlueprintCallable, Category = "Interchange | Node | StaticMesh")
+	bool SetSlotMaterialDependencyUid(const FString& SlotName, const FString& MaterialDependencyUid);
+
+	/**
+	 * Remove the Material dependency associated with the given slot name from this object.
+	 */
+	UFUNCTION(BlueprintCallable, Category = "Interchange | Node | StaticMesh")
+	bool RemoveSlotMaterialDependencyUid(const FString& SlotName);
+
+private:
+	UE::Interchange::TMapAttributeHelper<FString, FString> SlotMaterialDependencies;
+};
\ No newline at end of file
diff --git a/Engine/Plugins/Experimental/Interchange/Runtime/Source/FactoryNodes/Public/InterchangeMeshFactoryNode.h b/Engine/Plugins/Experimental/Interchange/Runtime/Source/FactoryNodes/Public/InterchangeMeshFactoryNode.h
new file mode 100644
--- /dev/null
+++ b/Engine/Plugins/Experimental/Interchange/Runtime/Source/FactoryNodes/Public/InterchangeMeshFactoryNode.h
@@ -0,0 +1,118 @@
+// Copyright Epic Games, Inc. All Rights Reserved.
+
+#pragma once
+
+#include "CoreMinimal.h"
+#include "Nodes/InterchangeFactoryBaseNode.h"
+
+#include "InterchangeMeshFactoryNode.generated.h"
+
+
+namespace UE::Interchange
+{
+	struct INTERCHANGEFACTORYNODES_API FMeshFactoryNodeStaticData : public FBaseNodeStaticData
+	{
+		static const FAttributeKey& GetLodDependenciesBaseKey();
+		static const FAttributeKey& GetSlotMaterialDependencyBaseKey();
+	};
+} // namespace Interchange
+
+
+UCLASS(BlueprintType, Experimental, Abstract)
+class INTERCHANGEFACTORYNODES_API UInterchangeMeshFactoryNode : public UInterchangeFactoryBaseNode
+{
+	GENERATED_BODY()
+
+public:
+	UInterchangeMeshFactoryNode();
+
+	/**
+	 * Override serialize to restore SlotMaterialDependencies on load.
+	 */
+	virtual void Serialize(FArchive& Ar) override
+	{
+		Super::Serialize(Ar);
+
+		if (Ar.IsLoading() && bIsInitialized)
+		{
+			SlotMaterialDependencies.RebuildCache();
+#if WITH_ENGINE
+			// Make sure the class is properly set when we compile with engine, this will set the bIsNodeClassInitialized to true.
+			SetNodeClassFromClassAttribute();
+#endif
+		}
+	}
+
+	virtual FString GetKeyDisplayName(const UE::Interchange::FAttributeKey& NodeAttributeKey) const override;
+
+public:
+	/** Return The number of LOD this static mesh has.*/
+	UFUNCTION(BlueprintCallable, Category = "Interchange | FactoryNode | Mesh")
+	int32 GetLodDataCount() const;
+
+	UFUNCTION(BlueprintCallable, Category = "Interchange | FactoryNode | Mesh")
+	void GetLodDataUniqueIds(TArray<FString>& OutLodDataUniqueIds) const;
+
+	UFUNCTION(BlueprintCallable, Category = "Interchange | FactoryNode | Mesh")
+	bool AddLodDataUniqueId(const FString& LodDataUniqueId);
+
+	UFUNCTION(BlueprintCallable, Category = "Interchange | FactoryNode | Mesh")
+	bool RemoveLodDataUniqueId(const FString& LodDataUniqueId);
+
+	/** Query weather the static mesh factory should replace the vertex color. Return false if the attribute was not set.*/
+	UFUNCTION(BlueprintCallable, Category = "Interchange | FactoryNode | Mesh")
+	bool GetCustomVertexColorReplace(bool& AttributeValue) const;
+
+	/** Set weather the static mesh factory should replace the vertex color. Return false if the attribute cannot be set.*/
+	UFUNCTION(BlueprintCallable, Category = "Interchange | FactoryNode | Mesh")
+	bool SetCustomVertexColorReplace(const bool& AttributeValue);
+
+	/** Query weather the static mesh factory should ignore the vertex color. Return false if the attribute was not set.*/
+	UFUNCTION(BlueprintCallable, Category = "Interchange | FactoryNode | Mesh")
+	bool GetCustomVertexColorIgnore(bool& AttributeValue) const;
+
+	/** Set weather the static mesh factory should ignore the vertex color. Return false if the attribute cannot be set.*/
+	UFUNCTION(BlueprintCallable, Category = "Interchange | FactoryNode | Mesh")
+	bool SetCustomVertexColorIgnore(const bool& AttributeValue);
+
+	/** Query weather the static mesh factory should override the vertex color. Return false if the attribute was not set.*/
+	UFUNCTION(BlueprintCallable, Category = "Interchange | FactoryNode | Mesh")
+	bool GetCustomVertexColorOverride(FColor& AttributeValue) const;
+
+	/** Set weather the static mesh factory should override the vertex color. Return false if the attribute cannot be set.*/
+	UFUNCTION(BlueprintCallable, Category = "Interchange | FactoryNode | Mesh")
+	bool SetCustomVertexColorOverride(const FColor& AttributeValue);
+
+	/** Allow to retrieve the correspondence table between slot names and assigned materials for this object. */
+	UFUNCTION(BlueprintCallable, Category = "Interchange | FactoryNode | Mesh")
+	void GetSlotMaterialDependencies(TMap<FString, FString>& OutMaterialDependencies) const;
+
+	/** Allow to retrieve one Material dependency for a given slot of this object. */
+	UFUNCTION(BlueprintCallable, Category = "Interchange | FactoryNode | Mesh")
+	bool GetSlotMaterialDependencyUid(const FString& SlotName, FString& OutMaterialDependency) const;
+
+	/** Add one Material dependency to a specific slot name of this object. */
+	UFUNCTION(BlueprintCallable, Category = "Interchange | FactoryNode | Mesh")
+	bool SetSlotMaterialDependencyUid(const FString& SlotName, const FString& MaterialDependencyUid);
+
+	/** Remove the Material dependency associated with the given slot name from this object. */
+	UFUNCTION(BlueprintCallable, Category = "Interchange | FactoryNode | Mesh")
+	bool RemoveSlotMaterialDependencyUid(const FString& SlotName);
+
+private:
+
+	const UE::Interchange::FAttributeKey Macro_CustomVertexColorReplaceKey = UE::Interchange::FAttributeKey(TEXT("VertexColorReplace"));
+	const UE::Interchange::FAttributeKey Macro_CustomVertexColorIgnoreKey = UE::Interchange::FAttributeKey(TEXT("VertexColorIgnore"));
+	const UE::Interchange::FAttributeKey Macro_CustomVertexColorOverrideKey = UE::Interchange::FAttributeKey(TEXT("VertexColorOverride"));
+
+	UE::Interchange::TArrayAttributeHelper<FString> LodDependencies;
+	UE::Interchange::TMapAttributeHelper<FString, FString> SlotMaterialDependencies;
+
+protected:
+	virtual void FillAssetClassFromAttribute() PURE_VIRTUAL("FillAssetClassFromAttribute");
+	virtual bool SetNodeClassFromClassAttribute() PURE_VIRTUAL("SetNodeClassFromClassAttribute", return false;);
+
+	const UE::Interchange::FAttributeKey ClassNameAttributeKey = UE::Interchange::FBaseNodeStaticData::ClassTypeAttributeKey();
+
+	bool bIsNodeClassInitialized = false;
+};
diff --git a/Engine/Plugins/Experimental/Interchange/Runtime/Source/FactoryNodes/Public/InterchangeSkeletalMeshFactoryNode.h b/Engine/Plugins/Experimental/Interchange/Runtime/Source/FactoryNodes/Public/InterchangeSkeletalMeshFactoryNode.h
--- a/Engine/Plugins/Experimental/Interchange/Runtime/Source/FactoryNodes/Public/InterchangeSkeletalMeshFactoryNode.h
+++ b/Engine/Plugins/Experimental/Interchange/Runtime/Source/FactoryNodes/Public/InterchangeSkeletalMeshFactoryNode.h
@@ -3,6 +3,7 @@
 #pragma once
 
 #include "CoreMinimal.h"
+#include "InterchangeMeshFactoryNode.h"
 #include "Nodes/InterchangeFactoryBaseNode.h"
 
 #if WITH_ENGINE
@@ -11,15 +12,6 @@
 
 #include "InterchangeSkeletalMeshFactoryNode.generated.h"
 
-//Interchange namespace
-namespace UE::Interchange
-{
-	struct INTERCHANGEFACTORYNODES_API FSkeletalMeshNodeStaticData : public FBaseNodeStaticData
-	{
-		static const FAttributeKey& GetLodDependenciesBaseKey();
-	};
-}//ns UE::Interchange
-
 UENUM(BlueprintType)
 enum class EInterchangeSkeletalMeshContentType : uint8
 {
@@ -30,7 +22,7 @@ enum class EInterchangeSkeletalMeshContentType : uint8
 };
 
 UCLASS(BlueprintType, Experimental)
-class INTERCHANGEFACTORYNODES_API UInterchangeSkeletalMeshFactoryNode : public UInterchangeFactoryBaseNode
+class INTERCHANGEFACTORYNODES_API UInterchangeSkeletalMeshFactoryNode : public UInterchangeMeshFactoryNode
 {
 	GENERATED_BODY()
 
@@ -46,33 +38,16 @@ public:
 	UFUNCTION(BlueprintCallable, Category = "Interchange | Node | SkeletalMesh")
 	void InitializeSkeletalMeshNode(const FString& UniqueID, const FString& DisplayLabel, const FString& InAssetClass);
 
-	virtual void Serialize(FArchive& Ar) override;
-
 	/**
 	 * Return the node type name of the class, we use this when reporting error
 	 */
 	virtual FString GetTypeName() const override;
 
-	virtual FString GetKeyDisplayName(const UE::Interchange::FAttributeKey& NodeAttributeKey) const override;
-
 	/** Get the class this node want to create */
 	UFUNCTION(BlueprintCallable, Category = "Interchange | Node | SkeletalMesh")
 	virtual class UClass* GetObjectClass() const override;
 
 public:
-	/** Return The number of LOD this skeletalmesh has.*/
-	UFUNCTION(BlueprintCallable, Category = "Interchange | Node | SkeletalMesh")
-	int32 GetLodDataCount() const;
-
-	UFUNCTION(BlueprintCallable, Category = "Interchange | Node | SkeletalMesh")
-	void GetLodDataUniqueIds(TArray<FString>& OutLodDataUniqueIds) const;
-
-	UFUNCTION(BlueprintCallable, Category = "Interchange | Node | SkeletalMesh")
-	bool AddLodDataUniqueId(const FString& LodDataUniqueId);
-
-	UFUNCTION(BlueprintCallable, Category = "Interchange | Node | SkeletalMesh")
-	bool RemoveLodDataUniqueId(const FString& LodDataUniqueId);
-	
 	/** Query the skeletal mesh factory skeleton UObject. Return false if the attribute was not set.*/
 	UFUNCTION(BlueprintCallable, Category = "Interchange | Node | SkeletalMesh")
 	bool GetCustomSkeletonSoftObjectPath(FSoftObjectPath& AttributeValue) const;
@@ -105,30 +80,6 @@ public:
 	UFUNCTION(BlueprintCallable, Category = "Interchange | Node | SkeletalMesh")
 	bool SetCustomPhysicAssetSoftObjectPath(const FSoftObjectPath& AttributeValue);
 
-	/** Query weather the skeletal mesh factory should replace the vertex color. Return false if the attribute was not set.*/
-	UFUNCTION(BlueprintCallable, Category = "Interchange | Node | SkeletalMesh")
-	bool GetCustomVertexColorReplace(bool& AttributeValue) const;
-
-	/** Set weather the skeletal mesh factory should replace the vertex color. Return false if the attribute cannot be set.*/
-	UFUNCTION(BlueprintCallable, Category = "Interchange | Node | SkeletalMesh")
-	bool SetCustomVertexColorReplace(const bool& AttributeValue);
-
-	/** Query weather the skeletal mesh factory should ignore the vertex color. Return false if the attribute was not set.*/
-	UFUNCTION(BlueprintCallable, Category = "Interchange | Node | SkeletalMesh")
-	bool GetCustomVertexColorIgnore(bool& AttributeValue) const;
-
-	/** Set weather the skeletal mesh factory should ignore the vertex color. Return false if the attribute cannot be set.*/
-	UFUNCTION(BlueprintCallable, Category = "Interchange | Node | SkeletalMesh")
-	bool SetCustomVertexColorIgnore(const bool& AttributeValue);
-
-	/** Query weather the skeletal mesh factory should override the vertex color. Return false if the attribute was not set.*/
-	UFUNCTION(BlueprintCallable, Category = "Interchange | Node | SkeletalMesh")
-	bool GetCustomVertexColorOverride(FColor& AttributeValue) const;
-
-	/** Set weather the skeletal mesh factory should override the vertex color. Return false if the attribute cannot be set.*/
-	UFUNCTION(BlueprintCallable, Category = "Interchange | Node | SkeletalMesh")
-	bool SetCustomVertexColorOverride(const FColor& AttributeValue);
-
 	/** Query the skeletal mesh import content type. The content type is use by the factory to import partial or full translated content. Return false if the attribute was not set.*/
 	UFUNCTION(BlueprintCallable, Category = "Interchange | Node | SkeletalMesh")
 	bool GetCustomImportContentType(EInterchangeSkeletalMeshContentType& AttributeValue) const;
@@ -143,26 +94,18 @@ public:
 	virtual void AppendAssetRegistryTags(TArray<FAssetRegistryTag>& OutTags) const override;
 private:
 
-	void FillAssetClassFromAttribute();
-
-	bool SetNodeClassFromClassAttribute();
-
-	bool IsEditorOnlyDataDefined();
+	virtual void FillAssetClassFromAttribute() override;
+	virtual bool SetNodeClassFromClassAttribute() override;
 
 	const UE::Interchange::FAttributeKey ClassNameAttributeKey = UE::Interchange::FBaseNodeStaticData::ClassTypeAttributeKey();
 	const UE::Interchange::FAttributeKey Macro_CustomImportMorphTargetKey = UE::Interchange::FAttributeKey(TEXT("ImportMorphTarget"));
 	const UE::Interchange::FAttributeKey Macro_CustomSkeletonSoftObjectPathKey = UE::Interchange::FAttributeKey(TEXT("SkeletonSoftObjectPath"));
 	const UE::Interchange::FAttributeKey Macro_CustomCreatePhysicsAssetKey = UE::Interchange::FAttributeKey(TEXT("CreatePhysicsAsset"));
 	const UE::Interchange::FAttributeKey Macro_CustomPhysicAssetSoftObjectPathKey = UE::Interchange::FAttributeKey(TEXT("PhysicAssetSoftObjectPath"));
-	const UE::Interchange::FAttributeKey Macro_CustomVertexColorReplaceKey = UE::Interchange::FAttributeKey(TEXT("VertexColorReplace"));
-	const UE::Interchange::FAttributeKey Macro_CustomVertexColorIgnoreKey = UE::Interchange::FAttributeKey(TEXT("VertexColorIgnore"));
-	const UE::Interchange::FAttributeKey Macro_CustomVertexColorOverrideKey = UE::Interchange::FAttributeKey(TEXT("VertexColorOverride"));
 	const UE::Interchange::FAttributeKey Macro_CustomImportContentTypeKey = UE::Interchange::FAttributeKey(TEXT("ImportContentType"));
 
-	UE::Interchange::TArrayAttributeHelper<FString> LodDependencies;
 protected:
 #if WITH_ENGINE
 	TSubclassOf<USkeletalMesh> AssetClass = nullptr;
 #endif
-	bool bIsNodeClassInitialized = false;
 };
diff --git a/Engine/Plugins/Experimental/Interchange/Runtime/Source/FactoryNodes/Public/InterchangeStaticMeshFactoryNode.h b/Engine/Plugins/Experimental/Interchange/Runtime/Source/FactoryNodes/Public/InterchangeStaticMeshFactoryNode.h
--- a/Engine/Plugins/Experimental/Interchange/Runtime/Source/FactoryNodes/Public/InterchangeStaticMeshFactoryNode.h
+++ b/Engine/Plugins/Experimental/Interchange/Runtime/Source/FactoryNodes/Public/InterchangeStaticMeshFactoryNode.h
@@ -3,6 +3,7 @@
 #pragma once
 
 #include "CoreMinimal.h"
+#include "InterchangeMeshFactoryNode.h"
 #include "Nodes/InterchangeFactoryBaseNode.h"
 
 #if WITH_ENGINE
@@ -18,7 +19,6 @@ namespace UE
 	{
 		struct INTERCHANGEFACTORYNODES_API FStaticMeshNodeStaticData : public FBaseNodeStaticData
 		{
-			static const FAttributeKey& GetLodDependenciesBaseKey();
 			static const FAttributeKey& GetSocketUidsBaseKey();
 		};
 	} // namespace Interchange
@@ -26,7 +26,7 @@ namespace UE
 
 
 UCLASS(BlueprintType, Experimental)
-class INTERCHANGEFACTORYNODES_API UInterchangeStaticMeshFactoryNode : public UInterchangeFactoryBaseNode
+class INTERCHANGEFACTORYNODES_API UInterchangeStaticMeshFactoryNode : public UInterchangeMeshFactoryNode
 {
 	GENERATED_BODY()
 
@@ -42,33 +42,16 @@ public:
 	UFUNCTION(BlueprintCallable, Category = "Interchange | Node | StaticMesh")
 	void InitializeStaticMeshNode(const FString& UniqueID, const FString& DisplayLabel, const FString& InAssetClass);
 
-	virtual void Serialize(FArchive& Ar) override;
-
 	/**
 	 * Return the node type name of the class, we use this when reporting error
 	 */
 	virtual FString GetTypeName() const override;
 
-	virtual FString GetKeyDisplayName(const UE::Interchange::FAttributeKey& NodeAttributeKey) const override;
-
 	/** Get the class this node want to create */
 	UFUNCTION(BlueprintCallable, Category = "Interchange | Node | StaticMesh")
 	virtual class UClass* GetObjectClass() const override;
 
 public:
-	/** Return The number of LOD this static mesh has.*/
-	UFUNCTION(BlueprintCallable, Category = "Interchange | Node | StaticMesh")
-	int32 GetLodDataCount() const;
-
-	UFUNCTION(BlueprintCallable, Category = "Interchange | Node | StaticMesh")
-	void GetLodDataUniqueIds(TArray<FString>& OutLodDataUniqueIds) const;
-
-	UFUNCTION(BlueprintCallable, Category = "Interchange | Node | StaticMesh")
-	bool AddLodDataUniqueId(const FString& LodDataUniqueId);
-
-	UFUNCTION(BlueprintCallable, Category = "Interchange | Node | StaticMesh")
-	bool RemoveLodDataUniqueId(const FString& LodDataUniqueId);
-
 	/** Return The number of socket UIDs this static mesh has.*/
 	UFUNCTION(BlueprintCallable, Category = "Interchange | Node | StaticMesh")
 	int32 GetSocketUidCount() const;
@@ -85,46 +68,15 @@ public:
 	UFUNCTION(BlueprintCallable, Category = "Interchange | Node | StaticMesh")
 	bool RemoveSocketUd(const FString& SocketUid);
 
-	/** Query weather the static mesh factory should replace the vertex color. Return false if the attribute was not set.*/
-	UFUNCTION(BlueprintCallable, Category = "Interchange | Node | StaticMesh")
-	bool GetCustomVertexColorReplace(bool& AttributeValue) const;
-
-	/** Set weather the static mesh factory should replace the vertex color. Return false if the attribute cannot be set.*/
-	UFUNCTION(BlueprintCallable, Category = "Interchange | Node | StaticMesh")
-	bool SetCustomVertexColorReplace(const bool& AttributeValue);
-
-	/** Query weather the static mesh factory should ignore the vertex color. Return false if the attribute was not set.*/
-	UFUNCTION(BlueprintCallable, Category = "Interchange | Node | StaticMesh")
-	bool GetCustomVertexColorIgnore(bool& AttributeValue) const;
-
-	/** Set weather the static mesh factory should ignore the vertex color. Return false if the attribute cannot be set.*/
-	UFUNCTION(BlueprintCallable, Category = "Interchange | Node | StaticMesh")
-	bool SetCustomVertexColorIgnore(const bool& AttributeValue);
-
-	/** Query weather the static mesh factory should override the vertex color. Return false if the attribute was not set.*/
-	UFUNCTION(BlueprintCallable, Category = "Interchange | Node | StaticMesh")
-	bool GetCustomVertexColorOverride(FColor& AttributeValue) const;
-
-	/** Set weather the static mesh factory should override the vertex color. Return false if the attribute cannot be set.*/
-	UFUNCTION(BlueprintCallable, Category = "Interchange | Node | StaticMesh")
-	bool SetCustomVertexColorOverride(const FColor& AttributeValue);
-
 private:
 
-	void FillAssetClassFromAttribute();
-	bool SetNodeClassFromClassAttribute();
-	bool IsEditorOnlyDataDefined();
-
-	const UE::Interchange::FAttributeKey ClassNameAttributeKey = UE::Interchange::FBaseNodeStaticData::ClassTypeAttributeKey();
-	const UE::Interchange::FAttributeKey Macro_CustomVertexColorReplaceKey = UE::Interchange::FAttributeKey(TEXT("VertexColorReplace"));
-	const UE::Interchange::FAttributeKey Macro_CustomVertexColorIgnoreKey = UE::Interchange::FAttributeKey(TEXT("VertexColorIgnore"));
-	const UE::Interchange::FAttributeKey Macro_CustomVertexColorOverrideKey = UE::Interchange::FAttributeKey(TEXT("VertexColorOverride"));
+	virtual void FillAssetClassFromAttribute() override;
+	virtual bool SetNodeClassFromClassAttribute() override;
 
-	UE::Interchange::TArrayAttributeHelper<FString> LodDependencies;
 	UE::Interchange::TArrayAttributeHelper<FString> SocketUids;
+
 protected:
 #if WITH_ENGINE
 	TSubclassOf<UStaticMesh> AssetClass = nullptr;
 #endif
-	bool bIsNodeClassInitialized = false;
 };
diff --git a/Engine/Plugins/Experimental/Interchange/Runtime/Source/FactoryNodes/Public/InterchangeTexture2DFactoryNode.h b/Engine/Plugins/Experimental/Interchange/Runtime/Source/FactoryNodes/Public/InterchangeTexture2DFactoryNode.h
--- a/Engine/Plugins/Experimental/Interchange/Runtime/Source/FactoryNodes/Public/InterchangeTexture2DFactoryNode.h
+++ b/Engine/Plugins/Experimental/Interchange/Runtime/Source/FactoryNodes/Public/InterchangeTexture2DFactoryNode.h
@@ -56,6 +56,19 @@ public:
 		IMPLEMENT_NODE_ATTRIBUTE_SETTER(UInterchangeTexture2DFactoryNode, AddressY, TEnumAsByte<TextureAddress>, UTexture2D)
 	}
 
+	/**
+	 * Override serialize to restore SourceBlocks on load.
+	 */
+	virtual void Serialize(FArchive& Ar) override
+	{
+		Super::Serialize(Ar);
+
+		if (Ar.IsLoading() && bIsInitialized)
+		{
+			SourceBlocks.RebuildCache();
+		}
+	}
+
 	//////////////////////////////////////////////////////////////////////////
 	// UDIMs begin here
 	// UDIM base texture use a different model for the source data
diff --git a/Engine/Plugins/Experimental/Interchange/Runtime/Source/Import/Private/Gltf/InterchangeGltfTranslator.cpp b/Engine/Plugins/Experimental/Interchange/Runtime/Source/Import/Private/Gltf/InterchangeGltfTranslator.cpp
--- a/Engine/Plugins/Experimental/Interchange/Runtime/Source/Import/Private/Gltf/InterchangeGltfTranslator.cpp
+++ b/Engine/Plugins/Experimental/Interchange/Runtime/Source/Import/Private/Gltf/InterchangeGltfTranslator.cpp
@@ -764,8 +764,9 @@ bool UInterchangeGltfTranslator::Translate( UInterchangeBaseNodeContainer& NodeC
 			{
 				if ( GltfAsset.Materials.IsValidIndex( Primitive.MaterialIndex ) )
 				{
-					const FString ShaderGraphNodeUid = UInterchangeShaderGraphNode::MakeNodeUid( TEXT("Material_") + GenerateNameForGltfNode( GltfAsset.Materials[ Primitive.MaterialIndex ].Name, Primitive.MaterialIndex ) );
-					MeshNode->SetMaterialDependencyUid( ShaderGraphNodeUid );
+					const FString MaterialName = GenerateNameForGltfNode(GltfAsset.Materials[Primitive.MaterialIndex].Name, Primitive.MaterialIndex);
+					const FString ShaderGraphNodeUid = UInterchangeShaderGraphNode::MakeNodeUid( TEXT("Material_") + MaterialName);
+					MeshNode->SetSlotMaterialDependencyUid(MaterialName, ShaderGraphNodeUid);
 				}
 			}
 		}
@@ -864,7 +865,8 @@ TFuture< TOptional< UE::Interchange::FStaticMeshPayloadData > > UInterchangeGltf
 
 			if ( GltfAsset.Materials.IsValidIndex( MaterialIndex ) )
 			{
-				StaticMeshAttributes.GetPolygonGroupMaterialSlotNames()[ MaterialSlotIndex ] = *(TEXT("Material_") + GenerateNameForGltfNode( GltfAsset.Materials[ MaterialIndex ].Name, MaterialIndex ));
+				const FString MaterialName = GenerateNameForGltfNode(GltfAsset.Materials[MaterialIndex].Name, MaterialIndex);
+				StaticMeshAttributes.GetPolygonGroupMaterialSlotNames()[ MaterialSlotIndex ] = *MaterialName;
 			}
 		}
 	}
diff --git a/Engine/Plugins/Experimental/Interchange/Runtime/Source/Import/Private/Mesh/InterchangeOBJTranslator.cpp b/Engine/Plugins/Experimental/Interchange/Runtime/Source/Import/Private/Mesh/InterchangeOBJTranslator.cpp
--- a/Engine/Plugins/Experimental/Interchange/Runtime/Source/Import/Private/Mesh/InterchangeOBJTranslator.cpp
+++ b/Engine/Plugins/Experimental/Interchange/Runtime/Source/Import/Private/Mesh/InterchangeOBJTranslator.cpp
@@ -939,7 +939,8 @@ bool UInterchangeOBJTranslator::Translate(UInterchangeBaseNodeContainer& BaseNod
 
 			if (!Group.Value.MaterialName.IsEmpty())
 			{
-				MeshNode->SetMaterialDependencyUid(UInterchangeShaderGraphNode::MakeNodeUid(MakeShaderGraphNodeName(Group.Value.MaterialName)));
+				const FString MaterialDependencyUid = UInterchangeShaderGraphNode::MakeNodeUid(MakeShaderGraphNodeName(Group.Value.MaterialName));
+				MeshNode->SetSlotMaterialDependencyUid(Group.Value.MaterialName, MaterialDependencyUid);
 			}
 		}
 
diff --git a/Engine/Plugins/Experimental/Interchange/Runtime/Source/Import/Private/Mesh/InterchangeSkeletalMeshFactory.cpp b/Engine/Plugins/Experimental/Interchange/Runtime/Source/Import/Private/Mesh/InterchangeSkeletalMeshFactory.cpp
--- a/Engine/Plugins/Experimental/Interchange/Runtime/Source/Import/Private/Mesh/InterchangeSkeletalMeshFactory.cpp
+++ b/Engine/Plugins/Experimental/Interchange/Runtime/Source/Import/Private/Mesh/InterchangeSkeletalMeshFactory.cpp
@@ -18,6 +18,7 @@
 #include "InterchangeSkeletonHelper.h"
 #include "InterchangeSourceData.h"
 #include "InterchangeTranslatorBase.h"
+#include "Materials/Material.h"
 #include "Math/GenericOctree.h"
 #include "Mesh/InterchangeSkeletalMeshPayload.h"
 #include "Mesh/InterchangeSkeletalMeshPayloadInterface.h"
@@ -542,72 +543,6 @@ namespace UE
 				CopyBlendShapesMeshDescriptionToSkeletalMeshImportData(BlendShapeMeshDescriptionsPerBlendShapeName, DestinationImportData);
 			}
 
-			/**
-			 * Fill the Materials array using the raw skeletalmesh geometry data (using material imported name)
-			 * Find the material from the dependencies of the skeletalmesh before searching in all package.
-			 */
-			//TODO: the pipeline should search for existing material and hook those before the factory is called
-			void ProcessImportMeshMaterials(TArray<FSkeletalMaterial>& Materials, FSkeletalMeshImportData& ImportData, TMap<FString, UMaterialInterface*>& AvailableMaterials)
-			{
-				TArray <SkeletalMeshImportData::FMaterial>& ImportedMaterials = ImportData.Materials;
-				// If direct linkup of materials is requested, try to find them here - to get a texture name from a 
-				// material name, cut off anything in front of the dot (beyond are special flags).
-				int32 SkinOffset = INDEX_NONE;
-				for (int32 MatIndex = 0; MatIndex < ImportedMaterials.Num(); ++MatIndex)
-				{
-					const SkeletalMeshImportData::FMaterial& ImportedMaterial = ImportedMaterials[MatIndex];
-
-					UMaterialInterface* Material = nullptr;
-
-					const FName SearchMaterialSlotName(*ImportedMaterial.MaterialImportName);
-					int32 MaterialIndex = 0;
-					FSkeletalMaterial* SkeletalMeshMaterialFind = Materials.FindByPredicate([&SearchMaterialSlotName, &MaterialIndex](const FSkeletalMaterial& ItemMaterial)
-					{
-						//Imported material slot name is available only WITH_EDITOR
-						FName ImportedMaterialSlot = NAME_None;
-						ImportedMaterialSlot = ItemMaterial.ImportedMaterialSlotName;
-						if (ImportedMaterialSlot != SearchMaterialSlotName)
-						{
-							MaterialIndex++;
-							return false;
-						}
-						return true;
-					});
-
-					if (SkeletalMeshMaterialFind != nullptr)
-					{
-						Material = SkeletalMeshMaterialFind->MaterialInterface;
-					}
-
-					if(!Material)
-					{
-						//Try to find the material in the skeletal mesh node dependencies (Materials are import before skeletal mesh when there is a dependency)
-						if (AvailableMaterials.Contains(ImportedMaterial.MaterialImportName))
-						{
-							Material = AvailableMaterials.FindChecked(ImportedMaterial.MaterialImportName);
-						}
-						else
-						{
-							//We did not found any material in the dependencies so try to find material everywhere
-							Material = FindObject<UMaterialInterface>(ANY_PACKAGE, *ImportedMaterial.MaterialImportName);
-						}
-
-						const bool bEnableShadowCasting = true;
-						const bool bInRecomputeTangent = false;
-						Materials.Add(FSkeletalMaterial(Material, bEnableShadowCasting, bInRecomputeTangent, Material != nullptr ? Material->GetFName() : FName(*ImportedMaterial.MaterialImportName), FName(*(ImportedMaterial.MaterialImportName))));
-					}
-				}
-
-				int32 NumMaterialsToAdd = FMath::Max<int32>(ImportedMaterials.Num(), ImportData.MaxMaterialIndex + 1);
-
-				// Pad the material pointers
-				while (NumMaterialsToAdd > Materials.Num())
-				{
-					UMaterialInterface* NullMaterialInterface = nullptr;
-					Materials.Add(FSkeletalMaterial(NullMaterialInterface));
-				}
-			}
-
 			void ProcessImportMeshInfluences(const int32 WedgeCount, TArray<SkeletalMeshImportData::FRawBoneInfluence>& Influences)
 			{
 
@@ -1086,6 +1021,44 @@ UObject* UInterchangeSkeletalMeshFactory::CreateAsset(const FCreateAssetParams&
 		SkeletalMeshFactoryNode->SetCustomVertexColorReplace(bFalseSetting);
 	}
 
+	// Update skeletal materials
+	TArray<FSkeletalMaterial>& Materials = SkeletalMesh->GetMaterials();
+
+	auto UpdateOrAddSkeletalMaterial = [&Materials](const FName& MaterialSlotName, UMaterialInterface* MaterialInterface)
+	{
+		FSkeletalMaterial* SkeletalMaterial = Materials.FindByPredicate([&MaterialSlotName](const FSkeletalMaterial& Material) { return Material.MaterialSlotName == MaterialSlotName; });
+		
+		if (SkeletalMaterial)
+		{
+			SkeletalMaterial->MaterialInterface = MaterialInterface;
+		}
+		else
+		{
+			const bool bEnableShadowCasting = true;
+			const bool bInRecomputeTangent = false;
+			Materials.Emplace(MaterialInterface, bEnableShadowCasting, bInRecomputeTangent, MaterialSlotName, MaterialSlotName);
+		}
+	};
+
+	TMap<FString, FString> SlotMaterialDependencies;
+	SkeletalMeshFactoryNode->GetSlotMaterialDependencies(SlotMaterialDependencies);
+	Materials.Reserve(SlotMaterialDependencies.Num());
+
+	for (TPair<FString, FString>& SlotMaterialDependency : SlotMaterialDependencies)
+	{
+		FName MaterialSlotName = *SlotMaterialDependency.Key;
+
+		const UInterchangeBaseMaterialFactoryNode* MaterialFactoryNode = Cast<UInterchangeBaseMaterialFactoryNode>(Arguments.NodeContainer->GetNode(SlotMaterialDependency.Value));
+		if (!MaterialFactoryNode || !MaterialFactoryNode->ReferenceObject.IsValid() || !MaterialFactoryNode->IsEnabled())
+		{
+			UpdateOrAddSkeletalMaterial(MaterialSlotName, UMaterial::GetDefaultMaterial(MD_Surface));
+			continue;
+		}
+
+		UMaterialInterface* MaterialInterface = Cast<UMaterialInterface>(MaterialFactoryNode->ReferenceObject.ResolveObject());
+		UpdateOrAddSkeletalMaterial(MaterialSlotName, MaterialInterface ? MaterialInterface : UMaterial::GetDefaultMaterial(MD_Surface));
+	}
+
 	for (int32 LodIndex = 0; LodIndex < LodCount; ++LodIndex)
 	{
 		ESkeletalMeshGeoImportVersions GeoImportVersion = ESkeletalMeshGeoImportVersions::LatestVersion;
@@ -1297,31 +1270,9 @@ UObject* UInterchangeSkeletalMeshFactory::CreateAsset(const FCreateAssetParams&
 		{
 			ensure(ImportedResource->LODModels.Add(new FSkeletalMeshLODModel()) == CurrentLodIndex);
 		}
-		FSkeletalMeshLODModel& LODModel = ImportedResource->LODModels[CurrentLodIndex];
 
-		TMap<FString, UMaterialInterface*> AvailableMaterials;
-		TArray<FString> FactoryDependencies;
-		SkeletalMeshFactoryNode->GetFactoryDependencies(FactoryDependencies);
-		for (int32 DependencyIndex = 0; DependencyIndex < FactoryDependencies.Num(); ++DependencyIndex)
-		{
-			const UInterchangeMaterialFactoryNode* MaterialFactoryNode = Cast<UInterchangeMaterialFactoryNode>(Arguments.NodeContainer->GetNode(FactoryDependencies[DependencyIndex]));
-			if (!MaterialFactoryNode || !MaterialFactoryNode->ReferenceObject.IsValid())
-			{
-				continue;
-			}
-			if (!MaterialFactoryNode->IsEnabled())
-			{
-				continue;
-			}
-			UMaterialInterface* MaterialInterface = Cast<UMaterialInterface>(MaterialFactoryNode->ReferenceObject.ResolveObject());
-			if (!MaterialInterface)
-			{
-				continue;
-			}
-			AvailableMaterials.Add(MaterialFactoryNode->GetDisplayLabel(), MaterialInterface);
-		}
+		FSkeletalMeshLODModel& LODModel = ImportedResource->LODModels[CurrentLodIndex];
 
-		UE::Interchange::Private::ProcessImportMeshMaterials(SkeletalMesh->GetMaterials(), SkeletalMeshImportData, AvailableMaterials);
 		UE::Interchange::Private::ProcessImportMeshInfluences(SkeletalMeshImportData.Wedges.Num(), SkeletalMeshImportData.Influences);
 
 		if (bApplyGeometryOnly)
@@ -1372,11 +1323,11 @@ UObject* UInterchangeSkeletalMeshFactory::CreateAsset(const FCreateAssetParams&
 			AddLodInfo();
 		}
 
-		TArray<FSkeletalMaterial>& Materials = SkeletalMesh->GetMaterials();
 		const TArray<SkeletalMeshImportData::FMaterial>& ImportedMaterials = SkeletalMeshImportData.Materials;
 		if (FSkeletalMeshLODInfo* LodInfo = SkeletalMesh->GetLODInfo(CurrentLodIndex))
 		{
 			LodInfo->LODMaterialMap.Empty();
+
 			// Now set up the material mapping array.
 			for (int32 ImportedMaterialIndex = 0; ImportedMaterialIndex < ImportedMaterials.Num(); ImportedMaterialIndex++)
 			{
@@ -1392,11 +1343,14 @@ UObject* UInterchangeSkeletalMeshFactory::CreateAsset(const FCreateAssetParams&
 						break;
 					}
 				}
-				// If we dont have a match, add a new entry to the material list.
-				if (LODMatIndex == INDEX_NONE)
+
+				// If we don't have a match, add a new entry to the material list.
+				// However this should not happen. The translator should have properly set the UInterchangeMeshNode with the slot names used by all LODs.
+				if (!ensure(LODMatIndex != INDEX_NONE))
 				{
 					LODMatIndex = Materials.Add(FSkeletalMaterial(ImportedMaterials[ImportedMaterialIndex].Material.Get(), true, false, ImportedMaterialName, ImportedMaterialName));
 				}
+
 				LodInfo->LODMaterialMap.Add(LODMatIndex);
 			}
 		}
diff --git a/Engine/Plugins/Experimental/Interchange/Runtime/Source/Import/Private/Mesh/InterchangeStaticMeshFactory.cpp b/Engine/Plugins/Experimental/Interchange/Runtime/Source/Import/Private/Mesh/InterchangeStaticMeshFactory.cpp
--- a/Engine/Plugins/Experimental/Interchange/Runtime/Source/Import/Private/Mesh/InterchangeStaticMeshFactory.cpp
+++ b/Engine/Plugins/Experimental/Interchange/Runtime/Source/Import/Private/Mesh/InterchangeStaticMeshFactory.cpp
@@ -95,7 +95,6 @@ UObject* UInterchangeStaticMeshFactory::CreateEmptyAsset(const FCreateAssetParam
 #endif //else !WITH_EDITOR || !WITH_EDITORONLY_DATA
 }
 
-
 UObject* UInterchangeStaticMeshFactory::CreateAsset(const FCreateAssetParams& Arguments)
 {
 #if !WITH_EDITOR || !WITH_EDITORONLY_DATA
@@ -151,52 +150,59 @@ UObject* UInterchangeStaticMeshFactory::CreateAsset(const FCreateAssetParams& Ar
 		return nullptr;
 	}
 			
-	int32 LodCount = StaticMeshFactoryNode->GetLodDataCount();
-	int32 FinalLodCount = LodCount;
+	const int32 LodCount = StaticMeshFactoryNode->GetLodDataCount();
+	const int32 PrevLodCount = StaticMesh->GetNumLODs();
+	const int32 FinalLodCount = FMath::Max(PrevLodCount, LodCount);
 
-	// If we are reimporting, cache the existing vertex colors so they can be optionally reapplied after reimport
-	TMap<FVector3f, FColor> ExisitingVertexColorData;
-	if (Arguments.ReimportObject)
+	StaticMesh->SetNumSourceModels(FinalLodCount);
+
+	// Make sure that mesh descriptions for added LODs are kept as is when the mesh is built
+	for (int32 LodIndex = PrevLodCount; LodIndex < FinalLodCount; ++LodIndex)
 	{
-		if(ExistingAsset)
-		{
-			StaticMesh->GetVertexColorData(ExisitingVertexColorData);
-		}
+		FStaticMeshSourceModel& SrcModel = StaticMesh->GetSourceModel(LodIndex);
 
-		// When we reimport we don't want to reduce the number of existing LODs
-		FinalLodCount = FMath::Max(StaticMesh->GetNumLODs(), LodCount);
+		SrcModel.ReductionSettings.MaxDeviation = 0.0f;
+		SrcModel.ReductionSettings.PercentTriangles = 1.0f;
+		SrcModel.ReductionSettings.PercentVertices = 1.0f;
 	}
 
-	StaticMesh->SetNumSourceModels(FinalLodCount);
+	// If we are reimporting, cache the existing vertex colors so they can be optionally reapplied after reimport
+	TMap<FVector3f, FColor> ExisitingVertexColorData;
+	if (Arguments.ReimportObject && ExistingAsset)
+	{
+		StaticMesh->GetVertexColorData(ExisitingVertexColorData);
+	}
 
 	// Set material slots from imported materials
-	TArray<FString> FactoryDependencies;
-	StaticMeshFactoryNode->GetFactoryDependencies(FactoryDependencies);
-
-	for (int32 DependencyIndex = 0; DependencyIndex < FactoryDependencies.Num(); ++DependencyIndex)
+	auto UpdateOrAddStaticMaterial = [&StaticMesh](const FName& MaterialSlotName, UMaterialInterface* MaterialInterface)
 	{
-		const UInterchangeBaseMaterialFactoryNode* MaterialFactoryNode = Cast<UInterchangeBaseMaterialFactoryNode>(Arguments.NodeContainer->GetNode(FactoryDependencies[DependencyIndex]));
-		if (!MaterialFactoryNode || !MaterialFactoryNode->ReferenceObject.IsValid())
+		int32 MaterialSlotIndex = StaticMesh->GetMaterialIndexFromImportedMaterialSlotName(MaterialSlotName);
+
+		if (MaterialSlotIndex == INDEX_NONE)
 		{
-			continue;
+			StaticMesh->GetStaticMaterials().Emplace(MaterialInterface, MaterialSlotName);
 		}
-		if (!MaterialFactoryNode->IsEnabled())
+		else
 		{
-			continue;
+			StaticMesh->GetStaticMaterials()[MaterialSlotIndex].MaterialInterface = MaterialInterface;
 		}
+	};
 
-		FName MaterialSlotName = *MaterialFactoryNode->GetDisplayLabel();
-		UMaterialInterface* MaterialInterface = Cast<UMaterialInterface>(MaterialFactoryNode->ReferenceObject.ResolveObject());
-		if (!MaterialInterface)
-		{
-			MaterialInterface = UMaterial::GetDefaultMaterial(MD_Surface);
-		}
+	TMap<FString, FString> SlotMaterialDependencies;
+	StaticMeshFactoryNode->GetSlotMaterialDependencies(SlotMaterialDependencies);
+	for (TPair<FString, FString>& SlotMaterialDependency : SlotMaterialDependencies)
+	{
+		FName MaterialSlotName = *SlotMaterialDependency.Key;
 
-		int32 MaterialSlotIndex = StaticMesh->GetMaterialIndexFromImportedMaterialSlotName(MaterialSlotName);
-		if (MaterialSlotIndex == INDEX_NONE)
+		const UInterchangeBaseMaterialFactoryNode* MaterialFactoryNode = Cast<UInterchangeBaseMaterialFactoryNode>(Arguments.NodeContainer->GetNode(SlotMaterialDependency.Value));
+		if (!MaterialFactoryNode || !MaterialFactoryNode->ReferenceObject.IsValid() || !MaterialFactoryNode->IsEnabled())
 		{
-			StaticMesh->GetStaticMaterials().Emplace(MaterialInterface, MaterialSlotName);
+			UpdateOrAddStaticMaterial(MaterialSlotName, UMaterial::GetDefaultMaterial(MD_Surface));
+			continue;
 		}
+
+		UMaterialInterface* MaterialInterface = Cast<UMaterialInterface>(MaterialFactoryNode->ReferenceObject.ResolveObject());
+		UpdateOrAddStaticMaterial(MaterialSlotName, MaterialInterface ? MaterialInterface : UMaterial::GetDefaultMaterial(MD_Surface));
 	}
 
 	// Now import geometry for each LOD
@@ -231,7 +237,10 @@ UObject* UInterchangeStaticMeshFactory::CreateAsset(const FCreateAssetParams& Ar
 
 		// Fill the lod mesh description using all combined mesh parts
 		TArray<FMeshPayload> MeshPayloads = GetMeshPayloads(Arguments, MeshUids);
-		for (const FMeshPayload& MeshPayload : MeshPayloads)
+
+		// Just move the mesh description from the first valid payload then append the rest
+		bool bFirstValidMoved = false;
+		for (FMeshPayload& MeshPayload : MeshPayloads)
 		{
 			const TOptional<UE::Interchange::FStaticMeshPayloadData>& LodMeshPayload = MeshPayload.PayloadData.Get();
 			if (!LodMeshPayload.IsSet())
@@ -240,9 +249,24 @@ UObject* UInterchangeStaticMeshFactory::CreateAsset(const FCreateAssetParams& Ar
 				continue;
 			}
 
-			// Bake the payload, with the provided transform
-			AppendSettings.MeshTransform = MeshPayload.Transform;
-			FStaticMeshOperations::AppendMeshDescription(LodMeshPayload->MeshDescription, *LodMeshDescription, AppendSettings);
+			if (!bFirstValidMoved)
+			{
+				FMeshDescription& MeshDescription = const_cast<FMeshDescription&>(LodMeshPayload->MeshDescription);
+				*LodMeshDescription = MoveTemp(MeshDescription);
+				bFirstValidMoved = true;
+
+				// Bake the payload mesh, with the provided transform
+				if (!MeshPayload.Transform.Equals(FTransform::Identity))
+				{
+					FStaticMeshOperations::ApplyTransform(*LodMeshDescription, MeshPayload.Transform);
+				}
+			}
+			else
+			{
+				// Bake the payload mesh, with the provided transform
+				AppendSettings.MeshTransform = MeshPayload.Transform;
+				FStaticMeshOperations::AppendMeshDescription(LodMeshPayload->MeshDescription, *LodMeshDescription, AppendSettings);
+			}
 		}
 
 		// Manage vertex color
@@ -312,11 +336,11 @@ UObject* UInterchangeStaticMeshFactory::CreateAsset(const FCreateAssetParams& Ar
 		{
 			int32 MaterialSlotIndex = StaticMesh->GetMaterialIndexFromImportedMaterialSlotName(SlotNames[PolygonGroupID]);
 					
-			if (MaterialSlotIndex == INDEX_NONE)
+			// If no material was found with this slot name, fill out a blank slot instead.
+			// However this should not happen. The translator should have properly set the UInterchangeMeshNode with the slot names used by all LODs.
+			if (!ensure(MaterialSlotIndex != INDEX_NONE))
 			{
-				// If no material was found with this slot name, it is probably because the pipeline is configured to not import materials.
-				// Fill out a blank slot instead.
-				MaterialSlotIndex = StaticMesh->GetStaticMaterials().Emplace(nullptr, SlotNames[PolygonGroupID]);
+				MaterialSlotIndex = StaticMesh->GetStaticMaterials().Emplace(UMaterial::GetDefaultMaterial(MD_Surface), SlotNames[PolygonGroupID]);
 			}
 
 			FMeshSectionInfo Info = StaticMesh->GetSectionInfoMap().Get(CurrentLodIndex, SectionIndex);
@@ -348,13 +372,6 @@ UObject* UInterchangeStaticMeshFactory::CreateAsset(const FCreateAssetParams& Ar
 			{
 				StaticMesh->SetLightMapCoordinateIndex(FirstOpenUVChannel);
 			}
-
-			if (CurrentLodIndex >= FinalLodCount)
-			{
-				SrcModel.ReductionSettings.MaxDeviation = 0.0f;
-				SrcModel.ReductionSettings.PercentTriangles = 1.0f;
-				SrcModel.ReductionSettings.PercentVertices = 1.0f;
-			}
 		}
 
 		CurrentLodIndex++;
diff --git a/Engine/Plugins/Experimental/Interchange/Runtime/Source/Import/Private/Scene/InterchangeActorHelper.cpp b/Engine/Plugins/Experimental/Interchange/Runtime/Source/Import/Private/Scene/InterchangeActorHelper.cpp
--- a/Engine/Plugins/Experimental/Interchange/Runtime/Source/Import/Private/Scene/InterchangeActorHelper.cpp
+++ b/Engine/Plugins/Experimental/Interchange/Runtime/Source/Import/Private/Scene/InterchangeActorHelper.cpp
@@ -3,12 +3,16 @@
 #include "Scene/InterchangeActorHelper.h"
 
 #include "InterchangeActorFactoryNode.h"
+#include "InterchangeMaterialFactoryNode.h"
+#include "InterchangeMeshActorFactoryNode.h"
 #include "InterchangeMeshNode.h"
 #include "InterchangeSceneNode.h"
 #include "Nodes/InterchangeBaseNode.h"
 #include "Nodes/InterchangeBaseNodeContainer.h"
 
+#include "Components/MeshComponent.h"
 #include "Engine/World.h"
+#include "Materials/Material.h"
 
 #if WITH_EDITOR
 #include "Editor/EditorEngine.h"
@@ -144,4 +148,32 @@ const UInterchangeFactoryBaseNode* UE::Interchange::ActorHelper::FindAssetInstan
 	TArray<FString> AssetTargetNodeIds;
 	AssetNode->GetTargetNodeUids(AssetTargetNodeIds);
 	return AssetTargetNodeIds.IsEmpty() ? nullptr : NodeContainer->GetFactoryNode(AssetTargetNodeIds[0]);
-}
\ No newline at end of file
+}
+
+void UE::Interchange::ActorHelper::ApplySlotMaterialDependencies(const UInterchangeBaseNodeContainer& NodeContainer, const UInterchangeMeshActorFactoryNode& MeshActorFactoryNode, UMeshComponent& MeshComponent)
+{
+	// Set material slots from imported materials
+	TMap<FString, FString> SlotMaterialDependencies;
+	MeshActorFactoryNode.GetSlotMaterialDependencies(SlotMaterialDependencies);
+	for (TPair<FString, FString>& SlotMaterialDependency : SlotMaterialDependencies)
+	{
+		const UInterchangeBaseMaterialFactoryNode* MaterialFactoryNode = Cast<UInterchangeBaseMaterialFactoryNode>(NodeContainer.GetNode(SlotMaterialDependency.Value));
+		if (!MaterialFactoryNode || !MaterialFactoryNode->ReferenceObject.IsValid() || !MaterialFactoryNode->IsEnabled())
+		{
+			continue;
+		}
+
+		FName MaterialSlotName = *SlotMaterialDependency.Key;
+		UMaterialInterface* MaterialInterface = Cast<UMaterialInterface>(MaterialFactoryNode->ReferenceObject.ResolveObject());
+		if (!MaterialInterface)
+		{
+			MaterialInterface = UMaterial::GetDefaultMaterial(MD_Surface);
+		}
+
+		int32 MaterialSlotIndex = MeshComponent.GetMaterialIndex(MaterialSlotName);
+		if (MaterialSlotIndex != INDEX_NONE)
+		{
+			MeshComponent.SetMaterial(MaterialSlotIndex, MaterialInterface);
+		}
+	}
+}
diff --git a/Engine/Plugins/Experimental/Interchange/Runtime/Source/Import/Private/Scene/InterchangeSkeletalMeshActorFactory.cpp b/Engine/Plugins/Experimental/Interchange/Runtime/Source/Import/Private/Scene/InterchangeSkeletalMeshActorFactory.cpp
--- a/Engine/Plugins/Experimental/Interchange/Runtime/Source/Import/Private/Scene/InterchangeSkeletalMeshActorFactory.cpp
+++ b/Engine/Plugins/Experimental/Interchange/Runtime/Source/Import/Private/Scene/InterchangeSkeletalMeshActorFactory.cpp
@@ -2,6 +2,7 @@
 #include "Scene/InterchangeSkeletalMeshActorFactory.h"
 
 #include "InterchangeActorFactoryNode.h"
+#include "InterchangeMeshActorFactoryNode.h"
 #include "Scene/InterchangeActorHelper.h"
 
 #include "Animation/SkeletalMeshActor.h"
@@ -39,12 +40,28 @@ void UInterchangeSkeletalMeshActorFactory::SetupSkeletalMeshActor(const UInterch
 {
 	USkeletalMeshComponent* SkeletalMeshComponent = SkeletalMeshActor->GetSkeletalMeshComponent();
 	SkeletalMeshComponent->UnregisterComponent();
+}
+
+void UInterchangeSkeletalMeshActorFactory::PostImportPreCompletedCallback(const FImportPreCompletedCallbackParams& Arguments)
+{
+	// Set the skeletal mesh on the component in the post import callback, once the skeletal mesh has been fully imported.
 
-	if (const UInterchangeFactoryBaseNode* MeshNode = UE::Interchange::ActorHelper::FindAssetInstanceFactoryNode(NodeContainer, ActorFactoryNode))
+	if (ASkeletalMeshActor* SkeletalMeshActor = Cast<ASkeletalMeshActor>(Arguments.ImportedObject))
 	{
-		if (USkeletalMesh* SkeletalMesh = Cast<USkeletalMesh>(MeshNode->ReferenceObject.TryLoad()))
+		if (USkeletalMeshComponent * SkeletalMeshComponent = SkeletalMeshActor->GetSkeletalMeshComponent())
 		{
-			SkeletalMeshComponent->SetSkeletalMesh(SkeletalMesh);
+			if (const UInterchangeFactoryBaseNode* MeshNode = UE::Interchange::ActorHelper::FindAssetInstanceFactoryNode(Arguments.NodeContainer, Arguments.FactoryNode))
+			{
+				if (USkeletalMesh* SkeletalMesh = Cast<USkeletalMesh>(MeshNode->ReferenceObject.TryLoad()))
+				{
+					SkeletalMeshComponent->SetSkeletalMesh(SkeletalMesh);
+
+					if (const UInterchangeMeshActorFactoryNode* MeshActorFactoryNode = Cast<UInterchangeMeshActorFactoryNode>(Arguments.FactoryNode))
+					{
+						UE::Interchange::ActorHelper::ApplySlotMaterialDependencies(*Arguments.NodeContainer, *MeshActorFactoryNode, *SkeletalMeshComponent);
+					}
+				}
+			}
 		}
 	}
-}
+}
\ No newline at end of file
diff --git a/Engine/Plugins/Experimental/Interchange/Runtime/Source/Import/Private/Scene/InterchangeStaticMeshActorFactory.cpp b/Engine/Plugins/Experimental/Interchange/Runtime/Source/Import/Private/Scene/InterchangeStaticMeshActorFactory.cpp
--- a/Engine/Plugins/Experimental/Interchange/Runtime/Source/Import/Private/Scene/InterchangeStaticMeshActorFactory.cpp
+++ b/Engine/Plugins/Experimental/Interchange/Runtime/Source/Import/Private/Scene/InterchangeStaticMeshActorFactory.cpp
@@ -2,6 +2,7 @@
 #include "Scene/InterchangeStaticMeshActorFactory.h"
 
 #include "InterchangeActorFactoryNode.h"
+#include "InterchangeMeshActorFactoryNode.h"
 #include "Nodes/InterchangeBaseNodeContainer.h"
 #include "Scene/InterchangeActorHelper.h"
 
@@ -61,6 +62,11 @@ void UInterchangeStaticMeshActorFactory::PostImportPreCompletedCallback(const FI
 				if (UStaticMesh* StaticMesh = Cast<UStaticMesh>(MeshNode->ReferenceObject.TryLoad()))
 				{
 					StaticMeshComponent->SetStaticMesh(StaticMesh);
+
+					if (const UInterchangeMeshActorFactoryNode* MeshActorFactoryNode = Cast<UInterchangeMeshActorFactoryNode>(Arguments.FactoryNode))
+					{
+						UE::Interchange::ActorHelper::ApplySlotMaterialDependencies(*Arguments.NodeContainer, *MeshActorFactoryNode, *StaticMeshComponent);
+					}
 				}
 			}
 		}
diff --git a/Engine/Plugins/Experimental/Interchange/Runtime/Source/Import/Public/Scene/InterchangeActorHelper.h b/Engine/Plugins/Experimental/Interchange/Runtime/Source/Import/Public/Scene/InterchangeActorHelper.h
--- a/Engine/Plugins/Experimental/Interchange/Runtime/Source/Import/Public/Scene/InterchangeActorHelper.h
+++ b/Engine/Plugins/Experimental/Interchange/Runtime/Source/Import/Public/Scene/InterchangeActorHelper.h
@@ -7,6 +7,8 @@
 class UInterchangeFactoryBaseNode;
 class UInterchangeBaseNodeContainer;
 class UInterchangeActorFactoryNode;
+class UInterchangeMeshActorFactoryNode;
+class UMeshComponent;
 
 namespace UE::Interchange::ActorHelper
 {
@@ -25,4 +27,9 @@ namespace UE::Interchange::ActorHelper
 	 * Returns the factory node of the asset instanced by ActorFactoryNode.
 	 */
 	INTERCHANGEIMPORT_API const UInterchangeFactoryBaseNode* FindAssetInstanceFactoryNode(const UInterchangeBaseNodeContainer* NodeContainer, const UInterchangeFactoryBaseNode* ActorFactoryNode);
+
+	/**
+	 * Applies material slot dependencies stored in ActorFactoryNode to MeshComponent.
+	 */
+	INTERCHANGEIMPORT_API void ApplySlotMaterialDependencies(const UInterchangeBaseNodeContainer& NodeContainer, const UInterchangeMeshActorFactoryNode& ActorFactoryNode, UMeshComponent& MeshComponent);
 }
diff --git a/Engine/Plugins/Experimental/Interchange/Runtime/Source/Import/Public/Scene/InterchangeSkeletalMeshActorFactory.h b/Engine/Plugins/Experimental/Interchange/Runtime/Source/Import/Public/Scene/InterchangeSkeletalMeshActorFactory.h
--- a/Engine/Plugins/Experimental/Interchange/Runtime/Source/Import/Public/Scene/InterchangeSkeletalMeshActorFactory.h
+++ b/Engine/Plugins/Experimental/Interchange/Runtime/Source/Import/Public/Scene/InterchangeSkeletalMeshActorFactory.h
@@ -31,6 +31,8 @@ public:
 		return false;
 	}
 
+	virtual void PostImportPreCompletedCallback(const FImportPreCompletedCallbackParams& Arguments) override;
+
 	// Interchange factory base interface end
 	//////////////////////////////////////////////////////////////////////////
 
diff --git a/Engine/Plugins/Experimental/Interchange/Runtime/Source/Nodes/Private/InterchangeMeshNode.cpp b/Engine/Plugins/Experimental/Interchange/Runtime/Source/Nodes/Private/InterchangeMeshNode.cpp
--- a/Engine/Plugins/Experimental/Interchange/Runtime/Source/Nodes/Private/InterchangeMeshNode.cpp
+++ b/Engine/Plugins/Experimental/Interchange/Runtime/Source/Nodes/Private/InterchangeMeshNode.cpp
@@ -38,12 +38,6 @@ namespace UE
 			return Dependencies_BaseKey;
 		}
 
-		const FAttributeKey& FMeshNodeStaticData::GetMaterialDependenciesKey()
-		{
-			static FAttributeKey Dependencies_BaseKey(TEXT("__MeshMaterialDependencies__"));
-			return Dependencies_BaseKey;
-		}
-
 		const FAttributeKey& FMeshNodeStaticData::GetShapeDependenciesKey()
 		{
 			static FAttributeKey Dependencies_BaseKey(TEXT("__MeshShapeDependencies__"));
@@ -56,15 +50,21 @@ namespace UE
 			return SceneInstanceUids_BaseKey;
 		}
 
+		const FAttributeKey& FMeshNodeStaticData::GetSlotMaterialDependenciesKey()
+		{
+			static FAttributeKey Dependencies_BaseKey(TEXT("__SlotMaterialDependencies__"));
+			return Dependencies_BaseKey;
+		}
+
 	}//ns Interchange
 }//ns UE
 
 UInterchangeMeshNode::UInterchangeMeshNode()
 {
 	SkeletonDependencies.Initialize(Attributes, UE::Interchange::FMeshNodeStaticData::GetSkeletonDependenciesKey().ToString());
-	MaterialDependencies.Initialize(Attributes, UE::Interchange::FMeshNodeStaticData::GetMaterialDependenciesKey().ToString());
 	ShapeDependencies.Initialize(Attributes, UE::Interchange::FMeshNodeStaticData::GetShapeDependenciesKey().ToString());
 	SceneInstancesUids.Initialize(Attributes, UE::Interchange::FMeshNodeStaticData::GetSceneInstancesUidsKey().ToString());
+	SlotMaterialDependencies.Initialize(Attributes.ToSharedRef(), UE::Interchange::FMeshNodeStaticData::GetSlotMaterialDependenciesKey().ToString());
 }
 
 FString UInterchangeMeshNode::GetKeyDisplayName(const UE::Interchange::FAttributeKey& NodeAttributeKey) const
@@ -102,21 +102,6 @@ FString UInterchangeMeshNode::GetKeyDisplayName(const UE::Interchange::FAttribut
 		}
 		return KeyDisplayName;
 	}
-	else if (NodeAttributeKey == UE::Interchange::FMeshNodeStaticData::GetMaterialDependenciesKey())
-	{
-		return KeyDisplayName = TEXT("Material Dependencies count");
-	}
-	else if (NodeAttributeKeyString.StartsWith(UE::Interchange::FMeshNodeStaticData::GetMaterialDependenciesKey().ToString()))
-	{
-		KeyDisplayName = TEXT("Material Dependencies Index ");
-		const FString IndexKey = UE::Interchange::TArrayAttributeHelper<FString>::IndexKey();
-		int32 IndexPosition = NodeAttributeKeyString.Find(IndexKey) + IndexKey.Len();
-		if (IndexPosition < NodeAttributeKeyString.Len())
-		{
-			KeyDisplayName += NodeAttributeKeyString.RightChop(IndexPosition);
-		}
-		return KeyDisplayName;
-	}
 	else if (NodeAttributeKey == UE::Interchange::FMeshNodeStaticData::GetShapeDependenciesKey())
 	{
 		return KeyDisplayName = TEXT("Shape Dependencies count");
@@ -147,6 +132,10 @@ FString UInterchangeMeshNode::GetKeyDisplayName(const UE::Interchange::FAttribut
 		}
 		return KeyDisplayName;
 	}
+	else if (NodeAttributeKeyString.StartsWith(UE::Interchange::FMeshNodeStaticData::GetSlotMaterialDependenciesKey().ToString()))
+	{
+		return KeyDisplayName = TEXT("Slot material dependencies");
+	}
 	return Super::GetKeyDisplayName(NodeAttributeKey);
 }
 
@@ -157,10 +146,6 @@ FString UInterchangeMeshNode::GetAttributeCategory(const UE::Interchange::FAttri
 	{
 		return FString(TEXT("SkeletonDependencies"));
 	}
-	else if (NodeAttributeKeyString.StartsWith(UE::Interchange::FMeshNodeStaticData::GetMaterialDependenciesKey().ToString()))
-	{
-		return FString(TEXT("MaterialDependencies"));
-	}
 	else if (NodeAttributeKeyString.StartsWith(UE::Interchange::FMeshNodeStaticData::GetShapeDependenciesKey().ToString()))
 	{
 		return FString(TEXT("ShapeDependencies"));
@@ -181,6 +166,10 @@ FString UInterchangeMeshNode::GetAttributeCategory(const UE::Interchange::FAttri
 	{
 		return FString(TEXT("MeshInfo"));
 	}
+	else if (NodeAttributeKeyString.StartsWith(UE::Interchange::FMeshNodeStaticData::GetSlotMaterialDependenciesKey().ToString()))
+	{
+		return FString(TEXT("SlotMaterialDependencies"));
+	}
 	return Super::GetAttributeCategory(NodeAttributeKey);
 }
 
@@ -434,31 +423,6 @@ bool UInterchangeMeshNode::RemoveSkeletonDependencyUid(const FString& Dependency
 	return SkeletonDependencies.RemoveItem(DependencyUid);
 }
 
-int32 UInterchangeMeshNode::GetMaterialDependeciesCount() const
-{
-	return MaterialDependencies.GetCount();
-}
-
-void UInterchangeMeshNode::GetMaterialDependencies(TArray<FString>& OutDependencies) const
-{
-	MaterialDependencies.GetItems(OutDependencies);
-}
-
-void UInterchangeMeshNode::GetMaterialDependency(const int32 Index, FString& OutDependency) const
-{
-	MaterialDependencies.GetItem(Index, OutDependency);
-}
-
-bool UInterchangeMeshNode::SetMaterialDependencyUid(const FString& DependencyUid)
-{
-	return MaterialDependencies.AddItem(DependencyUid);
-}
-
-bool UInterchangeMeshNode::RemoveMaterialDependencyUid(const FString& DependencyUid)
-{
-	return MaterialDependencies.RemoveItem(DependencyUid);
-}
-
 int32 UInterchangeMeshNode::GetShapeDependeciesCount() const
 {
 	return ShapeDependencies.GetCount();
@@ -507,4 +471,24 @@ bool UInterchangeMeshNode::SetSceneInstanceUid(const FString& DependencyUid)
 bool UInterchangeMeshNode::RemoveSceneInstanceUid(const FString& DependencyUid)
 {
 	return SceneInstancesUids.RemoveItem(DependencyUid);
-}
\ No newline at end of file
+}
+
+void UInterchangeMeshNode::GetSlotMaterialDependencies(TMap<FString, FString>& OutMaterialDependencies) const
+{
+	OutMaterialDependencies = SlotMaterialDependencies.ToMap();
+}
+
+bool UInterchangeMeshNode::GetSlotMaterialDependencyUid(const FString& SlotName, FString& OutMaterialDependency) const
+{
+	return SlotMaterialDependencies.GetValue(SlotName, OutMaterialDependency);
+}
+
+bool UInterchangeMeshNode::SetSlotMaterialDependencyUid(const FString& SlotName, const FString& MaterialDependencyUid)
+{
+	return SlotMaterialDependencies.SetKeyValue(SlotName, MaterialDependencyUid);
+}
+
+bool UInterchangeMeshNode::RemoveSlotMaterialDependencyUid(const FString& SlotName)
+{
+	return SlotMaterialDependencies.RemoveKey(SlotName);
+}
diff --git a/Engine/Plugins/Experimental/Interchange/Runtime/Source/Nodes/Private/InterchangeSceneNode.cpp b/Engine/Plugins/Experimental/Interchange/Runtime/Source/Nodes/Private/InterchangeSceneNode.cpp
--- a/Engine/Plugins/Experimental/Interchange/Runtime/Source/Nodes/Private/InterchangeSceneNode.cpp
+++ b/Engine/Plugins/Experimental/Interchange/Runtime/Source/Nodes/Private/InterchangeSceneNode.cpp
@@ -35,8 +35,14 @@ namespace UE
 
 		const FString& FSceneNodeStaticData::GetLodGroupSpecializeTypeString()
 		{
-			static FString JointSpecializeTypeString(TEXT("LodGroup"));
-			return JointSpecializeTypeString;
+			static FString LodGroupSpecializeTypeString(TEXT("LodGroup"));
+			return LodGroupSpecializeTypeString;
+		}
+
+		const FString& FSceneNodeStaticData::GetSlotMaterialDependenciesString()
+		{
+			static FString SlotMaterialDependenciesString(TEXT("__SlotMaterialDependencies__"));
+			return SlotMaterialDependenciesString;
 		}
 	}//ns Interchange
 }//ns UE
@@ -44,7 +50,7 @@ namespace UE
 UInterchangeSceneNode::UInterchangeSceneNode()
 {
 	NodeSpecializeTypes.Initialize(Attributes, UE::Interchange::FSceneNodeStaticData::GetNodeSpecializeTypeBaseKey().ToString());
-	MaterialDependencyUids.Initialize(Attributes, UE::Interchange::FSceneNodeStaticData::GetMaterialDependencyUidsBaseKey().ToString());
+	SlotMaterialDependencies.Initialize(Attributes.ToSharedRef(), UE::Interchange::FSceneNodeStaticData::GetSlotMaterialDependenciesString());
 }
 
 /**
@@ -180,31 +186,6 @@ bool UInterchangeSceneNode::RemoveSpecializedType(const FString& SpecializedType
 	return NodeSpecializeTypes.RemoveItem(SpecializedType);
 }
 
-int32 UInterchangeSceneNode::GetMaterialDependencyUidsCount() const
-{
-	return MaterialDependencyUids.GetCount();
-}
-
-void UInterchangeSceneNode::GetMaterialDependencyUid(const int32 Index, FString& OutMaterialDependencyUid) const
-{
-	MaterialDependencyUids.GetItem(Index, OutMaterialDependencyUid);
-}
-
-void UInterchangeSceneNode::GetMaterialDependencyUids(TArray<FString>& OutMaterialDependencyUids) const
-{
-	MaterialDependencyUids.GetItems(OutMaterialDependencyUids);
-}
-
-bool UInterchangeSceneNode::AddMaterialDependencyUid(const FString& MaterialDependencyUid)
-{
-	return MaterialDependencyUids.AddItem(MaterialDependencyUid);
-}
-
-bool UInterchangeSceneNode::RemoveMaterialDependencyUid(const FString& MaterialDependencyUid)
-{
-	return MaterialDependencyUids.RemoveItem(MaterialDependencyUid);
-}
-
 bool UInterchangeSceneNode::GetCustomLocalTransform(FTransform& AttributeValue) const
 {
 	IMPLEMENT_NODE_ATTRIBUTE_GETTER(LocalTransform, FTransform);
@@ -430,3 +411,23 @@ bool UInterchangeSceneNode::GetGlobalTransformInternal(const UE::Interchange::FA
 	AttributeValue = CacheTransform.GetValue();
 	return true;
 }
+
+void UInterchangeSceneNode::GetSlotMaterialDependencies(TMap<FString, FString>& OutMaterialDependencies) const
+{
+	OutMaterialDependencies = SlotMaterialDependencies.ToMap();
+}
+
+bool UInterchangeSceneNode::GetSlotMaterialDependencyUid(const FString& SlotName, FString& OutMaterialDependency) const
+{
+	return SlotMaterialDependencies.GetValue(SlotName, OutMaterialDependency);
+}
+
+bool UInterchangeSceneNode::SetSlotMaterialDependencyUid(const FString& SlotName, const FString& MaterialDependencyUid)
+{
+	return SlotMaterialDependencies.SetKeyValue(SlotName, MaterialDependencyUid);
+}
+
+bool UInterchangeSceneNode::RemoveSlotMaterialDependencyUid(const FString& SlotName)
+{
+	return SlotMaterialDependencies.RemoveKey(SlotName);
+}
diff --git a/Engine/Plugins/Experimental/Interchange/Runtime/Source/Nodes/Public/InterchangeMeshNode.h b/Engine/Plugins/Experimental/Interchange/Runtime/Source/Nodes/Public/InterchangeMeshNode.h
--- a/Engine/Plugins/Experimental/Interchange/Runtime/Source/Nodes/Public/InterchangeMeshNode.h
+++ b/Engine/Plugins/Experimental/Interchange/Runtime/Source/Nodes/Public/InterchangeMeshNode.h
@@ -20,9 +20,9 @@ namespace UE
 			static const FAttributeKey& IsBlendShapeKey();
 			static const FAttributeKey& BlendShapeNameKey();
 			static const FAttributeKey& GetSkeletonDependenciesKey();
-			static const FAttributeKey& GetMaterialDependenciesKey();
 			static const FAttributeKey& GetShapeDependenciesKey();
 			static const FAttributeKey& GetSceneInstancesUidsKey();
+			static const FAttributeKey& GetSlotMaterialDependenciesKey();
 		};
 
 	}//ns Interchange
@@ -49,6 +49,19 @@ public:
 	 * Icon name are simply create by adding "InterchangeIcon_" in front of the specialized type. If there is no special type the function will return NAME_None which will use the default icon.
 	 */
 	virtual FName GetIconName() const override;
+	
+	/**
+	 * Override serialize to restore SlotMaterialDependencies on load.
+	 */
+	virtual void Serialize(FArchive& Ar) override
+	{
+		Super::Serialize(Ar);
+
+		if (Ar.IsLoading() && bIsInitialized)
+		{
+			SlotMaterialDependencies.RebuildCache();
+		}
+	}
 
 	/**
 	 * Return true if this node represent a skinned mesh
@@ -195,36 +208,6 @@ public:
 	UFUNCTION(BlueprintCallable, Category = "Interchange | Node | Mesh")
 	bool RemoveSkeletonDependencyUid(const FString& DependencyUid);
 
-	/**
-	 * This function allow to retrieve the number of Material dependencies for this object.
-	 */
-	UFUNCTION(BlueprintCallable, Category = "Interchange | Node | Mesh")
-	int32 GetMaterialDependeciesCount() const;
-
-	/**
-	 * This function allow to retrieve the Material dependency for this object.
-	 */
-	UFUNCTION(BlueprintCallable, Category = "Interchange | Node | Mesh")
-	void GetMaterialDependencies(TArray<FString>& OutDependencies) const;
-
-	/**
-	 * This function allow to retrieve one Material dependency for this object.
-	 */
-	UFUNCTION(BlueprintCallable, Category = "Interchange | Node | Mesh")
-	void GetMaterialDependency(const int32 Index, FString& OutDependency) const;
-
-	/**
-	 * Add one Material dependency to this object.
-	 */
-	UFUNCTION(BlueprintCallable, Category = "Interchange | Node | Mesh")
-	bool SetMaterialDependencyUid(const FString& DependencyUid);
-
-	/**
-	 * Remove one Material dependency from this object.
-	 */
-	UFUNCTION(BlueprintCallable, Category = "Interchange | Node | Mesh")
-	bool RemoveMaterialDependencyUid(const FString& DependencyUid);
-
 	/**
 	 * This function allow to retrieve the number of Shape dependencies for this object.
 	 */
@@ -285,6 +268,30 @@ public:
 	UFUNCTION(BlueprintCallable, Category = "Interchange | Node | Mesh")
 	bool RemoveSceneInstanceUid(const FString& DependencyUid);
 
+	/**
+	 * Allow to retrieve the correspondence table between slot names and assigned materials for this object.
+	 */
+	UFUNCTION(BlueprintCallable, Category = "Interchange | Node | StaticMesh")
+	void GetSlotMaterialDependencies(TMap<FString, FString>& OutMaterialDependencies) const;
+
+	/**
+	 * Allow to retrieve one Material dependency for a given slot of this object.
+	 */
+	UFUNCTION(BlueprintCallable, Category = "Interchange | Node | StaticMesh")
+	bool GetSlotMaterialDependencyUid(const FString& SlotName, FString& OutMaterialDependency) const;
+
+	/**
+	 * Add one Material dependency to a specific slot name of this object.
+	 */
+	UFUNCTION(BlueprintCallable, Category = "Interchange | Node | StaticMesh")
+	bool SetSlotMaterialDependencyUid(const FString& SlotName, const FString& MaterialDependencyUid);
+
+	/**
+	 * Remove the Material dependency associated with the given slot name from this object.
+	 */
+	UFUNCTION(BlueprintCallable, Category = "Interchange | Node | StaticMesh")
+	bool RemoveSlotMaterialDependencyUid(const FString& SlotName);
+
 private:
 	const UE::Interchange::FAttributeKey Macro_CustomVertexCountKey = UE::Interchange::FAttributeKey(TEXT("VertexCount"));
 	const UE::Interchange::FAttributeKey Macro_CustomPolygonCountKey = UE::Interchange::FAttributeKey(TEXT("PolygonCount"));
@@ -300,4 +307,6 @@ private:
 	UE::Interchange::TArrayAttributeHelper<FString> MaterialDependencies;
 	UE::Interchange::TArrayAttributeHelper<FString> ShapeDependencies;
 	UE::Interchange::TArrayAttributeHelper<FString> SceneInstancesUids;
+	
+	UE::Interchange::TMapAttributeHelper<FString, FString> SlotMaterialDependencies;
 };
diff --git a/Engine/Plugins/Experimental/Interchange/Runtime/Source/Nodes/Public/InterchangeSceneNode.h b/Engine/Plugins/Experimental/Interchange/Runtime/Source/Nodes/Public/InterchangeSceneNode.h
--- a/Engine/Plugins/Experimental/Interchange/Runtime/Source/Nodes/Public/InterchangeSceneNode.h
+++ b/Engine/Plugins/Experimental/Interchange/Runtime/Source/Nodes/Public/InterchangeSceneNode.h
@@ -22,6 +22,7 @@ namespace UE
 			static const FString& GetTransformSpecializeTypeString();
 			static const FString& GetJointSpecializeTypeString();
 			static const FString& GetLodGroupSpecializeTypeString();
+			static const FString& GetSlotMaterialDependenciesString();
 		};
 
 	}//ns Interchange
@@ -35,6 +36,19 @@ class INTERCHANGENODES_API UInterchangeSceneNode : public UInterchangeBaseNode
 public:
 	UInterchangeSceneNode();
 
+	/**
+	 * Override serialize to restore SlotMaterialDependencies on load.
+	 */
+	virtual void Serialize(FArchive& Ar) override
+	{
+		Super::Serialize(Ar);
+
+		if (Ar.IsLoading() && bIsInitialized)
+		{
+			SlotMaterialDependencies.RebuildCache();
+		}
+	}
+
 	/**
 	 * Return the node type name of the class, we use this when reporting error
 	 */
@@ -68,22 +82,6 @@ public:
 	UFUNCTION(BlueprintCallable, Category = "Interchange | Node | Scene")
 	bool RemoveSpecializedType(const FString& SpecializedType);
 
-	/** Asset dependencies are the asset on which this node depend.*/
-	UFUNCTION(BlueprintCallable, Category = "Interchange | Node | Scene")
-	int32 GetMaterialDependencyUidsCount() const;
-
-	UFUNCTION(BlueprintCallable, Category = "Interchange | Node | Scene")
-	void GetMaterialDependencyUid(const int32 Index, FString& OutMaterialDependencyUid) const;
-
-	UFUNCTION(BlueprintCallable, Category = "Interchange | Node | Scene")
-	void GetMaterialDependencyUids(TArray<FString>& OutMaterialDependencyUids) const;
-
-	UFUNCTION(BlueprintCallable, Category = "Interchange | Node | Scene")
-	bool AddMaterialDependencyUid(const FString& MaterialDependencyUid);
-
-	UFUNCTION(BlueprintCallable, Category = "Interchange | Node | Scene")
-	bool RemoveMaterialDependencyUid(const FString& MaterialDependencyUid);
-
 	//Default transform is the transform we have in the node (no bind pose, no time evaluation).
 
 	/** Return the default scene node local transform. */
@@ -190,6 +188,30 @@ public:
 	/** This static function make sure all the global transform caches are reset for all the UInterchangeSceneNode nodes children in the UInterchangeBaseNodeContainer */
 	static void ResetGlobalTransformCachesOfNodeAndAllChildren(const UInterchangeBaseNodeContainer* BaseNodeContainer, const UInterchangeBaseNode* ParentNode);
 
+	/**
+	 * Allow to retrieve the correspondence table between slot names and assigned materials for this object.
+	 */
+	UFUNCTION(BlueprintCallable, Category = "Interchange | Node | StaticMesh")
+	void GetSlotMaterialDependencies(TMap<FString, FString>& OutMaterialDependencies) const;
+
+	/**
+	 * Allow to retrieve one Material dependency for a given slot of this object.
+	 */
+	UFUNCTION(BlueprintCallable, Category = "Interchange | Node | StaticMesh")
+	bool GetSlotMaterialDependencyUid(const FString& SlotName, FString& OutMaterialDependency) const;
+
+	/**
+	 * Add one Material dependency to a specific slot name of this object.
+	 */
+	UFUNCTION(BlueprintCallable, Category = "Interchange | Node | StaticMesh")
+	bool SetSlotMaterialDependencyUid(const FString& SlotName, const FString& MaterialDependencyUid);
+
+	/**
+	 * Remove the Material dependency associated with the given slot name from this object.
+	 */
+	UFUNCTION(BlueprintCallable, Category = "Interchange | Node | StaticMesh")
+	bool RemoveSlotMaterialDependencyUid(const FString& SlotName);
+
 private:
 
 	bool GetGlobalTransformInternal(const UE::Interchange::FAttributeKey LocalTransformKey
@@ -213,7 +235,8 @@ private:
 	const UE::Interchange::FAttributeKey Macro_CustomTransformCurvePayloadKeyKey = UE::Interchange::FAttributeKey(TEXT("TransformCurvePayloadKey"));
 
 	UE::Interchange::TArrayAttributeHelper<FString> NodeSpecializeTypes;
-	UE::Interchange::TArrayAttributeHelper<FString> MaterialDependencyUids;
+
+	UE::Interchange::TMapAttributeHelper<FString, FString> SlotMaterialDependencies;
 
 	mutable TOptional<FTransform> CacheGlobalTransform;
 	mutable TOptional<FTransform> CacheBindPoseGlobalTransform;
diff --git a/Engine/Plugins/Experimental/Interchange/Runtime/Source/Nodes/Public/InterchangeTexture2DNode.h b/Engine/Plugins/Experimental/Interchange/Runtime/Source/Nodes/Public/InterchangeTexture2DNode.h
--- a/Engine/Plugins/Experimental/Interchange/Runtime/Source/Nodes/Public/InterchangeTexture2DNode.h
+++ b/Engine/Plugins/Experimental/Interchange/Runtime/Source/Nodes/Public/InterchangeTexture2DNode.h
@@ -54,6 +54,19 @@ public:
 		return TypeName;
 	}
 
+	/**
+	 * Override serialize to restore SourceBlocks on load.
+	 */
+	virtual void Serialize(FArchive& Ar) override
+	{
+		Super::Serialize(Ar);
+
+		if (Ar.IsLoading() && bIsInitialized)
+		{
+			SourceBlocks.RebuildCache();
+		}
+	}
+
 	//////////////////////////////////////////////////////////////////////////
 	// UDIMs begin here
 	// UDIM base texture use a different model for the source data
diff --git a/Engine/Plugins/Experimental/Interchange/Runtime/Source/Parsers/Fbx/Private/FbxMaterial.cpp b/Engine/Plugins/Experimental/Interchange/Runtime/Source/Parsers/Fbx/Private/FbxMaterial.cpp
--- a/Engine/Plugins/Experimental/Interchange/Runtime/Source/Parsers/Fbx/Private/FbxMaterial.cpp
+++ b/Engine/Plugins/Experimental/Interchange/Runtime/Source/Parsers/Fbx/Private/FbxMaterial.cpp
@@ -288,8 +288,7 @@ namespace UE
 				{
 					FbxSurfaceMaterial* SurfaceMaterial = ParentFbxNode->GetMaterial(MaterialIndex);
 					const UInterchangeShaderGraphNode* ShaderGraphNode = AddShaderGraphNode(SurfaceMaterial, NodeContainer);
-					//The dependencies order is important because mesh will use index in that order to determine material use by a face
-					SceneNode->AddMaterialDependencyUid(ShaderGraphNode->GetUniqueID());
+					SceneNode->SetSlotMaterialDependencyUid(FFbxHelper::GetFbxObjectName(SurfaceMaterial), ShaderGraphNode->GetUniqueID());
 				}
 			}
 
diff --git a/Engine/Plugins/Experimental/Interchange/Runtime/Source/Parsers/Fbx/Private/FbxMesh.cpp b/Engine/Plugins/Experimental/Interchange/Runtime/Source/Parsers/Fbx/Private/FbxMesh.cpp
--- a/Engine/Plugins/Experimental/Interchange/Runtime/Source/Parsers/Fbx/Private/FbxMesh.cpp
+++ b/Engine/Plugins/Experimental/Interchange/Runtime/Source/Parsers/Fbx/Private/FbxMesh.cpp
@@ -1344,7 +1344,7 @@ void FFbxMesh::AddAllMeshes(FbxScene* SDKScene, FbxGeometryConverter* SDKGeometr
 						FString MaterialUid = TEXT("\\Material\\") + MaterialName;
 						if (bAddAllNodeMaterials || MaterialIndexes.Contains(MaterialIndex))
 						{
-							MeshNode->SetMaterialDependencyUid(MaterialUid);
+							MeshNode->SetSlotMaterialDependencyUid(MaterialName, MaterialUid);
 						}
 					}
 				}
diff --git a/Engine/Plugins/Experimental/Interchange/Runtime/Source/Pipelines/Private/InterchangeGenericScenesPipeline.cpp b/Engine/Plugins/Experimental/Interchange/Runtime/Source/Pipelines/Private/InterchangeGenericScenesPipeline.cpp
--- a/Engine/Plugins/Experimental/Interchange/Runtime/Source/Pipelines/Private/InterchangeGenericScenesPipeline.cpp
+++ b/Engine/Plugins/Experimental/Interchange/Runtime/Source/Pipelines/Private/InterchangeGenericScenesPipeline.cpp
@@ -7,8 +7,10 @@
 #include "InterchangeCineCameraFactoryNode.h"
 #include "InterchangeCommonPipelineDataFactoryNode.h"
 #include "InterchangeLightNode.h"
+#include "InterchangeMeshActorFactoryNode.h"
 #include "InterchangeMeshNode.h"
 #include "InterchangePipelineLog.h"
+#include "InterchangePipelineMeshesUtilities.h"
 #include "InterchangeSceneNode.h"
 
 #include "Animation/SkeletalMeshActor.h"
@@ -115,7 +117,7 @@ void UInterchangeGenericLevelPipeline::ExecuteSceneNodePreImport(UInterchangeBas
 
 	if (TranslatedAssetNode)
 	{
-		SetUpFactoryNode(ActorFactoryNode, TranslatedAssetNode, FactoryNodeContainer);
+		SetUpFactoryNode(ActorFactoryNode, SceneNode, TranslatedAssetNode, FactoryNodeContainer);
 	}
 }
 
@@ -125,14 +127,23 @@ UInterchangeActorFactoryNode* UInterchangeGenericLevelPipeline::CreateActorFacto
 	{
 		return NewObject<UInterchangeCineCameraFactoryNode>(FactoryNodeContainer, NAME_None);
 	}
+	else if (TranslatedAssetNode && TranslatedAssetNode->IsA<UInterchangeMeshNode>())
+	{
+		return NewObject<UInterchangeMeshActorFactoryNode>(FactoryNodeContainer, NAME_None);
+	}
 	else
 	{
 		return NewObject<UInterchangeActorFactoryNode>(FactoryNodeContainer, NAME_None);
 	}
 }
 
-void UInterchangeGenericLevelPipeline::SetUpFactoryNode(UInterchangeActorFactoryNode* ActorFactoryNode, const UInterchangeBaseNode* TranslatedAssetNode, UInterchangeBaseNodeContainer* FactoryNodeContainer) const
+void UInterchangeGenericLevelPipeline::SetUpFactoryNode(UInterchangeActorFactoryNode* ActorFactoryNode, const UInterchangeSceneNode* SceneNode, const UInterchangeBaseNode* TranslatedAssetNode, UInterchangeBaseNodeContainer* FactoryNodeContainer) const
 {
+	if (!ensure(ActorFactoryNode && SceneNode && TranslatedAssetNode && FactoryNodeContainer))
+	{
+		return;
+	}
+
 	if (const UInterchangeMeshNode* MeshNode = Cast<UInterchangeMeshNode>(TranslatedAssetNode))
 	{
 		if (MeshNode->IsSkinnedMesh())
@@ -144,6 +155,14 @@ void UInterchangeGenericLevelPipeline::SetUpFactoryNode(UInterchangeActorFactory
 		{
 			ActorFactoryNode->SetCustomActorClassName(AStaticMeshActor::StaticClass()->GetPathName());
 		}
+
+		if (UInterchangeMeshActorFactoryNode* MeshActorFactoryNode = Cast<UInterchangeMeshActorFactoryNode>(ActorFactoryNode))
+		{
+			TMap<FString, FString> SlotMaterialDependencies;
+			SceneNode->GetSlotMaterialDependencies(SlotMaterialDependencies);
+
+			UE::Interchange::MeshesUtilities::ApplySlotMaterialDependencies(*MeshActorFactoryNode, SlotMaterialDependencies, *FactoryNodeContainer);
+		}
 	}
 	else if (const UInterchangeLightNode* LightNode = Cast<UInterchangeLightNode>(TranslatedAssetNode))
 	{
diff --git a/Engine/Plugins/Experimental/Interchange/Runtime/Source/Pipelines/Private/InterchangeGenericSkeletalMeshPipeline.cpp b/Engine/Plugins/Experimental/Interchange/Runtime/Source/Pipelines/Private/InterchangeGenericSkeletalMeshPipeline.cpp
--- a/Engine/Plugins/Experimental/Interchange/Runtime/Source/Pipelines/Private/InterchangeGenericSkeletalMeshPipeline.cpp
+++ b/Engine/Plugins/Experimental/Interchange/Runtime/Source/Pipelines/Private/InterchangeGenericSkeletalMeshPipeline.cpp
@@ -537,7 +537,7 @@ void UInterchangeGenericMeshPipeline::AddLodDataToSkeletalMesh(const UInterchang
 		constexpr bool bAddSourceNodeName = true;
 		for (const FString& NodeUid : NodeUids)
 		{
-			TArray<FString> MaterialDependencies;
+			TMap<FString, FString> SlotMaterialDependencies;
 			if (const UInterchangeSceneNode* SceneNode = Cast<UInterchangeSceneNode>(BaseNodeContainer->GetNode(NodeUid)))
 			{
 				FString MeshDependency;
@@ -548,11 +548,12 @@ void UInterchangeGenericMeshPipeline::AddLodDataToSkeletalMesh(const UInterchang
 					UInterchangeUserDefinedAttributesAPI::DuplicateAllUserDefinedAttribute(MeshDependencyNode, SkeletalMeshFactoryNode, bAddSourceNodeName);
 					SkeletalMeshFactoryNode->AddTargetNodeUid(MeshDependency);
 					MeshDependencyNode->AddTargetNodeUid(SkeletalMeshFactoryNode->GetUniqueID());
-					MeshDependencyNode->GetMaterialDependencies(MaterialDependencies);
+					
+					MeshDependencyNode->GetSlotMaterialDependencies(SlotMaterialDependencies);
 				}
 				else
 				{
-					SceneNode->GetMaterialDependencyUids(MaterialDependencies);
+					SceneNode->GetSlotMaterialDependencies(SlotMaterialDependencies);
 				}
 				UInterchangeUserDefinedAttributesAPI::DuplicateAllUserDefinedAttribute(SceneNode, SkeletalMeshFactoryNode, bAddSourceNodeName);
 			}
@@ -561,24 +562,12 @@ void UInterchangeGenericMeshPipeline::AddLodDataToSkeletalMesh(const UInterchang
 				UInterchangeUserDefinedAttributesAPI::DuplicateAllUserDefinedAttribute(MeshNode, SkeletalMeshFactoryNode, bAddSourceNodeName);
 				SkeletalMeshFactoryNode->AddTargetNodeUid(NodeUid);
 				MeshNode->AddTargetNodeUid(SkeletalMeshFactoryNode->GetUniqueID());
-				MeshNode->GetMaterialDependencies(MaterialDependencies);
-			}
-			const int32 MaterialCount = MaterialDependencies.Num();
-			for (int32 MaterialIndex = 0; MaterialIndex < MaterialCount; ++MaterialIndex)
-			{
-				const FString MaterialFactoryNodeUid = UInterchangeMaterialFactoryNode::GetMaterialFactoryNodeUidFromMaterialNodeUid(MaterialDependencies[MaterialIndex]);
-				if (BaseNodeContainer->IsNodeUidValid(MaterialFactoryNodeUid))
-				{
-					BaseNodeContainer->GetFactoryNode(MaterialFactoryNodeUid)->SetEnabled(true);
-					//Create a factory dependency so Material asset are import before the skeletal mesh asset
-					TArray<FString> FactoryDependencies;
-					SkeletalMeshFactoryNode->GetFactoryDependencies(FactoryDependencies);
-					if (!FactoryDependencies.Contains(MaterialFactoryNodeUid))
-					{
-						SkeletalMeshFactoryNode->AddFactoryDependencyUid(MaterialFactoryNodeUid);
-					}
-				}
+
+				MeshNode->GetSlotMaterialDependencies(SlotMaterialDependencies);
 			}
+
+			UE::Interchange::MeshesUtilities::ApplySlotMaterialDependencies(*SkeletalMeshFactoryNode, SlotMaterialDependencies, *BaseNodeContainer);
+
 			LodDataNode->AddMeshUid(NodeUid);
 		}
 	}
diff --git a/Engine/Plugins/Experimental/Interchange/Runtime/Source/Pipelines/Private/InterchangeGenericStaticMeshPipeline.cpp b/Engine/Plugins/Experimental/Interchange/Runtime/Source/Pipelines/Private/InterchangeGenericStaticMeshPipeline.cpp
--- a/Engine/Plugins/Experimental/Interchange/Runtime/Source/Pipelines/Private/InterchangeGenericStaticMeshPipeline.cpp
+++ b/Engine/Plugins/Experimental/Interchange/Runtime/Source/Pipelines/Private/InterchangeGenericStaticMeshPipeline.cpp
@@ -470,7 +470,7 @@ void UInterchangeGenericMeshPipeline::AddLodDataToStaticMesh(UInterchangeStaticM
 		constexpr bool bAddSourceNodeName = true;
 		for (const FString& NodeUid : NodeUids)
 		{
-			TArray<FString> MaterialDependencies;
+			TMap<FString, FString> SlotMaterialDependencies;
 			if (const UInterchangeSceneNode* SceneNode = Cast<UInterchangeSceneNode>(BaseNodeContainer->GetNode(NodeUid)))
 			{
 				FString MeshDependency;
@@ -482,11 +482,12 @@ void UInterchangeGenericMeshPipeline::AddLodDataToStaticMesh(UInterchangeStaticM
 					StaticMeshFactoryNode->AddTargetNodeUid(MeshDependency);
 					StaticMeshFactoryNode->AddSocketUids(PipelineMeshesUtilities->GetMeshGeometryByUid(MeshDependency).AttachedSocketUids);
 					MeshDependencyNode->AddTargetNodeUid(StaticMeshFactoryNode->GetUniqueID());
-					MeshDependencyNode->GetMaterialDependencies(MaterialDependencies);
+
+					MeshDependencyNode->GetSlotMaterialDependencies(SlotMaterialDependencies);
 				}
 				else
 				{
-					SceneNode->GetMaterialDependencyUids(MaterialDependencies);
+					SceneNode->GetSlotMaterialDependencies(SlotMaterialDependencies);
 				}
 
 				UInterchangeUserDefinedAttributesAPI::DuplicateAllUserDefinedAttribute(SceneNode, StaticMeshFactoryNode, bAddSourceNodeName);
@@ -497,25 +498,11 @@ void UInterchangeGenericMeshPipeline::AddLodDataToStaticMesh(UInterchangeStaticM
 				StaticMeshFactoryNode->AddTargetNodeUid(NodeUid);
 				StaticMeshFactoryNode->AddSocketUids(PipelineMeshesUtilities->GetMeshGeometryByUid(NodeUid).AttachedSocketUids);
 				MeshNode->AddTargetNodeUid(StaticMeshFactoryNode->GetUniqueID());
-				MeshNode->GetMaterialDependencies(MaterialDependencies);
+				
+				MeshNode->GetSlotMaterialDependencies(SlotMaterialDependencies);
 			}
 
-			const int32 MaterialCount = MaterialDependencies.Num();
-			for (int32 MaterialIndex = 0; MaterialIndex < MaterialCount; ++MaterialIndex)
-			{
-				const FString MaterialFactoryNodeUid = UInterchangeMaterialFactoryNode::GetMaterialFactoryNodeUidFromMaterialNodeUid(MaterialDependencies[MaterialIndex]);
-				if (BaseNodeContainer->IsNodeUidValid(MaterialFactoryNodeUid))
-				{
-					BaseNodeContainer->GetFactoryNode(MaterialFactoryNodeUid)->SetEnabled(true);
-					// Create a factory dependency so Material asset are imported before the skeletal mesh asset
-					TArray<FString> FactoryDependencies;
-					StaticMeshFactoryNode->GetFactoryDependencies(FactoryDependencies);
-					if (!FactoryDependencies.Contains(MaterialFactoryNodeUid))
-					{
-						StaticMeshFactoryNode->AddFactoryDependencyUid(MaterialFactoryNodeUid);
-					}
-				}
-			}
+			UE::Interchange::MeshesUtilities::ApplySlotMaterialDependencies(*StaticMeshFactoryNode, SlotMaterialDependencies, *BaseNodeContainer);
 
 			if (bImportCollisionAccordingToMeshName)
 			{
diff --git a/Engine/Plugins/Experimental/Interchange/Runtime/Source/Pipelines/Public/InterchangeGenericScenesPipeline.h b/Engine/Plugins/Experimental/Interchange/Runtime/Source/Pipelines/Public/InterchangeGenericScenesPipeline.h
--- a/Engine/Plugins/Experimental/Interchange/Runtime/Source/Pipelines/Public/InterchangeGenericScenesPipeline.h
+++ b/Engine/Plugins/Experimental/Interchange/Runtime/Source/Pipelines/Public/InterchangeGenericScenesPipeline.h
@@ -42,7 +42,7 @@ protected:
 	/**
 	 * Use to set up the given factory node's attributes after its initialization.
 	 */
-	virtual void SetUpFactoryNode(UInterchangeActorFactoryNode* ActorFactoryNode, const UInterchangeBaseNode* TranslatedAssetNode, UInterchangeBaseNodeContainer* FactoryNodeContainer) const;
+	virtual void SetUpFactoryNode(UInterchangeActorFactoryNode* ActorFactoryNode, const UInterchangeSceneNode* SceneNode, const UInterchangeBaseNode* TranslatedAssetNode, UInterchangeBaseNodeContainer* FactoryNodeContainer) const;
 	
 };
 
diff --git a/Engine/Plugins/Experimental/Interchange/Runtime/Source/Pipelines/Public/InterchangePipelineMeshesUtilities.h b/Engine/Plugins/Experimental/Interchange/Runtime/Source/Pipelines/Public/InterchangePipelineMeshesUtilities.h
--- a/Engine/Plugins/Experimental/Interchange/Runtime/Source/Pipelines/Public/InterchangePipelineMeshesUtilities.h
+++ b/Engine/Plugins/Experimental/Interchange/Runtime/Source/Pipelines/Public/InterchangePipelineMeshesUtilities.h
@@ -6,6 +6,10 @@
 #include "UObject/Object.h"
 #include "UObject/ObjectMacros.h"
 
+#include "InterchangeMaterialFactoryNode.h"
+#include "Nodes/InterchangeBaseNodeContainer.h"
+#include "Nodes/InterchangeFactoryBaseNode.h"
+
 #include "InterchangePipelineMeshesUtilities.generated.h"
 
 class UInterchangeBaseNodeContainer;
@@ -272,3 +276,33 @@ protected:
 	TMap<FString, FInterchangeMeshInstance> MeshInstancesPerMeshInstanceUid;
 	TMap<FString, FString> SkeletonRootUidPerMeshUid;
 };
+
+namespace UE::Interchange::MeshesUtilities
+{
+
+	/**
+	 * Applies material slot dependencies stored in SlotMaterialDependencies to FactoryNode.
+	 */
+	template<class T>
+	void ApplySlotMaterialDependencies(T& FactoryNode, const TMap<FString, FString>& SlotMaterialDependencies, const UInterchangeBaseNodeContainer& NodeContainer)
+	{
+		for (const TPair<FString, FString>& SlotMaterialDependency : SlotMaterialDependencies)
+		{
+			const FString MaterialFactoryNodeUid = UInterchangeMaterialFactoryNode::GetMaterialFactoryNodeUidFromMaterialNodeUid(SlotMaterialDependency.Value);
+			if (UInterchangeMaterialFactoryNode* MaterialFactoryNode = Cast< UInterchangeMaterialFactoryNode>(NodeContainer.GetFactoryNode(MaterialFactoryNodeUid)))
+			{
+				FactoryNode.SetSlotMaterialDependencyUid(SlotMaterialDependency.Key, MaterialFactoryNodeUid);
+				MaterialFactoryNode->SetEnabled(true);
+
+				// Create a factory dependency so Material asset are imported before the static mesh asset
+				TArray<FString> FactoryDependencies;
+				FactoryNode.GetFactoryDependencies(FactoryDependencies);
+				if (!FactoryDependencies.Contains(MaterialFactoryNodeUid))
+				{
+					FactoryNode.AddFactoryDependencyUid(MaterialFactoryNodeUid);
+				}
+			}
+		}
+
+	}
+}
\ No newline at end of file
diff --git a/Engine/Source/Runtime/Experimental/Interchange/Core/Public/Nodes/InterchangeBaseNodeUtilities.h b/Engine/Source/Runtime/Experimental/Interchange/Core/Public/Nodes/InterchangeBaseNodeUtilities.h
--- a/Engine/Source/Runtime/Experimental/Interchange/Core/Public/Nodes/InterchangeBaseNodeUtilities.h
+++ b/Engine/Source/Runtime/Experimental/Interchange/Core/Public/Nodes/InterchangeBaseNodeUtilities.h
@@ -310,21 +310,7 @@ public:
 		else
 		{
 			KeyCountHandle = InAttributes->GetAttributeHandle<int32>(KeyCountKey);
-
-			//Init map
-			int32 KeyCount;
-			KeyCountHandle.Get(KeyCount);
-
-			CachedKeysAndValues.Reserve(KeyCount);
-
-			for (int32 Index = 0; Index < KeyCount; Index++)
-			{
-				TAttributeHandle<KeyType> KeyAttribute = InAttributes->GetAttributeHandle<KeyType>(GetKeyAttribute(Index));
-				KeyType Key;
-				KeyAttribute.Get(Key);
-				TAttributeHandle<ValueType> ValueAttribute = InAttributes->GetAttributeHandle<ValueType>(GetValueAttribute(Key));
-				CachedKeysAndValues.Add(Key, TPair<TAttributeHandle<KeyType>, TAttributeHandle<ValueType>>(KeyAttribute, ValueAttribute));
-			}
+			RebuildCache();
 		}
 	}
 
@@ -339,10 +325,10 @@ public:
 		Attributes.Reset();
 	}
 
-	void SetKeyValue(const KeyType& InKey, const ValueType& InValue)
+	bool SetKeyValue(const KeyType& InKey, const ValueType& InValue)
 	{
 		const uint32 Hash = GetTypeHash(InKey);
-		SetKeyValueByHash(Hash, InKey, InValue);
+		return SetKeyValueByHash(Hash, InKey, InValue);
 	}
 
 	bool GetValue(const KeyType& InKey, ValueType& OutValue) const
@@ -363,12 +349,12 @@ public:
 		return RemoveKeyAndGetValueByHash(Hash, InKey, OutValue);
 	}
 
-	void SetKeyValueByHash(uint32 Hash, const KeyType& InKey, const ValueType& InValue)
+	bool SetKeyValueByHash(uint32 Hash, const KeyType& InKey, const ValueType& InValue)
 	{
 		TSharedPtr<FAttributeStorage, ESPMode::ThreadSafe> AttributesPtr = Attributes.Pin();
 		if (!ensure(AttributesPtr.IsValid()))
 		{
-			return;
+			return false;
 		}
 
 		if (TPair<TAttributeHandle<KeyType>, TAttributeHandle<ValueType>>* Pair = CachedKeysAndValues.FindByHash(Hash, InKey))
@@ -390,6 +376,8 @@ public:
 					));
 			KeyCountHandle.Set(CachedKeysAndValues.Num());
 		}
+
+		return true;
 	}
 
 
@@ -494,6 +482,33 @@ public:
 		return Map;
 	}
 
+	void RebuildCache()
+	{
+		TSharedPtr<FAttributeStorage, ESPMode::ThreadSafe> AttributesPtr = Attributes.Pin();
+		if (AttributesPtr.IsValid() && KeyCountHandle.IsValid())
+		{
+			int32 KeyCount;
+			if (KeyCountHandle.Get(KeyCount) == EAttributeStorageResult::Operation_Success)
+			{
+				CachedKeysAndValues.Empty(KeyCount);
+
+				for (int32 Index = 0; Index < KeyCount; Index++)
+				{
+					TAttributeHandle<KeyType> KeyAttribute = AttributesPtr->GetAttributeHandle<KeyType>(GetKeyAttribute(Index));
+					ensure(KeyAttribute.IsValid());
+
+					KeyType Key;
+					KeyAttribute.Get(Key);
+
+					TAttributeHandle<ValueType> ValueAttribute = AttributesPtr->GetAttributeHandle<ValueType>(GetValueAttribute(Key));
+					ensure(ValueAttribute.IsValid());
+
+					CachedKeysAndValues.Add(Key, TPair<TAttributeHandle<KeyType>, TAttributeHandle<ValueType>>(KeyAttribute, ValueAttribute));
+				}
+			}
+		}
+	}
+
 private:
 	FAttributeKey GetKeyAttribute(int32 Index) const
 	{
@@ -538,12 +553,14 @@ private:
 		KeyType LastKey;
 		LastKeyIndex.Get(LastKey);
 		CachedKeysAndValues[LastKey].Key = CachedPair.Key;
-		CachedPair->Key.Set(MoveTemp(LastKey));
+		CachedPair.Key.Set(MoveTemp(LastKey));
 
 		AttributesPtr->UnregisterAttribute(LastKeyIndex.GetKey());
 		AttributesPtr->UnregisterAttribute(CachedPair.Value.GetKey());
 		CachedKeysAndValues.RemoveByHash(Hash, InKey);
 		KeyCountHandle.Set(CachedKeysAndValues.Num());
+
+		return true;
 	}
 
 	void EmptyInternal(int32 NumOfExpectedElements)
diff --git a/Engine/Source/Runtime/Experimental/Interchange/Engine/Private/Tasks/InterchangeTaskParsing.cpp b/Engine/Source/Runtime/Experimental/Interchange/Engine/Private/Tasks/InterchangeTaskParsing.cpp
--- a/Engine/Source/Runtime/Experimental/Interchange/Engine/Private/Tasks/InterchangeTaskParsing.cpp
+++ b/Engine/Source/Runtime/Experimental/Interchange/Engine/Private/Tasks/InterchangeTaskParsing.cpp
@@ -176,12 +176,16 @@ void UE::Interchange::FTaskParsing::DoTask(ENamedThreads::Type CurrentThread, co
 					{
 						return true;
 					}
+					// Cache number of B's dependencies as reference on TSet can become stale
+					const int32 BDependenciesNum = BDependencies.Num();
+
 					const TSet<FString>& ADependencies = DependencyCache.GetAccumulatedDependencies(BaseNodeContainer, A.UniqueID);
 					if (ADependencies.Contains(B.UniqueID))
 					{
 						return false;
 					}
-					return ADependencies.Num() <= BDependencies.Num();
+
+					return ADependencies.Num() <= BDependenciesNum;
 				};
 
 				// Nodes cannot depend on a node from another source, so it's faster to sort the dependencies per-source and then append those to the TaskData arrays.
@@ -256,7 +260,7 @@ void UE::Interchange::FTaskParsing::DoTask(ENamedThreads::Type CurrentThread, co
 				);
 				CreateAssetPrerequistes.Add(AsyncHelper->CreateAssetTasks[CreateTaskIndex]);
 
-				CreatedTasksAssetNames.Add(AssetFullPath);
+					CreatedTasksAssetNames.Add(AssetFullPath);
 				return AsyncHelper->CreateAssetTasks[CreateTaskIndex];
 			}
 			else
