Fix gimbal lock by avoiding euler angle when importing fbx bone animation, we instead bake the data like the legacy importer did.
To achieve this a new type of payload was added to the payload interface which allow the client asking for payload to describe the bake settings (Hz, start time, end time)
#rnx

diff --git a/Engine/Plugins/Experimental/Interchange/Runtime/Source/Dispatcher/Private/InterchangeDispatcherTask.cpp b/Engine/Plugins/Experimental/Interchange/Runtime/Source/Dispatcher/Private/InterchangeDispatcherTask.cpp
--- a/Engine/Plugins/Experimental/Interchange/Runtime/Source/Dispatcher/Private/InterchangeDispatcherTask.cpp
+++ b/Engine/Plugins/Experimental/Interchange/Runtime/Source/Dispatcher/Private/InterchangeDispatcherTask.cpp
@@ -206,5 +206,93 @@ namespace UE
 			}
 			return ResultObject->TryGetStringField(GetResultFilenameJsonKey(), ResultFilename);
 		}
+
+		FString FJsonFetchAnimationBakeTransformPayloadCmd::ToJson() const
+		{
+			//Code should not do a ToJson if the data was not set before
+			ensure(bIsDataInitialize);
+
+			TSharedPtr<FJsonObject> CmdObject = MakeShared<FJsonObject>();
+			TSharedPtr<FJsonObject> ActionDataObject = MakeShared<FJsonObject>();
+			//CmdObject
+			CmdObject->SetStringField(GetCommandIDJsonKey(), GetAction());
+			CmdObject->SetStringField(GetTranslatorIDJsonKey(), GetTranslatorID());
+			ActionDataObject->SetStringField(GetPayloadKeyJsonKey(), GetPayloadKey());
+			
+			//Bake settings
+			ActionDataObject->SetNumberField(GetBakeFrequencyJsonKey(), GetBakeFrequency());
+			ActionDataObject->SetNumberField(GetRangeStartTimeJsonKey(), GetRangeStartTime());
+			ActionDataObject->SetNumberField(GetRangeEndTimeJsonKey(), GetRangeEndTime());
+			
+			CmdObject->SetObjectField(GetCommandDataJsonKey(), ActionDataObject);
+
+			FString FetchPayloadCmd;
+			TSharedRef< TJsonWriter< TCHAR, TPrettyJsonPrintPolicy<TCHAR> > > JsonWriter = TJsonWriterFactory< TCHAR, TPrettyJsonPrintPolicy<TCHAR> >::Create(&FetchPayloadCmd);
+			if (!FJsonSerializer::Serialize(CmdObject.ToSharedRef(), JsonWriter))
+			{
+				//Error creating the json cmd string 
+				return FString();
+			}
+			return FetchPayloadCmd;
+		}
+
+		bool FJsonFetchAnimationBakeTransformPayloadCmd::FromJson(const FString& JsonString)
+		{
+			TSharedRef<TJsonReader<TCHAR>> Reader = FJsonStringReader::Create(JsonString);
+
+			TSharedPtr<FJsonObject> CmdObject;
+			if (!FJsonSerializer::Deserialize(Reader, CmdObject) || !CmdObject.IsValid())
+			{
+				//Cannot read the json file
+				return false;
+			}
+			FString JsonActionValue;
+			if (!CmdObject->TryGetStringField(GetCommandIDJsonKey(), JsonActionValue))
+			{
+				//The json cmd id key is missing
+				return false;
+			}
+
+			if (!JsonActionValue.Equals(GetAction()))
+			{
+				//This json do not represent a load command
+				return false;
+			}
+
+			//Read the json
+			if (!CmdObject->TryGetStringField(GetTranslatorIDJsonKey(), TranslatorID))
+			{
+				//Missing Load command translator ID
+				return false;
+			}
+			const TSharedPtr<FJsonObject>* ActionDataObject = nullptr;
+			if (!CmdObject->TryGetObjectField(GetCommandDataJsonKey(), ActionDataObject))
+			{
+				//Missing Load Action data object
+				return false;
+			}
+			if (!((*ActionDataObject)->TryGetStringField(GetPayloadKeyJsonKey(), PayloadKey)))
+			{
+				return false;
+			}
+
+			//Bake settings
+			if (!(*ActionDataObject)->TryGetNumberField(GetBakeFrequencyJsonKey(), BakeFrequency))
+			{
+				return false;
+			}
+			if (!(*ActionDataObject)->TryGetNumberField(GetRangeStartTimeJsonKey(), RangeStartTime))
+			{
+				return false;
+			}
+			if (!(*ActionDataObject)->TryGetNumberField(GetRangeEndTimeJsonKey(), RangeEndTime))
+			{
+				return false;
+			}
+
+			//Since we filled the data from the json file, set the data has been initialize.
+			bIsDataInitialize = true;
+			return true;
+		}
 	} //ns Interchange
 }//ns UE
diff --git a/Engine/Plugins/Experimental/Interchange/Runtime/Source/Dispatcher/Public/InterchangeCommands.h b/Engine/Plugins/Experimental/Interchange/Runtime/Source/Dispatcher/Public/InterchangeCommands.h
--- a/Engine/Plugins/Experimental/Interchange/Runtime/Source/Dispatcher/Public/InterchangeCommands.h
+++ b/Engine/Plugins/Experimental/Interchange/Runtime/Source/Dispatcher/Public/InterchangeCommands.h
@@ -17,7 +17,7 @@ namespace UE
 		{
 		public:
 			//Major version should be updated when there is an existing API that has been change.
-			static int32 GetMajor() { return 8; }
+			static int32 GetMajor() { return 9; }
 			//Minor version must be updated when there is an addition to the API.
 			static int32 GetMinor() { return 0; }
 			//Patch version should be update if there is some bug fixes in the private code.
diff --git a/Engine/Plugins/Experimental/Interchange/Runtime/Source/Dispatcher/Public/InterchangeDispatcherTask.h b/Engine/Plugins/Experimental/Interchange/Runtime/Source/Dispatcher/Public/InterchangeDispatcherTask.h
--- a/Engine/Plugins/Experimental/Interchange/Runtime/Source/Dispatcher/Public/InterchangeDispatcherTask.h
+++ b/Engine/Plugins/Experimental/Interchange/Runtime/Source/Dispatcher/Public/InterchangeDispatcherTask.h
@@ -228,10 +228,74 @@ namespace UE
 				FString ResultFilename = FString();
 			};
 
-		private:
+		protected:
 			FString TranslatorID = FString();
 			FString PayloadKey = FString();
 		};
 
+		//Animation transform payload require transform to be bake by the translator
+		//The anim sequence API is not yet using curve to describe bone track animation
+		class INTERCHANGEDISPATCHER_API FJsonFetchAnimationBakeTransformPayloadCmd : public FJsonFetchPayloadCmd
+		{
+		public:
+			FJsonFetchAnimationBakeTransformPayloadCmd()
+			{
+				check(!bIsDataInitialize);
+			}
+
+			FJsonFetchAnimationBakeTransformPayloadCmd(const FString& InTranslatorID
+				, const FString& InPayloadKey
+				, double InBakeFrequency
+				, double InRangeStartTime
+				, double InRangeEndTime)
+				: FJsonFetchPayloadCmd(InTranslatorID, InPayloadKey)
+				, BakeFrequency(InBakeFrequency)
+				, RangeStartTime(InRangeStartTime)
+				, RangeEndTime(InRangeEndTime)
+			{}
+
+			virtual FString ToJson() const override;
+			virtual bool FromJson(const FString& JsonString) override;
+
+			double GetBakeFrequency() const
+			{
+				ensure(bIsDataInitialize);
+				return BakeFrequency;
+			}
+
+			static FString GetBakeFrequencyJsonKey()
+			{
+				static const FString Key = TEXT("BakeFrequency");
+				return Key;
+			}
+
+			double GetRangeStartTime() const
+			{
+				ensure(bIsDataInitialize);
+				return RangeStartTime;
+			}
+
+			static FString GetRangeStartTimeJsonKey()
+			{
+				static const FString Key = TEXT("RangeStartTime");
+				return Key;
+			}
+
+			double GetRangeEndTime() const
+			{
+				ensure(bIsDataInitialize);
+				return RangeEndTime;
+			}
+
+			static FString GetRangeEndTimeJsonKey()
+			{
+				static const FString Key = TEXT("RangeEndTime");
+				return Key;
+			}
+		protected:
+			double BakeFrequency = 30.0;
+			double RangeStartTime = 0.0;
+			double RangeEndTime = 1.0/BakeFrequency;
+		};
 	} //ns Interchange
 }//ns UE
diff --git a/Engine/Plugins/Experimental/Interchange/Runtime/Source/FactoryNodes/Private/InterchangeAnimSequenceFactoryNode.cpp b/Engine/Plugins/Experimental/Interchange/Runtime/Source/FactoryNodes/Private/InterchangeAnimSequenceFactoryNode.cpp
--- a/Engine/Plugins/Experimental/Interchange/Runtime/Source/FactoryNodes/Private/InterchangeAnimSequenceFactoryNode.cpp
+++ b/Engine/Plugins/Experimental/Interchange/Runtime/Source/FactoryNodes/Private/InterchangeAnimSequenceFactoryNode.cpp
@@ -3,6 +3,21 @@
 #include "InterchangeAnimSequenceFactoryNode.h"
 #include "Animation/AnimSequence.h"
 
+namespace UE::Interchange::Animation
+{
+	FFrameRate ConvertSampleRatetoFrameRate(double SampleRate)
+	{
+		double IntegralPart = FMath::FloorToDouble(SampleRate);
+		double FractionalPart = SampleRate - IntegralPart;
+		const int32 Tolerance = 1000000;
+		double Divisor = static_cast<double>(FMath::GreatestCommonDivisor(FMath::RoundToInt(FractionalPart * Tolerance), Tolerance));
+		int32 Denominator = static_cast<int32>(Tolerance / Divisor);
+		int32 Numerator = static_cast<int32>(IntegralPart * Denominator + FMath::RoundToDouble(FractionalPart * Tolerance) / Divisor);
+		check(Denominator != 0);
+		return FFrameRate(Numerator, Denominator);
+	}
+}
+
 UInterchangeAnimSequenceFactoryNode::UInterchangeAnimSequenceFactoryNode()
 {
 }
@@ -100,3 +115,78 @@ bool UInterchangeAnimSequenceFactoryNode::SetCustomSkeletonSoftObjectPath(const
 {
 	IMPLEMENT_NODE_ATTRIBUTE_SETTER_NODELEGATE(SkeletonSoftObjectPath, FSoftObjectPath)
 }
+
+/************************************************************************/
+/* Automation tests                                                     */
+/************************************************************************/
+#if WITH_EDITOR
+#if WITH_DEV_AUTOMATION_TESTS
+
+#include "Misc/AutomationTest.h"
+
+IMPLEMENT_SIMPLE_AUTOMATION_TEST(FInterchangeAnimSequenceTest, "System.Runtime.Interchange.ConvertSampleRatetoFrameRate", EAutomationTestFlags::EditorContext | EAutomationTestFlags::EngineFilter)
+
+bool FInterchangeAnimSequenceTest::RunTest(const FString& Parameters)
+{
+	FFrameRate FrameRate;
+	FrameRate = UE::Interchange::Animation::ConvertSampleRatetoFrameRate(120.0);
+	TestEqual(TEXT("`ConvertSampleRatetoFrameRate` Error converting 120.0 to FFrameRate"), FrameRate.Numerator, 120);
+	TestEqual(TEXT("`ConvertSampleRatetoFrameRate` Error converting 120.0 to FFrameRate"), FrameRate.Denominator, 1);
+
+	FrameRate = UE::Interchange::Animation::ConvertSampleRatetoFrameRate(100.0);
+	TestEqual(TEXT("`ConvertSampleRatetoFrameRate` Error converting 100.0 to FFrameRate"), FrameRate.Numerator, 100);
+	TestEqual(TEXT("`ConvertSampleRatetoFrameRate` Error converting 100.0 to FFrameRate"), FrameRate.Denominator, 1);
+
+	FrameRate = UE::Interchange::Animation::ConvertSampleRatetoFrameRate(60.0);
+	TestEqual(TEXT("`ConvertSampleRatetoFrameRate` Error converting 60.0 to FFrameRate"), FrameRate.Numerator, 60);
+	TestEqual(TEXT("`ConvertSampleRatetoFrameRate` Error converting 60.0 to FFrameRate"), FrameRate.Denominator, 1);
+
+	FrameRate = UE::Interchange::Animation::ConvertSampleRatetoFrameRate(50.0);
+	TestEqual(TEXT("`ConvertSampleRatetoFrameRate` Error converting 50.0 to FFrameRate"), FrameRate.Numerator, 50);
+	TestEqual(TEXT("`ConvertSampleRatetoFrameRate` Error converting 50.0 to FFrameRate"), FrameRate.Denominator, 1);
+
+	FrameRate = UE::Interchange::Animation::ConvertSampleRatetoFrameRate(48.0);
+	TestEqual(TEXT("`ConvertSampleRatetoFrameRate` Error converting 48.0 to FFrameRate"), FrameRate.Numerator, 48);
+	TestEqual(TEXT("`ConvertSampleRatetoFrameRate` Error converting 48.0 to FFrameRate"), FrameRate.Denominator, 1);
+
+	FrameRate = UE::Interchange::Animation::ConvertSampleRatetoFrameRate(30.0);
+	TestEqual(TEXT("`ConvertSampleRatetoFrameRate` Error converting 30.0 to FFrameRate"), FrameRate.Numerator, 30);
+	TestEqual(TEXT("`ConvertSampleRatetoFrameRate` Error converting 30.0 to FFrameRate"), FrameRate.Denominator, 1);
+
+	FrameRate = UE::Interchange::Animation::ConvertSampleRatetoFrameRate(29.97);
+	TestEqual(TEXT("`ConvertSampleRatetoFrameRate` Error converting 29.97 to FFrameRate"), FrameRate.Numerator, 2997);
+	TestEqual(TEXT("`ConvertSampleRatetoFrameRate` Error converting 29.97 to FFrameRate"), FrameRate.Denominator, 100);
+
+	FrameRate = UE::Interchange::Animation::ConvertSampleRatetoFrameRate(25.0);
+	TestEqual(TEXT("`ConvertSampleRatetoFrameRate` Error converting 25.0 to FFrameRate"), FrameRate.Numerator, 25);
+	TestEqual(TEXT("`ConvertSampleRatetoFrameRate` Error converting 25.0 to FFrameRate"), FrameRate.Denominator, 1);
+
+	FrameRate = UE::Interchange::Animation::ConvertSampleRatetoFrameRate(24.0);
+	TestEqual(TEXT("`ConvertSampleRatetoFrameRate` Error converting 24.0 to FFrameRate"), FrameRate.Numerator, 24);
+	TestEqual(TEXT("`ConvertSampleRatetoFrameRate` Error converting 24.0 to FFrameRate"), FrameRate.Denominator, 1);
+
+	FrameRate = UE::Interchange::Animation::ConvertSampleRatetoFrameRate(23.976);
+	TestEqual(TEXT("`ConvertSampleRatetoFrameRate` Error converting 23.976 to FFrameRate"), FrameRate.Numerator, 2997);
+	TestEqual(TEXT("`ConvertSampleRatetoFrameRate` Error converting 23.976 to FFrameRate"), FrameRate.Denominator, 125);
+
+	FrameRate = UE::Interchange::Animation::ConvertSampleRatetoFrameRate(96.0);
+	TestEqual(TEXT("`ConvertSampleRatetoFrameRate` Error converting 96.0 to FFrameRate"), FrameRate.Numerator, 96);
+	TestEqual(TEXT("`ConvertSampleRatetoFrameRate` Error converting 96.0 to FFrameRate"), FrameRate.Denominator, 1);
+
+	FrameRate = UE::Interchange::Animation::ConvertSampleRatetoFrameRate(72.0);
+	TestEqual(TEXT("`ConvertSampleRatetoFrameRate` Error converting 72.0 to FFrameRate"), FrameRate.Numerator, 72);
+	TestEqual(TEXT("`ConvertSampleRatetoFrameRate` Error converting 72.0 to FFrameRate"), FrameRate.Denominator, 1);
+
+	FrameRate = UE::Interchange::Animation::ConvertSampleRatetoFrameRate(59.94);
+	TestEqual(TEXT("`ConvertSampleRatetoFrameRate` Error converting 59.94 to FFrameRate"), FrameRate.Numerator, 2997);
+	TestEqual(TEXT("`ConvertSampleRatetoFrameRate` Error converting 59.94 to FFrameRate"), FrameRate.Denominator, 50);
+
+	FrameRate = UE::Interchange::Animation::ConvertSampleRatetoFrameRate(119.88);
+	TestEqual(TEXT("`ConvertSampleRatetoFrameRate` Error converting 119.88 to FFrameRate"), FrameRate.Numerator, 2997);
+	TestEqual(TEXT("`ConvertSampleRatetoFrameRate` Error converting 119.88 to FFrameRate"), FrameRate.Denominator, 25);
+
+	return true;
+}
+
+#endif //WITH_DEV_AUTOMATION_TESTS
+#endif //WITH_EDITOR
diff --git a/Engine/Plugins/Experimental/Interchange/Runtime/Source/FactoryNodes/Public/InterchangeAnimSequenceFactoryNode.h b/Engine/Plugins/Experimental/Interchange/Runtime/Source/FactoryNodes/Public/InterchangeAnimSequenceFactoryNode.h
--- a/Engine/Plugins/Experimental/Interchange/Runtime/Source/FactoryNodes/Public/InterchangeAnimSequenceFactoryNode.h
+++ b/Engine/Plugins/Experimental/Interchange/Runtime/Source/FactoryNodes/Public/InterchangeAnimSequenceFactoryNode.h
@@ -5,9 +5,14 @@
 #include "CoreMinimal.h"
 #include "Nodes/InterchangeFactoryBaseNode.h"
 #include "Animation/AnimSequence.h"
-
+#include "Misc/FrameRate.h"
 #include "InterchangeAnimSequenceFactoryNode.generated.h"
 
+namespace UE::Interchange::Animation
+{
+	INTERCHANGEFACTORYNODES_API FFrameRate ConvertSampleRatetoFrameRate(double SampleRate);
+}
+
 UCLASS(BlueprintType, Experimental)
 class INTERCHANGEFACTORYNODES_API UInterchangeAnimSequenceFactoryNode : public UInterchangeFactoryBaseNode
 {
diff --git a/Engine/Plugins/Experimental/Interchange/Runtime/Source/Import/Private/Animation/InterchangeAnimSequenceFactory.cpp b/Engine/Plugins/Experimental/Interchange/Runtime/Source/Import/Private/Animation/InterchangeAnimSequenceFactory.cpp
--- a/Engine/Plugins/Experimental/Interchange/Runtime/Source/Import/Private/Animation/InterchangeAnimSequenceFactory.cpp
+++ b/Engine/Plugins/Experimental/Interchange/Runtime/Source/Import/Private/Animation/InterchangeAnimSequenceFactory.cpp
@@ -45,7 +45,7 @@ namespace UE::Interchange::Private
 		, const IInterchangeAnimationPayloadInterface* AnimSequenceTranslatorPayloadInterface
 		, const FString& AssetName)
 	{
-		TMap<FString, TFuture<TOptional<UE::Interchange::FAnimationTransformPayloadData>>> AnimationPayloads;
+		TMap<FString, TFuture<TOptional<UE::Interchange::FAnimationBakeTransformPayloadData>>> AnimationPayloads;
 
 		FString SkeletonRootUid;
 		if (!SkeletonFactoryNode->GetCustomRootJointUid(SkeletonRootUid))
@@ -57,6 +57,20 @@ namespace UE::Interchange::Private
 		AnimSequenceFactoryNode->GetCustomImportBoneTracks(bImportBoneTracks);
 		if (bImportBoneTracks)
 		{
+			//Get the sample rate, default to 30Hz in case the attribute is missing
+			double SampleRate = 30.0;
+			AnimSequenceFactoryNode->GetCustomImportBoneTracksSampleRate(SampleRate);
+
+			double RangeStart = 0.0;
+			AnimSequenceFactoryNode->GetCustomImportBoneTracksRangeStart(RangeStart);
+
+			double RangeEnd = 1.0 / SampleRate; //One frame duration per default
+			AnimSequenceFactoryNode->GetCustomImportBoneTracksRangeStop(RangeEnd);
+
+			const double BakeInterval = 1.0 / SampleRate;
+			const double SequenceLength = FMath::Max<double>(RangeEnd - RangeStart, MINIMUM_ANIMATION_LENGTH);
+			int32 BakeKeyCount = (SequenceLength / BakeInterval) + 1;
+
 			TArray<FString> SkeletonNodes;
 			GetSkeletonSceneNodeFlatListRecursive(NodeContainer, SkeletonRootUid, SkeletonNodes);
 			for (const FString& NodeUid : SkeletonNodes)
@@ -66,23 +80,10 @@ namespace UE::Interchange::Private
 					FString PayloadKey;
 					if (SkeletonSceneNode->GetCustomTransformCurvePayloadKey(PayloadKey))
 					{
-						AnimationPayloads.Add(PayloadKey, AnimSequenceTranslatorPayloadInterface->GetAnimationTransformPayloadData(PayloadKey));
+						AnimationPayloads.Add(PayloadKey, AnimSequenceTranslatorPayloadInterface->GetAnimationBakeTransformPayloadData(PayloadKey, SampleRate, RangeStart, RangeEnd));
 					}
 				}
 			}
-			//Get the sample rate, default to 30Hz in case the attribute is missing
-			double SampleRate = 30.0;
-			AnimSequenceFactoryNode->GetCustomImportBoneTracksSampleRate(SampleRate);
-
-			double RangeStart = 0.0;
-			AnimSequenceFactoryNode->GetCustomImportBoneTracksRangeStart(RangeStart);
-
-			double RangeStop = 1.0 / SampleRate; //One frame duration per default
-			AnimSequenceFactoryNode->GetCustomImportBoneTracksRangeStop(RangeStop);
-
-			const double BakeInterval = 1.0 / SampleRate;
-			const double SequenceLength = RangeStop - RangeStart;
-			int32 BakeKeyCount = (SequenceLength / BakeInterval)+1;
 
 			IAnimationDataController& Controller = AnimSequence->GetController();
 			//This destroy all previously imported animation raw data
@@ -116,82 +117,49 @@ namespace UE::Interchange::Private
 					FString PayloadKey;
 					if (SkeletonSceneNode->GetCustomTransformCurvePayloadKey(PayloadKey))
 					{
-						TOptional<UE::Interchange::FAnimationTransformPayloadData> AnimationTransformPayload = AnimationPayloads.FindChecked(PayloadKey).Get();
+						TOptional<UE::Interchange::FAnimationBakeTransformPayloadData> AnimationTransformPayload = AnimationPayloads.FindChecked(PayloadKey).Get();
 						if (!AnimationTransformPayload.IsSet())
 						{
 							UE_LOG(LogInterchangeImport, Warning, TEXT("Invalid animation transform payload key [%s] AnimSequence asset %s"), *PayloadKey, *AssetName);
 							continue;
 						}
 
+						if (AnimationTransformPayload->Transforms.Num() == 0)
+						{
+							//We need at least one transform
+							AnimationTransformPayload->Transforms.Add(FTransform::Identity);
+						}
+
 						FRawAnimSequenceTrack RawTrack;
 						RawTrack.PosKeys.Reserve(BakeKeyCount);
 						RawTrack.RotKeys.Reserve(BakeKeyCount);
 						RawTrack.ScaleKeys.Reserve(BakeKeyCount);
 						TArray<float> TimeKeys;
 						TimeKeys.Reserve(BakeKeyCount);
-						for (double CurrentTime = RangeStart; CurrentTime <= RangeStop + SMALL_NUMBER; CurrentTime += BakeInterval)
+						
+						//Everything should match Key count, sample rate and range
+						check(AnimationTransformPayload->Transforms.Num() == BakeKeyCount);
+						check(FMath::IsNearlyEqual(AnimationTransformPayload->BakeFrequency, SampleRate, UE_DOUBLE_KINDA_SMALL_NUMBER));
+						check(FMath::IsNearlyEqual(AnimationTransformPayload->RangeStartTime, RangeStart, UE_DOUBLE_KINDA_SMALL_NUMBER));
+						check(FMath::IsNearlyEqual(AnimationTransformPayload->RangeEndTime, RangeEnd, UE_DOUBLE_KINDA_SMALL_NUMBER));
+
+						int32 TransformIndex = 0;
+						for (double CurrentTime = RangeStart; CurrentTime <= RangeEnd + UE_DOUBLE_SMALL_NUMBER; CurrentTime += BakeInterval, ++TransformIndex)
 						{
-							//Default value to identity
-							FVector3f Position(0.0f);
-							FVector3f Euler(0.0f);
-							FVector3f Scale(1.0f);
-							//Fill the anim sequence raw data
-							for (int32 CurveIndex = 0; CurveIndex < AnimationTransformPayload->TransformCurves.Num(); ++CurveIndex)
-							{
-								const UE::Interchange::FAnimationCurveTransformPayloadData& CurveData = AnimationTransformPayload->TransformCurves[CurveIndex];
-								if (CurveData.Curve.GetNumKeys() <= 0)
-								{
-									//skip empty curve, the default value will be applied
-									continue;
-								}
-
-								if (CurveData.TransformChannel == EInterchangeTransformCurveChannel::TranslationX)
-								{
-									Position.X = CurveData.Curve.Eval(CurrentTime);
-								}
-								else if (CurveData.TransformChannel == EInterchangeTransformCurveChannel::TranslationY)
-								{
-									Position.Y = CurveData.Curve.Eval(CurrentTime);
-								}
-								else if (CurveData.TransformChannel == EInterchangeTransformCurveChannel::TranslationZ)
-								{
-									Position.Z = CurveData.Curve.Eval(CurrentTime);
-								}
-								else if (CurveData.TransformChannel == EInterchangeTransformCurveChannel::EulerX)
-								{
-									Euler.X = CurveData.Curve.Eval(CurrentTime);
-								}
-								else if (CurveData.TransformChannel == EInterchangeTransformCurveChannel::EulerY)
-								{
-									Euler.Y = CurveData.Curve.Eval(CurrentTime);
-								}
-								else if (CurveData.TransformChannel == EInterchangeTransformCurveChannel::EulerZ)
-								{
-									Euler.Z = CurveData.Curve.Eval(CurrentTime);
-								}
-								else if (CurveData.TransformChannel == EInterchangeTransformCurveChannel::ScaleX)
-								{
-									Scale.X = CurveData.Curve.Eval(CurrentTime);
-								}
-								else if (CurveData.TransformChannel == EInterchangeTransformCurveChannel::ScaleY)
-								{
-									Scale.Y = CurveData.Curve.Eval(CurrentTime);
-								}
-								else if (CurveData.TransformChannel == EInterchangeTransformCurveChannel::ScaleZ)
-								{
-									Scale.Z = CurveData.Curve.Eval(CurrentTime);
-								}
-							}
-							FQuat4f Quaternion = FQuat4f::MakeFromEuler(Euler);
-							
+							check(AnimationTransformPayload->Transforms.IsValidIndex(TransformIndex));
+							FTransform3f AnimKeyTransform = FTransform3f(AnimationTransformPayload->Transforms[TransformIndex]);
 							if (bApplyGlobalOffset)
 							{
-								FTransform3f AnimKeyTransform(Quaternion, Position, Scale);
 								AnimKeyTransform = AnimKeyTransform * GlobalOffsetTransform;
-								Position = AnimKeyTransform.GetLocation();
-								Quaternion = AnimKeyTransform.GetRotation();
-								Scale = AnimKeyTransform.GetScale3D();
 							}
+							//Default value to identity
+							FVector3f Position(0.0f);
+							FQuat4f Quaternion(0.0f);
+							FVector3f Scale(1.0f);
+							
+							Position = AnimKeyTransform.GetLocation();
+							Quaternion = AnimKeyTransform.GetRotation();
+							Scale = AnimKeyTransform.GetScale3D();
 							RawTrack.ScaleKeys.Add(Scale);
 							RawTrack.PosKeys.Add(Position);
 							RawTrack.RotKeys.Add(Quaternion);
@@ -213,19 +181,6 @@ namespace UE::Interchange::Private
 			}
 		}
 	}
-
-	FFrameRate ConvertSampleRatetoFrameRate(double SampleRate)
-	{
-		double IntegralPart = FMath::FloorToDouble(SampleRate);
-		double FractionalPart = SampleRate - IntegralPart;
-		const int32 Tolerance = 1000000;
-		double Divisor = static_cast<double>(FMath::GreatestCommonDivisor(FMath::RoundToInt(FractionalPart * Tolerance), Tolerance));
-		int32 Denominator = static_cast<int32>(Tolerance / Divisor);
-		int32 Numerator = static_cast<int32>(IntegralPart * Denominator + FMath::RoundToDouble(FractionalPart * Tolerance) / Divisor);
-		check(Denominator != 0);
-		return FFrameRate(Numerator, Denominator);
-	}
-
 } //namespace UE::Interchange::Private
 #endif //WITH_EDITOR
 
@@ -252,6 +207,12 @@ UObject* UInterchangeAnimSequenceFactory::CreateEmptyAsset(const FCreateAssetPar
 	{
 		return nullptr;
 	}
+	
+	//Verify if the bone track animation is valid (sequence length versus framerate ...)
+	if(!IsBoneTrackAnimationValid(AnimSequenceFactoryNode, Arguments))
+	{
+		return nullptr;
+	}
 
 	// create an asset if it doesn't exist
 	UObject* ExistingAsset = StaticFindObject(nullptr, Arguments.Parent, *Arguments.AssetName);
@@ -299,6 +260,12 @@ UObject* UInterchangeAnimSequenceFactory::CreateAsset(const FCreateAssetParams&
 		return nullptr;
 	}
 
+	//Verify if the bone track animation is valid (sequence length versus framerate ...)
+	if (!IsBoneTrackAnimationValid(AnimSequenceFactoryNode, Arguments))
+	{
+		return nullptr;
+	}
+
 	FString SkeletonUid;
 	if (!AnimSequenceFactoryNode->GetCustomSkeletonFactoryNodeUid(SkeletonUid))
 	{
@@ -398,20 +365,15 @@ UObject* UInterchangeAnimSequenceFactory::CreateAsset(const FCreateAssetParams&
 	{
 		FFrameRate FrameRate(30, 1);
 		double SampleRate = 30.0;
-		
-		//If we import bone track we will set the framerate according to the bone track settings
-		//The default is 30Hz when there is no bone tracks, i.e. like when we just have morphtarget attribute curve animation.
-		//Curve do not need framerate, framerate is needed if we want to bake curve animation data.
-		//TODO: discuss with animation team so we can pass curves for bones transform animation instead of baked data.
+
 		bool bImportBoneTracks = false;
 		if (AnimSequenceFactoryNode->GetCustomImportBoneTracks(bImportBoneTracks) && bImportBoneTracks)
 		{
 			if (AnimSequenceFactoryNode->GetCustomImportBoneTracksSampleRate(SampleRate))
 			{
-				FrameRate = UE::Interchange::Private::ConvertSampleRatetoFrameRate(SampleRate);
+				FrameRate = UE::Interchange::Animation::ConvertSampleRatetoFrameRate(SampleRate);
 			}
 		}
-
 		IAnimationDataController& Controller = AnimSequence->GetController();
 		Controller.OpenBracket(NSLOCTEXT("InterchangeAnimSequenceFactory", "ImportAnimationInterchange_Bracket", "Importing Animation (Interchange)"));
 		AnimSequence->SetSkeleton(Skeleton);
@@ -505,77 +467,40 @@ bool UInterchangeAnimSequenceFactory::SetSourceFilename(const UObject* Object, c
 	return false;
 }
 
-/************************************************************************/
-/* Automation tests                                                     */
-/************************************************************************/
-#if WITH_EDITOR
-#if WITH_DEV_AUTOMATION_TESTS
-
-#include "Misc/AutomationTest.h"
-
-IMPLEMENT_SIMPLE_AUTOMATION_TEST(FInterchangeAnimSequenceTest, "System.Runtime.Interchange.ConvertSampleRatetoFrameRate", EAutomationTestFlags::EditorContext | EAutomationTestFlags::EngineFilter)
-
-bool FInterchangeAnimSequenceTest::RunTest(const FString& Parameters)
+bool UInterchangeAnimSequenceFactory::IsBoneTrackAnimationValid(const UInterchangeAnimSequenceFactoryNode* AnimSequenceFactoryNode, const FCreateAssetParams& Arguments)
 {
-	FFrameRate FrameRate;
-	FrameRate = UE::Interchange::Private::ConvertSampleRatetoFrameRate(120.0);
-	TestEqual(TEXT("`ConvertSampleRatetoFrameRate` Error converting 120.0 to FFrameRate"), FrameRate.Numerator, 120);
-	TestEqual(TEXT("`ConvertSampleRatetoFrameRate` Error converting 120.0 to FFrameRate"), FrameRate.Denominator, 1);
-	
-	FrameRate = UE::Interchange::Private::ConvertSampleRatetoFrameRate(100.0);
-	TestEqual(TEXT("`ConvertSampleRatetoFrameRate` Error converting 100.0 to FFrameRate"), FrameRate.Numerator, 100);
-	TestEqual(TEXT("`ConvertSampleRatetoFrameRate` Error converting 100.0 to FFrameRate"), FrameRate.Denominator, 1);
-
-	FrameRate = UE::Interchange::Private::ConvertSampleRatetoFrameRate(60.0);
-	TestEqual(TEXT("`ConvertSampleRatetoFrameRate` Error converting 60.0 to FFrameRate"), FrameRate.Numerator, 60);
-	TestEqual(TEXT("`ConvertSampleRatetoFrameRate` Error converting 60.0 to FFrameRate"), FrameRate.Denominator, 1);
-
-	FrameRate = UE::Interchange::Private::ConvertSampleRatetoFrameRate(50.0);
-	TestEqual(TEXT("`ConvertSampleRatetoFrameRate` Error converting 50.0 to FFrameRate"), FrameRate.Numerator, 50);
-	TestEqual(TEXT("`ConvertSampleRatetoFrameRate` Error converting 50.0 to FFrameRate"), FrameRate.Denominator, 1);
-
-	FrameRate = UE::Interchange::Private::ConvertSampleRatetoFrameRate(48.0);
-	TestEqual(TEXT("`ConvertSampleRatetoFrameRate` Error converting 48.0 to FFrameRate"), FrameRate.Numerator, 48);
-	TestEqual(TEXT("`ConvertSampleRatetoFrameRate` Error converting 48.0 to FFrameRate"), FrameRate.Denominator, 1);
+	bool bResult = true;
+	FFrameRate FrameRate(30, 1);
+	double SampleRate = 30.0;
 
-	FrameRate = UE::Interchange::Private::ConvertSampleRatetoFrameRate(30.0);
-	TestEqual(TEXT("`ConvertSampleRatetoFrameRate` Error converting 30.0 to FFrameRate"), FrameRate.Numerator, 30);
-	TestEqual(TEXT("`ConvertSampleRatetoFrameRate` Error converting 30.0 to FFrameRate"), FrameRate.Denominator, 1);
-
-	FrameRate = UE::Interchange::Private::ConvertSampleRatetoFrameRate(29.97);
-	TestEqual(TEXT("`ConvertSampleRatetoFrameRate` Error converting 29.97 to FFrameRate"), FrameRate.Numerator, 2997);
-	TestEqual(TEXT("`ConvertSampleRatetoFrameRate` Error converting 29.97 to FFrameRate"), FrameRate.Denominator, 100);
-
-	FrameRate = UE::Interchange::Private::ConvertSampleRatetoFrameRate(25.0);
-	TestEqual(TEXT("`ConvertSampleRatetoFrameRate` Error converting 25.0 to FFrameRate"), FrameRate.Numerator, 25);
-	TestEqual(TEXT("`ConvertSampleRatetoFrameRate` Error converting 25.0 to FFrameRate"), FrameRate.Denominator, 1);
-
-	FrameRate = UE::Interchange::Private::ConvertSampleRatetoFrameRate(24.0);
-	TestEqual(TEXT("`ConvertSampleRatetoFrameRate` Error converting 24.0 to FFrameRate"), FrameRate.Numerator, 24);
-	TestEqual(TEXT("`ConvertSampleRatetoFrameRate` Error converting 24.0 to FFrameRate"), FrameRate.Denominator, 1);
+	bool bImportBoneTracks = false;
+	if (AnimSequenceFactoryNode->GetCustomImportBoneTracks(bImportBoneTracks) && bImportBoneTracks)
+	{
+		if (AnimSequenceFactoryNode->GetCustomImportBoneTracksSampleRate(SampleRate))
+		{
+			FrameRate = UE::Interchange::Animation::ConvertSampleRatetoFrameRate(SampleRate);
+		}
 
-	FrameRate = UE::Interchange::Private::ConvertSampleRatetoFrameRate(23.976);
-	TestEqual(TEXT("`ConvertSampleRatetoFrameRate` Error converting 23.976 to FFrameRate"), FrameRate.Numerator, 2997);
-	TestEqual(TEXT("`ConvertSampleRatetoFrameRate` Error converting 23.976 to FFrameRate"), FrameRate.Denominator, 125);
+		double RangeStart = 0.0;
+		AnimSequenceFactoryNode->GetCustomImportBoneTracksRangeStart(RangeStart);
 
-	FrameRate = UE::Interchange::Private::ConvertSampleRatetoFrameRate(96.0);
-	TestEqual(TEXT("`ConvertSampleRatetoFrameRate` Error converting 96.0 to FFrameRate"), FrameRate.Numerator, 96);
-	TestEqual(TEXT("`ConvertSampleRatetoFrameRate` Error converting 96.0 to FFrameRate"), FrameRate.Denominator, 1);
+		double RangeEnd = 1.0 / SampleRate; //One frame duration per default
+		AnimSequenceFactoryNode->GetCustomImportBoneTracksRangeStop(RangeEnd);
 
-	FrameRate = UE::Interchange::Private::ConvertSampleRatetoFrameRate(72.0);
-	TestEqual(TEXT("`ConvertSampleRatetoFrameRate` Error converting 72.0 to FFrameRate"), FrameRate.Numerator, 72);
-	TestEqual(TEXT("`ConvertSampleRatetoFrameRate` Error converting 72.0 to FFrameRate"), FrameRate.Denominator, 1);
+		const double SequenceLength = FMath::Max<double>(RangeEnd - RangeStart, MINIMUM_ANIMATION_LENGTH);
 
-	FrameRate = UE::Interchange::Private::ConvertSampleRatetoFrameRate(59.94);
-	TestEqual(TEXT("`ConvertSampleRatetoFrameRate` Error converting 59.94 to FFrameRate"), FrameRate.Numerator, 2997);
-	TestEqual(TEXT("`ConvertSampleRatetoFrameRate` Error converting 59.94 to FFrameRate"), FrameRate.Denominator, 50);
+		const float SubFrame = FrameRate.AsFrameTime(SequenceLength).GetSubFrame();
 
-	FrameRate = UE::Interchange::Private::ConvertSampleRatetoFrameRate(119.88);
-	TestEqual(TEXT("`ConvertSampleRatetoFrameRate` Error converting 119.88 to FFrameRate"), FrameRate.Numerator, 2997);
-	TestEqual(TEXT("`ConvertSampleRatetoFrameRate` Error converting 119.88 to FFrameRate"), FrameRate.Denominator, 25);
-	
-	return true;
+		if (!FMath::IsNearlyZero(SubFrame, KINDA_SMALL_NUMBER) && !FMath::IsNearlyEqual(SubFrame, 1.0f, KINDA_SMALL_NUMBER))
+		{
+			UInterchangeResultError_Generic* Message = AddMessage<UInterchangeResultError_Generic>();
+			Message->SourceAssetName = Arguments.SourceData->GetFilename();
+			Message->DestinationAssetName = Arguments.AssetName;
+			Message->AssetType = UAnimSequence::StaticClass();
+			Message->Text = FText::Format(NSLOCTEXT("UInterchangeAnimSequenceFactory", "WrongSequenceLength", "Animation length {0} is not compatible with import frame-rate {1} (sub frame {2}), animation has to be frame-border aligned."),
+				FText::AsNumber(SequenceLength), FrameRate.ToPrettyText(), FText::AsNumber(SubFrame));
+			bResult = false;
+		}
+	}
+	return bResult;
 }
-
-#endif //WITH_DEV_AUTOMATION_TESTS
-#endif //WITH_EDITOR
diff --git a/Engine/Plugins/Experimental/Interchange/Runtime/Source/Import/Private/Fbx/InterchangeFbxTranslator.cpp b/Engine/Plugins/Experimental/Interchange/Runtime/Source/Import/Private/Fbx/InterchangeFbxTranslator.cpp
--- a/Engine/Plugins/Experimental/Interchange/Runtime/Source/Import/Private/Fbx/InterchangeFbxTranslator.cpp
+++ b/Engine/Plugins/Experimental/Interchange/Runtime/Source/Import/Private/Fbx/InterchangeFbxTranslator.cpp
@@ -486,6 +486,79 @@ TFuture<TOptional<UE::Interchange::FAnimationTransformPayloadData>> UInterchange
 	return Promise->GetFuture();
 }
 
+TFuture<TOptional<UE::Interchange::FAnimationBakeTransformPayloadData>> UInterchangeFbxTranslator::GetAnimationBakeTransformPayloadData(const FString& PayLoadKey, const double BakeFrequency, const double RangeStartSecond, const double RangeStopSecond) const
+{
+	TSharedPtr<TPromise<TOptional<UE::Interchange::FAnimationBakeTransformPayloadData>>> Promise = MakeShared<TPromise<TOptional<UE::Interchange::FAnimationBakeTransformPayloadData>>>();
+
+	if (!Dispatcher.IsValid())
+	{
+		Promise->SetValue(TOptional<UE::Interchange::FAnimationBakeTransformPayloadData>());
+		return Promise->GetFuture();
+	}
+
+
+	// Create a json command to read the fbx file
+	FString JsonCommand = CreateFetchAnimationBakeTransformPayloadFbxCommand(PayLoadKey, BakeFrequency, RangeStartSecond, RangeStopSecond);
+	const int32 CreatedTaskIndex = Dispatcher->AddTask(JsonCommand, FInterchangeDispatcherTaskCompleted::CreateLambda([this, Promise](const int32 TaskIndex)
+		{
+			UE::Interchange::ETaskState TaskState;
+			FString JsonResult;
+			TArray<FString> JsonMessages;
+			Dispatcher->GetTaskState(TaskIndex, TaskState, JsonResult, JsonMessages);
+
+			// Parse the Json messages into UInterchangeResults
+			for (const FString& JsonMessage : JsonMessages)
+			{
+				AddMessage(UInterchangeResult::FromJson(JsonMessage));
+			}
+
+			if (TaskState != UE::Interchange::ETaskState::ProcessOk)
+			{
+				Promise->SetValue(TOptional<UE::Interchange::FAnimationBakeTransformPayloadData>());
+				return;
+			}
+
+			// Grab the result file and fill the BaseNodeContainer
+			UE::Interchange::FJsonFetchPayloadCmd::JsonResultParser ResultParser;
+			ResultParser.FromJson(JsonResult);
+			FString AnimationTransformPayloadFilename = ResultParser.GetResultFilename();
+
+			if (!ensure(FPaths::FileExists(AnimationTransformPayloadFilename)))
+			{
+				// TODO log an error saying the payload file does not exist even if the get payload command succeeded
+				Promise->SetValue(TOptional<UE::Interchange::FAnimationBakeTransformPayloadData>());
+				return;
+			}
+
+			UE::Interchange::FAnimationBakeTransformPayloadData AnimationTransformPayload;
+
+			// All sub object should be gone with the reset
+			TArray64<uint8> Buffer;
+			FFileHelper::LoadFileToArray(Buffer, *AnimationTransformPayloadFilename);
+			uint8* FileData = Buffer.GetData();
+			int64 FileDataSize = Buffer.Num();
+			if (FileDataSize < 1)
+			{
+				// Nothing to load from this file
+				Promise->SetValue(TOptional<UE::Interchange::FAnimationBakeTransformPayloadData>());
+				return;
+			}
+
+			// Buffer keeps the ownership of the data, the large memory reader is use to serialize the TMap
+			FLargeMemoryReader Ar(FileData, FileDataSize);
+			AnimationTransformPayload.Serialize(Ar);
+			Promise->SetValue(MoveTemp(AnimationTransformPayload));
+		}));
+
+	// The task was not added to the dispatcher
+	if (CreatedTaskIndex == INDEX_NONE)
+	{
+		Promise->SetValue(TOptional<UE::Interchange::FAnimationBakeTransformPayloadData>{});
+	}
+
+	return Promise->GetFuture();
+}
+
 FString UInterchangeFbxTranslator::CreateLoadFbxFileCommand(const FString& FbxFilePath) const
 {
 	UE::Interchange::FJsonLoadSourceCmd LoadSourceCommand(TEXT("FBX"), FbxFilePath);
@@ -496,4 +569,10 @@ FString UInterchangeFbxTranslator::CreateFetchPayloadFbxCommand(const FString& F
 {
 	UE::Interchange::FJsonFetchPayloadCmd PayloadCommand(TEXT("FBX"), FbxPayloadKey);
 	return PayloadCommand.ToJson();
-}
\ No newline at end of file
+}
+
+FString UInterchangeFbxTranslator::CreateFetchAnimationBakeTransformPayloadFbxCommand(const FString& FbxPayloadKey, const double BakeFrequency, const double RangeStartTime, const double RangeEndTime) const
+{
+	UE::Interchange::FJsonFetchAnimationBakeTransformPayloadCmd PayloadCommand(TEXT("FBX"), FbxPayloadKey, BakeFrequency, RangeStartTime, RangeEndTime);
+	return PayloadCommand.ToJson();
+}
diff --git a/Engine/Plugins/Experimental/Interchange/Runtime/Source/Import/Public/Animation/InterchangeAnimSequenceFactory.h b/Engine/Plugins/Experimental/Interchange/Runtime/Source/Import/Public/Animation/InterchangeAnimSequenceFactory.h
--- a/Engine/Plugins/Experimental/Interchange/Runtime/Source/Import/Public/Animation/InterchangeAnimSequenceFactory.h
+++ b/Engine/Plugins/Experimental/Interchange/Runtime/Source/Import/Public/Animation/InterchangeAnimSequenceFactory.h
@@ -10,6 +10,7 @@
 #include "InterchangeAnimSequenceFactory.generated.h"
 
 class UAnimSequence;
+class UInterchangeAnimSequenceFactoryNode;
 
 UCLASS(BlueprintType, Experimental)
 class INTERCHANGEIMPORT_API UInterchangeAnimSequenceFactory : public UInterchangeFactoryBase
@@ -33,6 +34,9 @@ public:
 	virtual bool GetSourceFilenames(const UObject* Object, TArray<FString>& OutSourceFilenames) const override;
 	virtual bool SetSourceFilename(const UObject* Object, const FString& SourceFilename, int32 SourceIndex) const override;
 
+private:
+	bool IsBoneTrackAnimationValid(const UInterchangeAnimSequenceFactoryNode* AnimSequenceFactoryNode, const FCreateAssetParams& Arguments);
+
 	// Interchange factory base interface end
 	//////////////////////////////////////////////////////////////////////////
 };
diff --git a/Engine/Plugins/Experimental/Interchange/Runtime/Source/Import/Public/Animation/InterchangeAnimationPayloadInterface.h b/Engine/Plugins/Experimental/Interchange/Runtime/Source/Import/Public/Animation/InterchangeAnimationPayloadInterface.h
--- a/Engine/Plugins/Experimental/Interchange/Runtime/Source/Import/Public/Animation/InterchangeAnimationPayloadInterface.h
+++ b/Engine/Plugins/Experimental/Interchange/Runtime/Source/Import/Public/Animation/InterchangeAnimationPayloadInterface.h
@@ -3,6 +3,7 @@
 
 #include "CoreMinimal.h"
 #include "Animation/InterchangeAnimationPayload.h"
+#include "Misc/FrameRate.h"
 #include "UObject/Interface.h"
 #include "UObject/Object.h"
 #include "UObject/ObjectMacros.h"
@@ -31,6 +32,8 @@ public:
 	 * 
 	 */
 	virtual TFuture<TOptional<UE::Interchange::FAnimationTransformPayloadData>> GetAnimationTransformPayloadData(const FString& PayLoadKey) const = 0;
+
+	virtual TFuture<TOptional<UE::Interchange::FAnimationBakeTransformPayloadData>> GetAnimationBakeTransformPayloadData(const FString& PayLoadKey, const double BakeFrequency, const double RangeStartSecond, const double RangeStopSecond) const = 0;
 };
 
 
diff --git a/Engine/Plugins/Experimental/Interchange/Runtime/Source/Import/Public/Fbx/InterchangeFbxTranslator.h b/Engine/Plugins/Experimental/Interchange/Runtime/Source/Import/Public/Fbx/InterchangeFbxTranslator.h
--- a/Engine/Plugins/Experimental/Interchange/Runtime/Source/Import/Public/Fbx/InterchangeFbxTranslator.h
+++ b/Engine/Plugins/Experimental/Interchange/Runtime/Source/Import/Public/Fbx/InterchangeFbxTranslator.h
@@ -84,12 +84,16 @@ public:
 
 	virtual TFuture<TOptional<UE::Interchange::FAnimationTransformPayloadData>> GetAnimationTransformPayloadData(const FString& PayLoadKey) const override;
 
+	virtual TFuture<TOptional<UE::Interchange::FAnimationBakeTransformPayloadData>> GetAnimationBakeTransformPayloadData(const FString& PayLoadKey, const double BakeFrequency, const double RangeStartSecond, const double RangeStopSecond) const override;
+
 	/* IInterchangeAnimationPayloadInterface End */
 private:
 	FString CreateLoadFbxFileCommand(const FString& FbxFilePath) const;
 
 	FString CreateFetchPayloadFbxCommand(const FString& FbxPayloadKey) const;
 
+	FString CreateFetchAnimationBakeTransformPayloadFbxCommand(const FString& FbxPayloadKey, const double BakeFrequency, const double RangeStartTime, const double RangeEndTime) const;
+	
 	//Dispatcher is mutable since it is create during the Translate operation
 	//We do not want to allocate the dispatcher and start the InterchangeWorker process
 	//in the constructor because Archetype, CDO and registered translators will
diff --git a/Engine/Plugins/Experimental/Interchange/Runtime/Source/Parsers/CommonParser/Public/InterchangeCommonAnimationPayload.h b/Engine/Plugins/Experimental/Interchange/Runtime/Source/Parsers/CommonParser/Public/InterchangeCommonAnimationPayload.h
--- a/Engine/Plugins/Experimental/Interchange/Runtime/Source/Parsers/CommonParser/Public/InterchangeCommonAnimationPayload.h
+++ b/Engine/Plugins/Experimental/Interchange/Runtime/Source/Parsers/CommonParser/Public/InterchangeCommonAnimationPayload.h
@@ -2,6 +2,8 @@
 
 #pragma once
 
+#include "Math/Transform.h"
+#include "Serialization/Archive.h"
 #include "UObject/Class.h"
 #include "UObject/Object.h"
 #include "UObject/ObjectMacros.h"
@@ -12,6 +14,26 @@
 
 #include "InterchangeCommonAnimationPayload.generated.h"
 
+namespace UE::Interchange
+{
+
+	struct INTERCHANGECOMMONPARSER_API FAnimationBakeTransformPayloadData
+	{
+		double BakeFrequency = 30.0;
+		double RangeStartTime = 0.0;
+		double RangeEndTime = 1.0 / BakeFrequency;
+		TArray<FTransform> Transforms;
+
+		void Serialize(FArchive& Ar)
+		{
+			Ar << BakeFrequency;
+			Ar << RangeStartTime;
+			Ar << RangeEndTime;
+			Ar << Transforms;
+		}
+	};
+}
+
 UENUM()
 enum class EInterchangeTransformCurveChannel
 {
diff --git a/Engine/Plugins/Experimental/Interchange/Runtime/Source/Parsers/Fbx/Private/FbxAPI.cpp b/Engine/Plugins/Experimental/Interchange/Runtime/Source/Parsers/Fbx/Private/FbxAPI.cpp
--- a/Engine/Plugins/Experimental/Interchange/Runtime/Source/Parsers/Fbx/Private/FbxAPI.cpp
+++ b/Engine/Plugins/Experimental/Interchange/Runtime/Source/Parsers/Fbx/Private/FbxAPI.cpp
@@ -129,8 +129,29 @@ namespace UE
 					return false;
 				}
 
-				TSharedPtr<FPayloadContextBase>& PayloadContext = PayloadContexts.FindChecked(PayloadKey);
-				return PayloadContext->FetchPayloadToFile(*this, PayloadFilepath);
+				{
+					//Critical section to force payload to be fetch one by one with no concurrency.
+					FScopeLock Lock(&PayloadCriticalSection);
+					TSharedPtr<FPayloadContextBase>& PayloadContext = PayloadContexts.FindChecked(PayloadKey);
+					return PayloadContext->FetchPayloadToFile(*this, PayloadFilepath);
+				}
+			}
+
+			bool FFbxParser::FetchAnimationBakeTransformPayload(const FString& PayloadKey, const double BakeFrequency, const double RangeStartTime, const double RangeEndTime, const FString& PayloadFilepath)
+			{
+				if (!PayloadContexts.Contains(PayloadKey))
+				{
+					UInterchangeResultError_Generic* Message = AddMessage<UInterchangeResultError_Generic>();
+					Message->Text = LOCTEXT("CannotRetrievePayload", "Cannot retrieve payload; payload key doesn't have any context.");
+					return false;
+				}
+
+				{
+					//Critical section to force payload to be fetch one by one with no concurrency.
+					FScopeLock Lock(&PayloadCriticalSection);
+					TSharedPtr<FPayloadContextBase>& PayloadContext = PayloadContexts.FindChecked(PayloadKey);
+					return PayloadContext->FetchAnimationBakeTransformPayloadToFile(*this, BakeFrequency, RangeStartTime, RangeEndTime, PayloadFilepath);
+				}
 			}
 
 			void FFbxParser::CleanupFbxData()
diff --git a/Engine/Plugins/Experimental/Interchange/Runtime/Source/Parsers/Fbx/Private/FbxAPI.h b/Engine/Plugins/Experimental/Interchange/Runtime/Source/Parsers/Fbx/Private/FbxAPI.h
--- a/Engine/Plugins/Experimental/Interchange/Runtime/Source/Parsers/Fbx/Private/FbxAPI.h
+++ b/Engine/Plugins/Experimental/Interchange/Runtime/Source/Parsers/Fbx/Private/FbxAPI.h
@@ -46,7 +46,9 @@ namespace UE
 
 				/* Extract the fbx data from the sdk into our node container */
 				bool FetchPayloadData(const FString& PayloadKey, const FString& PayloadFilepath);
-			
+
+				/* Extract the fbx data from the sdk into our node container */
+				bool FetchAnimationBakeTransformPayload(const FString& PayloadKey, const double BakeFrequency, const double RangeStartTime, const double RangeEndTime, const FString& PayloadFilepath);
 				/**
 				 * This function is used to add the given message object directly into the results for this operation.
 				 */
@@ -68,12 +70,11 @@ namespace UE
 				}
 
 				/**
-				 * Critical section to avoid getting animation payload curves asynchronously.
+				 * Critical section to avoid getting multiple payload in same time.
 				 * The FBX evaluator use a cache mechanism for evaluating global transform that is not thread safe.
-				 * We avoid evaluating global transform asynchronously by using this CriticalSection when we import curves payload.
-				 * Each scene nodes representing a joint have is own payload for the transform animation.
+				 * There si other stuff in the sdk which are not thread safe, so all fbx payload should be fetch one by one
 				 */
-				FCriticalSection AnimationTransformPayloadCriticalSection;
+				FCriticalSection PayloadCriticalSection;
 			private:
 
 				void CleanupFbxData();
diff --git a/Engine/Plugins/Experimental/Interchange/Runtime/Source/Parsers/Fbx/Private/FbxAnimation.cpp b/Engine/Plugins/Experimental/Interchange/Runtime/Source/Parsers/Fbx/Private/FbxAnimation.cpp
--- a/Engine/Plugins/Experimental/Interchange/Runtime/Source/Parsers/Fbx/Private/FbxAnimation.cpp
+++ b/Engine/Plugins/Experimental/Interchange/Runtime/Source/Parsers/Fbx/Private/FbxAnimation.cpp
@@ -409,6 +409,47 @@ namespace UE::Interchange::Private
 		return bResult;
 	}
 
+	bool ImportBakeTransforms(FbxNode* Node, FAnimationBakeTransformPayloadData& AnimationBakeTransformPayloadData)
+	{
+		if (!ensure(!FMath::IsNearlyZero(AnimationBakeTransformPayloadData.BakeFrequency)))
+		{
+			return false;
+		}
+		
+		FbxTime StartTime;
+		StartTime.SetSecondDouble(AnimationBakeTransformPayloadData.RangeStartTime);
+		FbxTime EndTime;
+		EndTime.SetSecondDouble(AnimationBakeTransformPayloadData.RangeEndTime);
+		if (!ensure(AnimationBakeTransformPayloadData.RangeEndTime > AnimationBakeTransformPayloadData.RangeStartTime))
+		{
+			return false;
+		}
+
+		const double TimeStepSecond = 1.0 / AnimationBakeTransformPayloadData.BakeFrequency;
+		FbxTime TimeStep = 0;
+		TimeStep.SetSecondDouble(TimeStepSecond);
+
+		const int32 NumFrame = FMath::RoundToInt32((AnimationBakeTransformPayloadData.RangeEndTime - AnimationBakeTransformPayloadData.RangeStartTime * AnimationBakeTransformPayloadData.BakeFrequency));
+		
+		//Add a threshold when we compare if we have reach the end of the animation
+		const FbxTime TimeComparisonThreshold = (UE_DOUBLE_KINDA_SMALL_NUMBER * static_cast<double>(FBXSDK_TC_SECOND));
+		AnimationBakeTransformPayloadData.Transforms.Empty(NumFrame);
+
+		for (FbxTime CurTime = StartTime; CurTime < (EndTime + TimeComparisonThreshold); CurTime += TimeStep)
+		{
+			FbxAMatrix NodeTransform = Node->EvaluateGlobalTransform(CurTime);
+			FbxNode* ParentNode = Node->GetParent();
+			if (ParentNode)
+			{
+				FbxAMatrix ParentTransform = ParentNode->EvaluateGlobalTransform(CurTime);
+				NodeTransform = ParentTransform.Inverse() * NodeTransform;
+			}
+			AnimationBakeTransformPayloadData.Transforms.Add(UE::Interchange::Private::FFbxConvert::ConvertTransform(NodeTransform));
+		}
+
+		return true;
+	}
+
 	struct FGetFbxTransformCurvesParameters
 	{
 		FGetFbxTransformCurvesParameters(FbxScene* InSDKScene, FbxNode* InNode)
@@ -521,6 +562,12 @@ namespace UE::Interchange::Private
 	}
 
 	bool FAnimationPayloadContextTransform::FetchPayloadToFile(FFbxParser& Parser, const FString& PayloadFilepath)
+	{
+		//Will need this one to import animated curves
+		return false;
+	}
+
+	bool FAnimationPayloadContextTransform::FetchAnimationBakeTransformPayloadToFile(FFbxParser& Parser, const double BakeFrequency, const double RangeStartTime, const double RangeEndTime, const FString& PayloadFilepath)
 	{
 		if (!ensure(SDKScene != nullptr))
 		{
@@ -538,23 +585,16 @@ namespace UE::Interchange::Private
 			return false;
 		}
 
-		FGetFbxTransformCurvesParameters Parameters(SDKScene, Node);
-		GetFbxTransformCurves(Parameters);
-
-		TArray<FInterchangeCurve> Curves;
-		Curves.AddDefaulted(Parameters.TransformChannelCount);
-
-		if (Parameters.IsNodeAnimated)
-		{
-			//Fbx node global transform evaluator is not thread safe, so we acquire a lock to import the curves
-			FScopeLock Lock(&Parser.AnimationTransformPayloadCriticalSection);
-			ImportTransformCurves(Node, Parameters.TransformChannelCurves, Curves);
-		}
+		bool bBakeTransform = false;
+		FAnimationBakeTransformPayloadData AnimationBakeTransformPayloadData;
+		AnimationBakeTransformPayloadData.BakeFrequency = BakeFrequency;
+		AnimationBakeTransformPayloadData.RangeStartTime = RangeStartTime;
+		AnimationBakeTransformPayloadData.RangeEndTime = RangeEndTime;
 
-		//Dump all curves to a file
+		ImportBakeTransforms(Node, AnimationBakeTransformPayloadData);
 		{
 			FLargeMemoryWriter Ar;
-			Ar << Curves;
+			AnimationBakeTransformPayloadData.Serialize(Ar);
 			uint8* ArchiveData = Ar.GetData();
 			int64 ArchiveSize = Ar.TotalSize();
 			TArray64<uint8> Buffer(ArchiveData, ArchiveSize);
diff --git a/Engine/Plugins/Experimental/Interchange/Runtime/Source/Parsers/Fbx/Private/FbxAnimation.h b/Engine/Plugins/Experimental/Interchange/Runtime/Source/Parsers/Fbx/Private/FbxAnimation.h
--- a/Engine/Plugins/Experimental/Interchange/Runtime/Source/Parsers/Fbx/Private/FbxAnimation.h
+++ b/Engine/Plugins/Experimental/Interchange/Runtime/Source/Parsers/Fbx/Private/FbxAnimation.h
@@ -20,6 +20,7 @@ namespace UE::Interchange::Private
 		virtual ~FAnimationPayloadContextTransform() {}
 		virtual FString GetPayloadType() const override { return TEXT("TransformAnimation-PayloadContext"); }
 		virtual bool FetchPayloadToFile(FFbxParser& Parser, const FString& PayloadFilepath) override;
+		virtual bool FetchAnimationBakeTransformPayloadToFile(FFbxParser& Parser, const double BakeFrequency, const double RangeStartTime, const double RangeEndTime, const FString& PayloadFilepath) override;
 		FbxNode* Node = nullptr;
 		FbxScene* SDKScene = nullptr;
 	};
diff --git a/Engine/Plugins/Experimental/Interchange/Runtime/Source/Parsers/Fbx/Private/FbxHelper.h b/Engine/Plugins/Experimental/Interchange/Runtime/Source/Parsers/Fbx/Private/FbxHelper.h
--- a/Engine/Plugins/Experimental/Interchange/Runtime/Source/Parsers/Fbx/Private/FbxHelper.h
+++ b/Engine/Plugins/Experimental/Interchange/Runtime/Source/Parsers/Fbx/Private/FbxHelper.h
@@ -19,6 +19,7 @@ namespace UE
 				virtual ~FPayloadContextBase() {}
 				virtual FString GetPayloadType() const { return FString(); }
 				virtual bool FetchPayloadToFile(FFbxParser& Parser, const FString& PayloadFilepath) { return false; }
+				virtual bool FetchAnimationBakeTransformPayloadToFile(FFbxParser& Parser, const double BakeFrequency, const double RangeStartTime, const double RangeEndTime, const FString& PayloadFilepath) { return false; }
 			};
 
 			struct FFbxHelper
diff --git a/Engine/Plugins/Experimental/Interchange/Runtime/Source/Parsers/Fbx/Private/InterchangeFbxParser.cpp b/Engine/Plugins/Experimental/Interchange/Runtime/Source/Parsers/Fbx/Private/InterchangeFbxParser.cpp
--- a/Engine/Plugins/Experimental/Interchange/Runtime/Source/Parsers/Fbx/Private/InterchangeFbxParser.cpp
+++ b/Engine/Plugins/Experimental/Interchange/Runtime/Source/Parsers/Fbx/Private/InterchangeFbxParser.cpp
@@ -84,6 +84,28 @@ namespace UE
 			}
 		}
 
+		void FInterchangeFbxParser::FetchAnimationBakeTransformPayload(const FString& PayloadKey, const double BakeFrequency, const double RangeStartTime, const double RangeEndTime, const FString& ResultFolder)
+		{
+			check(FbxParserPrivate.IsValid());
+			ResultsContainer->Empty();
+			FString PayloadFilepathCopy;
+			{
+				FScopeLock Lock(&ResultPayloadsCriticalSection);
+				FString& PayloadFilepath = ResultPayloads.FindOrAdd(PayloadKey);
+				PayloadFilepath = ResultFolder + TEXT("/") + PayloadKey + TEXT(".payload");
+
+				//Copy the map filename key because we are multithreaded and the TMap can be reallocated
+				PayloadFilepathCopy = PayloadFilepath;
+			}
+			if (!FbxParserPrivate->FetchAnimationBakeTransformPayload(PayloadKey, BakeFrequency, RangeStartTime, RangeEndTime, PayloadFilepathCopy))
+			{
+				UInterchangeResultError_Generic* Error = AddMessage<UInterchangeResultError_Generic>();
+				Error->SourceAssetName = SourceFilename;
+				Error->Text = LOCTEXT("CantFetchPayload", "Cannot fetch FBX payload data.");
+				return;
+			}
+		}
+
 		TArray<FString> FInterchangeFbxParser::GetJsonLoadMessages() const
 		{
 			TArray<FString> JsonResults;
diff --git a/Engine/Plugins/Experimental/Interchange/Runtime/Source/Parsers/Fbx/Public/InterchangeFbxParser.h b/Engine/Plugins/Experimental/Interchange/Runtime/Source/Parsers/Fbx/Public/InterchangeFbxParser.h
--- a/Engine/Plugins/Experimental/Interchange/Runtime/Source/Parsers/Fbx/Public/InterchangeFbxParser.h
+++ b/Engine/Plugins/Experimental/Interchange/Runtime/Source/Parsers/Fbx/Public/InterchangeFbxParser.h
@@ -37,6 +37,16 @@ namespace UE
 			 */
 			void FetchPayload(const FString& PayloadKey, const FString& ResultFolder);
 
+			/**
+			 * Extract bake transform animation payload data from the fbx, the key tell the translator what payload the client ask
+			 * @param - PayloadKey is the key that describe the payload data to extract from the fbx file
+			 * @param - BakeFrequency is the Hz at which we should bake the transform
+			 * @param - RangeStartTime is the start time of the bake
+			 * @param - RangeEndTime is the end time of the bake
+			 * @param - ResultFolder is the folder where we must put any result file
+			 */
+			void FetchAnimationBakeTransformPayload(const FString& PayloadKey, const double BakeFrequency, const double RangeStartTime, const double RangeEndTime, const FString& ResultFolder);
+
 			FString GetResultFilepath() const { return ResultFilepath; }
 			FString GetResultPayloadFilepath(const FString& PayloadKey) const
 			{
diff --git a/Engine/Plugins/Experimental/Interchange/Runtime/Source/Pipelines/Private/InterchangeGenericAnimationPipeline.cpp b/Engine/Plugins/Experimental/Interchange/Runtime/Source/Pipelines/Private/InterchangeGenericAnimationPipeline.cpp
--- a/Engine/Plugins/Experimental/Interchange/Runtime/Source/Pipelines/Private/InterchangeGenericAnimationPipeline.cpp
+++ b/Engine/Plugins/Experimental/Interchange/Runtime/Source/Pipelines/Private/InterchangeGenericAnimationPipeline.cpp
@@ -197,6 +197,28 @@ void UInterchangeGenericAnimationPipeline::ExecutePreImportPipeline(UInterchange
 			AnimSequenceName.LeftChopInline(9);
 		}
 		AnimSequenceName += TEXT("_Anim");
+
+		if (bImportBoneTracks)
+		{
+			FFrameRate FrameRate = UE::Interchange::Animation::ConvertSampleRatetoFrameRate(SampleRate);
+
+			const double SequenceLength = FMath::Max<double>(RangeStop - RangeStart, MINIMUM_ANIMATION_LENGTH);
+
+			const float SubFrame = FrameRate.AsFrameTime(SequenceLength).GetSubFrame();
+
+			if (!FMath::IsNearlyZero(SubFrame, KINDA_SMALL_NUMBER) && !FMath::IsNearlyEqual(SubFrame, 1.0f, KINDA_SMALL_NUMBER))
+			{
+				UInterchangeResultError_Generic* Message = AddMessage<UInterchangeResultError_Generic>();
+				Message->SourceAssetName = SourceDatas[0]->GetFilename();
+				Message->DestinationAssetName = AnimSequenceName;
+				Message->AssetType = UAnimSequence::StaticClass();
+				Message->Text = FText::Format(NSLOCTEXT("UInterchangeGenericAnimationPipeline", "WrongSequenceLength", "Animation length {0} is not compatible with import frame-rate {1} (sub frame {2}), animation has to be frame-border aligned."),
+					FText::AsNumber(SequenceLength), FrameRate.ToPrettyText(), FText::AsNumber(SubFrame));
+				//Skip this anim sequence factory node
+				continue;
+			}
+		}
+
 		UInterchangeAnimSequenceFactoryNode* AnimSequenceFactoryNode = NewObject<UInterchangeAnimSequenceFactoryNode>(BaseNodeContainer, NAME_None);
 		AnimSequenceFactoryNode->InitializeAnimSequenceNode(AnimSequenceUid, AnimSequenceName);
 		
diff --git a/Engine/Source/Programs/InterchangeWorker/Private/InterchangeWorkerImpl.cpp b/Engine/Source/Programs/InterchangeWorker/Private/InterchangeWorkerImpl.cpp
--- a/Engine/Source/Programs/InterchangeWorker/Private/InterchangeWorkerImpl.cpp
+++ b/Engine/Source/Programs/InterchangeWorker/Private/InterchangeWorkerImpl.cpp
@@ -218,6 +218,7 @@ void FInterchangeWorkerImpl::ProcessCommand(const TSharedPtr<UE::Interchange::IC
 	FString JSonResult;
 	TArray<FString> JSonMessages;
 	FJsonLoadSourceCmd LoadSourceCommand;
+	FJsonFetchAnimationBakeTransformPayloadCmd FetchAnimationBakeTransform;
 	FJsonFetchPayloadCmd FetchPayloadCommand;
 	//Any command FromJson function return true if the Json descibe the command
 	if (LoadSourceCommand.FromJson(JsonToProcess))
@@ -229,6 +230,15 @@ void FInterchangeWorkerImpl::ProcessCommand(const TSharedPtr<UE::Interchange::IC
 			ProcessResult = LoadFbxFile(LoadSourceCommand, JSonResult, JSonMessages);
 		}
 	}
+	else if (FetchAnimationBakeTransform.FromJson(JsonToProcess))
+	{
+		//Load file command
+		if (FetchAnimationBakeTransform.GetTranslatorID().Equals(TEXT("FBX"), ESearchCase::IgnoreCase))
+		{
+			//We want to load an FBX file
+			ProcessResult = FetchFbxPayload(FetchAnimationBakeTransform, JSonResult, JSonMessages);
+		}
+	}
 	else if (FetchPayloadCommand.FromJson(JsonToProcess))
 	{
 		//Load file command
@@ -286,4 +296,22 @@ ETaskState FInterchangeWorkerImpl::FetchFbxPayload(const FJsonFetchPayloadCmd& F
 	OutJSonResult = ResultParser.ToJson();
 	ResultState = ETaskState::ProcessOk;
 	return ResultState;
+}
+
+ETaskState FInterchangeWorkerImpl::FetchFbxPayload(const FJsonFetchAnimationBakeTransformPayloadCmd& FetchAnimationBakeTransformPayloadCommand, FString& OutJSonResult, TArray<FString>& OutJSonMessages)
+{
+	ETaskState ResultState = ETaskState::Unknown;
+	FString PayloadKey = FetchAnimationBakeTransformPayloadCommand.GetPayloadKey();
+
+	FbxParser.FetchAnimationBakeTransformPayload(PayloadKey
+		, FetchAnimationBakeTransformPayloadCommand.GetBakeFrequency()
+		, FetchAnimationBakeTransformPayloadCommand.GetRangeStartTime()
+		, FetchAnimationBakeTransformPayloadCommand.GetRangeEndTime()
+		, ResultFolder);
+	FJsonLoadSourceCmd::JsonResultParser ResultParser;
+	ResultParser.SetResultFilename(FbxParser.GetResultPayloadFilepath(PayloadKey));
+	OutJSonMessages = FbxParser.GetJsonLoadMessages();
+	OutJSonResult = ResultParser.ToJson();
+	ResultState = ETaskState::ProcessOk;
+	return ResultState;
 }
\ No newline at end of file
diff --git a/Engine/Source/Programs/InterchangeWorker/Private/InterchangeWorkerImpl.h b/Engine/Source/Programs/InterchangeWorker/Private/InterchangeWorkerImpl.h
--- a/Engine/Source/Programs/InterchangeWorker/Private/InterchangeWorkerImpl.h
+++ b/Engine/Source/Programs/InterchangeWorker/Private/InterchangeWorkerImpl.h
@@ -27,7 +27,7 @@ private:
 
 	UE::Interchange::ETaskState LoadFbxFile(const UE::Interchange::FJsonLoadSourceCmd& LoadSourceCommand, FString& OutJSonResult, TArray<FString>& OutJSonMessages);
 	UE::Interchange::ETaskState FetchFbxPayload(const UE::Interchange::FJsonFetchPayloadCmd& FetchPayloadCommand, FString& OutJSonResult, TArray<FString>& OutJSonMessages);
-
+	UE::Interchange::ETaskState FetchFbxPayload(const UE::Interchange::FJsonFetchAnimationBakeTransformPayloadCmd& FetchAnimationBakeTransformPayloadCommand, FString& OutJSonResult, TArray<FString>& OutJSonMessages);
 private:
 	UE::Interchange::FNetworkClientNode NetworkInterface;
 	UE::Interchange::FCommandQueue CommandIO;
