Interchange:
 - Removed use of ensure based on wrong assumption about presence of slot information  in factory nodes for meshes
 - Fixed logic in ApplySlotMaterialDependencies to check for UInterchangeBaseMaterialFactoryNode instead of UInterchangeMaterialFactoryNode

Datasmith:
 - Removed use of InterchangeGenericAssetPipeline for Datasmith import
 - Added logic to the translator to pro
 - Added Datasmith static mesh pipeline to handle case where allslots of a mesh have the same material
 - Modified level pipeline to handle case where all overrides of a mesh actor have the same material
 - Temporarily forced creation of all assets during a Datasmith import
 - Limit compilation of DatasmithInterchange module to Win64, Mac and Linux

diff --git a/Engine/Plugins/Experimental/Interchange/Runtime/Source/Import/Private/Mesh/InterchangeSkeletalMeshFactory.cpp b/Engine/Plugins/Experimental/Interchange/Runtime/Source/Import/Private/Mesh/InterchangeSkeletalMeshFactory.cpp
--- a/Engine/Plugins/Experimental/Interchange/Runtime/Source/Import/Private/Mesh/InterchangeSkeletalMeshFactory.cpp
+++ b/Engine/Plugins/Experimental/Interchange/Runtime/Source/Import/Private/Mesh/InterchangeSkeletalMeshFactory.cpp
@@ -1345,8 +1345,7 @@ UObject* UInterchangeSkeletalMeshFactory::CreateAsset(const FCreateAssetParams&
 				}
 
 				// If we don't have a match, add a new entry to the material list.
-				// However this should not happen. The translator should have properly set the UInterchangeMeshNode with the slot names used by all LODs.
-				if (!ensure(LODMatIndex != INDEX_NONE))
+				if (LODMatIndex == INDEX_NONE)
 				{
 					LODMatIndex = Materials.Add(FSkeletalMaterial(ImportedMaterials[ImportedMaterialIndex].Material.Get(), true, false, ImportedMaterialName, ImportedMaterialName));
 				}
diff --git a/Engine/Plugins/Experimental/Interchange/Runtime/Source/Import/Private/Mesh/InterchangeStaticMeshFactory.cpp b/Engine/Plugins/Experimental/Interchange/Runtime/Source/Import/Private/Mesh/InterchangeStaticMeshFactory.cpp
--- a/Engine/Plugins/Experimental/Interchange/Runtime/Source/Import/Private/Mesh/InterchangeStaticMeshFactory.cpp
+++ b/Engine/Plugins/Experimental/Interchange/Runtime/Source/Import/Private/Mesh/InterchangeStaticMeshFactory.cpp
@@ -337,8 +337,7 @@ UObject* UInterchangeStaticMeshFactory::CreateAsset(const FCreateAssetParams& Ar
 			int32 MaterialSlotIndex = StaticMesh->GetMaterialIndexFromImportedMaterialSlotName(SlotNames[PolygonGroupID]);
 					
 			// If no material was found with this slot name, fill out a blank slot instead.
-			// However this should not happen. The translator should have properly set the UInterchangeMeshNode with the slot names used by all LODs.
-			if (!ensure(MaterialSlotIndex != INDEX_NONE))
+			if (MaterialSlotIndex == INDEX_NONE)
 			{
 				MaterialSlotIndex = StaticMesh->GetStaticMaterials().Emplace(UMaterial::GetDefaultMaterial(MD_Surface), SlotNames[PolygonGroupID]);
 			}
diff --git a/Engine/Plugins/Experimental/Interchange/Runtime/Source/Pipelines/Public/InterchangePipelineMeshesUtilities.h b/Engine/Plugins/Experimental/Interchange/Runtime/Source/Pipelines/Public/InterchangePipelineMeshesUtilities.h
--- a/Engine/Plugins/Experimental/Interchange/Runtime/Source/Pipelines/Public/InterchangePipelineMeshesUtilities.h
+++ b/Engine/Plugins/Experimental/Interchange/Runtime/Source/Pipelines/Public/InterchangePipelineMeshesUtilities.h
@@ -288,12 +288,13 @@ namespace UE::Interchange::MeshesUtilities
 	{
 		for (const TPair<FString, FString>& SlotMaterialDependency : SlotMaterialDependencies)
 		{
-			const FString MaterialFactoryNodeUid = UInterchangeMaterialFactoryNode::GetMaterialFactoryNodeUidFromMaterialNodeUid(SlotMaterialDependency.Value);
-			if (UInterchangeMaterialFactoryNode* MaterialFactoryNode = Cast< UInterchangeMaterialFactoryNode>(NodeContainer.GetFactoryNode(MaterialFactoryNodeUid)))
+			const FString MaterialFactoryNodeUid = UInterchangeBaseMaterialFactoryNode::GetMaterialFactoryNodeUidFromMaterialNodeUid(SlotMaterialDependency.Value);
+			if (UInterchangeBaseMaterialFactoryNode* MaterialFactoryNode = Cast<UInterchangeBaseMaterialFactoryNode>(NodeContainer.GetFactoryNode(MaterialFactoryNodeUid)))
 			{
-				FactoryNode.SetSlotMaterialDependencyUid(SlotMaterialDependency.Key, MaterialFactoryNodeUid);
 				MaterialFactoryNode->SetEnabled(true);
 
+				FactoryNode.SetSlotMaterialDependencyUid(SlotMaterialDependency.Key, MaterialFactoryNodeUid);
+
 				// Create a factory dependency so Material asset are imported before the static mesh asset
 				TArray<FString> FactoryDependencies;
 				FactoryNode.GetFactoryDependencies(FactoryDependencies);
