Add support of Tangent and SubsurfaceColor inputs for materials.

diff --git a/Engine/Plugins/Experimental/Interchange/Runtime/Source/FactoryNodes/Private/InterchangeMaterialFactoryNode.cpp b/Engine/Plugins/Experimental/Interchange/Runtime/Source/FactoryNodes/Private/InterchangeMaterialFactoryNode.cpp
--- a/Engine/Plugins/Experimental/Interchange/Runtime/Source/FactoryNodes/Private/InterchangeMaterialFactoryNode.cpp
+++ b/Engine/Plugins/Experimental/Interchange/Runtime/Source/FactoryNodes/Private/InterchangeMaterialFactoryNode.cpp
@@ -137,6 +137,36 @@ bool UInterchangeMaterialFactoryNode::ConnectOutputToNormal(const FString& Expre
 	return UInterchangeShaderPortsAPI::ConnectOuputToInput(this, UE::Interchange::Materials::PBR::Parameters::Normal.ToString(), ExpressionNodeUid, OutputName);
 }
 
+bool UInterchangeMaterialFactoryNode::GetTangentConnection(FString& ExpressionNodeUid, FString& OutputName) const
+{
+	return UInterchangeShaderPortsAPI::GetInputConnection(this, UE::Interchange::Materials::Common::Parameters::Tangent.ToString(), ExpressionNodeUid, OutputName);
+}
+
+bool UInterchangeMaterialFactoryNode::ConnectToTangent(const FString& ExpressionNodeUid)
+{
+	return UInterchangeShaderPortsAPI::ConnectDefaultOuputToInput(this, UE::Interchange::Materials::Common::Parameters::Tangent.ToString(), ExpressionNodeUid);
+}
+
+bool UInterchangeMaterialFactoryNode::ConnectOutputToTangent(const FString& ExpressionNodeUid, const FString& OutputName)
+{
+	return UInterchangeShaderPortsAPI::ConnectOuputToInput(this, UE::Interchange::Materials::PBR::Parameters::Tangent.ToString(), ExpressionNodeUid, OutputName);
+}
+
+bool UInterchangeMaterialFactoryNode::GetSubsurfaceConnection(FString& ExpressionNodeUid, FString& OutputName) const
+{
+	return UInterchangeShaderPortsAPI::GetInputConnection(this, UE::Interchange::Materials::Subsurface::Parameters::SubsurfaceColor.ToString(), ExpressionNodeUid, OutputName);
+}
+
+bool UInterchangeMaterialFactoryNode::ConnectToSubsurface(const FString& ExpressionNodeUid)
+{
+	return UInterchangeShaderPortsAPI::ConnectDefaultOuputToInput(this, UE::Interchange::Materials::Subsurface::Parameters::SubsurfaceColor.ToString(), ExpressionNodeUid);
+}
+
+bool UInterchangeMaterialFactoryNode::ConnectOutputToSubsurface(const FString& ExpressionNodeUid, const FString& OutputName)
+{
+	return UInterchangeShaderPortsAPI::ConnectOuputToInput(this, UE::Interchange::Materials::Subsurface::Parameters::SubsurfaceColor.ToString(), ExpressionNodeUid, OutputName);
+}
+
 bool UInterchangeMaterialFactoryNode::GetOpacityConnection(FString& ExpressionNodeUid, FString& OutputName) const
 {
 	return UInterchangeShaderPortsAPI::GetInputConnection(this, UE::Interchange::Materials::Common::Parameters::Opacity.ToString(), ExpressionNodeUid, OutputName);
diff --git a/Engine/Plugins/Experimental/Interchange/Runtime/Source/FactoryNodes/Public/InterchangeMaterialFactoryNode.h b/Engine/Plugins/Experimental/Interchange/Runtime/Source/FactoryNodes/Public/InterchangeMaterialFactoryNode.h
--- a/Engine/Plugins/Experimental/Interchange/Runtime/Source/FactoryNodes/Public/InterchangeMaterialFactoryNode.h
+++ b/Engine/Plugins/Experimental/Interchange/Runtime/Source/FactoryNodes/Public/InterchangeMaterialFactoryNode.h
@@ -99,6 +99,24 @@ public:
 	UFUNCTION(BlueprintCallable, Category = "Interchange | Node | MaterialFactory")
 	bool ConnectOutputToNormal(const FString& ExpressionNodeUid, const FString& OutputName);
 
+	UFUNCTION(BlueprintCallable, Category = "Interchange | Node | MaterialFactory")
+	bool GetTangentConnection(FString& ExpressionNodeUid, FString& OutputName) const;
+
+	UFUNCTION(BlueprintCallable, Category = "Interchange | Node | MaterialFactory")
+	bool ConnectToTangent(const FString& ExpressionNodeUid);
+
+	UFUNCTION(BlueprintCallable, Category = "Interchange | Node | MaterialFactory")
+	bool ConnectOutputToTangent(const FString& ExpressionNodeUid, const FString& OutputName);
+
+	UFUNCTION(BlueprintCallable, Category = "Interchange | Node | MaterialFactory")
+	bool GetSubsurfaceConnection(FString& ExpressionNodeUid, FString& OutputName) const;
+
+	UFUNCTION(BlueprintCallable, Category = "Interchange | Node | MaterialFactory")
+	bool ConnectToSubsurface(const FString& ExpressionNodeUid);
+
+	UFUNCTION(BlueprintCallable, Category = "Interchange | Node | MaterialFactory")
+	bool ConnectOutputToSubsurface(const FString& ExpressionNodeUid, const FString& OutputName);
+
 	UFUNCTION(BlueprintCallable, Category = "Interchange | Node | MaterialFactory")
 	bool GetOpacityConnection(FString& ExpressionNodeUid, FString& OutputName) const;
 
diff --git a/Engine/Plugins/Experimental/Interchange/Runtime/Source/Import/Private/Material/InterchangeMaterialFactory.cpp b/Engine/Plugins/Experimental/Interchange/Runtime/Source/Import/Private/Material/InterchangeMaterialFactory.cpp
--- a/Engine/Plugins/Experimental/Interchange/Runtime/Source/Import/Private/Material/InterchangeMaterialFactory.cpp
+++ b/Engine/Plugins/Experimental/Interchange/Runtime/Source/Import/Private/Material/InterchangeMaterialFactory.cpp
@@ -541,6 +541,50 @@ void UInterchangeMaterialFactory::SetupMaterial(UMaterial* Material, const FCrea
 		}
 	}
 
+	// Tangent
+	{
+		FString TangentUid;
+		FString OutputName;
+
+		if(MaterialFactoryNode->GetTangentConnection(TangentUid, OutputName))
+		{
+			const UInterchangeMaterialExpressionFactoryNode* TangentNode = Cast<UInterchangeMaterialExpressionFactoryNode>(Arguments.NodeContainer->GetNode(TangentUid));
+
+			if(TangentNode)
+			{
+				if(UMaterialExpression* TangentExpression = CreateExpressionsForNode(Material, Arguments, TangentNode, Expressions))
+				{
+					if(FExpressionInput* TangentInput = Material->GetExpressionInputForProperty(MP_Tangent))
+					{
+						TangentExpression->ConnectExpression(TangentInput, GetOutputIndex(*TangentExpression, OutputName));
+					}
+				}
+			}
+		}
+	}
+
+	// Subsurface
+	{
+		FString SubsurfaceColorUid;
+		FString OutputName;
+
+		if(MaterialFactoryNode->GetSubsurfaceConnection(SubsurfaceColorUid, OutputName))
+		{
+			const UInterchangeMaterialExpressionFactoryNode* SubsurfaceColorNode = Cast<UInterchangeMaterialExpressionFactoryNode>(Arguments.NodeContainer->GetNode(SubsurfaceColorUid));
+
+			if(SubsurfaceColorNode)
+			{
+				if(UMaterialExpression* SubsurfaceColorExpression = CreateExpressionsForNode(Material, Arguments, SubsurfaceColorNode, Expressions))
+				{
+					if(FExpressionInput* SubsurfaceColorInput = Material->GetExpressionInputForProperty(MP_SubsurfaceColor))
+					{
+						SubsurfaceColorExpression->ConnectExpression(SubsurfaceColorInput, GetOutputIndex(*SubsurfaceColorExpression, OutputName));
+					}
+				}
+			}
+		}
+	}
+
 	// Opacity
 	{
 		FString OpacityUid;
diff --git a/Engine/Plugins/Experimental/Interchange/Runtime/Source/Nodes/Public/InterchangeMaterialDefinitions.h b/Engine/Plugins/Experimental/Interchange/Runtime/Source/Nodes/Public/InterchangeMaterialDefinitions.h
--- a/Engine/Plugins/Experimental/Interchange/Runtime/Source/Nodes/Public/InterchangeMaterialDefinitions.h
+++ b/Engine/Plugins/Experimental/Interchange/Runtime/Source/Nodes/Public/InterchangeMaterialDefinitions.h
@@ -130,6 +130,7 @@ namespace Interchange
 			{
 				const FName EmissiveColor = TEXT("EmissiveColor"); // Type: linear color
 				const FName Normal = TEXT("Normal"); // Type: vector3f
+				const FName Tangent = TEXT("Tangent"); // Type: vector3f
 				const FName Opacity = TEXT("Opacity"); // Type: float
 				const FName Occlusion = TEXT("Occlusion"); // Type: float
 				const FName IndexOfRefraction = TEXT("IOR"); // Type: float
@@ -204,6 +205,16 @@ namespace Interchange
 			}
 		}
 
+		namespace Subsurface
+		{
+			namespace Parameters
+			{
+				using namespace PBR::Parameters;
+
+				const FName SubsurfaceColor = TEXT("SubsurfaceColor"); // Type: linear color
+			}
+		}
+
 		namespace StandardSurface
 		{
 			const FName Name = TEXT("standard_surface");
diff --git a/Engine/Plugins/Experimental/Interchange/Runtime/Source/Pipelines/Public/InterchangeGenericMaterialPipeline.h b/Engine/Plugins/Experimental/Interchange/Runtime/Source/Pipelines/Public/InterchangeGenericMaterialPipeline.h
--- a/Engine/Plugins/Experimental/Interchange/Runtime/Source/Pipelines/Public/InterchangeGenericMaterialPipeline.h
+++ b/Engine/Plugins/Experimental/Interchange/Runtime/Source/Pipelines/Public/InterchangeGenericMaterialPipeline.h
@@ -67,6 +67,9 @@ private:
 	/** True if the shader graph has a sheen color input. */
 	bool IsSheenModel(const UInterchangeShaderGraphNode* ShaderGraphNode) const;
 
+	/** True if the shader graph has a subsurface color input. */
+	bool IsSubsurfaceModel(const UInterchangeShaderGraphNode* ShaderGraphNode) const;
+
 	/** True if the shader graph has a transmission color input. */
 	bool IsThinTranslucentModel(const UInterchangeShaderGraphNode* ShaderGraphNode) const;
 
@@ -87,6 +90,7 @@ private:
 	bool HandlePBRModel(const UInterchangeShaderGraphNode* ShaderGraphNode, UInterchangeMaterialFactoryNode* MaterialFactoryNode);
 	bool HandleStandardSurfaceModel(const UInterchangeShaderGraphNode* ShaderGraphNode, UInterchangeMaterialFactoryNode* MaterialFactoryNode);
 	bool HandleClearCoat(const UInterchangeShaderGraphNode* ShaderGraphNode, UInterchangeMaterialFactoryNode* MaterialFactoryNode);
+	bool HandleSubsurface(const UInterchangeShaderGraphNode* ShaderGraphNode, UInterchangeMaterialFactoryNode* MaterialFactoryNode);
 	bool HandleSheen(const UInterchangeShaderGraphNode* ShaderGraphNode, UInterchangeMaterialFactoryNode* MaterialFactoryNode);
 	bool HandleThinTranslucent(const UInterchangeShaderGraphNode* ShaderGraphNode, UInterchangeMaterialFactoryNode* MaterialFactoryNode);
 	void HandleCommonParameters(const UInterchangeShaderGraphNode* ShaderGraphNode, UInterchangeMaterialFactoryNode* MaterialFactoryNode);
