Add Subsurface into InterchangeGenericMaterialPipeline.cpp

diff --git a/Engine/Plugins/Experimental/Interchange/Runtime/Source/Pipelines/Private/InterchangeGenericMaterialPipeline.cpp b/Engine/Plugins/Experimental/Interchange/Runtime/Source/Pipelines/Private/InterchangeGenericMaterialPipeline.cpp
--- a/Engine/Plugins/Experimental/Interchange/Runtime/Source/Pipelines/Private/InterchangeGenericMaterialPipeline.cpp
+++ b/Engine/Plugins/Experimental/Interchange/Runtime/Source/Pipelines/Private/InterchangeGenericMaterialPipeline.cpp
@@ -166,6 +166,15 @@ bool UInterchangeGenericMaterialPipeline::IsSheenModel(const UInterchangeShaderG
 	return bHasSheenColorInput;
 }
 
+bool UInterchangeGenericMaterialPipeline::IsSubsurfaceModel(const UInterchangeShaderGraphNode* ShaderGraphNode) const
+{
+	using namespace UE::Interchange::Materials::Subsurface;
+
+	const bool bHasSubsurfaceColorInput = UInterchangeShaderPortsAPI::HasInput(ShaderGraphNode, Parameters::SubsurfaceColor);
+
+	return bHasSubsurfaceColorInput;
+}
+
 bool UInterchangeGenericMaterialPipeline::IsThinTranslucentModel(const UInterchangeShaderGraphNode* ShaderGraphNode) const
 {
 	using namespace UE::Interchange::Materials::ThinTranslucent;
@@ -578,6 +587,38 @@ bool UInterchangeGenericMaterialPipeline::HandleClearCoat(const UInterchangeShad
 	return bShadingModelHandled;
 }
 
+bool UInterchangeGenericMaterialPipeline::HandleSubsurface(const UInterchangeShaderGraphNode* ShaderGraphNode, UInterchangeMaterialFactoryNode* MaterialFactoryNode)
+{
+	using namespace UE::Interchange::Materials::Subsurface;
+
+	bool bShadingModelHandled = false;
+
+	// Subsurface Color
+	{
+		const bool bHasInput = UInterchangeShaderPortsAPI::HasInput(ShaderGraphNode, Parameters::SubsurfaceColor);
+
+		if(bHasInput)
+		{
+			TTuple<UInterchangeMaterialExpressionFactoryNode*, FString> ExpressionFactoryNode =
+				CreateMaterialExpressionForInput(MaterialFactoryNode, ShaderGraphNode, Parameters::SubsurfaceColor.ToString(), MaterialFactoryNode->GetUniqueID());
+
+			if(ExpressionFactoryNode.Get<0>())
+			{
+				MaterialFactoryNode->ConnectOutputToSubsurface(ExpressionFactoryNode.Get<0>()->GetUniqueID(), ExpressionFactoryNode.Get<1>());
+			}
+
+			bShadingModelHandled = true;
+		}
+	}
+
+	if(bShadingModelHandled)
+	{
+		MaterialFactoryNode->SetCustomShadingModel(EMaterialShadingModel::MSM_Subsurface);
+	}
+
+	return bShadingModelHandled;
+}
+
 bool UInterchangeGenericMaterialPipeline::HandleSheen(const UInterchangeShaderGraphNode* ShaderGraphNode, UInterchangeMaterialFactoryNode* MaterialFactoryNode)
 {
 	using namespace UE::Interchange::Materials::Sheen;
@@ -1378,7 +1419,10 @@ UInterchangeMaterialFactoryNode* UInterchangeGenericMaterialPipeline::CreateMate
 	{
 		if (!HandleClearCoat(ShaderGraphNode, MaterialFactoryNode))
 		{
-			HandleSheen(ShaderGraphNode, MaterialFactoryNode);
+			if(!HandleSheen(ShaderGraphNode, MaterialFactoryNode))
+			{
+				HandleSubsurface(ShaderGraphNode, MaterialFactoryNode);
+			}
 		}
 	}
 
@@ -1408,6 +1452,10 @@ UInterchangeMaterialInstanceFactoryNode* UInterchangeGenericMaterialPipeline::Cr
 	{
 		MaterialInstanceFactoryNode->SetCustomParent(TEXT("Material'/Interchange/Materials/SheenMaterial.SheenMaterial'"));
 	}
+	else if(IsSubsurfaceModel(ShaderGraphNode))
+	{
+		MaterialInstanceFactoryNode->SetCustomParent(TEXT("Material'/Interchange/Materials/SubsurfaceMaterial.SubsurfaceMaterial'"));
+	}
 	else if (IsPBRModel(ShaderGraphNode))
 	{
 		MaterialInstanceFactoryNode->SetCustomParent(TEXT("Material'/Interchange/Materials/PBRSurfaceMaterial.PBRSurfaceMaterial'"));
