Added support for MaterialFunction:
  - Added UInterchangeMateriaFunctionlFactoryNode to store data to create a MaterialFunction object
  - Added UInterchangeMaterialFunctionFactory to create aMaterialFunction from a UInterchangeMateriaFunctionlFactoryNode
  -Added UInterchangeFunctionCallShaderNode to carry over the unique id of the material function to be called by the resulting function call expression
  - Reworked CreateExpressionsForNode & CreateExpression in InterchangeMaterialFactory.cpp  as helper functions to serve UInterchangeMaterialFactory and UInterchangeMaterialFunctionFactory
  -
Added support for materials using the 'material attributes' input
  - Added MaterialAttributes namespace under UE::Interchange::Materials::Standard::Nodes to list MaterialAttributes expression's inputs
  - Add method to properly handle material using MakeMaterialAttributes

diff --git a/Engine/Plugins/Experimental/Interchange/Runtime/Source/FactoryNodes/Private/InterchangeMaterialFactoryNode.cpp b/Engine/Plugins/Experimental/Interchange/Runtime/Source/FactoryNodes/Private/InterchangeMaterialFactoryNode.cpp
--- a/Engine/Plugins/Experimental/Interchange/Runtime/Source/FactoryNodes/Private/InterchangeMaterialFactoryNode.cpp
+++ b/Engine/Plugins/Experimental/Interchange/Runtime/Source/FactoryNodes/Private/InterchangeMaterialFactoryNode.cpp
@@ -5,6 +5,7 @@
 #if WITH_ENGINE
 
 #include "Materials/Material.h"
+#include "Materials/MaterialFunction.h"
 #include "Materials/MaterialInstance.h"
 #include "Materials/MaterialInstanceConstant.h"
 #include "Materials/MaterialInterface.h"
@@ -429,3 +430,23 @@ bool UInterchangeMaterialFunctionCallExpressionFactoryNode::SetCustomMaterialFun
 	IMPLEMENT_NODE_ATTRIBUTE_SETTER_NODELEGATE(MaterialFunctionDependency, FString);
 	AddFactoryDependencyUid(AttributeValue);
 }
+
+FString UInterchangeMaterialFunctionFactoryNode::GetTypeName() const
+{
+	const FString TypeName = TEXT("MaterialFunctionFactoryNode");
+	return TypeName;
+}
+
+UClass* UInterchangeMaterialFunctionFactoryNode::GetObjectClass() const
+{
+#if WITH_ENGINE
+	return UMaterialFunction::StaticClass();
+#else
+	return nullptr;
+#endif
+}
+
+bool UInterchangeMaterialFunctionFactoryNode::GetInputConnection(const FString& InputName, FString& ExpressionNodeUid, FString& OutputName) const
+{
+	return UInterchangeShaderPortsAPI::GetInputConnection(this, InputName, ExpressionNodeUid, OutputName);
+}
\ No newline at end of file
diff --git a/Engine/Plugins/Experimental/Interchange/Runtime/Source/FactoryNodes/Public/InterchangeMaterialFactoryNode.h b/Engine/Plugins/Experimental/Interchange/Runtime/Source/FactoryNodes/Public/InterchangeMaterialFactoryNode.h
--- a/Engine/Plugins/Experimental/Interchange/Runtime/Source/FactoryNodes/Public/InterchangeMaterialFactoryNode.h
+++ b/Engine/Plugins/Experimental/Interchange/Runtime/Source/FactoryNodes/Public/InterchangeMaterialFactoryNode.h
@@ -312,3 +312,21 @@ private:
 	const UE::Interchange::FAttributeKey Macro_CustomMaterialFunctionDependencyKey = UE::Interchange::FAttributeKey(TEXT("MaterialFunctionDependency"));
 
 };
+
+UCLASS(BlueprintType, Experimental)
+class INTERCHANGEFACTORYNODES_API UInterchangeMaterialFunctionFactoryNode : public UInterchangeBaseMaterialFactoryNode
+{
+	GENERATED_BODY()
+
+public:
+
+	virtual FString GetTypeName() const override;
+
+	UFUNCTION(BlueprintCallable, Category = "Interchange | Node | MaterialFactory")
+	virtual class UClass* GetObjectClass() const override;
+
+	UFUNCTION(BlueprintCallable, Category = "Interchange | Node | MaterialFactory")
+	bool GetInputConnection(const FString& InputName, FString& ExpressionNodeUid, FString& OutputName) const;
+
+private:
+};
diff --git a/Engine/Plugins/Experimental/Interchange/Runtime/Source/Import/Private/InterchangeImportModule.cpp b/Engine/Plugins/Experimental/Interchange/Runtime/Source/Import/Private/InterchangeImportModule.cpp
--- a/Engine/Plugins/Experimental/Interchange/Runtime/Source/Import/Private/InterchangeImportModule.cpp
+++ b/Engine/Plugins/Experimental/Interchange/Runtime/Source/Import/Private/InterchangeImportModule.cpp
@@ -74,6 +74,7 @@ void FInterchangeImportModule::StartupModule()
 		//Register the factories
 		InterchangeManager.RegisterFactory(UInterchangeTextureFactory::StaticClass());
 		InterchangeManager.RegisterFactory(UInterchangeMaterialFactory::StaticClass());
+		InterchangeManager.RegisterFactory(UInterchangeMaterialFunctionFactory::StaticClass());
 		InterchangeManager.RegisterFactory(UInterchangeSkeletonFactory::StaticClass());
 		InterchangeManager.RegisterFactory(UInterchangeSkeletalMeshFactory::StaticClass());
 		InterchangeManager.RegisterFactory(UInterchangeStaticMeshFactory::StaticClass());
diff --git a/Engine/Plugins/Experimental/Interchange/Runtime/Source/Import/Private/Material/InterchangeMaterialFactory.cpp b/Engine/Plugins/Experimental/Interchange/Runtime/Source/Import/Private/Material/InterchangeMaterialFactory.cpp
--- a/Engine/Plugins/Experimental/Interchange/Runtime/Source/Import/Private/Material/InterchangeMaterialFactory.cpp
+++ b/Engine/Plugins/Experimental/Interchange/Runtime/Source/Import/Private/Material/InterchangeMaterialFactory.cpp
@@ -11,11 +11,15 @@
 #include "Materials/Material.h"
 #include "Materials/MaterialExpression.h"
 #include "Materials/MaterialExpressionClearCoatNormalCustomOutput.h"
+#include "Materials/MaterialExpressionFunctionInput.h"
+#include "Materials/MaterialExpressionFunctionOutput.h"
+#include "Materials/MaterialExpressionMakeMaterialAttributes.h"
 #include "Materials/MaterialExpressionScalarParameter.h"
 #include "Materials/MaterialExpressionTextureCoordinate.h"
 #include "Materials/MaterialExpressionTextureSample.h"
 #include "Materials/MaterialExpressionThinTranslucentMaterialOutput.h"
 #include "Materials/MaterialExpressionVectorParameter.h"
+#include "Materials/MaterialFunction.h"
 #include "Materials/MaterialInterface.h"
 #include "Materials/MaterialInstanceConstant.h"
 #include "Materials/MaterialInstanceDynamic.h"
@@ -116,7 +120,7 @@ namespace UE
 						FString MaterialFactoryNodeUid;
 						FunctionCallFactoryNode->GetCustomMaterialFunctionDependency(MaterialFactoryNodeUid);
 
-						if (const UInterchangeMaterialFactoryNode* MaterialFactoryNode = Cast<UInterchangeMaterialFactoryNode>(Arguments.NodeContainer->GetNode(MaterialFactoryNodeUid)))
+						if (const UInterchangeMaterialFunctionFactoryNode* MaterialFactoryNode = Cast<UInterchangeMaterialFunctionFactoryNode>(Arguments.NodeContainer->GetNode(MaterialFactoryNodeUid)))
 						{
 							if (UMaterialFunctionInterface* MaterialFunction = Cast<UMaterialFunctionInterface>(MaterialFactoryNode->ReferenceObject.TryLoad()))
 							{
@@ -147,9 +151,8 @@ namespace UE
 					TextureExpression->AutoSetSampleType();
 				}
 
-				UMaterialExpression* CreateMaterialExpression(UMaterial* Material, const TSubclassOf<UMaterialExpression>& ExpressionClass)
+				UMaterialExpression* CreateMaterialExpression(UMaterial* Material, UMaterialFunction* MaterialFunction, const TSubclassOf<UMaterialExpression>& ExpressionClass)
 				{
-					UMaterialFunction* const MaterialFunction = nullptr;
 					UObject* const SelectedAsset = nullptr;
 					const int32 NodePosX = 0;
 					const int32 NodePosY = 0;
@@ -157,6 +160,143 @@ namespace UE
 
 					return UMaterialEditingLibrary::CreateMaterialExpressionEx(Material, MaterialFunction, ExpressionClass, SelectedAsset, NodePosX, NodePosY, bAllowMarkingPackageDirty);
 				}
+
+				template<class T>
+				T* CreateMaterialExpression(UMaterial* Material, UMaterialFunction* MaterialFunction)
+				{
+					return Cast<T>(CreateMaterialExpression(Material, MaterialFunction, T::StaticClass()));
+				}
+
+
+				class FMaterialExpressionBuilder
+				{
+				public:
+					FMaterialExpressionBuilder(UMaterial* InMaterial, UMaterialFunction* InMaterialFunction, const UInterchangeFactoryBase::FCreateAssetParams& InArguments)
+						: Material(InMaterial)
+						, MaterialFunction(InMaterialFunction)
+						, Arguments(InArguments)
+					{
+						check((Material || MaterialFunction) && !(Material && MaterialFunction));
+					}
+
+					UMaterialExpression* CreateExpressionsForNode(const UInterchangeMaterialExpressionFactoryNode& ExpressionNode)
+					{
+						UMaterialExpression* MaterialExpression = Expressions.FindRef(ExpressionNode.GetUniqueID());
+
+						if (MaterialExpression)
+						{
+							return MaterialExpression;
+						}
+
+						MaterialExpression = CreateExpression(ExpressionNode);
+						if (!MaterialExpression)
+						{
+							return nullptr;
+						}
+
+						Expressions.Add(ExpressionNode.GetUniqueID()) = MaterialExpression;
+
+						TArray<FString> Inputs;
+						UInterchangeShaderPortsAPI::GatherInputs(&ExpressionNode, Inputs);
+
+						for (const FString& InputName : Inputs)
+						{
+							FString ConnectedExpressionUid;
+							FString OutputName;
+							if (UInterchangeShaderPortsAPI::GetInputConnection(&ExpressionNode, InputName, ConnectedExpressionUid, OutputName))
+							{
+								const UInterchangeMaterialExpressionFactoryNode* ConnectedExpressionNode = Cast<UInterchangeMaterialExpressionFactoryNode>(Arguments.NodeContainer->GetNode(ConnectedExpressionUid));
+								if (ConnectedExpressionNode)
+								{
+									UMaterialExpression* ConnectedExpression = Expressions.FindRef(ConnectedExpressionUid);
+									if (!ConnectedExpression)
+									{
+										ConnectedExpression = CreateExpressionsForNode(*ConnectedExpressionNode);
+									}
+
+									const int32 InputIndex = GetInputIndex(*MaterialExpression, InputName);
+									const int32 OutputIndex = GetOutputIndex(*ConnectedExpression, OutputName);
+
+									if (InputIndex != INDEX_NONE)
+									{
+										FExpressionInput* ExpressionInput = MaterialExpression->GetInput(InputIndex);
+										if (ExpressionInput)
+										{
+											ConnectedExpression->ConnectExpression(ExpressionInput, OutputIndex);
+										}
+									}
+									else
+									{
+										Messages.Add(FText::Format(LOCTEXT("InputNotFound", "Invalid input {0} for material expression node {1}."),
+											FText::FromString(InputName),
+											FText::FromString(ExpressionNode.GetDisplayLabel())));
+									}
+								}
+							}
+						}
+
+						return MaterialExpression;
+					}
+
+					TArray<FText> Messages;
+
+				private:
+					UMaterialExpression* CreateExpression(const UInterchangeMaterialExpressionFactoryNode& ExpressionNode)
+					{
+						FString ExpressionClassName;
+						ExpressionNode.GetCustomExpressionClassName(ExpressionClassName);
+
+						TSubclassOf<UMaterialExpression> ExpressionClass = FindExpressionClass(*ExpressionClassName);
+
+						if (!ExpressionClass.Get())
+						{
+							Messages.Add(FText::Format(LOCTEXT("ExpressionClassNotFound", "Invalid class {0} for material expression node {1}."),
+								FText::FromString(ExpressionClassName),
+								FText::FromString(ExpressionNode.GetDisplayLabel())));
+
+							return nullptr;
+						}
+
+						UMaterialExpression* MaterialExpression = CreateMaterialExpression(Material, MaterialFunction, ExpressionClass);
+
+						if (!MaterialExpression)
+						{
+							Messages.Add(FText::Format(LOCTEXT("MaterialExpressionCreationFailed", "Failed to create {0} object for material expression node {1}."),
+								FText::FromString(ExpressionClassName),
+								FText::FromString(ExpressionNode.GetDisplayLabel())));
+
+							return nullptr;
+						}
+
+						if (Material)
+						{
+							// Set the parameter name if the material expression has one (some material expressions don't inherit from UMaterialExpressionParameter, ie: UMaterialExpressionTextureSampleParameter
+							if (FNameProperty* Property = FindFProperty<FNameProperty>(MaterialExpression->GetClass(), GET_MEMBER_NAME_CHECKED(UMaterialExpressionParameter, ParameterName)))
+							{
+								*(Property->ContainerPtrToValuePtr<FName>(MaterialExpression)) = FName(*(ExpressionNode.GetDisplayLabel() + LexToString(Material->GetExpressions().Num())));
+							}
+						}
+
+						ExpressionNode.ApplyAllCustomAttributeToObject(MaterialExpression);
+
+						if (UMaterialExpressionTextureBase* TextureExpression = Cast<UMaterialExpressionTextureBase>(MaterialExpression))
+						{
+							SetupTextureExpression(Arguments, &ExpressionNode, TextureExpression);
+						}
+						else if (UMaterialExpressionMaterialFunctionCall* FunctionCallExpression = Cast<UMaterialExpressionMaterialFunctionCall>(MaterialExpression))
+						{
+							SetupFunctionCallExpression(Arguments, ExpressionNode, FunctionCallExpression);
+						}
+
+						return MaterialExpression;
+					}
+
+
+					UMaterial* Material = nullptr;
+					UMaterialFunction* MaterialFunction = nullptr;
+					const UInterchangeFactoryBase::FCreateAssetParams& Arguments;
+					TMap<FString, UMaterialExpression*> Expressions;
+				};
 #endif // #if WITH_EDITOR
 			}
 		}
@@ -381,11 +521,38 @@ bool UInterchangeMaterialFactory::SetSourceFilename(const UObject* Object, const
 #if WITH_EDITOR
 void UInterchangeMaterialFactory::SetupMaterial(UMaterial* Material, const FCreateAssetParams& Arguments, const UInterchangeBaseMaterialFactoryNode* BaseMaterialFactoryNode)
 {
+	using namespace UE::Interchange::Materials;
 	using namespace UE::Interchange::MaterialFactory::Internal;
 
 	const UInterchangeMaterialFactoryNode* MaterialFactoryNode = Cast<UInterchangeMaterialFactoryNode>(BaseMaterialFactoryNode);
 
-	TMap<FString, UMaterialExpression*> Expressions;
+	FMaterialExpressionBuilder Builder(Material, nullptr, Arguments);
+
+	if (UInterchangeShaderPortsAPI::HasInput(MaterialFactoryNode, Common::Parameters::BxDF))
+	{
+		FString ExpressionNodeUid;
+		FString OutputName;
+
+		UInterchangeShaderPortsAPI::GetInputConnection(MaterialFactoryNode, Common::Parameters::BxDF.ToString(), ExpressionNodeUid, OutputName);
+
+		const UInterchangeMaterialExpressionFactoryNode* MaterialAttributes = Cast<UInterchangeMaterialExpressionFactoryNode>(Arguments.NodeContainer->GetNode(ExpressionNodeUid));
+
+		if (MaterialAttributes)
+		{
+			if (UMaterialExpression* MaterialAttributesExpression = Builder.CreateExpressionsForNode(*MaterialAttributes))
+			{
+				if (FExpressionInput* MaterialAttributesInput = Material->GetExpressionInputForProperty(MP_MaterialAttributes))
+				{
+					MaterialAttributesExpression->ConnectExpression(MaterialAttributesInput, GetOutputIndex(*MaterialAttributesExpression, OutputName));
+					Material->bUseMaterialAttributes = true;
+				}
+			}
+		}
+
+		UMaterialEditingLibrary::LayoutMaterialExpressions(Material);
+
+		return;
+	}
 
 	// Base Color
 	{
@@ -398,7 +565,7 @@ void UInterchangeMaterialFactory::SetupMaterial(UMaterial* Material, const FCrea
 
 			if (BaseColor)
 			{
-				if (UMaterialExpression* BaseColorExpression = CreateExpressionsForNode(Material, Arguments, BaseColor, Expressions))
+				if (UMaterialExpression* BaseColorExpression = Builder.CreateExpressionsForNode(*BaseColor))
 				{
 					if (FExpressionInput* BaseColorInput = Material->GetExpressionInputForProperty(MP_BaseColor))
 					{
@@ -420,7 +587,7 @@ void UInterchangeMaterialFactory::SetupMaterial(UMaterial* Material, const FCrea
 
 			if (MetallicNode)
 			{
-				if (UMaterialExpression* MetallicExpression = CreateExpressionsForNode(Material, Arguments, MetallicNode, Expressions))
+				if (UMaterialExpression* MetallicExpression = Builder.CreateExpressionsForNode(*MetallicNode))
 				{
 					if (FExpressionInput* MetallicInput = Material->GetExpressionInputForProperty(MP_Metallic))
 					{
@@ -442,7 +609,7 @@ void UInterchangeMaterialFactory::SetupMaterial(UMaterial* Material, const FCrea
 
 			if (SpecularNode)
 			{
-				if (UMaterialExpression* SpecularExpression = CreateExpressionsForNode(Material, Arguments, SpecularNode, Expressions))
+				if (UMaterialExpression* SpecularExpression = Builder.CreateExpressionsForNode(*SpecularNode))
 				{
 					if (FExpressionInput* SpecularInput = Material->GetExpressionInputForProperty(MP_Specular))
 					{
@@ -464,7 +631,7 @@ void UInterchangeMaterialFactory::SetupMaterial(UMaterial* Material, const FCrea
 
 			if (Roughness)
 			{
-				if (UMaterialExpression* RoughnessExpression = CreateExpressionsForNode(Material, Arguments, Roughness, Expressions))
+				if (UMaterialExpression* RoughnessExpression = Builder.CreateExpressionsForNode(*Roughness))
 				{
 					if (FExpressionInput* RoughnessInput = Material->GetExpressionInputForProperty(MP_Roughness))
 					{
@@ -486,7 +653,7 @@ void UInterchangeMaterialFactory::SetupMaterial(UMaterial* Material, const FCrea
 
 			if (Anisotropy)
 			{
-				if (UMaterialExpression* AnisotropyExpression = CreateExpressionsForNode(Material, Arguments, Anisotropy, Expressions))
+				if (UMaterialExpression* AnisotropyExpression = Builder.CreateExpressionsForNode(*Anisotropy))
 				{
 					if (FExpressionInput* AnisotropyInput = Material->GetExpressionInputForProperty(MP_Anisotropy))
 					{
@@ -508,7 +675,7 @@ void UInterchangeMaterialFactory::SetupMaterial(UMaterial* Material, const FCrea
 
 			if (EmissiveColor)
 			{
-				if (UMaterialExpression* EmissiveExpression = CreateExpressionsForNode(Material, Arguments, EmissiveColor, Expressions))
+				if (UMaterialExpression* EmissiveExpression = Builder.CreateExpressionsForNode(*EmissiveColor))
 				{
 					if (FExpressionInput* EmissiveInput = Material->GetExpressionInputForProperty(MP_EmissiveColor))
 					{
@@ -530,7 +697,7 @@ void UInterchangeMaterialFactory::SetupMaterial(UMaterial* Material, const FCrea
 
 			if (NormalNode)
 			{
-				if (UMaterialExpression* NormalExpression = CreateExpressionsForNode(Material, Arguments, NormalNode, Expressions))
+				if (UMaterialExpression* NormalExpression = Builder.CreateExpressionsForNode(*NormalNode))
 				{
 					if (FExpressionInput* NormalInput = Material->GetExpressionInputForProperty(MP_Normal))
 					{
@@ -552,7 +719,7 @@ void UInterchangeMaterialFactory::SetupMaterial(UMaterial* Material, const FCrea
 
 			if(TangentNode)
 			{
-				if(UMaterialExpression* TangentExpression = CreateExpressionsForNode(Material, Arguments, TangentNode, Expressions))
+				if(UMaterialExpression* TangentExpression = Builder.CreateExpressionsForNode(*TangentNode))
 				{
 					if(FExpressionInput* TangentInput = Material->GetExpressionInputForProperty(MP_Tangent))
 					{
@@ -574,7 +741,7 @@ void UInterchangeMaterialFactory::SetupMaterial(UMaterial* Material, const FCrea
 
 			if(SubsurfaceColorNode)
 			{
-				if(UMaterialExpression* SubsurfaceColorExpression = CreateExpressionsForNode(Material, Arguments, SubsurfaceColorNode, Expressions))
+				if(UMaterialExpression* SubsurfaceColorExpression = Builder.CreateExpressionsForNode(*SubsurfaceColorNode))
 				{
 					if(FExpressionInput* SubsurfaceColorInput = Material->GetExpressionInputForProperty(MP_SubsurfaceColor))
 					{
@@ -596,7 +763,7 @@ void UInterchangeMaterialFactory::SetupMaterial(UMaterial* Material, const FCrea
 
 			if (OpacityNode)
 			{
-				if (UMaterialExpression* OpacityExpression = CreateExpressionsForNode(Material, Arguments, OpacityNode, Expressions))
+				if (UMaterialExpression* OpacityExpression = Builder.CreateExpressionsForNode(*OpacityNode))
 				{
 					FExpressionInput* OpacityInput = Material->GetExpressionInputForProperty(MP_Opacity);
 
@@ -629,7 +796,7 @@ void UInterchangeMaterialFactory::SetupMaterial(UMaterial* Material, const FCrea
 
 			if (OcclusionNode)
 			{
-				if (UMaterialExpression* OcclusionExpression = CreateExpressionsForNode(Material, Arguments, OcclusionNode, Expressions))
+				if (UMaterialExpression* OcclusionExpression = Builder.CreateExpressionsForNode(*OcclusionNode))
 				{
 					if (FExpressionInput* OcclusionInput = Material->GetExpressionInputForProperty(MP_AmbientOcclusion))
 					{
@@ -651,7 +818,7 @@ void UInterchangeMaterialFactory::SetupMaterial(UMaterial* Material, const FCrea
 
 			if (RefractionNode)
 			{
-				if (UMaterialExpression* RefractionExpression = CreateExpressionsForNode(Material, Arguments, RefractionNode, Expressions))
+				if (UMaterialExpression* RefractionExpression = Builder.CreateExpressionsForNode(*RefractionNode))
 				{
 					if (FExpressionInput* RefractionInput = Material->GetExpressionInputForProperty(MP_Refraction))
 					{
@@ -673,7 +840,7 @@ void UInterchangeMaterialFactory::SetupMaterial(UMaterial* Material, const FCrea
 
 			if (ClearCoatNode)
 			{
-				if (UMaterialExpression* ClearCoatExpression = CreateExpressionsForNode(Material, Arguments, ClearCoatNode, Expressions))
+				if (UMaterialExpression* ClearCoatExpression = Builder.CreateExpressionsForNode(*ClearCoatNode))
 				{
 					if (FExpressionInput* ClearCoatInput = Material->GetExpressionInputForProperty(MP_CustomData0))
 					{
@@ -695,7 +862,7 @@ void UInterchangeMaterialFactory::SetupMaterial(UMaterial* Material, const FCrea
 
 			if (ClearCoatRoughnessNode)
 			{
-				if (UMaterialExpression* ClearCoatRoughnessExpression = CreateExpressionsForNode(Material, Arguments, ClearCoatRoughnessNode, Expressions))
+				if (UMaterialExpression* ClearCoatRoughnessExpression = Builder.CreateExpressionsForNode(*ClearCoatRoughnessNode))
 				{
 					if (FExpressionInput* ClearCoatRoughnessInput = Material->GetExpressionInputForProperty(MP_CustomData1))
 					{
@@ -717,9 +884,9 @@ void UInterchangeMaterialFactory::SetupMaterial(UMaterial* Material, const FCrea
 
 			if (ClearCoatNormalNode)
 			{
-				if (UMaterialExpression* ClearCoatNormalExpression = CreateExpressionsForNode(Material, Arguments, ClearCoatNormalNode, Expressions))
+				if (UMaterialExpression* ClearCoatNormalExpression = Builder.CreateExpressionsForNode(*ClearCoatNormalNode))
 				{
-					UMaterialExpression* ClearCoatNormalCustomOutput = CreateMaterialExpression(Material, UMaterialExpressionClearCoatNormalCustomOutput::StaticClass());
+					UMaterialExpression* ClearCoatNormalCustomOutput = CreateMaterialExpression(Material, nullptr, UMaterialExpressionClearCoatNormalCustomOutput::StaticClass());
 					
 					ClearCoatNormalExpression->ConnectExpression(ClearCoatNormalCustomOutput->GetInput(0), GetOutputIndex(*ClearCoatNormalExpression, OutputName));
 				}
@@ -738,9 +905,9 @@ void UInterchangeMaterialFactory::SetupMaterial(UMaterial* Material, const FCrea
 
 			if (TransmissionColorNode)
 			{
-				if (UMaterialExpression* TransmissionColorExpression = CreateExpressionsForNode(Material, Arguments, TransmissionColorNode, Expressions))
+				if (UMaterialExpression* TransmissionColorExpression = Builder.CreateExpressionsForNode(*TransmissionColorNode))
 				{
-					UMaterialExpression* ThinTranslucentMaterialOutput = CreateMaterialExpression(Material, UMaterialExpressionThinTranslucentMaterialOutput::StaticClass());
+					UMaterialExpression* ThinTranslucentMaterialOutput = CreateMaterialExpression(Material, nullptr, UMaterialExpressionThinTranslucentMaterialOutput::StaticClass());
 					
 					TransmissionColorExpression->ConnectExpression(ThinTranslucentMaterialOutput->GetInput(0), GetOutputIndex(*TransmissionColorExpression, OutputName));
 				}
@@ -759,7 +926,7 @@ void UInterchangeMaterialFactory::SetupMaterial(UMaterial* Material, const FCrea
 
 			if (FuzzColorNode)
 			{
-				if (UMaterialExpression* FuzzColorExpression = CreateExpressionsForNode(Material, Arguments, FuzzColorNode, Expressions))
+				if (UMaterialExpression* FuzzColorExpression = Builder.CreateExpressionsForNode(*FuzzColorNode))
 				{
 					if (FExpressionInput* FuzzColorInput = Material->GetExpressionInputForProperty(MP_SubsurfaceColor))
 					{
@@ -781,7 +948,7 @@ void UInterchangeMaterialFactory::SetupMaterial(UMaterial* Material, const FCrea
 
 			if (ClothNode)
 			{
-				if (UMaterialExpression* ClothExpression = CreateExpressionsForNode(Material, Arguments, ClothNode, Expressions))
+				if (UMaterialExpression* ClothExpression = Builder.CreateExpressionsForNode(*ClothNode))
 				{
 					if (FExpressionInput* ClothInput = Material->GetExpressionInputForProperty(MP_CustomData0))
 					{
@@ -794,116 +961,6 @@ void UInterchangeMaterialFactory::SetupMaterial(UMaterial* Material, const FCrea
 
 	UMaterialEditingLibrary::LayoutMaterialExpressions(Material);
 }
-
-UMaterialExpression* UInterchangeMaterialFactory::CreateExpression(UMaterial* Material, const UInterchangeFactoryBase::FCreateAssetParams& Arguments, const UInterchangeMaterialExpressionFactoryNode& ExpressionNode)
-{
-	using namespace UE::Interchange::MaterialFactory::Internal;
-
-	FString ExpressionClassName;
-	ExpressionNode.GetCustomExpressionClassName(ExpressionClassName);
-
-	TSubclassOf<UMaterialExpression> ExpressionClass = FindExpressionClass(*ExpressionClassName);
-
-	if (!ExpressionClass.Get())
-	{
-		UInterchangeResultWarning_Generic* Message = AddMessage<UInterchangeResultWarning_Generic>();
-		Message->Text = FText::Format(LOCTEXT("ExpressionClassNotFound", "Invalid class {0} for material expression node {1}."),
-			FText::FromString(ExpressionClassName),
-			FText::FromString(ExpressionNode.GetDisplayLabel()));
-
-		return nullptr;
-	}
-
-	UMaterialExpression* MaterialExpression = CreateMaterialExpression(Material, ExpressionClass);
-
-	if (!MaterialExpression)
-	{
-		return nullptr;
-	}
-
-	// Set the parameter name if the material expression has one (some material expressions don't inherit from UMaterialExpressionParameter, ie: UMaterialExpressionTextureSampleParameter
-	if (FNameProperty* Property = FindFProperty<FNameProperty>(MaterialExpression->GetClass(), GET_MEMBER_NAME_CHECKED(UMaterialExpressionParameter, ParameterName)))
-	{
-		*(Property->ContainerPtrToValuePtr<FName>(MaterialExpression)) = FName(*(ExpressionNode.GetDisplayLabel() + LexToString(Material->GetExpressions().Num())));
-	}
-
-	ExpressionNode.ApplyAllCustomAttributeToObject(MaterialExpression);
-
-	if (UMaterialExpressionTextureBase* TextureExpression = Cast<UMaterialExpressionTextureBase>(MaterialExpression))
-	{
-		SetupTextureExpression(Arguments, &ExpressionNode, TextureExpression);	
-	}
-	else if (UMaterialExpressionMaterialFunctionCall* FunctionCallExpression = Cast<UMaterialExpressionMaterialFunctionCall>(MaterialExpression))
-	{
-		SetupFunctionCallExpression(Arguments, ExpressionNode, FunctionCallExpression);
-	}
-
-	return MaterialExpression;
-}
-
-UMaterialExpression* UInterchangeMaterialFactory::CreateExpressionsForNode(UMaterial* Material, const UInterchangeFactoryBase::FCreateAssetParams& Arguments,
-	const UInterchangeMaterialExpressionFactoryNode* Expression, TMap<FString, UMaterialExpression*>& Expressions)
-{
-	using namespace UE::Interchange::MaterialFactory::Internal;
-
-	UMaterialExpression* MaterialExpression = Expressions.FindRef(Expression->GetUniqueID());
-
-	if (MaterialExpression)
-	{
-		return MaterialExpression;
-	}
-
-	MaterialExpression = CreateExpression(Material, Arguments, *Expression);
-
-	if (!MaterialExpression)
-	{
-		return nullptr;
-	}
-
-	Expressions.Add(Expression->GetUniqueID()) = MaterialExpression;
-
-	TArray<FString> Inputs;
-	UInterchangeShaderPortsAPI::GatherInputs(Expression, Inputs);
-
-	for (const FString& InputName : Inputs)
-	{
-		FString ConnectedExpressionUid;
-		FString OutputName;
-		if (UInterchangeShaderPortsAPI::GetInputConnection(Expression, InputName, ConnectedExpressionUid, OutputName))
-		{
-			const UInterchangeMaterialExpressionFactoryNode* ConnectedExpressionNode = Cast<UInterchangeMaterialExpressionFactoryNode>(Arguments.NodeContainer->GetNode(ConnectedExpressionUid));
-			if (ConnectedExpressionNode)
-			{
-				UMaterialExpression* ConnectedExpression = Expressions.FindRef(ConnectedExpressionUid);
-				if (!ConnectedExpression)
-				{
-					ConnectedExpression = CreateExpressionsForNode(Material, Arguments, ConnectedExpressionNode, Expressions);
-				}
-
-				const int32 InputIndex = GetInputIndex(*MaterialExpression, InputName);
-				const int32 OutputIndex = GetOutputIndex(*ConnectedExpression, OutputName);
-
-				if (InputIndex != INDEX_NONE)
-				{
-					FExpressionInput* ExpressionInput = MaterialExpression->GetInput(InputIndex);
-					if (ExpressionInput)
-					{
-						ConnectedExpression->ConnectExpression(ExpressionInput, OutputIndex);
-					}
-				}
-				else
-				{
-					UInterchangeResultWarning_Generic* Message = AddMessage<UInterchangeResultWarning_Generic>();
-					Message->Text = FText::Format(LOCTEXT("InputNotFound", "Invalid input {0} for material expression node {1}."),
-						FText::FromString(InputName),
-						FText::FromString(Expression->GetDisplayLabel()));
-				}
-			}
-		}
-	}
-
-	return MaterialExpression;
-}
 #endif // #if WITH_EDITOR
 
 void UInterchangeMaterialFactory::SetupMaterialInstance(UMaterialInstance* MaterialInstance, const UInterchangeBaseNodeContainer* NodeContainer, const UInterchangeBaseMaterialFactoryNode* MaterialFactoryNode)
@@ -1009,4 +1066,208 @@ void UInterchangeMaterialFactory::SetupMaterialInstance(UMaterialInstance* Mater
 	}
 }
 
+UClass* UInterchangeMaterialFunctionFactory::GetFactoryClass() const
+{
+	return UMaterialFunctionInterface::StaticClass();
+}
+
+UObject* UInterchangeMaterialFunctionFactory::CreateEmptyAsset(const FCreateAssetParams& Arguments)
+{
+	UObject* Material = nullptr;
+
+	if (!Arguments.AssetNode || !Arguments.AssetNode->GetObjectClass()->IsChildOf(GetFactoryClass()))
+	{
+		return nullptr;
+	}
+
+	const UInterchangeMaterialFunctionFactoryNode* MaterialFactoryNode = Cast<UInterchangeMaterialFunctionFactoryNode>(Arguments.AssetNode);
+	if (!ensure(MaterialFactoryNode))
+	{
+		return nullptr;
+	}
+
+	const UClass* MaterialClass = MaterialFactoryNode->GetObjectClass();
+	if (!ensure(MaterialClass && MaterialClass->IsChildOf(GetFactoryClass())))
+	{
+		return nullptr;
+	}
+
+	// create an asset if it doesn't exist
+	UObject* ExistingAsset = StaticFindObject(nullptr, Arguments.Parent, *Arguments.AssetName);
+
+	// create a new material or overwrite existing asset, if possible
+	if (!ExistingAsset)
+	{
+		Material = NewObject<UObject>(Arguments.Parent, MaterialClass, *Arguments.AssetName, RF_Public | RF_Standalone);
+	}
+	else if (ExistingAsset->GetClass()->IsChildOf(MaterialClass))
+	{
+		//This is a reimport, we are just re-updating the source data
+		Material = ExistingAsset;
+	}
+
+	if (!Material)
+	{
+		UE_LOG(LogInterchangeImport, Warning, TEXT("Could not create Material asset %s"), *Arguments.AssetName);
+		return nullptr;
+	}
+
+#if WITH_EDITOR
+	Material->PreEditChange(nullptr);
+#endif //WITH_EDITOR
+
+	return Material;
+}
+
+UObject* UInterchangeMaterialFunctionFactory::CreateAsset(const FCreateAssetParams& Arguments)
+{
+	if (!Arguments.AssetNode || !Arguments.AssetNode->GetObjectClass()->IsChildOf(GetFactoryClass()))
+	{
+		return nullptr;
+	}
+
+	const UInterchangeMaterialFunctionFactoryNode* MaterialFactoryNode = Cast<UInterchangeMaterialFunctionFactoryNode>(Arguments.AssetNode);
+	if (!ensure(MaterialFactoryNode))
+	{
+		return nullptr;
+	}
+
+	const UClass* MaterialClass = MaterialFactoryNode->GetObjectClass();
+	check(MaterialClass && MaterialClass->IsChildOf(GetFactoryClass()));
+
+	// create an asset if it doesn't exist
+	UObject* ExistingAsset = StaticFindObject(nullptr, Arguments.Parent, *Arguments.AssetName);
+
+	UObject* MaterialObject = nullptr;
+	// create a new material or overwrite existing asset, if possible
+	if (!ExistingAsset)
+	{
+		//NewObject is not thread safe, the asset registry directory watcher tick on the main thread can trig before we finish initializing the UObject and will crash
+		//The UObject should have been created by calling CreateEmptyAsset on the main thread.
+		check(IsInGameThread());
+		MaterialObject = NewObject<UObject>(Arguments.Parent, MaterialClass, *Arguments.AssetName, RF_Public | RF_Standalone);
+	}
+	else if (ExistingAsset->GetClass()->IsChildOf(MaterialClass))
+	{
+		//This is a reimport, we are just re-updating the source data
+		MaterialObject = ExistingAsset;
+	}
+
+	if (!MaterialObject)
+	{
+		UE_LOG(LogInterchangeImport, Warning, TEXT("Could not create Material asset %s"), *Arguments.AssetName);
+		return nullptr;
+	}
+
+	if (MaterialObject)
+	{
+		//Currently material re-import will not touch the material at all
+		//TODO design a re-import process for the material (expressions and input connections)
+		if (!Arguments.ReimportObject)
+		{
+			if (UMaterialFunction* MaterialFunction = Cast<UMaterialFunction>(MaterialObject))
+			{
+#if WITH_EDITOR
+				SetupMaterial(MaterialFunction, Arguments, MaterialFactoryNode);
+#endif // #if WITH_EDITOR
+
+				MaterialFactoryNode->ApplyAllCustomAttributeToObject(MaterialObject);
+			}
+		}
+
+		//Getting the file Hash will cache it into the source data
+		Arguments.SourceData->GetFileContentHash();
+
+		//The interchange completion task (call in the GameThread after the factories pass), will call PostEditChange which will trig another asynchronous system that will build all material in parallel
+	}
+
+	return MaterialObject;
+}
+
+/* This function is call in the completion task on the main thread, use it to call main thread post creation step for your assets*/
+void UInterchangeMaterialFunctionFactory::PreImportPreCompletedCallback(const FImportPreCompletedCallbackParams& Arguments)
+{
+	Super::PreImportPreCompletedCallback(Arguments);
+
+#if WITH_EDITORONLY_DATA
+	if (ensure(Arguments.ImportedObject && Arguments.SourceData))
+	{
+		UMaterialFunction* ImportedMaterialFunction = CastChecked<UMaterialFunction>(Arguments.ImportedObject);
+
+		if (UMaterialFunctionEditorOnlyData* Data = ImportedMaterialFunction->GetEditorOnlyData())
+		{
+			//Update the samplers type in case the textures were changed during their PreImportPreCompletedCallback
+			for (UMaterialExpression* Expression : Data->ExpressionCollection.Expressions)
+			{
+				if (UMaterialExpressionTextureBase* TextureSample = Cast<UMaterialExpressionTextureBase>(Expression))
+				{
+					TextureSample->AutoSetSampleType();
+				}
+			}
+		}
+	}
+#endif
+}
+
+#if WITH_EDITOR
+void UInterchangeMaterialFunctionFactory::SetupMaterial(UMaterialFunction* MaterialFunction, const FCreateAssetParams& Arguments, const UInterchangeMaterialFunctionFactoryNode* MaterialFunctionFactoryNode)
+{
+	using namespace UE::Interchange::MaterialFactory::Internal;
+	using namespace UE::Interchange::Materials;
+
+	FMaterialExpressionBuilder Builder(nullptr, MaterialFunction, Arguments);
+
+	UMaterialExpressionFunctionOutput* Output = CreateMaterialExpression<UMaterialExpressionFunctionOutput>(nullptr, MaterialFunction);
+
+	UMaterialExpressionMakeMaterialAttributes* Attrib = CreateMaterialExpression<UMaterialExpressionMakeMaterialAttributes>(nullptr, MaterialFunction);
+	Attrib->ConnectExpression(Output->GetInput(0), 0);
+
+	TFunction<bool(const FString& , FExpressionInput*)> ConnectInput;
+	ConnectInput = [&MaterialFunctionFactoryNode, &Builder, NodeContainer = Arguments.NodeContainer](const FString& InputName, FExpressionInput* Input) -> bool
+	{
+		FString ExpressionNodeUid;
+		FString OutputName;
+
+		if (MaterialFunctionFactoryNode->GetInputConnection(InputName, ExpressionNodeUid, OutputName))
+		{
+			const UInterchangeMaterialExpressionFactoryNode* ExpressionFactoryNode = Cast<UInterchangeMaterialExpressionFactoryNode>(NodeContainer->GetNode(ExpressionNodeUid));
+
+			if (ExpressionFactoryNode)
+			{
+				if (UMaterialExpression* MaterialExpression = Builder.CreateExpressionsForNode(*ExpressionFactoryNode))
+				{
+					MaterialExpression->ConnectExpression(Input, 0);
+					return true;
+				}
+			}
+		}
+
+		return false;
+	};
+
+	// Only material functions with BSDF output are supported for now
+	if (!ConnectInput(Common::Parameters::BxDF.ToString(), Attrib->GetInput(0)))
+	{
+		ConnectInput(PBR::Parameters::BaseColor.ToString(), &Attrib->BaseColor);
+		ConnectInput(PBR::Parameters::Metallic.ToString(), &Attrib->Metallic);
+		ConnectInput(PBR::Parameters::Specular.ToString(), &Attrib->Specular);
+		ConnectInput(PBR::Parameters::Roughness.ToString(), &Attrib->Roughness);
+		ConnectInput(PBR::Parameters::EmissiveColor.ToString(), &Attrib->EmissiveColor);
+		ConnectInput(PBR::Parameters::Normal.ToString(), &Attrib->Normal);
+		ConnectInput(PBR::Parameters::Anisotropy.ToString(), &Attrib->Anisotropy);
+		ConnectInput(PBR::Parameters::Tangent.ToString(), &Attrib->Tangent);
+		ConnectInput(PBR::Parameters::IndexOfRefraction.ToString(), &Attrib->Refraction);
+		ConnectInput(PBR::Parameters::Occlusion.ToString(), &Attrib->AmbientOcclusion);
+		ConnectInput(ClearCoat::Parameters::ClearCoat.ToString(), &Attrib->ClearCoat);
+		ConnectInput(ClearCoat::Parameters::ClearCoatRoughness.ToString(), &Attrib->ClearCoatRoughness);
+		ConnectInput(Subsurface::Parameters::SubsurfaceColor.ToString(), &Attrib->SubsurfaceColor);
+	}
+
+	UMaterialEditingLibrary::LayoutMaterialFunctionExpressions(MaterialFunction);
+
+	MaterialFunction->UpdateDependentFunctionCandidates();
+}
+
+#endif // #if WITH_EDITOR
+
 #undef LOCTEXT_NAMESPACE
diff --git a/Engine/Plugins/Experimental/Interchange/Runtime/Source/Import/Public/Material/InterchangeMaterialFactory.h b/Engine/Plugins/Experimental/Interchange/Runtime/Source/Import/Public/Material/InterchangeMaterialFactory.h
--- a/Engine/Plugins/Experimental/Interchange/Runtime/Source/Import/Public/Material/InterchangeMaterialFactory.h
+++ b/Engine/Plugins/Experimental/Interchange/Runtime/Source/Import/Public/Material/InterchangeMaterialFactory.h
@@ -36,14 +36,32 @@ public:
 private:
 #if WITH_EDITOR
 	void SetupMaterial(UMaterial* Material, const FCreateAssetParams& Arguments, const UInterchangeBaseMaterialFactoryNode* MaterialFactoryNode);
-
-	UMaterialExpression* CreateExpressionsForNode(UMaterial* Material, const UInterchangeFactoryBase::FCreateAssetParams& Arguments,
-	const UInterchangeMaterialExpressionFactoryNode* Expression, TMap<FString, UMaterialExpression*>& Expressions);
-
-	UMaterialExpression* CreateExpression(UMaterial* Material, const UInterchangeFactoryBase::FCreateAssetParams& Arguments, const UInterchangeMaterialExpressionFactoryNode& ExpressionNode);
 #endif // #if WITH_EDITOR
 
 	void SetupMaterialInstance(UMaterialInstance* MaterialInstance, const UInterchangeBaseNodeContainer* NodeContainer, const UInterchangeBaseMaterialFactoryNode* MaterialFactoryNode);
 };
 
+UCLASS(BlueprintType, Experimental)
+class INTERCHANGEIMPORT_API UInterchangeMaterialFunctionFactory : public UInterchangeFactoryBase
+{
+	GENERATED_BODY()
+
+public:
+	//////////////////////////////////////////////////////////////////////////
+	// Interchange factory base interface begin
+
+	virtual UClass* GetFactoryClass() const override;
+	virtual UObject* CreateEmptyAsset(const FCreateAssetParams& Arguments) override;
+	virtual UObject* CreateAsset(const FCreateAssetParams& Arguments) override;
+	virtual void PreImportPreCompletedCallback(const FImportPreCompletedCallbackParams& Arguments) override;
+
+	// Interchange factory base interface end
+	//////////////////////////////////////////////////////////////////////////
+
+private:
+#if WITH_EDITOR
+	void SetupMaterial(class UMaterialFunction* Material, const FCreateAssetParams& Arguments, const class UInterchangeMaterialFunctionFactoryNode* MaterialFactoryNode);
+#endif
+};
+
 
diff --git a/Engine/Plugins/Experimental/Interchange/Runtime/Source/Nodes/Private/InterchangeShaderGraphNode.cpp b/Engine/Plugins/Experimental/Interchange/Runtime/Source/Nodes/Private/InterchangeShaderGraphNode.cpp
--- a/Engine/Plugins/Experimental/Interchange/Runtime/Source/Nodes/Private/InterchangeShaderGraphNode.cpp
+++ b/Engine/Plugins/Experimental/Interchange/Runtime/Source/Nodes/Private/InterchangeShaderGraphNode.cpp
@@ -204,3 +204,29 @@ bool UInterchangeShaderGraphNode::SetCustomOpacityMaskClipValue(const float& Att
 {
 	IMPLEMENT_NODE_ATTRIBUTE_SETTER_NODELEGATE(OpacityMaskClipValue, float);
 }
+
+bool UInterchangeShaderGraphNode::GetCustomIsAShaderFunction(bool& AttributeValue) const
+{
+	IMPLEMENT_NODE_ATTRIBUTE_GETTER(IsAShaderFunction, bool);
+}
+
+bool UInterchangeShaderGraphNode::SetCustomIsAShaderFunction(const bool& AttributeValue)
+{
+	IMPLEMENT_NODE_ATTRIBUTE_SETTER_NODELEGATE(IsAShaderFunction, bool);
+}
+
+FString UInterchangeFunctionCallShaderNode::GetTypeName() const
+{
+	const FString TypeName = TEXT("FunctionCallShaderNode");
+	return TypeName;
+}
+
+bool UInterchangeFunctionCallShaderNode::GetCustomMaterialFunction(FString& AttributeValue) const
+{
+	IMPLEMENT_NODE_ATTRIBUTE_GETTER(MaterialFunction, FString);
+}
+
+bool UInterchangeFunctionCallShaderNode::SetCustomMaterialFunction(const FString& AttributeValue)
+{
+	IMPLEMENT_NODE_ATTRIBUTE_SETTER_NODELEGATE(MaterialFunction, FString);
+}
diff --git a/Engine/Plugins/Experimental/Interchange/Runtime/Source/Nodes/Public/InterchangeMaterialDefinitions.h b/Engine/Plugins/Experimental/Interchange/Runtime/Source/Nodes/Public/InterchangeMaterialDefinitions.h
--- a/Engine/Plugins/Experimental/Interchange/Runtime/Source/Nodes/Public/InterchangeMaterialDefinitions.h
+++ b/Engine/Plugins/Experimental/Interchange/Runtime/Source/Nodes/Public/InterchangeMaterialDefinitions.h
@@ -134,6 +134,7 @@ namespace Interchange
 				const FName Opacity = TEXT("Opacity"); // Type: float
 				const FName Occlusion = TEXT("Occlusion"); // Type: float
 				const FName IndexOfRefraction = TEXT("IOR"); // Type: float
+				const FName BxDF = TEXT("BxDF"); // input/output of BSDF or BRDF or BXDF or BTDF data
 			}
 		}
 
diff --git a/Engine/Plugins/Experimental/Interchange/Runtime/Source/Nodes/Public/InterchangeShaderGraphNode.h b/Engine/Plugins/Experimental/Interchange/Runtime/Source/Nodes/Public/InterchangeShaderGraphNode.h
--- a/Engine/Plugins/Experimental/Interchange/Runtime/Source/Nodes/Public/InterchangeShaderGraphNode.h
+++ b/Engine/Plugins/Experimental/Interchange/Runtime/Source/Nodes/Public/InterchangeShaderGraphNode.h
@@ -60,7 +60,7 @@ public:
 	 * @param InterchangeNode	The Node to create the input on.
 	 * @param InputName			The name to give to the input.
 	 * @param ExpressionUid		The unique id of the node to connect to the input.
-	 * @return					true if the input connection was succesfully added to the node.
+	 * @return					true if the input connection was successfully added to the node.
 	 */
 	UFUNCTION(BlueprintCallable, Category = "Interchange | Node | Material")
 	static bool ConnectDefaultOuputToInput(UInterchangeBaseNode* InterchangeNode, const FString& InputName, const FString& ExpressionUid);
@@ -127,7 +127,32 @@ public:
 private:
 	const UE::Interchange::FAttributeKey Macro_CustomShaderTypeKey = UE::Interchange::FAttributeKey(TEXT("ShaderType"));
 };
- 
+
+/**
+ * A function call shader node has a named set of inputs and outputs which corresponds to the inputs and outputs of the shader function it instances.
+ */
+UCLASS(BlueprintType, Experimental)
+class INTERCHANGENODES_API UInterchangeFunctionCallShaderNode : public UInterchangeShaderNode
+{
+	GENERATED_BODY()
+
+public:
+	virtual FString GetTypeName() const override;
+
+public:
+	UFUNCTION(BlueprintCallable, Category = "Interchange | Node | Material")
+	bool GetCustomMaterialFunction(FString& AttributeValue) const;
+
+	/**
+	 * Sets the unique id of the material function referenced by the function call expression.
+	 */
+	UFUNCTION(BlueprintCallable, Category = "Interchange | Node | Material")
+	bool SetCustomMaterialFunction(const FString& AttributeValue);
+
+private:
+	const UE::Interchange::FAttributeKey Macro_CustomMaterialFunctionKey = UE::Interchange::FAttributeKey(TEXT("MaterialFunction"));
+};
+
 /**
  * A shader graph has its own set of inputs on which shader nodes can be connected to.
  */
@@ -171,7 +196,17 @@ public:
 	UFUNCTION(BlueprintCallable, Category = "Interchange | Node | Material")
 	bool SetCustomOpacityMaskClipValue(const float& AttributeValue, bool bAddApplyDelegate = true);
 
+	UFUNCTION(BlueprintCallable, Category = "Interchange | Node | Material")
+	bool GetCustomIsAShaderFunction(bool& AttributeValue) const;
+
+	/**
+	 * Sets if this shader graph should be considered as a material, false, or a material function, true.
+	 */
+	UFUNCTION(BlueprintCallable, Category = "Interchange | Node | Material")
+	bool SetCustomIsAShaderFunction(const bool& AttributeValue);
+
 private:
 	const UE::Interchange::FAttributeKey Macro_CustomTwoSidedKey = UE::Interchange::FAttributeKey(TEXT("TwoSided"));
 	const UE::Interchange::FAttributeKey Macro_CustomOpacityMaskClipValueKey = UE::Interchange::FAttributeKey(TEXT("OpacityMaskClipValue"));
+	const UE::Interchange::FAttributeKey Macro_CustomIsAShaderFunctionKey = UE::Interchange::FAttributeKey(TEXT("IsAShaderFunction"));
 };
\ No newline at end of file
diff --git a/Engine/Plugins/Experimental/Interchange/Runtime/Source/Pipelines/Private/InterchangeGenericMaterialPipeline.cpp b/Engine/Plugins/Experimental/Interchange/Runtime/Source/Pipelines/Private/InterchangeGenericMaterialPipeline.cpp
--- a/Engine/Plugins/Experimental/Interchange/Runtime/Source/Pipelines/Private/InterchangeGenericMaterialPipeline.cpp
+++ b/Engine/Plugins/Experimental/Interchange/Runtime/Source/Pipelines/Private/InterchangeGenericMaterialPipeline.cpp
@@ -98,9 +98,21 @@ void UInterchangeGenericMaterialPipeline::ExecutePreImportPipeline(UInterchangeB
 	{
 		for (const UInterchangeShaderGraphNode* ShaderGraphNode : MaterialNodes)
 		{
-			if (UInterchangeMaterialFactoryNode* MaterialFactoryNode = CreateMaterialFactoryNode(ShaderGraphNode))
+			UInterchangeBaseMaterialFactoryNode* MaterialBaseFactoryNode = nullptr;
+
+			bool bIsAShaderFunction;
+			if (ShaderGraphNode->GetCustomIsAShaderFunction(bIsAShaderFunction) && bIsAShaderFunction)
+			{
+				MaterialBaseFactoryNode = CreateMaterialFunctionFactoryNode(ShaderGraphNode);
+			}
+			else
 			{
-				MaterialFactoryNode->SetEnabled(bImportUnusedMaterial);
+				MaterialBaseFactoryNode = CreateMaterialFactoryNode(ShaderGraphNode);
+			}
+
+			if (MaterialBaseFactoryNode)
+			{
+				MaterialBaseFactoryNode->SetEnabled(bImportUnusedMaterial);
 			}
 		}
 	}
@@ -860,7 +872,7 @@ void UInterchangeGenericMaterialPipeline::HandleCommonParameters(const UIntercha
 	}
 }
 
-void UInterchangeGenericMaterialPipeline::HandleFlattenNormalNode(const UInterchangeShaderNode* ShaderNode, UInterchangeMaterialFactoryNode* MaterialFactoryNode,
+void UInterchangeGenericMaterialPipeline::HandleFlattenNormalNode(const UInterchangeShaderNode* ShaderNode, UInterchangeBaseMaterialFactoryNode* MaterialFactoryNode,
 	UInterchangeMaterialExpressionFactoryNode* FlattenNormalFactoryNode)
 {
 	using namespace UE::Interchange::Materials::Standard::Nodes::FlattenNormal;
@@ -896,7 +908,7 @@ void UInterchangeGenericMaterialPipeline::HandleFlattenNormalNode(const UInterch
 	}
 }
 
-void UInterchangeGenericMaterialPipeline::HandleTextureSampleNode(const UInterchangeShaderNode* ShaderNode, UInterchangeMaterialFactoryNode* MaterialFactoryNode, UInterchangeMaterialExpressionFactoryNode* TextureSampleFactoryNode)
+void UInterchangeGenericMaterialPipeline::HandleTextureSampleNode(const UInterchangeShaderNode* ShaderNode, UInterchangeBaseMaterialFactoryNode* MaterialFactoryNode, UInterchangeMaterialExpressionFactoryNode* TextureSampleFactoryNode)
 {
 	using namespace UE::Interchange::Materials::Standard::Nodes::TextureSample;
 
@@ -970,7 +982,7 @@ void UInterchangeGenericMaterialPipeline::HandleTextureSampleNode(const UInterch
 	}
 }
 
-void UInterchangeGenericMaterialPipeline::HandleTextureCoordinateNode(const UInterchangeShaderNode* ShaderNode, UInterchangeMaterialFactoryNode* MaterialFactoryNode,
+void UInterchangeGenericMaterialPipeline::HandleTextureCoordinateNode(const UInterchangeShaderNode* ShaderNode, UInterchangeBaseMaterialFactoryNode* MaterialFactoryNode,
 	UInterchangeMaterialExpressionFactoryNode*& TexCoordFactoryNode)
 {
 	using namespace UE::Interchange::Materials::Standard;
@@ -1080,7 +1092,7 @@ void UInterchangeGenericMaterialPipeline::HandleTextureCoordinateNode(const UInt
 	}
 }
 
-void UInterchangeGenericMaterialPipeline::HandleLerpNode(const UInterchangeShaderNode* ShaderNode, UInterchangeMaterialFactoryNode* MaterialFactoryNode, UInterchangeMaterialExpressionFactoryNode* LerpFactoryNode)
+void UInterchangeGenericMaterialPipeline::HandleLerpNode(const UInterchangeShaderNode* ShaderNode, UInterchangeBaseMaterialFactoryNode* MaterialFactoryNode, UInterchangeMaterialExpressionFactoryNode* LerpFactoryNode)
 {
 	using namespace UE::Interchange::Materials::Standard;
 
@@ -1123,7 +1135,7 @@ void UInterchangeGenericMaterialPipeline::HandleLerpNode(const UInterchangeShade
 	}
 }
 
-void UInterchangeGenericMaterialPipeline::HandleMaskNode(const UInterchangeShaderNode* ShaderNode, UInterchangeMaterialFactoryNode* MaterialFactoryNode, UInterchangeMaterialExpressionFactoryNode* MaskFactoryNode)
+void UInterchangeGenericMaterialPipeline::HandleMaskNode(const UInterchangeShaderNode* ShaderNode, UInterchangeBaseMaterialFactoryNode* MaterialFactoryNode, UInterchangeMaterialExpressionFactoryNode* MaskFactoryNode)
 {
 	using namespace UE::Interchange::Materials::Standard::Nodes;
 
@@ -1183,7 +1195,7 @@ void UInterchangeGenericMaterialPipeline::HandleMaskNode(const UInterchangeShade
 	}
 }
 
-UInterchangeMaterialExpressionFactoryNode* UInterchangeGenericMaterialPipeline::CreateMaterialExpressionForShaderNode(UInterchangeMaterialFactoryNode* MaterialFactoryNode,
+UInterchangeMaterialExpressionFactoryNode* UInterchangeGenericMaterialPipeline::CreateMaterialExpressionForShaderNode(UInterchangeBaseMaterialFactoryNode* MaterialFactoryNode,
 	const UInterchangeShaderNode* ShaderNode, const FString& ParentUid)
 {
 	using namespace UE::Interchange::Materials::Standard;
@@ -1202,13 +1214,31 @@ UInterchangeMaterialExpressionFactoryNode* UInterchangeGenericMaterialPipeline::
 		return MaterialExpression;
 	}
 
-	MaterialExpression = NewObject<UInterchangeMaterialExpressionFactoryNode>(BaseNodeContainer);
-	MaterialExpression->InitializeNode(MaterialExpressionUid, ShaderNode->GetDisplayLabel(), EInterchangeNodeContainerType::FactoryData);
-	BaseNodeContainer->AddNode(MaterialExpression);
+	if (const UInterchangeFunctionCallShaderNode* FunctionCallShaderNode = Cast<UInterchangeFunctionCallShaderNode>(ShaderNode))
+	{
+		UInterchangeMaterialFunctionCallExpressionFactoryNode* MaterialFunctionCallExpression = NewObject<UInterchangeMaterialFunctionCallExpressionFactoryNode>(BaseNodeContainer);
+		
+		FString MaterialFunctionUid;
+		if (FunctionCallShaderNode->GetCustomMaterialFunction(MaterialFunctionUid))
+		{
+			const FString MaterialFunctionFactoryNodeUid = UInterchangeMaterialFactoryNode::GetMaterialFactoryNodeUidFromMaterialNodeUid(MaterialFunctionUid);
+			MaterialFunctionCallExpression->SetCustomMaterialFunctionDependency(MaterialFunctionFactoryNodeUid);
+			MaterialFunctionCallExpression->AddFactoryDependencyUid(MaterialFunctionFactoryNodeUid);
+		}
+
+		MaterialExpression = MaterialFunctionCallExpression;
+	}
+	else
+	{
+		MaterialExpression = NewObject<UInterchangeMaterialExpressionFactoryNode>(BaseNodeContainer);
+	}
 
 	FString ShaderType;
 	ShaderNode->GetCustomShaderType(ShaderType);
 
+	MaterialExpression->InitializeNode(MaterialExpressionUid, ShaderNode->GetDisplayLabel(), EInterchangeNodeContainerType::FactoryData);
+	BaseNodeContainer->AddNode(MaterialExpression);
+
 	if (*ShaderType == Nodes::FlattenNormal::Name)
 	{
 		HandleFlattenNormalNode(ShaderNode, MaterialFactoryNode, MaterialExpression);
@@ -1353,7 +1383,7 @@ UInterchangeMaterialExpressionFactoryNode* UInterchangeGenericMaterialPipeline::
 	return nullptr;
 }
 
-TTuple<UInterchangeMaterialExpressionFactoryNode*, FString> UInterchangeGenericMaterialPipeline::CreateMaterialExpressionForInput(UInterchangeMaterialFactoryNode* MaterialFactoryNode, const UInterchangeShaderNode* ShaderNode, const FString& InputName, const FString& ParentUid)
+TTuple<UInterchangeMaterialExpressionFactoryNode*, FString> UInterchangeGenericMaterialPipeline::CreateMaterialExpressionForInput(UInterchangeBaseMaterialFactoryNode* MaterialFactoryNode, const UInterchangeShaderNode* ShaderNode, const FString& InputName, const FString& ParentUid)
 {
 	// Make sure we don't create an expression for an input if it already has one
 	if (UInterchangeShaderPortsAPI::HasInput(MaterialFactoryNode, *InputName))
@@ -1406,6 +1436,13 @@ UInterchangeMaterialFactoryNode* UInterchangeGenericMaterialPipeline::CreateMate
 		return MaterialFactoryNode;
 	}
 
+	// Handle the case where the material will be connected through the material attributes input
+	if (HandleBxDFInput(ShaderGraphNode, MaterialFactoryNode))
+	{
+		// No need to proceed any further
+		return MaterialFactoryNode;
+	}
+
 	if (!HandlePhongModel(ShaderGraphNode, MaterialFactoryNode))
 	{
 		HandleLambertModel(ShaderGraphNode, MaterialFactoryNode);
@@ -1772,3 +1809,51 @@ TVariant<FString, FLinearColor, float> UInterchangeGenericMaterialPipeline::Visi
 
 	return Result;
 }
+
+bool UInterchangeGenericMaterialPipeline::HandleBxDFInput(const UInterchangeShaderGraphNode* ShaderGraphNode, UInterchangeMaterialFactoryNode* MaterialFactoryNode)
+{
+	using namespace UE::Interchange::Materials;
+
+	if (!ShaderGraphNode || !UInterchangeShaderPortsAPI::HasInput(ShaderGraphNode, Common::Parameters::BxDF))
+	{
+		return false;
+	}
+
+	TTuple<UInterchangeMaterialExpressionFactoryNode*, FString> ExpressionFactoryNode =
+		CreateMaterialExpressionForInput(MaterialFactoryNode, ShaderGraphNode, Common::Parameters::BxDF.ToString(), MaterialFactoryNode->GetUniqueID());
+	ensure(ExpressionFactoryNode.Get<0>());
+
+	if (ExpressionFactoryNode.Get<0>())
+	{
+		UInterchangeShaderPortsAPI::ConnectOuputToInput(MaterialFactoryNode, Common::Parameters::BxDF.ToString(), ExpressionFactoryNode.Get<0>()->GetUniqueID(), ExpressionFactoryNode.Get<1>());
+	}
+
+	// Make sure the bUseMaterialAttributes property of the material is set to true
+	static const FName UseMaterialAttributesMemberName = GET_MEMBER_NAME_CHECKED(UMaterial, bUseMaterialAttributes);
+
+	MaterialFactoryNode->AddBooleanAttribute(UseMaterialAttributesMemberName, true);
+	MaterialFactoryNode->AddApplyAndFillDelegates<FString>(UseMaterialAttributesMemberName, UMaterialExpressionMaterialFunctionCall::StaticClass(), UseMaterialAttributesMemberName);
+
+	return true;
+}
+
+UInterchangeMaterialFunctionFactoryNode* UInterchangeGenericMaterialPipeline::CreateMaterialFunctionFactoryNode(const UInterchangeShaderGraphNode* ShaderGraphNode)
+{
+	UInterchangeMaterialFunctionFactoryNode* FactoryNode = Cast<UInterchangeMaterialFunctionFactoryNode>(CreateBaseMaterialFactoryNode(ShaderGraphNode, UInterchangeMaterialFunctionFactoryNode::StaticClass()));
+	
+	TArray<FString> InputNames;
+	UInterchangeShaderPortsAPI::GatherInputs(ShaderGraphNode, InputNames);
+
+	for (const FString& InputName : InputNames)
+	{
+		TTuple<UInterchangeMaterialExpressionFactoryNode*, FString> ExpressionFactoryNode =
+			CreateMaterialExpressionForInput(FactoryNode, ShaderGraphNode, InputName, FactoryNode->GetUniqueID());
+
+		if (ExpressionFactoryNode.Get<0>())
+		{
+			UInterchangeShaderPortsAPI::ConnectOuputToInput(FactoryNode, InputName, ExpressionFactoryNode.Get<0>()->GetUniqueID(), ExpressionFactoryNode.Get<1>());
+		}
+	}
+
+	return FactoryNode;
+}
diff --git a/Engine/Plugins/Experimental/Interchange/Runtime/Source/Pipelines/Public/InterchangeGenericMaterialPipeline.h b/Engine/Plugins/Experimental/Interchange/Runtime/Source/Pipelines/Public/InterchangeGenericMaterialPipeline.h
--- a/Engine/Plugins/Experimental/Interchange/Runtime/Source/Pipelines/Public/InterchangeGenericMaterialPipeline.h
+++ b/Engine/Plugins/Experimental/Interchange/Runtime/Source/Pipelines/Public/InterchangeGenericMaterialPipeline.h
@@ -12,10 +12,12 @@
 #include "InterchangeGenericMaterialPipeline.generated.h"
 
 class UInterchangeBaseMaterialFactoryNode;
+class UInterchangeFunctionCallShaderNode;
 class UInterchangeShaderGraphNode;
 class UInterchangeShaderNode;
 class UInterchangeMaterialFactoryNode;
 class UInterchangeMaterialExpressionFactoryNode;
+class UInterchangeMaterialFunctionFactoryNode;
 class UInterchangeMaterialInstanceFactoryNode;
 class UInterchangeResult;
 
@@ -59,6 +61,7 @@ private:
 
 	UInterchangeBaseMaterialFactoryNode* CreateBaseMaterialFactoryNode(const UInterchangeBaseNode* MaterialNode, TSubclassOf<UInterchangeBaseMaterialFactoryNode> NodeType);
 	UInterchangeMaterialFactoryNode* CreateMaterialFactoryNode(const UInterchangeShaderGraphNode* ShaderGraphNode);
+	UInterchangeMaterialFunctionFactoryNode* CreateMaterialFunctionFactoryNode(const UInterchangeShaderGraphNode* FunctionCallShaderNode);
 	UInterchangeMaterialInstanceFactoryNode* CreateMaterialInstanceFactoryNode(const UInterchangeShaderGraphNode* ShaderGraphNode);
 
 	/** True if the shader graph has a clear coat input. */
@@ -84,7 +87,7 @@ private:
 
 	/** True if the shader graph has the standard surface's shader type name */
 	bool IsStandardSurfaceModel(const UInterchangeShaderGraphNode* ShaderGraphNode) const;
- 
+
 	bool HandlePhongModel(const UInterchangeShaderGraphNode* ShaderGraphNode, UInterchangeMaterialFactoryNode* MaterialFactoryNode);
 	bool HandleLambertModel(const UInterchangeShaderGraphNode* ShaderGraphNode, UInterchangeMaterialFactoryNode* MaterialFactoryNode);
 	bool HandlePBRModel(const UInterchangeShaderGraphNode* ShaderGraphNode, UInterchangeMaterialFactoryNode* MaterialFactoryNode);
@@ -94,15 +97,16 @@ private:
 	bool HandleSheen(const UInterchangeShaderGraphNode* ShaderGraphNode, UInterchangeMaterialFactoryNode* MaterialFactoryNode);
 	bool HandleThinTranslucent(const UInterchangeShaderGraphNode* ShaderGraphNode, UInterchangeMaterialFactoryNode* MaterialFactoryNode);
 	void HandleCommonParameters(const UInterchangeShaderGraphNode* ShaderGraphNode, UInterchangeMaterialFactoryNode* MaterialFactoryNode);
+	bool HandleBxDFInput(const UInterchangeShaderGraphNode* ShaderGraphNode, UInterchangeMaterialFactoryNode* MaterialFactoryNode);
 
-	void HandleFlattenNormalNode(const UInterchangeShaderNode* ShaderNode, UInterchangeMaterialFactoryNode* MaterialFactoryNode, UInterchangeMaterialExpressionFactoryNode* FlattenNormalFactoryNode);
-	void HandleTextureSampleNode(const UInterchangeShaderNode* ShaderNode, UInterchangeMaterialFactoryNode* MaterialFactoryNode, UInterchangeMaterialExpressionFactoryNode* TextureSampleFactoryNode);
-	void HandleTextureCoordinateNode(const UInterchangeShaderNode* ShaderNode, UInterchangeMaterialFactoryNode* MaterialFactoryNode, UInterchangeMaterialExpressionFactoryNode*& TextureSampleFactoryNode);
-	void HandleLerpNode(const UInterchangeShaderNode* ShaderNode, UInterchangeMaterialFactoryNode* MaterialFactoryNode, UInterchangeMaterialExpressionFactoryNode* LerpFactoryNode);
-	void HandleMaskNode(const UInterchangeShaderNode* ShaderNode, UInterchangeMaterialFactoryNode* MaterialFactoryNode, UInterchangeMaterialExpressionFactoryNode* MaskFactoryNode);
+	void HandleFlattenNormalNode(const UInterchangeShaderNode* ShaderNode, UInterchangeBaseMaterialFactoryNode* MaterialFactoryNode, UInterchangeMaterialExpressionFactoryNode* FlattenNormalFactoryNode);
+	void HandleTextureSampleNode(const UInterchangeShaderNode* ShaderNode, UInterchangeBaseMaterialFactoryNode* MaterialFactoryNode, UInterchangeMaterialExpressionFactoryNode* TextureSampleFactoryNode);
+	void HandleTextureCoordinateNode(const UInterchangeShaderNode* ShaderNode, UInterchangeBaseMaterialFactoryNode* MaterialFactoryNode, UInterchangeMaterialExpressionFactoryNode*& TextureSampleFactoryNode);
+	void HandleLerpNode(const UInterchangeShaderNode* ShaderNode, UInterchangeBaseMaterialFactoryNode* MaterialFactoryNode, UInterchangeMaterialExpressionFactoryNode* LerpFactoryNode);
+	void HandleMaskNode(const UInterchangeShaderNode* ShaderNode, UInterchangeBaseMaterialFactoryNode* MaterialFactoryNode, UInterchangeMaterialExpressionFactoryNode* MaskFactoryNode);
 
-	UInterchangeMaterialExpressionFactoryNode* CreateMaterialExpressionForShaderNode(UInterchangeMaterialFactoryNode* MaterialFactoryNode, const UInterchangeShaderNode* ShaderNode, const FString& ParentUid);
-	TTuple<UInterchangeMaterialExpressionFactoryNode*, FString> CreateMaterialExpressionForInput(UInterchangeMaterialFactoryNode* MaterialFactoryNode, const UInterchangeShaderNode* ShaderNode, const FString& InputName, const FString& ParentUid);
+	UInterchangeMaterialExpressionFactoryNode* CreateMaterialExpressionForShaderNode(UInterchangeBaseMaterialFactoryNode* MaterialFactoryNode, const UInterchangeShaderNode* ShaderNode, const FString& ParentUid);
+	TTuple<UInterchangeMaterialExpressionFactoryNode*, FString> CreateMaterialExpressionForInput(UInterchangeBaseMaterialFactoryNode* MaterialFactoryNode, const UInterchangeShaderNode* ShaderNode, const FString& InputName, const FString& ParentUid);
 
 	UInterchangeMaterialExpressionFactoryNode* CreateExpressionNode(const FString& ExpressionName, const FString& ParentUid, UClass* MaterialExpressionClass);
 	UInterchangeMaterialExpressionFactoryNode* CreateScalarParameterExpression(const UInterchangeShaderNode* ShaderNode, const FString& InputName, const FString& ParentUid);
