Adding a way to not connect a node to the Standard Surface when the node value (float/color) is equal to its default value
Adding new inputs for the Standard surface shading model:
- Diffuse Roughness
- Specular
- - IOR
- - Anisotropy
- - Rotation
- - Roughness
- Subsurface
- - Weight
- - Color
- - Radius
- - Scale
- Sheen
- - Color
- - Roughness
- Coat
- - Roughness
- - Normal
- Emission
- - EmissionColor
- Tangent
Adding outputs:
- Emissive color
- Anisotropy
- Opacity
- Fuzz Color
- Cloth
- Clear Coat
- Clear Coat Roughness
- Clear Coat Normal

diff --git a/Engine/Plugins/Experimental/Interchange/Runtime/Source/Import/Private/MaterialX/InterchangeMaterialXDefinitions.h b/Engine/Plugins/Experimental/Interchange/Runtime/Source/Import/Private/MaterialX/InterchangeMaterialXDefinitions.h
--- a/Engine/Plugins/Experimental/Interchange/Runtime/Source/Import/Private/MaterialX/InterchangeMaterialXDefinitions.h
+++ b/Engine/Plugins/Experimental/Interchange/Runtime/Source/Import/Private/MaterialX/InterchangeMaterialXDefinitions.h
@@ -50,6 +50,55 @@ namespace MaterialX
 			static constexpr const char* Normal = "normal";
 			static constexpr const char* Tangent = "tangent";
 		}
+
+		namespace DefaultValue
+		{
+			namespace Float
+			{
+				static float Base = 0.8f;
+				static float DiffuseRoughness = 0.f;
+				static float Metalness = 0.f;
+				static float Specular = 1.f;
+				static float SpecularRoughness = 0.2f;
+				static float SpecularIOR = 1.5f;
+				static float SpecularAnisotropy = 0.f;
+				static float SpecularRotation = 0.f;
+				static float Transmission = 0.f;
+				static float TransmissionDepth = 0.f;
+				static float TransmissionScatterAnisotropy = 0.f;
+				static float TransmissionDispersion = 0.f;
+				static float TransmissionExtraRoughness = 0.f;
+				static float Subsurface = 0.f;
+				static float SubsurfaceScale = 1.f;
+				static float SubsurfaceAnisotropy = 0.f;
+				static float Sheen = 0.f;
+				static float SheenRoughness = 0.3f;
+				static float Coat = 0.f;
+				static float CoatRoughness = 0.1f;
+				static float CoatAnisotropy = 0.f;
+				static float CoatRotation = 0.f;
+				static float CoatIOR = 1.5f;
+				static float CoatAffectColor = 0.f;
+				static float CoatAffectRoughness = 0.f;
+				static float ThinFilmThickness = 0.f;
+				static float ThinFilmIOR = 1.5f;
+				static float Emission = 0.f;
+			}
+
+			namespace Color3
+			{
+				constexpr FLinearColor BaseColor{ 1.0, 1.0, 1.0 };
+				constexpr FLinearColor SpecularColor{ 1, 1, 1 };
+				constexpr FLinearColor TransmissionColor{ 1, 1, 1 };
+				constexpr FLinearColor TransmissionScatter{ 0, 0, 0 };
+				constexpr FLinearColor SubsurfaceColor{ 1, 1, 1 };
+				constexpr FLinearColor SubsurfaceRadius{ 1, 1, 1 };
+				constexpr FLinearColor SheenColor{ 1, 1, 1 };
+				constexpr FLinearColor CoatColor{ 1, 1, 1 };
+				constexpr FLinearColor EmissionColor{ 1, 1, 1 };
+				constexpr FLinearColor Opacity{ 1, 1, 1 };
+			}
+		}
 	}
 
 	namespace Attributes
@@ -80,6 +129,9 @@ namespace MaterialX
 		static constexpr const char* NormalMap = "normalmap";
 		static constexpr const char* Extract = "extract";
 		static constexpr const char* TexCoord = "texcoord";
+		static constexpr const char* PointLight = "point_light";
+		static constexpr const char* DirectionalLight = "directional_light";
+		static constexpr const char* SpotLight = "spot_light";
 	}
 
 	namespace NodeDefinition
@@ -92,5 +144,6 @@ namespace MaterialX
 		static constexpr const char* Std = "stdlib";
 		static constexpr const char* Pbr= "pbrlib";
 		static constexpr const char* Bxdf = "bxdf";
+		static constexpr const char* Lights = "lights";
 	}
 }
\ No newline at end of file
diff --git a/Engine/Plugins/Experimental/Interchange/Runtime/Source/Import/Private/MaterialX/InterchangeMaterialXTranslator.cpp b/Engine/Plugins/Experimental/Interchange/Runtime/Source/Import/Private/MaterialX/InterchangeMaterialXTranslator.cpp
--- a/Engine/Plugins/Experimental/Interchange/Runtime/Source/Import/Private/MaterialX/InterchangeMaterialXTranslator.cpp
+++ b/Engine/Plugins/Experimental/Interchange/Runtime/Source/Import/Private/MaterialX/InterchangeMaterialXTranslator.cpp
@@ -5,8 +5,10 @@
 #include "MaterialXFormat/Util.h"
 #endif
 #include "InterchangeImportLog.h"
+#include "InterchangeLightNode.h"
 #include "InterchangeManager.h"
 #include "InterchangeMaterialDefinitions.h"
+#include "InterchangeSceneNode.h"
 #include "InterchangeShaderGraphNode.h"
 #include "InterchangeMaterialXDefinitions.h"
 #include "InterchangeTexture2DNode.h"
@@ -139,38 +141,41 @@ bool UInterchangeMaterialXTranslator::Translate(UInterchangeBaseNodeContainer& B
 
 		mx::NodePtr Node = Elem->asA<mx::Node>();
 
-		//for the moment the entry point is only on surfacematerial node
-		if(Node && Node->getType() == mx::MATERIAL_TYPE_STRING && Node->getCategory() == mx::SURFACE_MATERIAL_NODE_STRING)
+		if(Node)
 		{
-			if(!Node->getTypeDef())
+			//The entry point for materials is only on a surfacematerial node
+			if(Node->getType() == mx::MATERIAL_TYPE_STRING && Node->getCategory() == mx::SURFACE_MATERIAL_NODE_STRING)
 			{
-				UInterchangeResultError_Generic* Message = AddMessage<UInterchangeResultError_Generic>();
-				Message->Text = FText::Format(LOCTEXT("TypeDefNotFound", "<{0}> has no matching TypeDef"),
-					FText::FromString(Node->getName().c_str()));
+				if(!Node->getTypeDef())
+				{
+					UInterchangeResultError_Generic* Message = AddMessage<UInterchangeResultError_Generic>();
+					Message->Text = FText::Format(LOCTEXT("TypeDefNotFound", "<{0}> has no matching TypeDef"),
+						FText::FromString(Node->getName().c_str()));
 
-				bIsReferencesValid = false;
-				break;
-			}
+					bIsReferencesValid = false;
+					break;
+				}
 
-			bool bHasStandardSurface = false;
+				bool bHasStandardSurface = false;
 
-			for(mx::InputPtr Input : Node->getInputs())
-			{
-				//we only support standard_surface for now
-				mx::NodePtr ConnectedNode = Input->getConnectedNode();
-
-				if(ConnectedNode && ConnectedNode->getCategory() == "standard_surface")
+				for(mx::InputPtr Input : Node->getInputs())
 				{
-					ProcessStandardSurface(BaseNodeContainer, ConnectedNode, Document);
-					bHasStandardSurface = true;
+					//we only support standard_surface for now
+					mx::NodePtr ConnectedNode = Input->getConnectedNode();
+
+					if(ConnectedNode && ConnectedNode->getCategory() == "standard_surface")
+					{
+						ProcessStandardSurface(BaseNodeContainer, ConnectedNode, Document);
+						bHasStandardSurface = true;
+					}
 				}
-			}
 
-			if(!bHasStandardSurface)
-			{
-				UInterchangeResultWarning_Generic* Message = AddMessage<UInterchangeResultWarning_Generic>();
-				Message->Text = FText::Format(LOCTEXT("StandardSurfaceNotFound", "<{0}> has no standard_surface inputs"),
-					FText::FromString(Node->getName().c_str()));
+				if(!bHasStandardSurface)
+				{
+					UInterchangeResultWarning_Generic* Message = AddMessage<UInterchangeResultWarning_Generic>();
+					Message->Text = FText::Format(LOCTEXT("StandardSurfaceNotFound", "<{0}> has no standard_surface inputs"),
+						FText::FromString(Node->getName().c_str()));
+				}
 			}
 		}
 	}
@@ -211,34 +216,6 @@ TOptional<UE::Interchange::FImportImage> UInterchangeMaterialXTranslator::GetTex
 #if WITH_EDITOR
 void UInterchangeMaterialXTranslator::ProcessStandardSurface(UInterchangeBaseNodeContainer& NodeContainer, MaterialX::NodePtr StandardSurfaceNode, MaterialX::DocumentPtr Document) const
 {
-	auto GetInputFloatValue = [](mx::InputPtr Input)
-	{
-		if(Input->hasValueString())
-		{
-			return mx::fromValueString<float>(Input->getValueString());
-		}
-		else
-		{
-			return mx::fromValueString<float>(Input->getDefaultValue()->getValueString());
-		}
-	};
-
-	auto GetInputColor3Value = [](mx::InputPtr Input)
-	{
-		mx::Color3 Color;
-
-		if(Input->hasValueString())
-		{
-			Color = mx::fromValueString<mx::Color3>(Input->getValueString());
-		}
-		else
-		{
-			Color = mx::fromValueString<mx::Color3>(Input->getDefaultValue()->getValueString());
-		}
-
-		return FLinearColor(Color[0], Color[1], Color[2]);
-	};
-
 	using namespace UE::Interchange::Materials;
 
 	FString StandardSurfaceName{ FPaths::GetBaseFilename(StandardSurfaceNode->getActiveSourceUri().c_str()) };
@@ -254,30 +231,46 @@ void UInterchangeMaterialXTranslator::ProcessStandardSurface(UInterchangeBaseNod
 	//Base
 	{
 		//Weight
-		mx::InputPtr InputBase = GetStandardSurfaceInput(StandardSurfaceNode, mx::StandardSurface::Input::Base, Document);
-
-		if(!ConnectNodeGraphOutputToInput(InputBase, ShaderGraphNode, StandardSurface::Parameters::Base.ToString(), NamesToShaderNodes, NodeContainer))
 		{
-			ShaderGraphNode->AddFloatAttribute(UInterchangeShaderPortsAPI::MakeInputValueKey(StandardSurface::Parameters::Base.ToString()), GetInputFloatValue(InputBase));
+			mx::InputPtr InputBase = GetStandardSurfaceInput(StandardSurfaceNode, mx::StandardSurface::Input::Base, Document);
+
+			if(!ConnectNodeGraphOutputToInput(InputBase, ShaderGraphNode, StandardSurface::Parameters::Base.ToString(), NamesToShaderNodes, NodeContainer))
+			{
+				AddFloatAttribute(InputBase, StandardSurface::Parameters::Base.ToString(), ShaderGraphNode, mx::StandardSurface::DefaultValue::Float::Base);
+			}
 		}
 
 		//Color
-		mx::InputPtr InputBaseColor = GetStandardSurfaceInput(StandardSurfaceNode, mx::StandardSurface::Input::BaseColor, Document);
+		{
+			mx::InputPtr InputBaseColor = GetStandardSurfaceInput(StandardSurfaceNode, mx::StandardSurface::Input::BaseColor, Document);
+
+			if(!ConnectNodeGraphOutputToInput(InputBaseColor, ShaderGraphNode, StandardSurface::Parameters::BaseColor.ToString(), NamesToShaderNodes, NodeContainer))
+			{
+				AddLinearColorAttribute(InputBaseColor, StandardSurface::Parameters::BaseColor.ToString(), ShaderGraphNode, mx::StandardSurface::DefaultValue::Color3::BaseColor);
+			}
+		}
+	}
 
-		if(!ConnectNodeGraphOutputToInput(InputBaseColor, ShaderGraphNode, StandardSurface::Parameters::BaseColor.ToString(), NamesToShaderNodes, NodeContainer))
+	//DiffuseRoughness
+	{
+		mx::InputPtr InputDiffuseRoughness = GetStandardSurfaceInput(StandardSurfaceNode, mx::StandardSurface::Input::DiffuseRoughness, Document);
+
+		if(!ConnectNodeGraphOutputToInput(InputDiffuseRoughness, ShaderGraphNode, StandardSurface::Parameters::DiffuseRoughness.ToString(), NamesToShaderNodes, NodeContainer))
 		{
-			ShaderGraphNode->AddLinearColorAttribute(UInterchangeShaderPortsAPI::MakeInputValueKey(StandardSurface::Parameters::BaseColor.ToString()), GetInputColor3Value(InputBaseColor));
+			AddFloatAttribute(InputDiffuseRoughness, StandardSurface::Parameters::DiffuseRoughness.ToString(), ShaderGraphNode, mx::StandardSurface::DefaultValue::Float::DiffuseRoughness);
 		}
 	}
 
 	//Specular
 	{
 		//Weight
-		mx::InputPtr InputSpecular = GetStandardSurfaceInput(StandardSurfaceNode, mx::StandardSurface::Input::Specular, Document);
-
-		if(!ConnectNodeGraphOutputToInput(InputSpecular, ShaderGraphNode, StandardSurface::Parameters::Specular.ToString(), NamesToShaderNodes, NodeContainer))
 		{
-			ShaderGraphNode->AddFloatAttribute(UInterchangeShaderPortsAPI::MakeInputValueKey(StandardSurface::Parameters::Specular.ToString()), GetInputFloatValue(InputSpecular));
+			mx::InputPtr InputSpecular = GetStandardSurfaceInput(StandardSurfaceNode, mx::StandardSurface::Input::Specular, Document);
+
+			if(!ConnectNodeGraphOutputToInput(InputSpecular, ShaderGraphNode, StandardSurface::Parameters::Specular.ToString(), NamesToShaderNodes, NodeContainer))
+			{
+				AddFloatAttribute(InputSpecular, StandardSurface::Parameters::Specular.ToString(), ShaderGraphNode, mx::StandardSurface::DefaultValue::Float::Specular);
+			}
 		}
 
 		//Roughness
@@ -286,7 +279,37 @@ void UInterchangeMaterialXTranslator::ProcessStandardSurface(UInterchangeBaseNod
 
 			if(!ConnectNodeGraphOutputToInput(InputSpecularRoughness, ShaderGraphNode, StandardSurface::Parameters::SpecularRoughness.ToString(), NamesToShaderNodes, NodeContainer))
 			{
-				ShaderGraphNode->AddFloatAttribute(UInterchangeShaderPortsAPI::MakeInputValueKey(StandardSurface::Parameters::SpecularRoughness.ToString()), GetInputFloatValue(InputSpecularRoughness));
+				AddFloatAttribute(InputSpecularRoughness, StandardSurface::Parameters::SpecularRoughness.ToString(), ShaderGraphNode, mx::StandardSurface::DefaultValue::Float::SpecularRoughness);
+			}
+		}
+
+		//IOR
+		{
+			mx::InputPtr InputSpecularIOR = GetStandardSurfaceInput(StandardSurfaceNode, mx::StandardSurface::Input::SpecularIOR, Document);
+
+			if(!ConnectNodeGraphOutputToInput(InputSpecularIOR, ShaderGraphNode, StandardSurface::Parameters::SpecularIOR.ToString(), NamesToShaderNodes, NodeContainer))
+			{
+				AddFloatAttribute(InputSpecularIOR, StandardSurface::Parameters::SpecularIOR.ToString(), ShaderGraphNode, mx::StandardSurface::DefaultValue::Float::SpecularIOR);
+			}
+		}
+
+		//Anisotropy
+		{
+			mx::InputPtr InputSpecularAnisotropy = GetStandardSurfaceInput(StandardSurfaceNode, mx::StandardSurface::Input::SpecularAnisotropy, Document);
+
+			if(!ConnectNodeGraphOutputToInput(InputSpecularAnisotropy, ShaderGraphNode, StandardSurface::Parameters::SpecularAnisotropy.ToString(), NamesToShaderNodes, NodeContainer))
+			{
+				AddFloatAttribute(InputSpecularAnisotropy, StandardSurface::Parameters::SpecularAnisotropy.ToString(), ShaderGraphNode, mx::StandardSurface::DefaultValue::Float::SpecularAnisotropy);
+			}
+		}
+
+		//Rotation
+		{
+			mx::InputPtr InputSpecularRotation = GetStandardSurfaceInput(StandardSurfaceNode, mx::StandardSurface::Input::SpecularRotation, Document);
+
+			if(!ConnectNodeGraphOutputToInput(InputSpecularRotation, ShaderGraphNode, StandardSurface::Parameters::SpecularRotation.ToString(), NamesToShaderNodes, NodeContainer))
+			{
+				AddFloatAttribute(InputSpecularRotation, StandardSurface::Parameters::SpecularRotation.ToString(), ShaderGraphNode, mx::StandardSurface::DefaultValue::Float::SpecularRotation);
 			}
 		}
 	}
@@ -297,7 +320,150 @@ void UInterchangeMaterialXTranslator::ProcessStandardSurface(UInterchangeBaseNod
 
 		if(!ConnectNodeGraphOutputToInput(InputMetalness, ShaderGraphNode, StandardSurface::Parameters::Metalness.ToString(), NamesToShaderNodes, NodeContainer))
 		{
-			ShaderGraphNode->AddFloatAttribute(UInterchangeShaderPortsAPI::MakeInputValueKey(StandardSurface::Parameters::Metalness.ToString()), GetInputFloatValue(InputMetalness));
+			AddFloatAttribute(InputMetalness, StandardSurface::Parameters::Metalness.ToString(), ShaderGraphNode, mx::StandardSurface::DefaultValue::Float::Metalness);
+		}
+	}
+
+	//Subsurface
+	{
+		//Weight
+		{
+			mx::InputPtr InputSubsurface = GetStandardSurfaceInput(StandardSurfaceNode, mx::StandardSurface::Input::Subsurface, Document);
+
+			if(!ConnectNodeGraphOutputToInput(InputSubsurface, ShaderGraphNode, StandardSurface::Parameters::Subsurface.ToString(), NamesToShaderNodes, NodeContainer))
+			{
+				AddFloatAttribute(InputSubsurface, StandardSurface::Parameters::Subsurface.ToString(), ShaderGraphNode, mx::StandardSurface::DefaultValue::Float::Subsurface);
+			}
+		}
+
+		//Color
+		{
+			mx::InputPtr InputSubsurfaceColor = GetStandardSurfaceInput(StandardSurfaceNode, mx::StandardSurface::Input::SubsurfaceColor, Document);
+
+			if(!ConnectNodeGraphOutputToInput(InputSubsurfaceColor, ShaderGraphNode, StandardSurface::Parameters::SubsurfaceColor.ToString(), NamesToShaderNodes, NodeContainer))
+			{
+				AddLinearColorAttribute(InputSubsurfaceColor, StandardSurface::Parameters::SubsurfaceColor.ToString(), ShaderGraphNode, mx::StandardSurface::DefaultValue::Color3::SubsurfaceColor);
+			}
+		}
+
+		//Radius
+		{
+			mx::InputPtr InputSubsurfaceRadius = GetStandardSurfaceInput(StandardSurfaceNode, mx::StandardSurface::Input::SubsurfaceRadius, Document);
+
+			if(!ConnectNodeGraphOutputToInput(InputSubsurfaceRadius, ShaderGraphNode, StandardSurface::Parameters::SubsurfaceRadius.ToString(), NamesToShaderNodes, NodeContainer))
+			{
+				AddLinearColorAttribute(InputSubsurfaceRadius, StandardSurface::Parameters::SubsurfaceRadius.ToString(), ShaderGraphNode, mx::StandardSurface::DefaultValue::Color3::SubsurfaceRadius);
+			}
+		}
+
+		//Scale
+		{
+			mx::InputPtr InputSubsurfaceScale = GetStandardSurfaceInput(StandardSurfaceNode, mx::StandardSurface::Input::SubsurfaceScale, Document);
+
+			if(!ConnectNodeGraphOutputToInput(InputSubsurfaceScale, ShaderGraphNode, StandardSurface::Parameters::SubsurfaceScale.ToString(), NamesToShaderNodes, NodeContainer))
+			{
+				AddFloatAttribute(InputSubsurfaceScale, StandardSurface::Parameters::SubsurfaceScale.ToString(), ShaderGraphNode, mx::StandardSurface::DefaultValue::Float::SubsurfaceScale);
+			}
+		}
+	}
+
+	//Sheen
+	{
+		//Weight
+		{
+			mx::InputPtr InputSheen = GetStandardSurfaceInput(StandardSurfaceNode, mx::StandardSurface::Input::Sheen, Document);
+
+			if(!ConnectNodeGraphOutputToInput(InputSheen, ShaderGraphNode, StandardSurface::Parameters::Sheen.ToString(), NamesToShaderNodes, NodeContainer))
+			{
+				AddFloatAttribute(InputSheen, StandardSurface::Parameters::Sheen.ToString(), ShaderGraphNode, mx::StandardSurface::DefaultValue::Float::Sheen);
+			}
+		}
+
+		//Color
+		{
+			mx::InputPtr InputSheenColor = GetStandardSurfaceInput(StandardSurfaceNode, mx::StandardSurface::Input::SheenColor, Document);
+
+			if(!ConnectNodeGraphOutputToInput(InputSheenColor, ShaderGraphNode, StandardSurface::Parameters::SubsurfaceColor.ToString(), NamesToShaderNodes, NodeContainer))
+			{
+				AddLinearColorAttribute(InputSheenColor, StandardSurface::Parameters::SheenColor.ToString(), ShaderGraphNode, mx::StandardSurface::DefaultValue::Color3::SheenColor);
+			}
+		}
+
+		//Roughness
+		{
+			mx::InputPtr InputSheenRoughness = GetStandardSurfaceInput(StandardSurfaceNode, mx::StandardSurface::Input::SheenRoughness, Document);
+
+			if(!ConnectNodeGraphOutputToInput(InputSheenRoughness, ShaderGraphNode, StandardSurface::Parameters::SheenRoughness.ToString(), NamesToShaderNodes, NodeContainer))
+			{
+				AddFloatAttribute(InputSheenRoughness, StandardSurface::Parameters::SheenRoughness.ToString(), ShaderGraphNode, mx::StandardSurface::DefaultValue::Float::SheenRoughness);
+			}
+		}
+	}
+
+	//Coat
+	{
+		//Weight
+		{
+			mx::InputPtr InputCoat = GetStandardSurfaceInput(StandardSurfaceNode, mx::StandardSurface::Input::Coat, Document);
+
+			if(!ConnectNodeGraphOutputToInput(InputCoat, ShaderGraphNode, StandardSurface::Parameters::Coat.ToString(), NamesToShaderNodes, NodeContainer))
+			{
+				AddFloatAttribute(InputCoat, StandardSurface::Parameters::Coat.ToString(), ShaderGraphNode, mx::StandardSurface::DefaultValue::Float::Coat);
+			}
+		}
+
+		//Roughness
+		{
+			mx::InputPtr InputCoatRoughness = GetStandardSurfaceInput(StandardSurfaceNode, mx::StandardSurface::Input::CoatRoughness, Document);
+
+			if(!ConnectNodeGraphOutputToInput(InputCoatRoughness, ShaderGraphNode, StandardSurface::Parameters::CoatRoughness.ToString(), NamesToShaderNodes, NodeContainer))
+			{
+				AddFloatAttribute(InputCoatRoughness, StandardSurface::Parameters::CoatRoughness.ToString(), ShaderGraphNode, mx::StandardSurface::DefaultValue::Float::CoatRoughness);
+			}
+		}
+
+		//Normal
+		{
+			//No need to take the default input if there is no Normal input
+			mx::InputPtr InputCoatNormal = StandardSurfaceNode->getInput(mx::StandardSurface::Input::CoatNormal);
+
+			if(InputCoatNormal)
+			{
+				ConnectNodeGraphOutputToInput(InputCoatNormal, ShaderGraphNode, StandardSurface::Parameters::CoatNormal.ToString(), NamesToShaderNodes, NodeContainer);
+			}
+		}
+	}
+
+	//ThinFilmThickness
+	{
+		mx::InputPtr InputThinFilmThickness = GetStandardSurfaceInput(StandardSurfaceNode, mx::StandardSurface::Input::ThinFilmThickness, Document);
+
+		if(!ConnectNodeGraphOutputToInput(InputThinFilmThickness, ShaderGraphNode, StandardSurface::Parameters::ThinFilmThickness.ToString(), NamesToShaderNodes, NodeContainer))
+		{
+			AddFloatAttribute(InputThinFilmThickness, StandardSurface::Parameters::ThinFilmThickness.ToString(), ShaderGraphNode, mx::StandardSurface::DefaultValue::Float::ThinFilmThickness);
+		}
+	}
+
+	//Emission
+	{
+		//Weight
+		{
+			mx::InputPtr InputEmission = GetStandardSurfaceInput(StandardSurfaceNode, mx::StandardSurface::Input::Emission, Document);
+
+			if(!ConnectNodeGraphOutputToInput(InputEmission, ShaderGraphNode, StandardSurface::Parameters::Emission.ToString(), NamesToShaderNodes, NodeContainer))
+			{
+				AddFloatAttribute(InputEmission, StandardSurface::Parameters::Emission.ToString(), ShaderGraphNode, mx::StandardSurface::DefaultValue::Float::Emission);
+			}
+		}
+
+		//Color
+		{
+			mx::InputPtr InputEmissionColor = GetStandardSurfaceInput(StandardSurfaceNode, mx::StandardSurface::Input::EmissionColor, Document);
+
+			if(!ConnectNodeGraphOutputToInput(InputEmissionColor, ShaderGraphNode, StandardSurface::Parameters::EmissionColor.ToString(), NamesToShaderNodes, NodeContainer))
+			{
+				AddLinearColorAttribute(InputEmissionColor, StandardSurface::Parameters::EmissionColor.ToString(), ShaderGraphNode, mx::StandardSurface::DefaultValue::Color3::EmissionColor);
+			}
 		}
 	}
 
@@ -311,6 +477,17 @@ void UInterchangeMaterialXTranslator::ProcessStandardSurface(UInterchangeBaseNod
 			ConnectNodeGraphOutputToInput(InputNormal, ShaderGraphNode, StandardSurface::Parameters::Normal.ToString(), NamesToShaderNodes, NodeContainer);
 		}
 	}
+
+	//Tangent
+	{
+		//No need to take the default input if there is no Tangent input
+		mx::InputPtr InputTangent = StandardSurfaceNode->getInput(mx::StandardSurface::Input::Tangent);
+
+		if(InputTangent)
+		{
+			ConnectNodeGraphOutputToInput(InputTangent, ShaderGraphNode, StandardSurface::Parameters::Tangent.ToString(), NamesToShaderNodes, NodeContainer);
+		}
+	}
 }
 
 bool UInterchangeMaterialXTranslator::ConnectNodeGraphOutputToInput(MaterialX::InputPtr InputToNodeGraph, UInterchangeShaderNode* ShaderNode, const FString& ParentInputName, TMap<FString, UInterchangeShaderNode*>& NamesToShaderNodes, UInterchangeBaseNodeContainer& NodeContainer) const
@@ -620,6 +797,43 @@ bool UInterchangeMaterialXTranslator::AddAttribute(MaterialX::InputPtr Input, co
 
 	return false;
 }
+
+bool UInterchangeMaterialXTranslator::AddFloatAttribute(MaterialX::InputPtr Input, const FString& InputChannelName, UInterchangeShaderNode* ShaderNode, float DefaultValue) const
+{
+	if(Input)
+	{
+		if(Input->hasValueString())
+		{
+			float Value = mx::fromValueString<float>(Input->getValueString());
+
+			if(!FMath::IsNearlyEqual(Value, DefaultValue))
+			{
+				return ShaderNode->AddFloatAttribute(UInterchangeShaderPortsAPI::MakeInputValueKey(InputChannelName), Value);
+			}
+		}
+	}
+
+	return false;
+}
+
+bool UInterchangeMaterialXTranslator::AddLinearColorAttribute(MaterialX::InputPtr Input, const FString& InputChannelName, UInterchangeShaderNode* ShaderNode, const FLinearColor& DefaultValue) const
+{
+	if(Input)
+	{
+		if(Input->hasValueString())
+		{
+			mx::Color3 Color = mx::fromValueString<mx::Color3>(Input->getValueString());
+			FLinearColor Value{ Color[0], Color[1], Color[2] };
+
+			if(!Value.Equals(DefaultValue))
+			{
+				return ShaderNode->AddLinearColorAttribute(UInterchangeShaderPortsAPI::MakeInputValueKey(InputChannelName), Value);
+			}
+		}
+	}
+
+	return false;
+}
 #endif //WITH_EDITOR
 
 #undef LOCTEXT_NAMESPACE
\ No newline at end of file
diff --git a/Engine/Plugins/Experimental/Interchange/Runtime/Source/Import/Public/MaterialX/InterchangeMaterialXTranslator.h b/Engine/Plugins/Experimental/Interchange/Runtime/Source/Import/Public/MaterialX/InterchangeMaterialXTranslator.h
--- a/Engine/Plugins/Experimental/Interchange/Runtime/Source/Import/Public/MaterialX/InterchangeMaterialXTranslator.h
+++ b/Engine/Plugins/Experimental/Interchange/Runtime/Source/Import/Public/MaterialX/InterchangeMaterialXTranslator.h
@@ -152,6 +152,30 @@ protected:
 	 */
 	bool AddAttribute(MaterialX::InputPtr Input, const FString& InputChannelName, UInterchangeShaderNode* ShaderNode) const;
 
+	/**
+	 * Add a float attribute to a shader node only if its value taken from the input is not equal to its default value. Return false if the attribute does not exist or if we cannot add it
+	 *
+	 * @param Input - The MaterialX input to retrieve and add the value from, must be of type float/color/vector
+	 * @param InputChannelName - The name of the shader node's input to add the attribute
+	 * @param ShaderNode - The shader node to which we want to add the attribute
+	 * @param DefaultValue - the default value to test the input against
+	 *
+	 * @return true if the attribute was successfully added
+	 */
+	bool AddFloatAttribute(MaterialX::InputPtr Input, const FString& InputChannelName, UInterchangeShaderNode* ShaderNode, float DefaultValue) const;
+
+	/**
+	 * Add a FLinearColor attribute to a shader node only if its value taken from the input is not equal to its default value. Return false if the attribute does not exist or if we cannot add it
+	 *
+	 * @param Input - The MaterialX input to retrieve and add the value from, must be of type float/color/vector
+	 * @param InputChannelName - The name of the shader node's input to add the attribute
+	 * @param ShaderNode - The shader node to which we want to add the attribute
+	 * @param DefaultValue - the default value to test the input against
+	 *
+	 * @return true if the attribute was successfully added
+	 */
+	bool AddLinearColorAttribute(MaterialX::InputPtr Input, const FString& InputChannelName, UInterchangeShaderNode* ShaderNode, const FLinearColor& DefaultValue) const;
+
 private:
 
 	TMap<FString, FString> InputNamesMaterialX2UE;
diff --git a/Engine/Plugins/Experimental/Interchange/Runtime/Source/Pipelines/Private/InterchangeGenericMaterialPipeline.cpp b/Engine/Plugins/Experimental/Interchange/Runtime/Source/Pipelines/Private/InterchangeGenericMaterialPipeline.cpp
--- a/Engine/Plugins/Experimental/Interchange/Runtime/Source/Pipelines/Private/InterchangeGenericMaterialPipeline.cpp
+++ b/Engine/Plugins/Experimental/Interchange/Runtime/Source/Pipelines/Private/InterchangeGenericMaterialPipeline.cpp
@@ -476,15 +476,32 @@ bool UInterchangeGenericMaterialPipeline::HandleStandardSurfaceModel(const UInte
 		};
 
 		using ArgumentsType = TTuple<FString, bool, bool>;
-		// TODO: We need to support all standard surface inputs
+		
 		static ArgumentsType ArgumentsPerInputs[]
 		{
-			ArgumentsType{StandardSurface::Parameters::Base.ToString(), false, false},				// BaseWeight
+			ArgumentsType{StandardSurface::Parameters::Base.ToString(), true, false},				// Base
 			ArgumentsType{StandardSurface::Parameters::BaseColor.ToString(), false, false},			// BaseColor
-			ArgumentsType{StandardSurface::Parameters::Metalness.ToString(), true, false},			// Metallic
-			ArgumentsType{StandardSurface::Parameters::Specular.ToString(), true, false},			// SpecularWeight
+			ArgumentsType{StandardSurface::Parameters::DiffuseRoughness.ToString(), false, false},	// DiffuseRoughness
+			ArgumentsType{StandardSurface::Parameters::Metalness.ToString(), true, false},			// Metalness
+			ArgumentsType{StandardSurface::Parameters::Specular.ToString(), true, false},			// Specular
 			ArgumentsType{StandardSurface::Parameters::SpecularRoughness.ToString(), true, false},	// SpecularRoughness
-			ArgumentsType{StandardSurface::Parameters::Normal.ToString(), false, true}				// Normal
+			ArgumentsType{StandardSurface::Parameters::SpecularIOR.ToString(), true, false},        // SpecularIOR
+			ArgumentsType{StandardSurface::Parameters::SpecularAnisotropy.ToString(), true, false},	// SpecularAnisotropy
+			ArgumentsType{StandardSurface::Parameters::SpecularRotation.ToString(), true, false},	// SpecularRotation
+			ArgumentsType{StandardSurface::Parameters::Subsurface.ToString(), true, false},			// Subsurface
+			ArgumentsType{StandardSurface::Parameters::SubsurfaceColor.ToString(), true, false},	// SubsurfaceColor
+			ArgumentsType{StandardSurface::Parameters::SubsurfaceRadius.ToString(), true, false},	// SubsurfaceRadius
+			ArgumentsType{StandardSurface::Parameters::SubsurfaceScale.ToString(), true, false},	// SubsurfaceScale
+			ArgumentsType{StandardSurface::Parameters::Sheen.ToString(), true, false},				// Sheen
+			ArgumentsType{StandardSurface::Parameters::SheenColor.ToString(), true, false},			// SheenColor
+			ArgumentsType{StandardSurface::Parameters::SheenRoughness.ToString(), true, false},		// SheenRoughness
+			ArgumentsType{StandardSurface::Parameters::Coat.ToString(), true, false},				// Coat
+			ArgumentsType{StandardSurface::Parameters::CoatRoughness.ToString(), true, false},		// CoatRoughness
+			ArgumentsType{StandardSurface::Parameters::ThinFilmThickness.ToString(), true, false},	// ThinFilmThickness
+			ArgumentsType{StandardSurface::Parameters::Emission.ToString(), true, false},			// Emission
+			ArgumentsType{StandardSurface::Parameters::EmissionColor.ToString(), false, false},		// EmissionColor
+			ArgumentsType{StandardSurface::Parameters::Normal.ToString(), false, true},				// Normal
+			ArgumentsType{StandardSurface::Parameters::Tangent.ToString(), false, true}				// Tangent
 		};
 
 		for(const ArgumentsType & Arguments : ArgumentsPerInputs)
@@ -494,24 +511,35 @@ bool UInterchangeGenericMaterialPipeline::HandleStandardSurfaceModel(const UInte
 			bool bNormalParsingValue = Arguments.Get<2>();
 			ConnectMaterialExpressionOutputToInput(InputName, bLinearParsingValue, bNormalParsingValue);
 		}
-
-		MaterialFactoryNode->ConnectOutputToBaseColor(FunctionCallExpressionUid, PBR::Parameters::BaseColor.ToString());
+		
+		MaterialFactoryNode->ConnectOutputToBaseColor(FunctionCallExpressionUid,  PBR::Parameters::BaseColor.ToString());
 		MaterialFactoryNode->ConnectOutputToMetallic(FunctionCallExpressionUid, PBR::Parameters::Metallic.ToString());
-		MaterialFactoryNode->ConnectOutputToRoughness(FunctionCallExpressionUid, PBR::Parameters::Roughness.ToString());
 		MaterialFactoryNode->ConnectOutputToSpecular(FunctionCallExpressionUid, PBR::Parameters::Specular.ToString());
+		MaterialFactoryNode->ConnectOutputToRoughness(FunctionCallExpressionUid, PBR::Parameters::Roughness.ToString());
+		MaterialFactoryNode->ConnectOutputToEmissiveColor(FunctionCallExpressionUid, PBR::Parameters::EmissiveColor.ToString());
+		MaterialFactoryNode->ConnectOutputToAnisotropy(FunctionCallExpressionUid, PBR::Parameters::Anisotropy.ToString());
 		MaterialFactoryNode->ConnectOutputToNormal(FunctionCallExpressionUid, PBR::Parameters::Normal.ToString());
+		MaterialFactoryNode->ConnectOutputToTangent(FunctionCallExpressionUid, PBR::Parameters::Tangent.ToString());
+		MaterialFactoryNode->ConnectOutputToOpacity(FunctionCallExpressionUid, PBR::Parameters::Opacity.ToString());
+		
 
 		// We can't have all shading models at once, so we have to make a choice here
 		if(UInterchangeShaderPortsAPI::HasInput(ShaderGraphNode, StandardSurface::Parameters::Sheen))
 		{
+			MaterialFactoryNode->ConnectOutputToFuzzColor(FunctionCallExpressionUid, Sheen::Parameters::SheenColor.ToString());
+			MaterialFactoryNode->ConnectOutputToCloth(FunctionCallExpressionUid, Sheen::Parameters::SheenRoughness.ToString());
 			MaterialFactoryNode->SetCustomShadingModel(MSM_Cloth);
 		}
 		else if(UInterchangeShaderPortsAPI::HasInput(ShaderGraphNode, StandardSurface::Parameters::Coat))
 		{
+			MaterialFactoryNode->ConnectOutputToClearCoat(FunctionCallExpressionUid, ClearCoat::Parameters::ClearCoat.ToString());
+			MaterialFactoryNode->ConnectOutputToClearCoatRoughness(FunctionCallExpressionUid, ClearCoat::Parameters::ClearCoatRoughness.ToString());
+			MaterialFactoryNode->ConnectOutputToClearCoatNormal(FunctionCallExpressionUid, ClearCoat::Parameters::ClearCoatNormal.ToString());
 			MaterialFactoryNode->SetCustomShadingModel(MSM_ClearCoat);
 		}
 		else if(UInterchangeShaderPortsAPI::HasInput(ShaderGraphNode, StandardSurface::Parameters::Subsurface))
 		{
+			MaterialFactoryNode->ConnectOutputToSubsurface(FunctionCallExpressionUid, Subsurface::Parameters::SubsurfaceColor.ToString());
 			MaterialFactoryNode->SetCustomShadingModel(MSM_Subsurface);
 		}
 		else
