Interchange GLTF: Made sure scene nodes had unique names since they can end up being used as asset names when merging hierarchies.

diff --git a/Engine/Plugins/Experimental/Interchange/Runtime/Source/Import/Private/Gltf/InterchangeGltfTranslator.cpp b/Engine/Plugins/Experimental/Interchange/Runtime/Source/Import/Private/Gltf/InterchangeGltfTranslator.cpp
--- a/Engine/Plugins/Experimental/Interchange/Runtime/Source/Import/Private/Gltf/InterchangeGltfTranslator.cpp
+++ b/Engine/Plugins/Experimental/Interchange/Runtime/Source/Import/Private/Gltf/InterchangeGltfTranslator.cpp
@@ -66,12 +66,13 @@ void UInterchangeGltfTranslator::HandleGltfNode( UInterchangeBaseNodeContainer&
 {
 	using namespace UE::Interchange::Gltf::Private;
 
-	const FString NodeUid = ParentNodeUid + TEXT("\\") + GenerateNameForGltfNode( GltfNode.Name, NodeIndex );
+	const FString NodeName = GenerateNameForGltfNode( GltfNode.Name, NodeIndex );
+	const FString NodeUid = ParentNodeUid + TEXT("\\") + NodeName;
 
 	const UInterchangeSceneNode* ParentSceneNode = Cast< UInterchangeSceneNode >( NodeContainer.GetNode( ParentNodeUid ) );
 
 	UInterchangeSceneNode* InterchangeSceneNode = NewObject< UInterchangeSceneNode >( &NodeContainer );
-	InterchangeSceneNode->InitializeNode( NodeUid, GltfNode.Name, EInterchangeNodeContainerType::TranslatedScene );
+	InterchangeSceneNode->InitializeNode( NodeUid, NodeName, EInterchangeNodeContainerType::TranslatedScene );
 	NodeContainer.AddNode( InterchangeSceneNode );
 
 	FTransform Transform = GltfNode.Transform;
