Interchange:
Add blend shape animation curve support
Add material parameters animation curve support
Add bone user defined attribute animation curve support

#rnx

diff --git a/Engine/Plugins/Experimental/Interchange/Runtime/Source/FactoryNodes/Private/InterchangeAnimSequenceFactoryNode.cpp b/Engine/Plugins/Experimental/Interchange/Runtime/Source/FactoryNodes/Private/InterchangeAnimSequenceFactoryNode.cpp
--- a/Engine/Plugins/Experimental/Interchange/Runtime/Source/FactoryNodes/Private/InterchangeAnimSequenceFactoryNode.cpp
+++ b/Engine/Plugins/Experimental/Interchange/Runtime/Source/FactoryNodes/Private/InterchangeAnimSequenceFactoryNode.cpp
@@ -18,8 +18,40 @@ namespace UE::Interchange::Animation
 	}
 }
 
+namespace UE::Interchange
+{
+	const FAttributeKey& FAnimSequenceNodeStaticData::GetAnimatedMorphTargetDependenciesKey()
+	{
+		static FAttributeKey AttributeKey(TEXT("__AnimatedMeshMorphTargetDependencies__"));
+		return AttributeKey;
+	}
+	
+	const FAttributeKey& FAnimSequenceNodeStaticData::GetAnimatedAttributeCurveNamesKey()
+	{
+		static FAttributeKey AttributeKey(TEXT("__AnimatedAttributeCurveNames__"));
+		return AttributeKey;
+	}
+	
+	const FAttributeKey& FAnimSequenceNodeStaticData::GetAnimatedAttributeStepCurveNamesKey()
+	{
+		static FAttributeKey AttributeKey(TEXT("__AnimatedAttributeStepCurveNames__"));
+		return AttributeKey;
+	}
+
+	const FAttributeKey& FAnimSequenceNodeStaticData::GetAnimatedMaterialCurveSuffixesKey()
+	{
+		static FAttributeKey AttributeKey(TEXT("__AnimatedMaterialCurveSuffixes__"));
+		return AttributeKey;
+	}
+}//ns UE::Interchange
+
+
 UInterchangeAnimSequenceFactoryNode::UInterchangeAnimSequenceFactoryNode()
 {
+	AnimatedMorphTargetDependencies.Initialize(Attributes, UE::Interchange::FAnimSequenceNodeStaticData::GetAnimatedMorphTargetDependenciesKey().ToString());
+	AnimatedAttributeCurveNames.Initialize(Attributes, UE::Interchange::FAnimSequenceNodeStaticData::GetAnimatedAttributeCurveNamesKey().ToString());
+	AnimatedAttributeStepCurveNames.Initialize(Attributes, UE::Interchange::FAnimSequenceNodeStaticData::GetAnimatedAttributeStepCurveNamesKey().ToString());
+	AnimatedMaterialCurveSuffixes.Initialize(Attributes, UE::Interchange::FAnimSequenceNodeStaticData::GetAnimatedMaterialCurveSuffixesKey().ToString());
 }
 
 void UInterchangeAnimSequenceFactoryNode::InitializeAnimSequenceNode(const FString& UniqueID, const FString& DisplayLabel)
@@ -106,6 +138,46 @@ bool UInterchangeAnimSequenceFactoryNode::SetCustomImportBoneTracksRangeStop(con
 	IMPLEMENT_NODE_ATTRIBUTE_SETTER_NODELEGATE(ImportBoneTracksRangeStop, double);
 }
 
+bool UInterchangeAnimSequenceFactoryNode::GetCustomImportAttributeCurves(bool& AttributeValue) const
+{
+	IMPLEMENT_NODE_ATTRIBUTE_GETTER(ImportAttributeCurves, bool);
+}
+
+bool UInterchangeAnimSequenceFactoryNode::SetCustomImportAttributeCurves(const bool& AttributeValue)
+{
+	IMPLEMENT_NODE_ATTRIBUTE_SETTER_NODELEGATE(ImportAttributeCurves, bool);
+}
+
+bool UInterchangeAnimSequenceFactoryNode::GetCustomDoNotImportCurveWithZero(bool& AttributeValue) const
+{
+	IMPLEMENT_NODE_ATTRIBUTE_GETTER(DoNotImportCurveWithZero, bool);
+}
+
+bool UInterchangeAnimSequenceFactoryNode::SetCustomDoNotImportCurveWithZero(const bool& AttributeValue)
+{
+	IMPLEMENT_NODE_ATTRIBUTE_SETTER_NODELEGATE(DoNotImportCurveWithZero, bool);
+}
+
+bool UInterchangeAnimSequenceFactoryNode::GetCustomRemoveCurveRedundantKeys(bool& AttributeValue) const
+{
+	IMPLEMENT_NODE_ATTRIBUTE_GETTER(RemoveCurveRedundantKeys, bool);
+}
+
+bool UInterchangeAnimSequenceFactoryNode::SetCustomRemoveCurveRedundantKeys(const bool& AttributeValue)
+{
+	IMPLEMENT_NODE_ATTRIBUTE_SETTER_NODELEGATE(RemoveCurveRedundantKeys, bool);
+}
+
+bool UInterchangeAnimSequenceFactoryNode::GetCustomDeleteExistingMorphTargetCurves(bool& AttributeValue) const
+{
+	IMPLEMENT_NODE_ATTRIBUTE_GETTER(DeleteExistingMorphTargetCurves, bool);
+}
+
+bool UInterchangeAnimSequenceFactoryNode::SetCustomDeleteExistingMorphTargetCurves(const bool& AttributeValue)
+{
+	IMPLEMENT_NODE_ATTRIBUTE_SETTER_NODELEGATE(DeleteExistingMorphTargetCurves, bool);
+}
+
 bool UInterchangeAnimSequenceFactoryNode::GetCustomSkeletonSoftObjectPath(FSoftObjectPath& AttributeValue) const
 {
 	IMPLEMENT_NODE_ATTRIBUTE_GETTER(SkeletonSoftObjectPath, FSoftObjectPath)
@@ -116,6 +188,139 @@ bool UInterchangeAnimSequenceFactoryNode::SetCustomSkeletonSoftObjectPath(const
 	IMPLEMENT_NODE_ATTRIBUTE_SETTER_NODELEGATE(SkeletonSoftObjectPath, FSoftObjectPath)
 }
 
+int32 UInterchangeAnimSequenceFactoryNode::GetAnimatedMorphTargetDependeciesCount() const
+{
+	return AnimatedMorphTargetDependencies.GetCount();
+}
+
+void UInterchangeAnimSequenceFactoryNode::GetAnimatedMorphTargetDependencies(TArray<FString>& OutDependencies) const
+{
+	AnimatedMorphTargetDependencies.GetItems(OutDependencies);
+}
+
+void UInterchangeAnimSequenceFactoryNode::GetAnimatedMorphTargetDependency(const int32 Index, FString& OutDependency) const
+{
+	AnimatedMorphTargetDependencies.GetItem(Index, OutDependency);
+}
+
+bool UInterchangeAnimSequenceFactoryNode::SetAnimatedMorphTargetDependencyUid(const FString& DependencyUid)
+{
+	return AnimatedMorphTargetDependencies.AddItem(DependencyUid);
+}
+
+bool UInterchangeAnimSequenceFactoryNode::RemoveAnimatedMorphTargetDependencyUid(const FString& DependencyUid)
+{
+	return AnimatedMorphTargetDependencies.RemoveItem(DependencyUid);
+}
+
+
+int32 UInterchangeAnimSequenceFactoryNode::GetAnimatedAttributeCurveNamesCount() const
+{
+	return AnimatedAttributeCurveNames.GetCount();
+}
+
+void UInterchangeAnimSequenceFactoryNode::GetAnimatedAttributeCurveNames(TArray<FString>& OutAttributeCurveNames) const
+{
+	AnimatedAttributeCurveNames.GetItems(OutAttributeCurveNames);
+}
+
+void UInterchangeAnimSequenceFactoryNode::GetAnimatedAttributeCurveName(const int32 Index, FString& OutAttributeCurveName) const
+{
+	AnimatedAttributeCurveNames.GetItem(Index, OutAttributeCurveName);
+}
+
+bool UInterchangeAnimSequenceFactoryNode::SetAnimatedAttributeCurveName(const FString& AttributeCurveName)
+{
+	return AnimatedAttributeCurveNames.AddItem(AttributeCurveName);
+}
+
+bool UInterchangeAnimSequenceFactoryNode::RemoveAnimatedAttributeCurveName(const FString& AttributeCurveName)
+{
+	return AnimatedAttributeCurveNames.RemoveItem(AttributeCurveName);
+}
+
+bool UInterchangeAnimSequenceFactoryNode::GetCustomMaterialDriveParameterOnCustomAttribute(bool& AttributeValue) const
+{
+	IMPLEMENT_NODE_ATTRIBUTE_GETTER(MaterialDriveParameterOnCustomAttribute, bool);
+}
+
+bool UInterchangeAnimSequenceFactoryNode::SetCustomMaterialDriveParameterOnCustomAttribute(const bool& AttributeValue)
+{
+	IMPLEMENT_NODE_ATTRIBUTE_SETTER_NODELEGATE(MaterialDriveParameterOnCustomAttribute, bool);
+}
+
+int32 UInterchangeAnimSequenceFactoryNode::GetAnimatedMaterialCurveSuffixesCount() const
+{
+	return AnimatedMaterialCurveSuffixes.GetCount();
+}
+
+void UInterchangeAnimSequenceFactoryNode::GetAnimatedMaterialCurveSuffixes(TArray<FString>& OutMaterialCurveSuffixes) const
+{
+	AnimatedMaterialCurveSuffixes.GetItems(OutMaterialCurveSuffixes);
+}
+
+void UInterchangeAnimSequenceFactoryNode::GetAnimatedMaterialCurveSuffixe(const int32 Index, FString& OutMaterialCurveSuffixe) const
+{
+	AnimatedMaterialCurveSuffixes.GetItem(Index, OutMaterialCurveSuffixe);
+}
+
+bool UInterchangeAnimSequenceFactoryNode::SetAnimatedMaterialCurveSuffixe(const FString& MaterialCurveSuffixe)
+{
+	return AnimatedMaterialCurveSuffixes.AddItem(MaterialCurveSuffixe);
+}
+
+bool UInterchangeAnimSequenceFactoryNode::RemoveAnimatedMaterialCurveSuffixe(const FString& MaterialCurveSuffixe)
+{
+	return AnimatedMaterialCurveSuffixes.RemoveItem(MaterialCurveSuffixe);
+}
+
+
+int32 UInterchangeAnimSequenceFactoryNode::GetAnimatedAttributeStepCurveNamesCount() const
+{
+	return AnimatedAttributeStepCurveNames.GetCount();
+}
+
+void UInterchangeAnimSequenceFactoryNode::GetAnimatedAttributeStepCurveNames(TArray<FString>& OutAttributeStepCurveNames) const
+{
+	AnimatedAttributeStepCurveNames.GetItems(OutAttributeStepCurveNames);
+}
+
+void UInterchangeAnimSequenceFactoryNode::GetAnimatedAttributeStepCurveName(const int32 Index, FString& OutAttributeStepCurveName) const
+{
+	AnimatedAttributeStepCurveNames.GetItem(Index, OutAttributeStepCurveName);
+}
+
+bool UInterchangeAnimSequenceFactoryNode::SetAnimatedAttributeStepCurveName(const FString& AttributeStepCurveName)
+{
+	return AnimatedAttributeStepCurveNames.AddItem(AttributeStepCurveName);
+}
+
+bool UInterchangeAnimSequenceFactoryNode::RemoveAnimatedAttributeStepCurveName(const FString& AttributeStepCurveName)
+{
+	return AnimatedAttributeStepCurveNames.RemoveItem(AttributeStepCurveName);
+}
+
+
+bool UInterchangeAnimSequenceFactoryNode::GetCustomDeleteExistingCustomAttributeCurves(bool& AttributeValue) const
+{
+	IMPLEMENT_NODE_ATTRIBUTE_GETTER(DeleteExistingCustomAttributeCurves, bool);
+}
+
+bool UInterchangeAnimSequenceFactoryNode::SetCustomDeleteExistingCustomAttributeCurves(const bool& AttributeValue)
+{
+	IMPLEMENT_NODE_ATTRIBUTE_SETTER_NODELEGATE(DeleteExistingCustomAttributeCurves, bool);
+}
+
+bool UInterchangeAnimSequenceFactoryNode::GetCustomDeleteExistingNonCurveCustomAttributes(bool& AttributeValue) const
+{
+	IMPLEMENT_NODE_ATTRIBUTE_GETTER(DeleteExistingNonCurveCustomAttributes, bool);
+}
+
+bool UInterchangeAnimSequenceFactoryNode::SetCustomDeleteExistingNonCurveCustomAttributes(const bool& AttributeValue)
+{
+	IMPLEMENT_NODE_ATTRIBUTE_SETTER_NODELEGATE(DeleteExistingNonCurveCustomAttributes, bool);
+}
+
 /************************************************************************/
 /* Automation tests                                                     */
 /************************************************************************/
diff --git a/Engine/Plugins/Experimental/Interchange/Runtime/Source/FactoryNodes/Public/InterchangeAnimSequenceFactoryNode.h b/Engine/Plugins/Experimental/Interchange/Runtime/Source/FactoryNodes/Public/InterchangeAnimSequenceFactoryNode.h
--- a/Engine/Plugins/Experimental/Interchange/Runtime/Source/FactoryNodes/Public/InterchangeAnimSequenceFactoryNode.h
+++ b/Engine/Plugins/Experimental/Interchange/Runtime/Source/FactoryNodes/Public/InterchangeAnimSequenceFactoryNode.h
@@ -13,6 +13,17 @@ namespace UE::Interchange::Animation
 	INTERCHANGEFACTORYNODES_API FFrameRate ConvertSampleRatetoFrameRate(double SampleRate);
 }
 
+namespace UE::Interchange
+{
+	struct INTERCHANGEFACTORYNODES_API FAnimSequenceNodeStaticData : public FBaseNodeStaticData
+	{
+		static const FAttributeKey& GetAnimatedMorphTargetDependenciesKey();
+		static const FAttributeKey& GetAnimatedAttributeCurveNamesKey();
+		static const FAttributeKey& GetAnimatedAttributeStepCurveNamesKey();
+		static const FAttributeKey& GetAnimatedMaterialCurveSuffixesKey();
+	};
+}//ns UE::Interchange
+
 UCLASS(BlueprintType, Experimental)
 class INTERCHANGEFACTORYNODES_API UInterchangeAnimSequenceFactoryNode : public UInterchangeFactoryBaseNode
 {
@@ -41,45 +52,310 @@ public:
 	virtual class UClass* GetObjectClass() const override;
 
 public:
-	/** Return false if the Attribute was not set previously. */
+	/** Get the skeleton factory node unique id. Return false if the attribute is not set. */
 	UFUNCTION(BlueprintCallable, Category = "Interchange | Node | AnimSequence")
 	bool GetCustomSkeletonFactoryNodeUid(FString& AttributeValue) const;
 
+	/** Set the skeleton factory node unique id. Return false if the attribute cannot be set. */
 	UFUNCTION(BlueprintCallable, Category = "Interchange | Node | AnimSequence")
 	bool SetCustomSkeletonFactoryNodeUid(const FString& AttributeValue);
 
-	/** Query whether we must import the bone tracks. Return false if the Attribute was not set previously. */
+
+	/**********************************************************************************************
+	 * Import bone tracks API begin
+	 */
+
+	/**
+	 * Get the import bone tracks state. The attribute will be true if we need to import bone tracks.
+	 * False if we do not import bone tracks.
+	 * 
+	 * Note - Return false if the attribute is not set. Return true if the attribute exist and can be query.
+	 */
 	UFUNCTION(BlueprintCallable, Category = "Interchange | Node | AnimSequence")
 	bool GetCustomImportBoneTracks(bool& AttributeValue) const;
 
-	/** Set true if we must import the bone tracks, otherwise set it to false. Return false if the Attribute was not set previously. */
+	/** Set the import bone tracks state. Pass true to import bone tracks, false to not import bone tracks. */
 	UFUNCTION(BlueprintCallable, Category = "Interchange | Node | AnimSequence")
 	bool SetCustomImportBoneTracks(const bool& AttributeValue);
 
-	/** Query weather we must import the bone tracks. Return false if the Attribute was not set previously. */
+	/** Get the import bone tracks sample rate. Return false if the attribute is not set. */
 	UFUNCTION(BlueprintCallable, Category = "Interchange | Node | AnimSequence")
 	bool GetCustomImportBoneTracksSampleRate(double& AttributeValue) const;
 
-	/** Set true if we must import the bone tracks, otherwise set it to false. Return false if the Attribute was not set previously. */
+	/** Set the import bone tracks sample rate. Return false if the attribute cannot be set. */
 	UFUNCTION(BlueprintCallable, Category = "Interchange | Node | AnimSequence")
 	bool SetCustomImportBoneTracksSampleRate(const double& AttributeValue);
 
-	/** Query weather we must import the bone tracks. Return false if the Attribute was not set previously. */
+	/** Get the import bone tracks start time in second. Return false if the attribute is not set. */
 	UFUNCTION(BlueprintCallable, Category = "Interchange | Node | AnimSequence")
 	bool GetCustomImportBoneTracksRangeStart(double& AttributeValue) const;
 
-	/** Set true if we must import the bone tracks, otherwise set it to false. Return false if the Attribute was not set previously. */
+	/** Set the import bone tracks start time in second. Return false if the attribute cannot be set. */
 	UFUNCTION(BlueprintCallable, Category = "Interchange | Node | AnimSequence")
 	bool SetCustomImportBoneTracksRangeStart(const double& AttributeValue);
 
-	/** Query weather we must import the bone tracks. Return false if the Attribute was not set previously. */
+	/** Get the import bone tracks end time in second. Return false if the attribute is not set. */
 	UFUNCTION(BlueprintCallable, Category = "Interchange | Node | AnimSequence")
 	bool GetCustomImportBoneTracksRangeStop(double& AttributeValue) const;
 
-	/** Set true if we must import the bone tracks, otherwise set it to false. Return false if the Attribute was not set previously. */
+	/** Set the import bone tracks end time in second. Return false if the attribute cannot be set. */
 	UFUNCTION(BlueprintCallable, Category = "Interchange | Node | AnimSequence")
 	bool SetCustomImportBoneTracksRangeStop(const double& AttributeValue);
 
+	/*
+	 * Import bone tracks API end
+	 **********************************************************************************************/
+
+
+	/**********************************************************************************************
+	 * Curves API begin
+	 */
+
+
+	 /**
+	  * Get the import attribute curves state. If true this mean we want to import all user custom attributes
+	  * we can find on a node.
+	  * 
+	  * Return false if the attribute is not set.
+	  */
+	UFUNCTION(BlueprintCallable, Category = "Interchange | Node | AnimSequence")
+	bool GetCustomImportAttributeCurves(bool& AttributeValue) const;
+
+	/** Set the import attribute curves state. Return false if the attribute cannot be set. */
+	UFUNCTION(BlueprintCallable, Category = "Interchange | Node | AnimSequence")
+	bool SetCustomImportAttributeCurves(const bool& AttributeValue);
+
+	/**
+	 * Get the custom attribute DoNotImportCurveWithZero, return false if the attribute is not set.
+	 * 
+	 * Note - If value is true, do not import if it doens't have any value other than zero. This is to avoid adding extra curves to evaluate.
+	 */
+	UFUNCTION(BlueprintCallable, Category = "Interchange | Node | AnimSequence")
+	bool GetCustomDoNotImportCurveWithZero(bool& AttributeValue) const;
+
+	/**
+	 * Set the custom attribute DoNotImportCurveWithZero. Return false if the attribute cannot be set.
+	 * 
+	 * Note - If value is true, do not import if it doens't have any value other than zero. This is to avoid adding extra curves to evaluate.
+	 */
+	UFUNCTION(BlueprintCallable, Category = "Interchange | Node | AnimSequence")
+	bool SetCustomDoNotImportCurveWithZero(const bool& AttributeValue);
+
+	/**
+	 * Get the custom attribute RemoveCurveRedundantKeys, return false if the attribute is not set.
+	 */
+	UFUNCTION(BlueprintCallable, Category = "Interchange | Node | AnimSequence")
+	bool GetCustomRemoveCurveRedundantKeys(bool& AttributeValue) const;
+
+	/**
+	 * Set the custom attribute RemoveCurveRedundantKeys. Return false if the attribute cannot be set.
+	 */
+	UFUNCTION(BlueprintCallable, Category = "Interchange | Node | AnimSequence")
+	bool SetCustomRemoveCurveRedundantKeys(const bool& AttributeValue);
+
+	
+
+
+	/*********************************************************************************************************
+	 * Morph target curve API Begin
+	 *
+	 * Note - Morpgh target curve payload is FRichCurve.
+	 */
+
+	/**
+	 * Get the custom attribute DeleteExistingMorphTargetCurves, return false if the attribute is not set.
+	 * 
+	 * Note - If true, all previous moprh target curves will be deleted when doing a re-import.
+	 */
+	UFUNCTION(BlueprintCallable, Category = "Interchange | Node | AnimSequence")
+	bool GetCustomDeleteExistingMorphTargetCurves(bool& AttributeValue) const;
+
+	/**
+	 * Set the custom attribute DeleteExistingMorphTargetCurves. Return false if the attribute cannot be set.
+	 * 
+	 * Note - If true, all previous moprh target curves will be deleted when doing a re-import.
+	 */
+	UFUNCTION(BlueprintCallable, Category = "Interchange | Node | AnimSequence")
+	bool SetCustomDeleteExistingMorphTargetCurves(const bool& AttributeValue);
+
+	/** return how many animated moprh target (moprh target are translated into a mesh node and can be animated with a float curve) this anim sequence depends on. */
+	UFUNCTION(BlueprintCallable, Category = "Interchange | Node | AnimSequence")
+	int32 GetAnimatedMorphTargetDependeciesCount() const;
+
+	/** Get all animated moprh target unique id (moprh target are translated into a mesh node and can be animated with a float curve) this anim sequence depends on. */
+	UFUNCTION(BlueprintCallable, Category = "Interchange | Node | AnimSequence")
+	void GetAnimatedMorphTargetDependencies(TArray<FString>& OutDependencies) const;
+
+	/** Get an animated moprh target unique id (moprh target are translated into a mesh node and can be animated with a float curve) point by the specified index. */
+	UFUNCTION(BlueprintCallable, Category = "Interchange | Node | AnimSequence")
+	void GetAnimatedMorphTargetDependency(const int32 Index, FString& OutDependency) const;
+
+	/** Add an animated moprh target unique id (moprh target are translated into a mesh node and can be animated with a float curve). */
+	UFUNCTION(BlueprintCallable, Category = "Interchange | Node | AnimSequence")
+	bool SetAnimatedMorphTargetDependencyUid(const FString& DependencyUid);
+
+	/** Remove one animated moprh target unique id (moprh target are translated into a mesh node and can be animated with a float curve). Return false if we cannot remove the moprh target unique id. */
+	UFUNCTION(BlueprintCallable, Category = "Interchange | Node | AnimSequence")
+	bool RemoveAnimatedMorphTargetDependencyUid(const FString& DependencyUid);
+
+	/**
+	 * Morph target curve API End
+	 *********************************************************************************************************/
+
+
+	
+	/*********************************************************************************************************
+	 * Attribute curve API Begin
+	 * 
+	 * Note - Attribute curve payload information can be retrieve via the UInterchangeUserDefinedAttributesAPI.
+	 *        Attribute curves are import has float FRichCurve
+	 */
+
+	 /** Return how many animated attribute curve names this anim sequence drive (curve are FRichCurve of type float). */
+	UFUNCTION(BlueprintCallable, Category = "Interchange | Node | AnimSequence")
+	int32 GetAnimatedAttributeCurveNamesCount() const;
+
+	/** Get all animated attribute curve names. */
+	UFUNCTION(BlueprintCallable, Category = "Interchange | Node | AnimSequence")
+	void GetAnimatedAttributeCurveNames(TArray<FString>& OutAttributeCurveNames) const;
+
+	/** Get an animated attribute curve name point by the specified index. */
+	UFUNCTION(BlueprintCallable, Category = "Interchange | Node | AnimSequence")
+	void GetAnimatedAttributeCurveName(const int32 Index, FString& OutAttributeCurveName) const;
+
+	/** Add an animated attribute curve name. */
+	UFUNCTION(BlueprintCallable, Category = "Interchange | Node | AnimSequence")
+	bool SetAnimatedAttributeCurveName(const FString& AttributeCurveName);
+
+	/** Remove one animated attribute curve name. */
+	UFUNCTION(BlueprintCallable, Category = "Interchange | Node | AnimSequence")
+	bool RemoveAnimatedAttributeCurveName(const FString& AttributeCurveName);
+
+	/**
+	 * Attribute curve API End
+	 *********************************************************************************************************/
+
+
+
+	 /*********************************************************************************************************
+	  * Material curve API Begin
+	  *
+	  * Note - Material curve are attribute curve that can animate a material parameter.
+	  */
+
+	/**
+	 * Get the custom attribute MaterialDriveParameterOnCustomAttribute, return false if the attribute is not set.
+	 *
+	 * Note - If true, Set Material Curve Type for all custom attributes that exists.
+	 */
+	UFUNCTION(BlueprintCallable, Category = "Interchange | Node | AnimSequence")
+	bool GetCustomMaterialDriveParameterOnCustomAttribute(bool& AttributeValue) const;
+
+	/**
+	 * Set the custom attribute MaterialDriveParameterOnCustomAttribute. Return false if the attribute cannot be set.
+	 *
+	 * Note - If true, Set Material Curve Type for all custom attributes that exists.
+	 */
+	UFUNCTION(BlueprintCallable, Category = "Interchange | Node | AnimSequence")
+	bool SetCustomMaterialDriveParameterOnCustomAttribute(const bool& AttributeValue);
+
+	/** Return how many animated material curve suffixes this anim sequence drive (curve are FRichCurve of type float). */
+	UFUNCTION(BlueprintCallable, Category = "Interchange | Node | AnimSequence")
+	int32 GetAnimatedMaterialCurveSuffixesCount() const;
+
+	/** Get all animated material curve suffixes. */
+	UFUNCTION(BlueprintCallable, Category = "Interchange | Node | AnimSequence")
+	void GetAnimatedMaterialCurveSuffixes(TArray<FString>& OutMaterialCurveSuffixes) const;
+
+	/** Get an animated material curve suffixe point by the specified index. */
+	UFUNCTION(BlueprintCallable, Category = "Interchange | Node | AnimSequence")
+	void GetAnimatedMaterialCurveSuffixe(const int32 Index, FString& OutMaterialCurveSuffixe) const;
+
+	/** Add an animated material curve suffixe. */
+	UFUNCTION(BlueprintCallable, Category = "Interchange | Node | AnimSequence")
+	bool SetAnimatedMaterialCurveSuffixe(const FString& MaterialCurveSuffixe);
+
+	/** Remove one animated material curve suffixe. */
+	UFUNCTION(BlueprintCallable, Category = "Interchange | Node | AnimSequence")
+	bool RemoveAnimatedMaterialCurveSuffixe(const FString& MaterialCurveSuffixe);
+
+	/**
+	 * Material curve API End
+	 *********************************************************************************************************/
+
+
+
+	/*********************************************************************************************************
+	 * Attribute step curve API Begin
+	 * 
+	 * Note - Attribute step curve payload information can be retrieve via the UInterchangeUserDefinedAttributesAPI.
+	 *        Attribute step curves are import has a TArray<float> for key time and a TArray<ValueType> for the value.
+	 *        Supported value type are: int32, float, FString.
+	 */
+
+	 /** Return how many animated attribute step curve names this anim sequence drive. */
+	UFUNCTION(BlueprintCallable, Category = "Interchange | Node | AnimSequence")
+	int32 GetAnimatedAttributeStepCurveNamesCount() const;
+
+	/** Get all animated attribute step curve names. */
+	UFUNCTION(BlueprintCallable, Category = "Interchange | Node | AnimSequence")
+	void GetAnimatedAttributeStepCurveNames(TArray<FString>& OutAttributeStepCurveNames) const;
+
+	/** Get an animated attribute step curve name point by the specified index. */
+	UFUNCTION(BlueprintCallable, Category = "Interchange | Node | AnimSequence")
+	void GetAnimatedAttributeStepCurveName(const int32 Index, FString& OutAttributeStepCurveName) const;
+
+	/** Add an animated attribute step curve name. */
+	UFUNCTION(BlueprintCallable, Category = "Interchange | Node | AnimSequence")
+	bool SetAnimatedAttributeStepCurveName(const FString& AttributeStepCurveName);
+
+	/** Remove one animated attribute step curve name. */
+	UFUNCTION(BlueprintCallable, Category = "Interchange | Node | AnimSequence")
+	bool RemoveAnimatedAttributeStepCurveName(const FString& AttributeStepCurveName);
+
+	/**
+	 * Attribute step curve API End
+	 *********************************************************************************************************/
+
+
+
+	/**
+	 * Get the custom attribute DeleteExistingCustomAttributeCurves, return false if the attribute is not set.
+	 * 
+	 * Note - If true, all previous custom attribute curves will be deleted when doing a re-import.
+	 */
+	UFUNCTION(BlueprintCallable, Category = "Interchange | Node | AnimSequence")
+	bool GetCustomDeleteExistingCustomAttributeCurves(bool& AttributeValue) const;
+
+	/**
+	 * Set the custom attribute DeleteExistingCustomAttributeCurves. Return false if the attribute cannot be set.
+	 * 
+	 * Note - If true, all previous custom attribute curves will be deleted when doing a re-import.
+	 */
+	UFUNCTION(BlueprintCallable, Category = "Interchange | Node | AnimSequence")
+	bool SetCustomDeleteExistingCustomAttributeCurves(const bool& AttributeValue);
+	
+	/*
+	 * Curves API end
+	 **********************************************************************************************/
+
+
+	/**
+	 * Get the custom attribute DeleteExistingNonCurveCustomAttributes, return false if the attribute is not set.
+	 * 
+	 * Note - If true, all previous non-curve custom attributes will be deleted when doing a re-import.
+	 */
+	UFUNCTION(BlueprintCallable, Category = "Interchange | Node | AnimSequence")
+	bool GetCustomDeleteExistingNonCurveCustomAttributes(bool& AttributeValue) const;
+
+	/**
+	 * Set the custom attribute DeleteExistingNonCurveCustomAttributes. Return false if the attribute cannot be set.
+	 * 
+	 * Note - If true, all previous non-curve custom attributes will be deleted when doing a re-import.
+	 */
+	UFUNCTION(BlueprintCallable, Category = "Interchange | Node | AnimSequence")
+	bool SetCustomDeleteExistingNonCurveCustomAttributes(const bool& AttributeValue);
+
 	/**
 	 * Query the optional existing USkeleton this anim must use. The animsequence factory will use this skeleton instead of the imported one
 	 * (from GetCustomSkeletonFactoryNodeUid) if this attribute is set and the skeleton pointer is valid.
@@ -99,11 +375,25 @@ public:
 
 private:
 	const UE::Interchange::FAttributeKey Macro_CustomSkeletonFactoryNodeUidKey = UE::Interchange::FAttributeKey(TEXT("SkeletonFactoryNodeUid"));
-	
+
 	const UE::Interchange::FAttributeKey Macro_CustomImportBoneTracksKey = UE::Interchange::FAttributeKey(TEXT("ImportBoneTracks"));
 	const UE::Interchange::FAttributeKey Macro_CustomImportBoneTracksSampleRateKey = UE::Interchange::FAttributeKey(TEXT("ImportBoneTracksSampleRate"));
 	const UE::Interchange::FAttributeKey Macro_CustomImportBoneTracksRangeStartKey = UE::Interchange::FAttributeKey(TEXT("ImportBoneTracksRangeStart"));
 	const UE::Interchange::FAttributeKey Macro_CustomImportBoneTracksRangeStopKey = UE::Interchange::FAttributeKey(TEXT("ImportBoneTracksRangeStop"));
 
+	const UE::Interchange::FAttributeKey Macro_CustomImportAttributeCurvesKey = UE::Interchange::FAttributeKey(TEXT("ImportAttributeCurves"));
+	const UE::Interchange::FAttributeKey Macro_CustomDoNotImportCurveWithZeroKey = UE::Interchange::FAttributeKey(TEXT("DoNotImportCurveWithZero"));
+	const UE::Interchange::FAttributeKey Macro_CustomRemoveCurveRedundantKeysKey = UE::Interchange::FAttributeKey(TEXT("RemoveCurveRedundantKeys"));
+	const UE::Interchange::FAttributeKey Macro_CustomMaterialDriveParameterOnCustomAttributeKey = UE::Interchange::FAttributeKey(TEXT("MaterialDriveParameterOnCustomAttribute"));
+	const UE::Interchange::FAttributeKey Macro_CustomDeleteExistingMorphTargetCurvesKey = UE::Interchange::FAttributeKey(TEXT("DeleteExistingMorphTargetCurves"));
+	const UE::Interchange::FAttributeKey Macro_CustomDeleteExistingCustomAttributeCurvesKey = UE::Interchange::FAttributeKey(TEXT("DeleteExistingCustomAttributeCurves"));
+
+	const UE::Interchange::FAttributeKey Macro_CustomDeleteExistingNonCurveCustomAttributesKey = UE::Interchange::FAttributeKey(TEXT("DeleteExistingNonCurveCustomAttributes"));
+
 	const UE::Interchange::FAttributeKey Macro_CustomSkeletonSoftObjectPathKey = UE::Interchange::FAttributeKey(TEXT("SkeletonSoftObjectPath"));
+
+	UE::Interchange::TArrayAttributeHelper<FString> AnimatedMorphTargetDependencies;
+	UE::Interchange::TArrayAttributeHelper<FString> AnimatedAttributeCurveNames;
+	UE::Interchange::TArrayAttributeHelper<FString> AnimatedAttributeStepCurveNames;
+	UE::Interchange::TArrayAttributeHelper<FString> AnimatedMaterialCurveSuffixes;
 };
diff --git a/Engine/Plugins/Experimental/Interchange/Runtime/Source/Import/Private/Animation/InterchangeAnimSequenceFactory.cpp b/Engine/Plugins/Experimental/Interchange/Runtime/Source/Import/Private/Animation/InterchangeAnimSequenceFactory.cpp
--- a/Engine/Plugins/Experimental/Interchange/Runtime/Source/Import/Private/Animation/InterchangeAnimSequenceFactory.cpp
+++ b/Engine/Plugins/Experimental/Interchange/Runtime/Source/Import/Private/Animation/InterchangeAnimSequenceFactory.cpp
@@ -2,6 +2,7 @@
 #include "Animation/InterchangeAnimSequenceFactory.h"
 
 #include "Animation/AnimSequence.h"
+#include "Animation/BuiltInAttributeTypes.h"
 #include "Animation/InterchangeAnimationPayload.h"
 #include "Animation/InterchangeAnimationPayloadInterface.h"
 #include "InterchangeCommonPipelineDataFactoryNode.h"
@@ -9,14 +10,17 @@
 #include "InterchangeCommonAnimationPayload.h"
 #include "InterchangeImportCommon.h"
 #include "InterchangeImportLog.h"
+#include "InterchangeMeshNode.h"
 #include "InterchangeSceneNode.h"
 #include "InterchangeAnimSequenceFactoryNode.h"
 #include "InterchangeSkeletonFactoryNode.h"
 #include "InterchangeSourceData.h"
 #include "InterchangeTranslatorBase.h"
+#include "Nodes/InterchangeAnimationAPI.h"
 #include "Nodes/InterchangeBaseNode.h"
 #include "Nodes/InterchangeBaseNodeContainer.h"
 #include "Nodes/InterchangeSourceNode.h"
+#include "Nodes/InterchangeUserDefinedAttribute.h"
 
 #if WITH_EDITORONLY_DATA
 
@@ -36,14 +40,229 @@ namespace UE::Interchange::Private
 			GetSkeletonSceneNodeFlatListRecursive(NodeContainer, ChildUid, SkeletonSceneNodeUids);
 		}
 	}
+	
+	template<typename ValueType>
+	bool AreAllValuesZero(const TArray<ValueType>& Values, TFunctionRef<bool(const ValueType& Value)> CompareValueToZeroFunction)
+	{
+		const int32 KeyCount = Values.Num();
+		for (int32 KeyIndex = 1; KeyIndex < KeyCount; ++KeyIndex)
+		{
+			//Only add the not equal keys
+			if (!CompareValueToZeroFunction(Values[KeyIndex]))
+			{
+				return false;
+			}
+		}
+		return true;
+	}
+
+	bool CreateAttributeStepCurve(UAnimSequence* TargetSequence
+		, FAnimationStepCurvePayloadData& StepCurvePayload
+		, const FString& CurveName
+		, const FString& BoneName
+		, const int32 CurveFlags
+		, const bool bDoNotImportCurveWithZero)
+	{
+		//For bone attribute we support only single curve type (structured type like vector are not allowed)
+		if (!TargetSequence || StepCurvePayload.StepCurves.Num() != 1 || CurveName.IsEmpty())
+		{
+			return false;
+		}
+
+		if (bDoNotImportCurveWithZero)
+		{
+			bool bAllCurveValuesZero = true;
+			for (const FInterchangeStepCurve& StepCurve : StepCurvePayload.StepCurves)
+			{
+				if (StepCurve.FloatKeyValues.IsSet())
+				{
+					if (!AreAllValuesZero<float>(StepCurve.FloatKeyValues.GetValue(), [](const float& Value)
+						{
+							return FMath::IsNearlyZero(Value);
+						}))
+					{
+						bAllCurveValuesZero = false;
+						break;
+					}
+				}
+				else if (StepCurve.IntegerKeyValues.IsSet())
+				{
+					if (!AreAllValuesZero<int32>(StepCurve.IntegerKeyValues.GetValue(), [](const int32& Value)
+						{
+							return (Value == 0);
+						}))
+					{
+						bAllCurveValuesZero = false;
+						break;
+					}
+				}
+				else if (StepCurve.StringKeyValues.IsSet())
+				{
+					if (!AreAllValuesZero<FString>(StepCurve.StringKeyValues.GetValue(), [](const FString& Value)
+						{
+							return Value.IsEmpty();
+						}))
+					{
+						bAllCurveValuesZero = false;
+						break;
+					}
+				}
+			}
+			if (bAllCurveValuesZero)
+			{
+				return false;
+			}
+		}
+
+		bool bResult = false;
+		for (const FInterchangeStepCurve& StepCurve : StepCurvePayload.StepCurves)
+		{
+			if (StepCurve.FloatKeyValues.IsSet())
+			{
+				bResult |= UE::Anim::AddTypedCustomAttribute<FFloatAnimationAttribute, float>(FName(CurveName)
+					, FName(BoneName)
+					, TargetSequence
+					, MakeArrayView(StepCurve.KeyTimes)
+					, MakeArrayView(StepCurve.FloatKeyValues.GetValue()));
+			}
+			else if (StepCurve.IntegerKeyValues.IsSet())
+			{
+				bResult |= UE::Anim::AddTypedCustomAttribute<FIntegerAnimationAttribute, int32>(FName(CurveName)
+					, FName(BoneName)
+					, TargetSequence
+					, MakeArrayView(StepCurve.KeyTimes)
+					, MakeArrayView(StepCurve.IntegerKeyValues.GetValue()));
+			}
+			else if (StepCurve.StringKeyValues.IsSet())
+			{
+				bResult |= UE::Anim::AddTypedCustomAttribute<FStringAnimationAttribute, FString>(FName(CurveName)
+					, FName(BoneName)
+					, TargetSequence
+					, MakeArrayView(StepCurve.KeyTimes)
+					, MakeArrayView(StepCurve.StringKeyValues.GetValue()));
+			}
+		}
+		return bResult;
+	}
+
+	bool InternalCreateCurve(UAnimSequence* TargetSequence
+		, FAnimationCurvePayloadData& CurvePayload
+		, const FString& CurveName
+		, const int32 CurveFlags
+		, const bool bDoNotImportCurveWithZero
+		, const bool bMorphTargetCurve
+		, const bool bMaterialCurve)
+	{
+		if (!TargetSequence || CurvePayload.Curves.Num() != 1 || CurveName.IsEmpty())
+		{
+			return false;
+		}
+
+		if (bDoNotImportCurveWithZero)
+		{
+			bool bAllCurveValueAreZero = true;
+			for (const FRichCurve& Curve : CurvePayload.Curves)
+			{
+				FKeyHandle KeyHandle = Curve.GetFirstKeyHandle();
+				while (KeyHandle != FKeyHandle::Invalid())
+				{
+					if (!FMath::IsNearlyZero(Curve.GetKeyValue(KeyHandle)))
+					{
+						bAllCurveValueAreZero = false;
+						break;
+					}
+					KeyHandle = Curve.GetNextKey(KeyHandle);
+				}
+			}
+			if (bAllCurveValueAreZero)
+			{
+				//Avoid importing morph target curve with only zero value
+				return false;
+			}
+		}
+		FName Name = *CurveName;
+		USkeleton* Skeleton = TargetSequence->GetSkeleton();
+		const FSmartNameMapping* NameMapping = Skeleton->GetSmartNameContainer(USkeleton::AnimCurveMappingName);
+
+		// Add or retrieve curve
+		if (!NameMapping->Exists(Name))
+		{
+			// mark skeleton dirty
+			Skeleton->Modify();
+		}
+
+		FSmartName NewName;
+		Skeleton->AddSmartNameAndModify(USkeleton::AnimCurveMappingName, Name, NewName);
+
+		FAnimationCurveIdentifier FloatCurveId(NewName, ERawCurveTrackTypes::RCT_Float);
+
+
+		UAnimDataModel* DataModel = TargetSequence->GetDataModel();
+		IAnimationDataController& Controller = TargetSequence->GetController();
+
+		const FFloatCurve* TargetCurve = DataModel->FindFloatCurve(FloatCurveId);
+		if (TargetCurve == nullptr)
+		{
+			// Need to add the curve first
+			Controller.AddCurve(FloatCurveId, AACF_DefaultCurve | CurveFlags);
+			TargetCurve = DataModel->FindFloatCurve(FloatCurveId);
+		}
+		else
+		{
+			// Need to update any of the flags
+			Controller.SetCurveFlags(FloatCurveId, CurveFlags | TargetCurve->GetCurveTypeFlags());
+		}
+
+		// Should be valid at this point
+		ensure(TargetCurve);
+
+		Controller.UpdateCurveNamesFromSkeleton(Skeleton, ERawCurveTrackTypes::RCT_Float);
+
+		//MorphTarget curves are shift to 0 second
+		const float CurveStartTime = CurvePayload.Curves[0].GetFirstKey().Time;
+		if (CurveStartTime > UE_SMALL_NUMBER)
+		{
+			CurvePayload.Curves[0].ShiftCurve(-CurveStartTime);
+		}
+		// Set actual keys on curve within the model
+		Controller.SetCurveKeys(FloatCurveId, CurvePayload.Curves[0].GetConstRefOfKeys());
+
+		if (bMaterialCurve || bMorphTargetCurve)
+		{
+			Skeleton->AccumulateCurveMetaData(*CurveName, bMaterialCurve, bMorphTargetCurve);
+		}
+		return true;
+	}
+
+	bool CreateMorphTargetCurve(UAnimSequence* TargetSequence, FAnimationCurvePayloadData& CurvePayload, const FString& CurveName, int32 CurveFlags, bool bDoNotImportCurveWithZero)
+	{
+		constexpr bool bIsMorphTargetCurve = true;
+		constexpr bool bIsMaterialCurve = false;
+		return InternalCreateCurve(TargetSequence, CurvePayload, CurveName, CurveFlags, bDoNotImportCurveWithZero, bIsMorphTargetCurve, bIsMaterialCurve);
+	}
+
+	bool CreateMaterialCurve(UAnimSequence* TargetSequence, FAnimationCurvePayloadData& CurvePayload, const FString& CurveName, int32 CurveFlags, bool bDoNotImportCurveWithZero)
+	{
+		constexpr bool bIsMorphTargetCurve = false;
+		constexpr bool bIsMaterialCurve = true;
+		return InternalCreateCurve(TargetSequence, CurvePayload, CurveName, CurveFlags, bDoNotImportCurveWithZero, bIsMorphTargetCurve, bIsMaterialCurve);
+	}
+
+	bool CreateAttributeCurve(UAnimSequence* TargetSequence, FAnimationCurvePayloadData& CurvePayload, const FString& CurveName, int32 CurveFlags, bool bDoNotImportCurveWithZero)
+	{
+		//This curve don't animate morph target or material parameter.
+		constexpr bool bIsMorphTargetCurve = false;
+		constexpr bool bIsMaterialCurve = false;
+		return InternalCreateCurve(TargetSequence, CurvePayload, CurveName, CurveFlags, bDoNotImportCurveWithZero, bIsMorphTargetCurve, bIsMaterialCurve);
+	}
 
 	void RetrieveAnimationPayloads(UAnimSequence* AnimSequence
 		, const UInterchangeAnimSequenceFactoryNode* AnimSequenceFactoryNode
 		, const UInterchangeBaseNodeContainer* NodeContainer
 		, const UInterchangeSkeletonFactoryNode* SkeletonFactoryNode
-		, const USkeleton* Skeleton
 		, const IInterchangeAnimationPayloadInterface* AnimSequenceTranslatorPayloadInterface
-		, const FString& AssetName)
+		, const FString& AssetName
+		, TArray<FString> OutCurvesNotFound)
 	{
 		TMap<FString, TFuture<TOptional<UE::Interchange::FAnimationBakeTransformPayloadData>>> AnimationPayloads;
 
@@ -53,6 +272,14 @@ namespace UE::Interchange::Private
 			//Cannot import animation without a skeleton
 			return;
 		}
+		
+		IAnimationDataController& Controller = AnimSequence->GetController();
+		USkeleton* Skeleton = AnimSequence->GetSkeleton();
+		check(Skeleton);
+
+		TArray<FString> SkeletonNodes;
+		GetSkeletonSceneNodeFlatListRecursive(NodeContainer, SkeletonRootUid, SkeletonNodes);
+
 		bool bImportBoneTracks = false;
 		AnimSequenceFactoryNode->GetCustomImportBoneTracks(bImportBoneTracks);
 		if (bImportBoneTracks)
@@ -71,21 +298,18 @@ namespace UE::Interchange::Private
 			const double SequenceLength = FMath::Max<double>(RangeEnd - RangeStart, MINIMUM_ANIMATION_LENGTH);
 			int32 BakeKeyCount = (SequenceLength / BakeInterval) + 1;
 
-			TArray<FString> SkeletonNodes;
-			GetSkeletonSceneNodeFlatListRecursive(NodeContainer, SkeletonRootUid, SkeletonNodes);
 			for (const FString& NodeUid : SkeletonNodes)
 			{
 				if (const UInterchangeSceneNode* SkeletonSceneNode = Cast<UInterchangeSceneNode>(NodeContainer->GetNode(NodeUid)))
 				{
 					FString PayloadKey;
-					if (SkeletonSceneNode->GetCustomTransformCurvePayloadKey(PayloadKey))
+					if (UInterchangeAnimationAPI::GetCustomNodeTransformPayloadKey(SkeletonSceneNode, PayloadKey))
 					{
 						AnimationPayloads.Add(PayloadKey, AnimSequenceTranslatorPayloadInterface->GetAnimationBakeTransformPayloadData(PayloadKey, SampleRate, RangeStart, RangeEnd));
 					}
 				}
 			}
 
-			IAnimationDataController& Controller = AnimSequence->GetController();
 			//This destroy all previously imported animation raw data
 			Controller.RemoveAllBoneTracks();
 			Controller.SetPlayLength(FGenericPlatformMath::Max<float>(SequenceLength, MINIMUM_ANIMATION_LENGTH));
@@ -115,7 +339,7 @@ namespace UE::Interchange::Private
 					bool bApplyGlobalOffset = NodeUid.Equals(SkeletonRootUid);
 
 					FString PayloadKey;
-					if (SkeletonSceneNode->GetCustomTransformCurvePayloadKey(PayloadKey))
+					if (UInterchangeAnimationAPI::GetCustomNodeTransformPayloadKey(SkeletonSceneNode, PayloadKey))
 					{
 						TOptional<UE::Interchange::FAnimationBakeTransformPayloadData> AnimationTransformPayload = AnimationPayloads.FindChecked(PayloadKey).Get();
 						if (!AnimationTransformPayload.IsSet())
@@ -180,6 +404,298 @@ namespace UE::Interchange::Private
 				}
 			}
 		}
+
+		bool bDeleteExistingMorphTargetCurves = false;
+		AnimSequenceFactoryNode->GetCustomDeleteExistingMorphTargetCurves(bDeleteExistingMorphTargetCurves);
+		bool bDeleteExistingCustomAttributeCurves = false;
+		AnimSequenceFactoryNode->GetCustomDeleteExistingCustomAttributeCurves(bDeleteExistingCustomAttributeCurves);
+		bool bDeleteExistingNonCurveCustomAttributes = false;
+		AnimSequenceFactoryNode->GetCustomDeleteExistingNonCurveCustomAttributes(bDeleteExistingNonCurveCustomAttributes);
+		if (bDeleteExistingMorphTargetCurves || bDeleteExistingCustomAttributeCurves)
+		{
+			TArray<FSmartName> CurveSmartNamesToRemove;
+			for (const FFloatCurve& Curve : AnimSequence->GetDataModel()->GetFloatCurves())
+			{
+				const FCurveMetaData* MetaData = Skeleton->GetCurveMetaData(Curve.Name);
+				if (MetaData)
+				{
+					bool bDeleteCurve = MetaData->Type.bMorphtarget ? bDeleteExistingMorphTargetCurves : bDeleteExistingCustomAttributeCurves;
+					if (bDeleteCurve)
+					{
+						CurveSmartNamesToRemove.Add(Curve.Name);
+					}
+				}
+			}
+
+			for (const FSmartName& CurveName : CurveSmartNamesToRemove)
+			{
+				const FAnimationCurveIdentifier CurveId(CurveName, ERawCurveTrackTypes::RCT_Float);
+				Controller.RemoveCurve(CurveId);
+			}
+		}
+
+		if (bDeleteExistingNonCurveCustomAttributes)
+		{
+			Controller.RemoveAllAttributes();
+		}
+
+		
+		bool bImportAttributeCurves = false;
+		AnimSequenceFactoryNode->GetCustomImportAttributeCurves(bImportAttributeCurves);
+		if (bImportAttributeCurves)
+		{
+			const UAnimDataModel* DataModel = AnimSequence->GetDataModel();
+			const int32 NumFloatCurves = DataModel->GetNumberOfFloatCurves();
+			const FAnimationCurveData& CurveData = DataModel->GetCurveData();
+
+			OutCurvesNotFound.Reset(NumFloatCurves);
+
+			for (const FFloatCurve& FloatCurve : CurveData.FloatCurves)
+			{
+				const FCurveMetaData* MetaData = Skeleton->GetCurveMetaData(FloatCurve.Name);
+
+				if (MetaData && !MetaData->Type.bMorphtarget)
+				{
+					OutCurvesNotFound.Add(FloatCurve.Name.DisplayName.ToString());
+				}
+			}
+
+			bool bMaterialDriveParameterOnCustomAttribute = false;
+			AnimSequenceFactoryNode->GetCustomMaterialDriveParameterOnCustomAttribute(bMaterialDriveParameterOnCustomAttribute);
+			TArray<FString> MaterialSuffixes;
+			AnimSequenceFactoryNode->GetAnimatedMaterialCurveSuffixes(MaterialSuffixes);
+			bool bDoNotImportCurveWithZero = false;
+			AnimSequenceFactoryNode->GetCustomDoNotImportCurveWithZero(bDoNotImportCurveWithZero);
+			bool bRemoveCurveRedundantKeys = false;
+			AnimSequenceFactoryNode->GetCustomRemoveCurveRedundantKeys(bRemoveCurveRedundantKeys);
+
+			auto IsCurveHookToMaterial = [&MaterialSuffixes](const FString& CurveName)
+			{
+				for (const FString& MaterialSuffixe : MaterialSuffixes)
+				{
+					if (CurveName.EndsWith(MaterialSuffixe))
+					{
+						return true;
+					}
+				}
+				return false;
+			};
+
+			//Import morph target curves
+			{
+				TMap<FString, TFuture<TOptional<UE::Interchange::FAnimationCurvePayloadData>>> AnimationCurvesPayloads;
+				TMap<FString, FString> AnimationCurveMorphTargetNodeNames;
+				//Import morph target curves (FRichCurve is what anim api )
+				TArray<FString> AnimatedMorphTargetUids;
+				AnimSequenceFactoryNode->GetAnimatedMorphTargetDependencies(AnimatedMorphTargetUids);
+				for (const FString& MorphTargetUid : AnimatedMorphTargetUids)
+				{
+					if (const UInterchangeMeshNode* MorphTargetNode = Cast<UInterchangeMeshNode>(NodeContainer->GetNode(MorphTargetUid)))
+					{
+						TOptional<FString> PayloadKey = MorphTargetNode->GetAnimationCurvePayLoadKey();
+						if (PayloadKey.IsSet())
+						{
+							AnimationCurvesPayloads.Add(PayloadKey.GetValue(), AnimSequenceTranslatorPayloadInterface->GetAnimationCurvePayloadData(PayloadKey.GetValue()));
+							AnimationCurveMorphTargetNodeNames.Add(PayloadKey.GetValue(), MorphTargetNode->GetDisplayLabel());
+						}
+					}
+				}
+				for (TPair<FString, TFuture<TOptional<UE::Interchange::FAnimationCurvePayloadData>>>& CurveNameAndPayload : AnimationCurvesPayloads)
+				{
+					TOptional<UE::Interchange::FAnimationCurvePayloadData> AnimationCurvePayload = CurveNameAndPayload.Value.Get();
+					if (!AnimationCurvePayload.IsSet())
+					{
+						UE_LOG(LogInterchangeImport, Warning, TEXT("Invalid animation morph target curve payload key [%s] AnimSequence asset %s"), *CurveNameAndPayload.Key, *AssetName);
+						continue;
+					}
+					FAnimationCurvePayloadData& CurvePayload = AnimationCurvePayload.GetValue();
+					if (bRemoveCurveRedundantKeys)
+					{
+						for (FRichCurve& RichCurve : CurvePayload.Curves)
+						{
+							RichCurve.RemoveRedundantAutoTangentKeys(SMALL_NUMBER);
+						}
+					}
+					constexpr int32 CurveFlags = 0;
+					CreateMorphTargetCurve(AnimSequence, CurvePayload, AnimationCurveMorphTargetNodeNames.FindChecked(CurveNameAndPayload.Key), CurveFlags, bDoNotImportCurveWithZero);
+				}
+			}
+
+			//Import Attribute curves
+			{
+				//Utility to make sure the curve is compatible with FRichCurve
+				auto IsDecimalType = [](UE::Interchange::EAttributeTypes Type)
+				{
+					switch (Type)
+					{
+					case UE::Interchange::EAttributeTypes::Double:
+					case UE::Interchange::EAttributeTypes::Float:
+					case UE::Interchange::EAttributeTypes::Float16:
+					case UE::Interchange::EAttributeTypes::Vector2d:
+					case UE::Interchange::EAttributeTypes::Vector2f:
+					case UE::Interchange::EAttributeTypes::Vector3d:
+					case UE::Interchange::EAttributeTypes::Vector3f:
+					case UE::Interchange::EAttributeTypes::Vector4d:
+					case UE::Interchange::EAttributeTypes::Vector4f:
+						return true;
+					}
+					return false;
+				};
+				//Import Attribute curves (FRichCurve)
+				{
+					TArray<FString> AttributeCurveNames;
+					AnimSequenceFactoryNode->GetAnimatedAttributeCurveNames(AttributeCurveNames);
+					TMap<FString, TFuture<TOptional<UE::Interchange::FAnimationCurvePayloadData>>> AnimationCurvesPayloads;
+					for (const FString& NodeUid : SkeletonNodes)
+					{
+						if (const UInterchangeSceneNode* SkeletonSceneNode = Cast<UInterchangeSceneNode>(NodeContainer->GetNode(NodeUid)))
+						{
+							//Import material parameter curves (FRichCurve)
+							TMap<FString, FString> CurveNamePayloads;
+							TArray<FInterchangeUserDefinedAttributeInfo> AttributeInfos = UInterchangeUserDefinedAttributesAPI::GetUserDefinedAttributeInfos(SkeletonSceneNode);
+							for (const FInterchangeUserDefinedAttributeInfo& AttributeInfo : AttributeInfos)
+							{
+								//Material curve must be convertible to float since we need a FRichCurve
+								if (AttributeInfo.PayloadKey.IsSet() && IsDecimalType(AttributeInfo.Type) && AttributeCurveNames.Contains(AttributeInfo.Name))
+								{
+									CurveNamePayloads.Add(AttributeInfo.Name, AttributeInfo.PayloadKey.GetValue());
+								}
+							}
+							for (const TPair<FString, FString>& CurveNamePayload : CurveNamePayloads)
+							{
+								AnimationCurvesPayloads.Add(CurveNamePayload.Key, AnimSequenceTranslatorPayloadInterface->GetAnimationCurvePayloadData(CurveNamePayload.Value));
+							}
+						}
+					}
+
+					for (TPair<FString, TFuture<TOptional<UE::Interchange::FAnimationCurvePayloadData>>>& CurveNameAndPayload : AnimationCurvesPayloads)
+					{
+						const FString& CurveName = CurveNameAndPayload.Key;
+						TOptional<UE::Interchange::FAnimationCurvePayloadData> AnimationCurvePayload = CurveNameAndPayload.Value.Get();
+						if (!AnimationCurvePayload.IsSet())
+						{
+							UE_LOG(LogInterchangeImport, Warning, TEXT("Invalid animation material curve payload key [%s] AnimSequence asset %s"), *CurveName, *AssetName);
+							continue;
+						}
+
+						FAnimationCurvePayloadData& CurvePayload = AnimationCurvePayload.GetValue();
+						if (bRemoveCurveRedundantKeys)
+						{
+							for (FRichCurve& RichCurve : CurvePayload.Curves)
+							{
+								RichCurve.RemoveRedundantAutoTangentKeys(SMALL_NUMBER);
+							}
+						}
+						OutCurvesNotFound.Remove(CurveName);
+						constexpr int32 CurveFlags = 0;
+						if (bMaterialDriveParameterOnCustomAttribute || IsCurveHookToMaterial(CurveName))
+						{
+							CreateMaterialCurve(AnimSequence, CurvePayload, CurveName, CurveFlags, bDoNotImportCurveWithZero);
+						}
+						else
+						{
+							CreateAttributeCurve(AnimSequence, CurvePayload, CurveName, CurveFlags, bDoNotImportCurveWithZero);
+						}
+					}
+				}
+
+				//Import attribute step curves (step curves, only time and value array, no FRichCurve. The anim API expect the value array is of the type of the translated attribute)
+				{
+					TArray<FString> AttributeStepCurveNames;
+					AnimSequenceFactoryNode->GetAnimatedAttributeStepCurveNames(AttributeStepCurveNames);
+					TMap<FString, TFuture<TOptional<UE::Interchange::FAnimationCurvePayloadData>>> AnimationCurvesPayloads;
+					TMap<FString, TFuture<TOptional<UE::Interchange::FAnimationStepCurvePayloadData>>> AnimationStepCurvesPayloads;
+					TMap<FString, FString> AnimationBoneNames;
+					for (const FString& NodeUid : SkeletonNodes)
+					{
+						if (const UInterchangeSceneNode* SkeletonSceneNode = Cast<UInterchangeSceneNode>(NodeContainer->GetNode(NodeUid)))
+						{
+							FString BoneName = SkeletonSceneNode->GetDisplayLabel();
+							//Import material parameter curves (FRichCurve)
+							TMap<FString, FString> CurveNameFloatPayloads;
+							TMap<FString, FString> CurveNameStepCurvePayloads;
+							TArray<FInterchangeUserDefinedAttributeInfo> AttributeInfos = UInterchangeUserDefinedAttributesAPI::GetUserDefinedAttributeInfos(SkeletonSceneNode);
+							for (const FInterchangeUserDefinedAttributeInfo& AttributeInfo : AttributeInfos)
+							{
+								//Material curve must be convertible to float since we need a FRichCurve
+								if (AttributeInfo.PayloadKey.IsSet() && AttributeStepCurveNames.Contains(AttributeInfo.Name))
+								{
+									if (IsDecimalType(AttributeInfo.Type))
+									{
+										CurveNameFloatPayloads.Add(AttributeInfo.Name, AttributeInfo.PayloadKey.GetValue());
+									}
+									else
+									{
+										CurveNameStepCurvePayloads.Add(AttributeInfo.Name, AttributeInfo.PayloadKey.GetValue());
+									}
+								}
+							}
+							for (const TPair<FString, FString>& CurveNamePayload : CurveNameFloatPayloads)
+							{
+								AnimationCurvesPayloads.Add(CurveNamePayload.Key, AnimSequenceTranslatorPayloadInterface->GetAnimationCurvePayloadData(CurveNamePayload.Value));
+								AnimationBoneNames.Add(CurveNamePayload.Key, BoneName);
+							}
+							for (const TPair<FString, FString>& StepCurveNamePayload : CurveNameStepCurvePayloads)
+							{
+								AnimationStepCurvesPayloads.Add(StepCurveNamePayload.Key, AnimSequenceTranslatorPayloadInterface->GetAnimationStepCurvePayloadData(StepCurveNamePayload.Value));
+								AnimationBoneNames.Add(StepCurveNamePayload.Key, BoneName);
+							}
+						}
+					}
+
+					auto AddAttributeCurveToAnimSequence = [bRemoveCurveRedundantKeys, bDoNotImportCurveWithZero, &AnimSequence](FAnimationStepCurvePayloadData& StepCurvePayload, const FString& CurveName, const FString& BoneName)
+					{
+						if (bRemoveCurveRedundantKeys)
+						{
+							for (FInterchangeStepCurve& StepCurve : StepCurvePayload.StepCurves)
+							{
+								StepCurve.RemoveRedundantKeys(SMALL_NUMBER);
+							}
+						}
+
+						constexpr int32 CurveFlags = 0;
+						CreateAttributeStepCurve(AnimSequence, StepCurvePayload, CurveName, BoneName, CurveFlags, bDoNotImportCurveWithZero);
+					};
+
+					for (TPair<FString, TFuture<TOptional<UE::Interchange::FAnimationCurvePayloadData>>>& CurveNameAndPayload : AnimationCurvesPayloads)
+					{
+						TOptional<UE::Interchange::FAnimationCurvePayloadData> AnimationCurvePayload = CurveNameAndPayload.Value.Get();
+						if (!AnimationCurvePayload.IsSet())
+						{
+							UE_LOG(LogInterchangeImport, Warning, TEXT("Invalid animation material curve payload key [%s] AnimSequence asset %s"), *CurveNameAndPayload.Key, *AssetName);
+							continue;
+						}
+						FAnimationCurvePayloadData& CurvePayload = AnimationCurvePayload.GetValue();
+						FAnimationStepCurvePayloadData StepCurvePayload = FAnimationStepCurvePayloadData::FromAnimationCurvePayloadData(CurvePayload);
+						AddAttributeCurveToAnimSequence(StepCurvePayload, CurveNameAndPayload.Key, AnimationBoneNames.FindChecked(CurveNameAndPayload.Key));
+					}
+
+					for (TPair<FString, TFuture<TOptional<UE::Interchange::FAnimationStepCurvePayloadData>>>& StepCurveNameAndPayload : AnimationStepCurvesPayloads)
+					{
+						TOptional<UE::Interchange::FAnimationStepCurvePayloadData> AnimationStepCurvePayload = StepCurveNameAndPayload.Value.Get();
+						if (!AnimationStepCurvePayload.IsSet())
+						{
+							UE_LOG(LogInterchangeImport, Warning, TEXT("Invalid animation material curve payload key [%s] AnimSequence asset %s"), *StepCurveNameAndPayload.Key, *AssetName);
+							continue;
+						}
+
+						AddAttributeCurveToAnimSequence(AnimationStepCurvePayload.GetValue(), StepCurveNameAndPayload.Key, AnimationBoneNames.FindChecked(StepCurveNameAndPayload.Key));
+					}
+				}
+			}
+		}
+		else
+		{
+			// Store float curve tracks which use to exist on the animation
+			for (const FFloatCurve& Curve : AnimSequence->GetDataModel()->GetFloatCurves())
+			{
+				const FCurveMetaData* MetaData = Skeleton->GetCurveMetaData(Curve.Name);
+				if (MetaData && !MetaData->Type.bMorphtarget)
+				{
+					OutCurvesNotFound.Add(Curve.Name.DisplayName.ToString());
+				}
+			}
+		}
 	}
 } //namespace UE::Interchange::Private
 #endif //WITH_EDITOR
@@ -381,14 +897,26 @@ UObject* UInterchangeAnimSequenceFactory::CreateAsset(const FCreateAssetParams&
 		AnimSequence->ImportResampleFramerate = SampleRate;
 		Controller.SetFrameRate(FrameRate);
 
+		TArray<FString> CurvesNotFound;
 		UE::Interchange::Private::RetrieveAnimationPayloads(AnimSequence
 			, AnimSequenceFactoryNode
 			, Arguments.NodeContainer
 			, SkeletonFactoryNode
-			, Skeleton
 			, AnimSequenceTranslatorPayloadInterface
-			, Arguments.AssetName);
+			, Arguments.AssetName
+			, CurvesNotFound);
 
+		if (CurvesNotFound.Num())
+		{
+			for (const FString& CurveName : CurvesNotFound)
+			{
+				UInterchangeResultError_Generic* Message = AddMessage<UInterchangeResultError_Generic>();
+				Message->SourceAssetName = Arguments.SourceData->GetFilename();
+				Message->DestinationAssetName = Arguments.AssetName;
+				Message->AssetType = UAnimSequence::StaticClass();
+				Message->Text = FText::Format(NSLOCTEXT("UInterchangeAnimSequenceFactory", "NonExistingCurve", "Curve ({0}) was not found in the new Animation."), FText::FromString(CurveName));
+			}
+		}
 		Controller.NotifyPopulated();
 		Controller.CloseBracket(false);
 	}
diff --git a/Engine/Plugins/Experimental/Interchange/Runtime/Source/Import/Private/Fbx/InterchangeFbxTranslator.cpp b/Engine/Plugins/Experimental/Interchange/Runtime/Source/Import/Private/Fbx/InterchangeFbxTranslator.cpp
--- a/Engine/Plugins/Experimental/Interchange/Runtime/Source/Import/Private/Fbx/InterchangeFbxTranslator.cpp
+++ b/Engine/Plugins/Experimental/Interchange/Runtime/Source/Import/Private/Fbx/InterchangeFbxTranslator.cpp
@@ -334,12 +334,12 @@ TFuture<TOptional<UE::Interchange::FSkeletalMeshLodPayloadData>> UInterchangeFbx
 	return Promise->GetFuture();
 }
 
-TFuture<TOptional<UE::Interchange::FSkeletalMeshBlendShapePayloadData>> UInterchangeFbxTranslator::GetSkeletalMeshBlendShapePayloadData(const FString& PayLoadKey) const
+TFuture<TOptional<UE::Interchange::FSkeletalMeshMorphTargetPayloadData>> UInterchangeFbxTranslator::GetSkeletalMeshMorphTargetPayloadData(const FString& PayLoadKey) const
 {
-	TSharedPtr<TPromise<TOptional<UE::Interchange::FSkeletalMeshBlendShapePayloadData>>> Promise = MakeShared<TPromise<TOptional<UE::Interchange::FSkeletalMeshBlendShapePayloadData>>>();
+	TSharedPtr<TPromise<TOptional<UE::Interchange::FSkeletalMeshMorphTargetPayloadData>>> Promise = MakeShared<TPromise<TOptional<UE::Interchange::FSkeletalMeshMorphTargetPayloadData>>>();
 	if (!Dispatcher.IsValid())
 	{
-		Promise->SetValue(TOptional<UE::Interchange::FSkeletalMeshBlendShapePayloadData>{});
+		Promise->SetValue(TOptional<UE::Interchange::FSkeletalMeshMorphTargetPayloadData>{});
 		return Promise->GetFuture();
 	}
 
@@ -360,7 +360,7 @@ TFuture<TOptional<UE::Interchange::FSkeletalMeshBlendShapePayloadData>> UInterch
 
 		if (TaskState != UE::Interchange::ETaskState::ProcessOk)
 		{
-			Promise->SetValue(TOptional<UE::Interchange::FSkeletalMeshBlendShapePayloadData>{});
+			Promise->SetValue(TOptional<UE::Interchange::FSkeletalMeshMorphTargetPayloadData>{});
 			return;
 		}
 		//Grab the result file and fill the BaseNodeContainer
@@ -371,7 +371,7 @@ TFuture<TOptional<UE::Interchange::FSkeletalMeshBlendShapePayloadData>> UInterch
 		if (!ensure(FPaths::FileExists(SkeletalMeshPayloadFilename)))
 		{
 			//TODO log an error saying the payload file do not exist even if the get payload command succeed
-			Promise->SetValue(TOptional<UE::Interchange::FSkeletalMeshBlendShapePayloadData>{});
+			Promise->SetValue(TOptional<UE::Interchange::FSkeletalMeshMorphTargetPayloadData>{});
 			return;
 		}
 		//All sub object should be gone with the reset
@@ -382,33 +382,33 @@ TFuture<TOptional<UE::Interchange::FSkeletalMeshBlendShapePayloadData>> UInterch
 		if (FileDataSize < 1)
 		{
 			//Nothing to load from this file
-			Promise->SetValue(TOptional<UE::Interchange::FSkeletalMeshBlendShapePayloadData>{});
+			Promise->SetValue(TOptional<UE::Interchange::FSkeletalMeshMorphTargetPayloadData>{});
 			return;
 		}
-		UE::Interchange::FSkeletalMeshBlendShapePayloadData SkeletalMeshBlendShapePayload;
-		SkeletalMeshBlendShapePayload.LodMeshDescription.Empty();
+		UE::Interchange::FSkeletalMeshMorphTargetPayloadData SkeletalMeshMorphTargetPayload;
+		SkeletalMeshMorphTargetPayload.LodMeshDescription.Empty();
 		//Buffer keep the ownership of the data, the large memory reader is use to serialize the TMap
 		FLargeMemoryReader Ar(FileData, FileDataSize);
-		SkeletalMeshBlendShapePayload.LodMeshDescription.Serialize(Ar);
-		Promise->SetValue(MoveTemp(SkeletalMeshBlendShapePayload));
+		SkeletalMeshMorphTargetPayload.LodMeshDescription.Serialize(Ar);
+		Promise->SetValue(MoveTemp(SkeletalMeshMorphTargetPayload));
 	}));
 
 	//The task was not added to the dispatcher
 	if (CreatedTaskIndex == INDEX_NONE)
 	{
-		Promise->SetValue(TOptional<UE::Interchange::FSkeletalMeshBlendShapePayloadData>{});
+		Promise->SetValue(TOptional<UE::Interchange::FSkeletalMeshMorphTargetPayloadData>{});
 	}
 
 	return Promise->GetFuture();
 }
 
-TFuture<TOptional<UE::Interchange::FAnimationTransformPayloadData>> UInterchangeFbxTranslator::GetAnimationTransformPayloadData(const FString& PayLoadKey) const
+TFuture<TOptional<UE::Interchange::FAnimationCurvePayloadData>> UInterchangeFbxTranslator::GetAnimationCurvePayloadData(const FString& PayLoadKey) const
 {
-	TSharedPtr<TPromise<TOptional<UE::Interchange::FAnimationTransformPayloadData>>> Promise = MakeShared<TPromise<TOptional<UE::Interchange::FAnimationTransformPayloadData>>>();
+	TSharedPtr<TPromise<TOptional<UE::Interchange::FAnimationCurvePayloadData>>> Promise = MakeShared<TPromise<TOptional<UE::Interchange::FAnimationCurvePayloadData>>>();
 
 	if (!Dispatcher.IsValid())
 	{
-		Promise->SetValue(TOptional<UE::Interchange::FAnimationTransformPayloadData>());
+		Promise->SetValue(TOptional<UE::Interchange::FAnimationCurvePayloadData>());
 		return Promise->GetFuture();
 	}
 
@@ -429,7 +429,7 @@ TFuture<TOptional<UE::Interchange::FAnimationTransformPayloadData>> UInterchange
 
 			if (TaskState != UE::Interchange::ETaskState::ProcessOk)
 			{
-				Promise->SetValue(TOptional<UE::Interchange::FAnimationTransformPayloadData>());
+				Promise->SetValue(TOptional<UE::Interchange::FAnimationCurvePayloadData>());
 				return;
 			}
 
@@ -441,11 +441,11 @@ TFuture<TOptional<UE::Interchange::FAnimationTransformPayloadData>> UInterchange
 			if (!ensure(FPaths::FileExists(AnimationTransformPayloadFilename)))
 			{
 				// TODO log an error saying the payload file does not exist even if the get payload command succeeded
-				Promise->SetValue(TOptional<UE::Interchange::FAnimationTransformPayloadData>());
+				Promise->SetValue(TOptional<UE::Interchange::FAnimationCurvePayloadData>());
 				return;
 			}
 
-			UE::Interchange::FAnimationTransformPayloadData AnimationTransformPayload;
+			UE::Interchange::FAnimationCurvePayloadData AnimationTransformPayload;
 
 			// All sub object should be gone with the reset
 			TArray64<uint8> Buffer;
@@ -455,7 +455,7 @@ TFuture<TOptional<UE::Interchange::FAnimationTransformPayloadData>> UInterchange
 			if (FileDataSize < 1)
 			{
 				// Nothing to load from this file
-				Promise->SetValue(TOptional<UE::Interchange::FAnimationTransformPayloadData>());
+				Promise->SetValue(TOptional<UE::Interchange::FAnimationCurvePayloadData>());
 				return;
 			}
 
@@ -463,16 +463,11 @@ TFuture<TOptional<UE::Interchange::FAnimationTransformPayloadData>> UInterchange
 			FLargeMemoryReader Ar(FileData, FileDataSize);
 			TArray<FInterchangeCurve> InterchangeCurves;
 			Ar << InterchangeCurves;
-			AnimationTransformPayload.TransformCurves.AddDefaulted(static_cast<int32>(EInterchangeTransformCurveChannel::TransformChannelCount));
-			for (const FInterchangeCurve& InterchangeCurve : InterchangeCurves)
+			AnimationTransformPayload.Curves.AddDefaulted(InterchangeCurves.Num());
+			for (int32 CurveIndex = 0; CurveIndex < InterchangeCurves.Num(); ++CurveIndex)
 			{
-				int32 TransformChannelIndex = static_cast<int32>(InterchangeCurve.TransformChannel);
-				if (AnimationTransformPayload.TransformCurves.IsValidIndex(TransformChannelIndex))
-				{
-					UE::Interchange::FAnimationCurveTransformPayloadData& CurveTransformData = AnimationTransformPayload.TransformCurves[static_cast<int32>(InterchangeCurve.TransformChannel)];
-					CurveTransformData.TransformChannel = InterchangeCurve.TransformChannel;
-					InterchangeCurve.ToRichCurve(CurveTransformData.Curve);
-				}
+				const FInterchangeCurve& InterchangeCurve = InterchangeCurves[CurveIndex];
+				InterchangeCurve.ToRichCurve(AnimationTransformPayload.Curves[CurveIndex]);
 			}
 			Promise->SetValue(MoveTemp(AnimationTransformPayload));
 		}));
@@ -480,7 +475,79 @@ TFuture<TOptional<UE::Interchange::FAnimationTransformPayloadData>> UInterchange
 	// The task was not added to the dispatcher
 	if (CreatedTaskIndex == INDEX_NONE)
 	{
-		Promise->SetValue(TOptional<UE::Interchange::FAnimationTransformPayloadData>{});
+		Promise->SetValue(TOptional<UE::Interchange::FAnimationCurvePayloadData>{});
+	}
+
+	return Promise->GetFuture();
+}
+
+TFuture<TOptional<UE::Interchange::FAnimationStepCurvePayloadData>> UInterchangeFbxTranslator::GetAnimationStepCurvePayloadData(const FString& PayLoadKey) const
+{
+	TSharedPtr<TPromise<TOptional<UE::Interchange::FAnimationStepCurvePayloadData>>> Promise = MakeShared<TPromise<TOptional<UE::Interchange::FAnimationStepCurvePayloadData>>>();
+
+	if (!Dispatcher.IsValid())
+	{
+		Promise->SetValue(TOptional<UE::Interchange::FAnimationStepCurvePayloadData>());
+		return Promise->GetFuture();
+	}
+
+	// Create a json command to read the fbx file
+	FString JsonCommand = CreateFetchPayloadFbxCommand(PayLoadKey);
+	const int32 CreatedTaskIndex = Dispatcher->AddTask(JsonCommand, FInterchangeDispatcherTaskCompleted::CreateLambda([this, Promise](const int32 TaskIndex)
+		{
+			UE::Interchange::ETaskState TaskState;
+			FString JsonResult;
+			TArray<FString> JsonMessages;
+			Dispatcher->GetTaskState(TaskIndex, TaskState, JsonResult, JsonMessages);
+
+			// Parse the Json messages into UInterchangeResults
+			for (const FString& JsonMessage : JsonMessages)
+			{
+				AddMessage(UInterchangeResult::FromJson(JsonMessage));
+			}
+
+			if (TaskState != UE::Interchange::ETaskState::ProcessOk)
+			{
+				Promise->SetValue(TOptional<UE::Interchange::FAnimationStepCurvePayloadData>());
+				return;
+			}
+
+			// Grab the result file and fill the BaseNodeContainer
+			UE::Interchange::FJsonFetchPayloadCmd::JsonResultParser ResultParser;
+			ResultParser.FromJson(JsonResult);
+			FString AnimationTransformPayloadFilename = ResultParser.GetResultFilename();
+
+			if (!ensure(FPaths::FileExists(AnimationTransformPayloadFilename)))
+			{
+				// TODO log an error saying the payload file does not exist even if the get payload command succeeded
+				Promise->SetValue(TOptional<UE::Interchange::FAnimationStepCurvePayloadData>());
+				return;
+			}
+
+			UE::Interchange::FAnimationStepCurvePayloadData AnimationTransformPayload;
+
+			// All sub object should be gone with the reset
+			TArray64<uint8> Buffer;
+			FFileHelper::LoadFileToArray(Buffer, *AnimationTransformPayloadFilename);
+			uint8* FileData = Buffer.GetData();
+			int64 FileDataSize = Buffer.Num();
+			if (FileDataSize < 1)
+			{
+				// Nothing to load from this file
+				Promise->SetValue(TOptional<UE::Interchange::FAnimationStepCurvePayloadData>());
+				return;
+			}
+
+			// Buffer keeps the ownership of the data, the large memory reader is use to serialize the TMap
+			FLargeMemoryReader Ar(FileData, FileDataSize);
+			Ar << AnimationTransformPayload.StepCurves;
+			Promise->SetValue(MoveTemp(AnimationTransformPayload));
+		}));
+
+	// The task was not added to the dispatcher
+	if (CreatedTaskIndex == INDEX_NONE)
+	{
+		Promise->SetValue(TOptional<UE::Interchange::FAnimationStepCurvePayloadData>{});
 	}
 
 	return Promise->GetFuture();
diff --git a/Engine/Plugins/Experimental/Interchange/Runtime/Source/Import/Private/Mesh/InterchangeSkeletalMeshFactory.cpp b/Engine/Plugins/Experimental/Interchange/Runtime/Source/Import/Private/Mesh/InterchangeSkeletalMeshFactory.cpp
--- a/Engine/Plugins/Experimental/Interchange/Runtime/Source/Import/Private/Mesh/InterchangeSkeletalMeshFactory.cpp
+++ b/Engine/Plugins/Experimental/Interchange/Runtime/Source/Import/Private/Mesh/InterchangeSkeletalMeshFactory.cpp
@@ -53,98 +53,98 @@ namespace UE
 				FString TranslatorPayloadKey;
 			};
 
-			void FillBlendShapeMeshDescriptionsPerBlendShapeName(const FMeshNodeContext& MeshNodeContext
-																 , TMap<FString, TOptional<UE::Interchange::FSkeletalMeshBlendShapePayloadData>>& BlendShapeMeshDescriptionsPerBlendShapeName
+			void FillMorphTargetMeshDescriptionsPerMorphTargetName(const FMeshNodeContext& MeshNodeContext
+																 , TMap<FString, TOptional<UE::Interchange::FSkeletalMeshMorphTargetPayloadData>>& MorphTargetMeshDescriptionsPerMorphTargetName
 																 , const IInterchangeSkeletalMeshPayloadInterface* SkeletalMeshTranslatorPayloadInterface
 																 , const int32 VertexOffset
 																 , const UInterchangeBaseNodeContainer* NodeContainer
 																 , FString AssetName)
 			{
-				TArray<FString> BlendShapeUids;
-				MeshNodeContext.MeshNode->GetShapeDependencies(BlendShapeUids);
-				TMap<FString, TFuture<TOptional<UE::Interchange::FSkeletalMeshBlendShapePayloadData>>> TempBlendShapeMeshDescriptionsPerBlendShapeName;
-				TempBlendShapeMeshDescriptionsPerBlendShapeName.Reserve(BlendShapeUids.Num());
-				for (const FString& BlendShapeUid : BlendShapeUids)
+				TArray<FString> MorphTargetUids;
+				MeshNodeContext.MeshNode->GetMorphTargetDependencies(MorphTargetUids);
+				TMap<FString, TFuture<TOptional<UE::Interchange::FSkeletalMeshMorphTargetPayloadData>>> TempMorphTargetMeshDescriptionsPerMorphTargetName;
+				TempMorphTargetMeshDescriptionsPerMorphTargetName.Reserve(MorphTargetUids.Num());
+				for (const FString& MorphTargetUid : MorphTargetUids)
 				{
-					if (const UInterchangeMeshNode* BlendShapeMeshNode = Cast<UInterchangeMeshNode>(NodeContainer->GetNode(BlendShapeUid)))
+					if (const UInterchangeMeshNode* MorphTargetMeshNode = Cast<UInterchangeMeshNode>(NodeContainer->GetNode(MorphTargetUid)))
 					{
-						TOptional<FString> BlendShapePayloadKey = BlendShapeMeshNode->GetPayLoadKey();
-						if (!BlendShapePayloadKey.IsSet())
+						TOptional<FString> MorphTargetPayloadKey = MorphTargetMeshNode->GetPayLoadKey();
+						if (!MorphTargetPayloadKey.IsSet())
 						{
 							UE_LOG(LogInterchangeImport, Warning, TEXT("Empty LOD morph target mesh reference payload when importing SkeletalMesh asset %s"), *AssetName);
 							continue;
 						}
-						const FString PayloadKey = BlendShapePayloadKey.GetValue();
+						const FString PayloadKey = MorphTargetPayloadKey.GetValue();
 						//Add the map entry key, the translator will be call after to bulk get all the needed payload
-						TempBlendShapeMeshDescriptionsPerBlendShapeName.Add(PayloadKey, SkeletalMeshTranslatorPayloadInterface->GetSkeletalMeshBlendShapePayloadData(PayloadKey));
+						TempMorphTargetMeshDescriptionsPerMorphTargetName.Add(PayloadKey, SkeletalMeshTranslatorPayloadInterface->GetSkeletalMeshMorphTargetPayloadData(PayloadKey));
 					}
 				}
 
-				for (const FString& BlendShapeUid : BlendShapeUids)
+				for (const FString& MorphTargetUid : MorphTargetUids)
 				{
-					if (const UInterchangeMeshNode* BlendShapeMeshNode = Cast<UInterchangeMeshNode>(NodeContainer->GetNode(BlendShapeUid)))
+					if (const UInterchangeMeshNode* MorphTargetMeshNode = Cast<UInterchangeMeshNode>(NodeContainer->GetNode(MorphTargetUid)))
 					{
 						
-						TOptional<FString> BlendShapePayloadKey = BlendShapeMeshNode->GetPayLoadKey();
-						if (!BlendShapePayloadKey.IsSet())
+						TOptional<FString> MorphTargetPayloadKey = MorphTargetMeshNode->GetPayLoadKey();
+						if (!MorphTargetPayloadKey.IsSet())
 						{
 							continue;
 						}
-						const FString& BlendShapePayloadKeyString = BlendShapePayloadKey.GetValue();
-						if (!ensure(TempBlendShapeMeshDescriptionsPerBlendShapeName.Contains(BlendShapePayloadKeyString)))
+						const FString& MorphTargetPayloadKeyString = MorphTargetPayloadKey.GetValue();
+						if (!ensure(TempMorphTargetMeshDescriptionsPerMorphTargetName.Contains(MorphTargetPayloadKeyString)))
 						{
 							continue;
 						}
 
-						TOptional<UE::Interchange::FSkeletalMeshBlendShapePayloadData> BlendShapeMeshPayload = TempBlendShapeMeshDescriptionsPerBlendShapeName.FindChecked(BlendShapePayloadKeyString).Get();
-						if (!BlendShapeMeshPayload.IsSet())
+						TOptional<UE::Interchange::FSkeletalMeshMorphTargetPayloadData> MorphTargetMeshPayload = TempMorphTargetMeshDescriptionsPerMorphTargetName.FindChecked(MorphTargetPayloadKeyString).Get();
+						if (!MorphTargetMeshPayload.IsSet())
 						{
-							UE_LOG(LogInterchangeImport, Warning, TEXT("Invalid Skeletal mesh morph target payload key [%s] SkeletalMesh asset %s"), *BlendShapePayloadKeyString, *AssetName);
+							UE_LOG(LogInterchangeImport, Warning, TEXT("Invalid Skeletal mesh morph target payload key [%s] SkeletalMesh asset %s"), *MorphTargetPayloadKeyString, *AssetName);
 							continue;
 						}
-						BlendShapeMeshPayload->VertexOffset = VertexOffset;
+						MorphTargetMeshPayload->VertexOffset = VertexOffset;
 						//Use the Mesh node parent bake transform
 						if (MeshNodeContext.SceneGlobalTransform.IsSet())
 						{
-							BlendShapeMeshPayload->GlobalTransform = MeshNodeContext.SceneGlobalTransform;
+							MorphTargetMeshPayload->GlobalTransform = MeshNodeContext.SceneGlobalTransform;
 						}
 						else
 						{
-							BlendShapeMeshPayload->GlobalTransform.Reset();
+							MorphTargetMeshPayload->GlobalTransform.Reset();
 						}
 
-						if (!BlendShapeMeshNode->GetBlendShapeName(BlendShapeMeshPayload->BlendShapeName))
+						if (!MorphTargetMeshNode->GetMorphTargetName(MorphTargetMeshPayload->MorphTargetName))
 						{
-							BlendShapeMeshPayload->BlendShapeName = BlendShapePayloadKeyString;
+							MorphTargetMeshPayload->MorphTargetName = MorphTargetPayloadKeyString;
 						}
-						//Add the Blend shape to the blend shape map
-						BlendShapeMeshDescriptionsPerBlendShapeName.Add(BlendShapePayloadKeyString, BlendShapeMeshPayload);
+						//Add the morph target to the morph target map
+						MorphTargetMeshDescriptionsPerMorphTargetName.Add(MorphTargetPayloadKeyString, MorphTargetMeshPayload);
 					}
 				}
 			}
 
-			void CopyBlendShapesMeshDescriptionToSkeletalMeshImportData(const TMap<FString, TOptional<UE::Interchange::FSkeletalMeshBlendShapePayloadData>>& LodBlendShapeMeshDescriptions, FSkeletalMeshImportData& DestinationSkeletalMeshImportData)
+			void CopyMorphTargetsMeshDescriptionToSkeletalMeshImportData(const TMap<FString, TOptional<UE::Interchange::FSkeletalMeshMorphTargetPayloadData>>& LodMorphTargetMeshDescriptions, FSkeletalMeshImportData& DestinationSkeletalMeshImportData)
 			{
-				const int32 OriginalMorphTargetCount = LodBlendShapeMeshDescriptions.Num();
+				const int32 OriginalMorphTargetCount = LodMorphTargetMeshDescriptions.Num();
 				TArray<FString> Keys;
 				int32 MorphTargetCount = 0;
-				for (const TPair<FString, TOptional<UE::Interchange::FSkeletalMeshBlendShapePayloadData>>& Pair : LodBlendShapeMeshDescriptions)
+				for (const TPair<FString, TOptional<UE::Interchange::FSkeletalMeshMorphTargetPayloadData>>& Pair : LodMorphTargetMeshDescriptions)
 				{
-					const FString BlendShapeName(Pair.Key);
-					const TOptional<UE::Interchange::FSkeletalMeshBlendShapePayloadData>& BlendShapePayloadData = Pair.Value;
-					if (!BlendShapePayloadData.IsSet())
+					const FString MorphTargetName(Pair.Key);
+					const TOptional<UE::Interchange::FSkeletalMeshMorphTargetPayloadData>& MorphTargetPayloadData = Pair.Value;
+					if (!MorphTargetPayloadData.IsSet())
 					{
-						UE_LOG(LogInterchangeImport, Error, TEXT("Empty blend shape optional payload data [%s]"), *BlendShapeName);
+						UE_LOG(LogInterchangeImport, Error, TEXT("Empty morph target optional payload data [%s]"), *MorphTargetName);
 						continue;
 					}
 
-					const FMeshDescription& SourceMeshDescription = BlendShapePayloadData.GetValue().LodMeshDescription;
-					const int32 VertexOffset = BlendShapePayloadData->VertexOffset;
+					const FMeshDescription& SourceMeshDescription = MorphTargetPayloadData.GetValue().LodMeshDescription;
+					const int32 VertexOffset = MorphTargetPayloadData->VertexOffset;
 					const int32 SourceMeshVertexCount = SourceMeshDescription.Vertices().Num();
 					const int32 DestinationVertexIndexMax = VertexOffset + SourceMeshVertexCount;
 					if (!DestinationSkeletalMeshImportData.Points.IsValidIndex(DestinationVertexIndexMax-1))
 					{
-						UE_LOG(LogInterchangeImport, Error, TEXT("Corrupted blend shape optional payload data [%s]"), *BlendShapeName);
+						UE_LOG(LogInterchangeImport, Error, TEXT("Corrupted morph target optional payload data [%s]"), *MorphTargetName);
 						continue;
 					}
 					Keys.Add(Pair.Key);
@@ -171,7 +171,7 @@ namespace UE
 				ParallelFor(NumMorphGroup, [MorphTargetGroupSize,
 							MorphTargetCount,
 							NumMorphGroup,
-							&LodBlendShapeMeshDescriptions,
+							&LodMorphTargetMeshDescriptions,
 							&Keys,
 							&DestinationSkeletalMeshImportData](const int32 MorphTargetGroupIndex)
 				{
@@ -185,17 +185,17 @@ namespace UE
 							//Executing the last morph target group, in case we do not have a full last group.
 							break;
 						}
-						const FString BlendShapeKey(Keys[MorphTargetIndex]);
-						const TOptional<UE::Interchange::FSkeletalMeshBlendShapePayloadData>& BlendShapePayloadData = LodBlendShapeMeshDescriptions.FindChecked(BlendShapeKey);
-						if (!ensure(BlendShapePayloadData.IsSet()))
+						const FString MorphTargetKey(Keys[MorphTargetIndex]);
+						const TOptional<UE::Interchange::FSkeletalMeshMorphTargetPayloadData>& MorphTargetPayloadData = LodMorphTargetMeshDescriptions.FindChecked(MorphTargetKey);
+						if (!ensure(MorphTargetPayloadData.IsSet()))
 						{
 							//This error was suppose to be catch in the pre parallel for loop
 							break;
 						}
 
-						const FMeshDescription& SourceMeshDescription = BlendShapePayloadData.GetValue().LodMeshDescription;
-						const FTransform GlobalTransform = BlendShapePayloadData->GlobalTransform.IsSet() ? BlendShapePayloadData->GlobalTransform.GetValue() : FTransform::Identity;
-						const int32 VertexOffset = BlendShapePayloadData->VertexOffset;
+						const FMeshDescription& SourceMeshDescription = MorphTargetPayloadData.GetValue().LodMeshDescription;
+						const FTransform GlobalTransform = MorphTargetPayloadData->GlobalTransform.IsSet() ? MorphTargetPayloadData->GlobalTransform.GetValue() : FTransform::Identity;
+						const int32 VertexOffset = MorphTargetPayloadData->VertexOffset;
 						const int32 SourceMeshVertexCount = SourceMeshDescription.Vertices().Num();
 						const int32 DestinationVertexIndexMax = VertexOffset + SourceMeshVertexCount;
 						if (!ensure(DestinationSkeletalMeshImportData.Points.IsValidIndex(DestinationVertexIndexMax-1)))
@@ -210,7 +210,7 @@ namespace UE
 
 						//Create the morph target source data
 						FString& MorphTargetName = DestinationSkeletalMeshImportData.MorphTargetNames[MorphTargetIndex];
-						MorphTargetName = BlendShapePayloadData->BlendShapeName;
+						MorphTargetName = MorphTargetPayloadData->MorphTargetName;
 						TSet<uint32>& ModifiedPoints = DestinationSkeletalMeshImportData.MorphTargetModifiedPoints[MorphTargetIndex];
 						FSkeletalMeshImportData& MorphTargetData = DestinationSkeletalMeshImportData.MorphTargets[MorphTargetIndex];
 
@@ -416,17 +416,17 @@ namespace UE
 				TMap<FString, TFuture<TOptional<UE::Interchange::FSkeletalMeshLodPayloadData>>> LodMeshPayloadPerTranslatorPayloadKey;
 				LodMeshPayloadPerTranslatorPayloadKey.Reserve(MeshReferences.Num());
 
-				TMap<FString, TOptional<UE::Interchange::FSkeletalMeshBlendShapePayloadData>> BlendShapeMeshDescriptionsPerBlendShapeName;
-				int32 BlendShapeCount = 0;
+				TMap<FString, TOptional<UE::Interchange::FSkeletalMeshMorphTargetPayloadData>> MorphTargetMeshDescriptionsPerMorphTargetName;
+				int32 MorphTargetCount = 0;
 
 				for (const FMeshNodeContext& MeshNodeContext : MeshReferences)
 				{
 					//Add the payload entry key, the payload data will be fill later in bulk by the translator
 					LodMeshPayloadPerTranslatorPayloadKey.Add(MeshNodeContext.TranslatorPayloadKey, SkeletalMeshTranslatorPayloadInterface->GetSkeletalMeshLodPayloadData(MeshNodeContext.TranslatorPayloadKey));
-					//Count the blend shape dependencies so we can reserve the right amount
-					BlendShapeCount += (bImportMorphTarget && MeshNodeContext.MeshNode) ? MeshNodeContext.MeshNode->GetShapeDependeciesCount() : 0;
+					//Count the morph target dependencies so we can reserve the right amount
+					MorphTargetCount += (bImportMorphTarget && MeshNodeContext.MeshNode) ? MeshNodeContext.MeshNode->GetMorphTargetDependeciesCount() : 0;
 				}
-				BlendShapeMeshDescriptionsPerBlendShapeName.Reserve(BlendShapeCount);
+				MorphTargetMeshDescriptionsPerMorphTargetName.Reserve(MorphTargetCount);
 
 				//Fill the lod mesh description using all combined mesh part
 				for (const FMeshNodeContext& MeshNodeContext : MeshReferences)
@@ -527,8 +527,8 @@ namespace UE
 					}
 					if (bImportMorphTarget)
 					{
-						FillBlendShapeMeshDescriptionsPerBlendShapeName(MeshNodeContext
-																		, BlendShapeMeshDescriptionsPerBlendShapeName
+						FillMorphTargetMeshDescriptionsPerMorphTargetName(MeshNodeContext
+																		, MorphTargetMeshDescriptionsPerMorphTargetName
 																		, SkeletalMeshTranslatorPayloadInterface
 																		, VertexOffset
 																		, Arguments.NodeContainer
@@ -539,8 +539,8 @@ namespace UE
 				DestinationImportData = FSkeletalMeshImportData::CreateFromMeshDescription(LodMeshDescription);
 				DestinationImportData.RefBonesBinary = RefBonesBinary;
 
-				//Copy all the lod blend shapes data to the DestinationImportData.
-				CopyBlendShapesMeshDescriptionToSkeletalMeshImportData(BlendShapeMeshDescriptionsPerBlendShapeName, DestinationImportData);
+				//Copy all the lod morph targets data to the DestinationImportData.
+				CopyMorphTargetsMeshDescriptionToSkeletalMeshImportData(MorphTargetMeshDescriptionsPerMorphTargetName, DestinationImportData);
 			}
 
 			void ProcessImportMeshInfluences(const int32 WedgeCount, TArray<SkeletalMeshImportData::FRawBoneInfluence>& Influences)
@@ -1198,7 +1198,7 @@ UObject* UInterchangeSkeletalMeshFactory::CreateAsset(const FCreateAssetParams&
 		//Add the lod mesh data to the skeletalmesh
 		FSkeletalMeshImportData SkeletalMeshImportData;
 		const bool bSkinControlPointToTimeZero = bUseTimeZeroAsBindPose && bDiffPose;
-		//Get all meshes and blend shapes payload and fill the SkeletalMeshImportData structure
+		//Get all meshes and morph targets payload and fill the SkeletalMeshImportData structure
 		UE::Interchange::Private::RetrieveAllSkeletalMeshPayloadsAndFillImportData(SkeletalMeshFactoryNode
 																					, SkeletalMeshImportData
 																					, MeshReferences
diff --git a/Engine/Plugins/Experimental/Interchange/Runtime/Source/Import/Public/Animation/InterchangeAnimationPayload.h b/Engine/Plugins/Experimental/Interchange/Runtime/Source/Import/Public/Animation/InterchangeAnimationPayload.h
--- a/Engine/Plugins/Experimental/Interchange/Runtime/Source/Import/Public/Animation/InterchangeAnimationPayload.h
+++ b/Engine/Plugins/Experimental/Interchange/Runtime/Source/Import/Public/Animation/InterchangeAnimationPayload.h
@@ -8,13 +8,39 @@
 
 namespace UE::Interchange
 {
-	struct FAnimationCurveTransformPayloadData
+	/**
+	 * This payload class is use to get a user attribute curve for any translated node.
+	 * If the attribute is a multi channel like a position vector we should get 3 curve representing X, Y and Z.
+	 */
+	struct FAnimationCurvePayloadData
 	{
-		EInterchangeTransformCurveChannel TransformChannel;
-		FRichCurve Curve;
+		TArray<FRichCurve> Curves;
 	};
-	struct FAnimationTransformPayloadData
+
+	struct FAnimationStepCurvePayloadData
 	{
-		TArray<FAnimationCurveTransformPayloadData> TransformCurves;
+		TArray<FInterchangeStepCurve> StepCurves;
+
+		static FAnimationStepCurvePayloadData FromAnimationCurvePayloadData(const FAnimationCurvePayloadData& CurvePayload)
+		{
+			FAnimationStepCurvePayloadData AnimationStepCurvePayloadData;
+			AnimationStepCurvePayloadData.StepCurves.Reserve(CurvePayload.Curves.Num());
+			for (const FRichCurve& RichCurve : CurvePayload.Curves)
+			{
+				FInterchangeStepCurve& StepCurve = AnimationStepCurvePayloadData.StepCurves.AddDefaulted_GetRef();
+				const int32 KeyCount = RichCurve.GetNumKeys();
+				StepCurve.KeyTimes.AddZeroed(KeyCount);
+				TArray<float> KeyValues;
+				KeyValues.AddZeroed(KeyCount);
+				int32 KeyIndex = 0;
+				for (FKeyHandle KeyHandle = RichCurve.GetFirstKeyHandle(); KeyHandle != FKeyHandle::Invalid(); KeyHandle = RichCurve.GetNextKey(KeyHandle), ++KeyIndex )
+				{
+					StepCurve.KeyTimes[KeyIndex] = RichCurve.GetKeyTime(KeyHandle);
+					KeyValues[KeyIndex] = RichCurve.GetKeyValue(KeyHandle);
+				}
+				StepCurve.FloatKeyValues = MoveTemp(KeyValues);
+			}
+			return AnimationStepCurvePayloadData;
+		}
 	};
 }//ns UE::Interchange
diff --git a/Engine/Plugins/Experimental/Interchange/Runtime/Source/Import/Public/Animation/InterchangeAnimationPayloadInterface.h b/Engine/Plugins/Experimental/Interchange/Runtime/Source/Import/Public/Animation/InterchangeAnimationPayloadInterface.h
--- a/Engine/Plugins/Experimental/Interchange/Runtime/Source/Import/Public/Animation/InterchangeAnimationPayloadInterface.h
+++ b/Engine/Plugins/Experimental/Interchange/Runtime/Source/Import/Public/Animation/InterchangeAnimationPayloadInterface.h
@@ -25,13 +25,22 @@ class INTERCHANGEIMPORT_API IInterchangeAnimationPayloadInterface
 public:
 
 	/**
-	 * Get animation payload data for the specified payload key
+	 * Get animation payload data for the specified payload key. It return an array of FRichCurve (Rich curve are float curve we can interpolate).
 	 *
 	 * @param PayLoadKey - The key to retrieve the a particular payload contain into the specified source data.
 	 * @return - The resulting PayloadData as a TFuture point by the PayloadKey. The TOptional will not be set if there is an error retrieving the payload.
 	 * 
 	 */
-	virtual TFuture<TOptional<UE::Interchange::FAnimationTransformPayloadData>> GetAnimationTransformPayloadData(const FString& PayLoadKey) const = 0;
+	virtual TFuture<TOptional<UE::Interchange::FAnimationCurvePayloadData>> GetAnimationCurvePayloadData(const FString& PayLoadKey) const = 0;
+
+	/**
+	 * Get animation payload data for the specified payload key. It return an array of "Step" curve data (step curve are a combination of Key time and Key Value, no interpolation data)
+	 *
+	 * @param PayLoadKey - The key to retrieve the a particular payload contain into the specified source data.
+	 * @return - The resulting PayloadData as a TFuture point by the PayloadKey. The TOptional will not be set if there is an error retrieving the payload.
+	 *
+	 */
+	virtual TFuture<TOptional<UE::Interchange::FAnimationStepCurvePayloadData>> GetAnimationStepCurvePayloadData(const FString& PayLoadKey) const = 0;
 
 	virtual TFuture<TOptional<UE::Interchange::FAnimationBakeTransformPayloadData>> GetAnimationBakeTransformPayloadData(const FString& PayLoadKey, const double BakeFrequency, const double RangeStartSecond, const double RangeStopSecond) const = 0;
 };
diff --git a/Engine/Plugins/Experimental/Interchange/Runtime/Source/Import/Public/Fbx/InterchangeFbxTranslator.h b/Engine/Plugins/Experimental/Interchange/Runtime/Source/Import/Public/Fbx/InterchangeFbxTranslator.h
--- a/Engine/Plugins/Experimental/Interchange/Runtime/Source/Import/Public/Fbx/InterchangeFbxTranslator.h
+++ b/Engine/Plugins/Experimental/Interchange/Runtime/Source/Import/Public/Fbx/InterchangeFbxTranslator.h
@@ -75,14 +75,16 @@ public:
 	/* IInterchangeSkeletalMeshPayloadInterface Begin */
 
 	virtual TFuture<TOptional<UE::Interchange::FSkeletalMeshLodPayloadData>> GetSkeletalMeshLodPayloadData(const FString& PayLoadKey) const override;
-	virtual TFuture<TOptional<UE::Interchange::FSkeletalMeshBlendShapePayloadData>> GetSkeletalMeshBlendShapePayloadData(const FString& PayLoadKey) const override;
+	virtual TFuture<TOptional<UE::Interchange::FSkeletalMeshMorphTargetPayloadData>> GetSkeletalMeshMorphTargetPayloadData(const FString& PayLoadKey) const override;
 	
 	/* IInterchangeSkeletalMeshPayloadInterface End */
 
 	//////////////////////////////////////////////////////////////////////////
 	/* IInterchangeAnimationPayloadInterface Begin */
 
-	virtual TFuture<TOptional<UE::Interchange::FAnimationTransformPayloadData>> GetAnimationTransformPayloadData(const FString& PayLoadKey) const override;
+	virtual TFuture<TOptional<UE::Interchange::FAnimationCurvePayloadData>> GetAnimationCurvePayloadData(const FString& PayLoadKey) const override;
+
+	virtual TFuture<TOptional<UE::Interchange::FAnimationStepCurvePayloadData>> GetAnimationStepCurvePayloadData(const FString& PayLoadKey) const override;
 
 	virtual TFuture<TOptional<UE::Interchange::FAnimationBakeTransformPayloadData>> GetAnimationBakeTransformPayloadData(const FString& PayLoadKey, const double BakeFrequency, const double RangeStartSecond, const double RangeStopSecond) const override;
 
diff --git a/Engine/Plugins/Experimental/Interchange/Runtime/Source/Import/Public/Mesh/InterchangeSkeletalMeshPayload.h b/Engine/Plugins/Experimental/Interchange/Runtime/Source/Import/Public/Mesh/InterchangeSkeletalMeshPayload.h
--- a/Engine/Plugins/Experimental/Interchange/Runtime/Source/Import/Public/Mesh/InterchangeSkeletalMeshPayload.h
+++ b/Engine/Plugins/Experimental/Interchange/Runtime/Source/Import/Public/Mesh/InterchangeSkeletalMeshPayload.h
@@ -18,15 +18,15 @@ namespace UE
 			TArray<FString> JointNames;
 		};
 
-		struct FSkeletalMeshBlendShapePayloadData
+		struct FSkeletalMeshMorphTargetPayloadData
 		{
-			//BlendShape payload is a MeshDescription containing at least the vertex position.
+			//Morph target payload is a MeshDescription containing at least the vertex position.
 			FMeshDescription LodMeshDescription;
 			TOptional<FTransform> GlobalTransform;
-			//The vertex offset of the blendshape in case we combine mesh node together
+			//The vertex offset of the morph target in case we combine mesh node together
 			int32 VertexOffset;
-			//The name of the blendshape
-			FString BlendShapeName;
+			//The name of the morph target
+			FString MorphTargetName;
 		};
 	}//ns Interchange
 }//ns UE
diff --git a/Engine/Plugins/Experimental/Interchange/Runtime/Source/Import/Public/Mesh/InterchangeSkeletalMeshPayloadInterface.h b/Engine/Plugins/Experimental/Interchange/Runtime/Source/Import/Public/Mesh/InterchangeSkeletalMeshPayloadInterface.h
--- a/Engine/Plugins/Experimental/Interchange/Runtime/Source/Import/Public/Mesh/InterchangeSkeletalMeshPayloadInterface.h
+++ b/Engine/Plugins/Experimental/Interchange/Runtime/Source/Import/Public/Mesh/InterchangeSkeletalMeshPayloadInterface.h
@@ -40,7 +40,7 @@ public:
 	 * @return - The resulting PayloadData as a TFuture point by the payload key. The TOptional will not be set if there is an error retrieving the payload.
 	 * 
 	 */
-	virtual TFuture<TOptional<UE::Interchange::FSkeletalMeshBlendShapePayloadData>> GetSkeletalMeshBlendShapePayloadData(const FString& PayLoadKey) const = 0;
+	virtual TFuture<TOptional<UE::Interchange::FSkeletalMeshMorphTargetPayloadData>> GetSkeletalMeshMorphTargetPayloadData(const FString& PayLoadKey) const = 0;
 };
 
 
diff --git a/Engine/Plugins/Experimental/Interchange/Runtime/Source/Nodes/Private/InterchangeMeshNode.cpp b/Engine/Plugins/Experimental/Interchange/Runtime/Source/Nodes/Private/InterchangeMeshNode.cpp
--- a/Engine/Plugins/Experimental/Interchange/Runtime/Source/Nodes/Private/InterchangeMeshNode.cpp
+++ b/Engine/Plugins/Experimental/Interchange/Runtime/Source/Nodes/Private/InterchangeMeshNode.cpp
@@ -8,9 +8,15 @@ namespace UE
 	namespace Interchange
 	{
 
-		const FAttributeKey& FMeshNodeStaticData::PayloadSourceFileKey()
+		const FAttributeKey& FMeshNodeStaticData::PayloadSourceKey()
 		{
-			static FAttributeKey AttributeKey(TEXT("__PayloadSourceFile__"));
+			static FAttributeKey AttributeKey(TEXT("__PayloadSourceKey__"));
+			return AttributeKey;
+		}
+
+		const FAttributeKey& FMeshNodeStaticData::PayloadAnimationCurveKey()
+		{
+			static FAttributeKey AttributeKey(TEXT("__PayloadAnimationCurveKey__"));
 			return AttributeKey;
 		}
 
@@ -20,15 +26,15 @@ namespace UE
 			return AttributeKey;
 		}
 
-		const FAttributeKey& FMeshNodeStaticData::IsBlendShapeKey()
+		const FAttributeKey& FMeshNodeStaticData::IsMorphTargetKey()
 		{
-			static FAttributeKey AttributeKey(TEXT("__IsBlendShapeKey__"));
+			static FAttributeKey AttributeKey(TEXT("__IsMorphTargetKey__"));
 			return AttributeKey;
 		}
 
-		const FAttributeKey& FMeshNodeStaticData::BlendShapeNameKey()
+		const FAttributeKey& FMeshNodeStaticData::MorphTargetNameKey()
 		{
-			static FAttributeKey AttributeKey(TEXT("__BlendShapeNameKey__"));
+			static FAttributeKey AttributeKey(TEXT("__MorphTargetNameKey__"));
 			return AttributeKey;
 		}
 
@@ -38,9 +44,9 @@ namespace UE
 			return Dependencies_BaseKey;
 		}
 
-		const FAttributeKey& FMeshNodeStaticData::GetShapeDependenciesKey()
+		const FAttributeKey& FMeshNodeStaticData::GetMorphTargetDependenciesKey()
 		{
-			static FAttributeKey Dependencies_BaseKey(TEXT("__MeshShapeDependencies__"));
+			static FAttributeKey Dependencies_BaseKey(TEXT("__MeshMorphTargetDependencies__"));
 			return Dependencies_BaseKey;
 		}
 
@@ -62,7 +68,7 @@ namespace UE
 UInterchangeMeshNode::UInterchangeMeshNode()
 {
 	SkeletonDependencies.Initialize(Attributes, UE::Interchange::FMeshNodeStaticData::GetSkeletonDependenciesKey().ToString());
-	ShapeDependencies.Initialize(Attributes, UE::Interchange::FMeshNodeStaticData::GetShapeDependenciesKey().ToString());
+	MorphTargetDependencies.Initialize(Attributes, UE::Interchange::FMeshNodeStaticData::GetMorphTargetDependenciesKey().ToString());
 	SceneInstancesUids.Initialize(Attributes, UE::Interchange::FMeshNodeStaticData::GetSceneInstancesUidsKey().ToString());
 	SlotMaterialDependencies.Initialize(Attributes.ToSharedRef(), UE::Interchange::FMeshNodeStaticData::GetSlotMaterialDependenciesKey().ToString());
 }
@@ -71,21 +77,25 @@ FString UInterchangeMeshNode::GetKeyDisplayName(const UE::Interchange::FAttribut
 {
 	FString KeyDisplayName = NodeAttributeKey.ToString();
 	const FString NodeAttributeKeyString = KeyDisplayName;
-	if (NodeAttributeKey == UE::Interchange::FMeshNodeStaticData::PayloadSourceFileKey())
+	if (NodeAttributeKey == UE::Interchange::FMeshNodeStaticData::PayloadSourceKey())
 	{
 		return KeyDisplayName = TEXT("Payload Source Key");
 	}
+	else if (NodeAttributeKey == UE::Interchange::FMeshNodeStaticData::PayloadAnimationCurveKey())
+	{
+		return KeyDisplayName = TEXT("Payload Animation Curve Key");
+	}
 	else if (NodeAttributeKey == UE::Interchange::FMeshNodeStaticData::IsSkinnedMeshKey())
 	{
 		return KeyDisplayName = TEXT("Is a Skinned Mesh");
 	}
-	else if (NodeAttributeKey == UE::Interchange::FMeshNodeStaticData::IsBlendShapeKey())
+	else if (NodeAttributeKey == UE::Interchange::FMeshNodeStaticData::IsMorphTargetKey())
 	{
-		return KeyDisplayName = TEXT("Is a Blend Shape");
+		return KeyDisplayName = TEXT("Is a Morph Target");
 	}
-	else if (NodeAttributeKey == UE::Interchange::FMeshNodeStaticData::BlendShapeNameKey())
+	else if (NodeAttributeKey == UE::Interchange::FMeshNodeStaticData::MorphTargetNameKey())
 	{
-		return KeyDisplayName = TEXT("Blend Shape Name");
+		return KeyDisplayName = TEXT("Morph Target Name");
 	}
 	else if (NodeAttributeKey == UE::Interchange::FMeshNodeStaticData::GetSkeletonDependenciesKey())
 	{
@@ -102,13 +112,13 @@ FString UInterchangeMeshNode::GetKeyDisplayName(const UE::Interchange::FAttribut
 		}
 		return KeyDisplayName;
 	}
-	else if (NodeAttributeKey == UE::Interchange::FMeshNodeStaticData::GetShapeDependenciesKey())
+	else if (NodeAttributeKey == UE::Interchange::FMeshNodeStaticData::GetMorphTargetDependenciesKey())
 	{
-		return KeyDisplayName = TEXT("Shape Dependencies count");
+		return KeyDisplayName = TEXT("Morph Target Dependencies Count");
 	}
-	else if (NodeAttributeKeyString.StartsWith(UE::Interchange::FMeshNodeStaticData::GetShapeDependenciesKey().ToString()))
+	else if (NodeAttributeKeyString.StartsWith(UE::Interchange::FMeshNodeStaticData::GetMorphTargetDependenciesKey().ToString()))
 	{
-		KeyDisplayName = TEXT("Shape Dependencies Index ");
+		KeyDisplayName = TEXT("Morph Target Dependencies Index ");
 		const FString IndexKey = UE::Interchange::TArrayAttributeHelper<FString>::IndexKey();
 		int32 IndexPosition = NodeAttributeKeyString.Find(IndexKey) + IndexKey.Len();
 		if (IndexPosition < NodeAttributeKeyString.Len())
@@ -146,9 +156,9 @@ FString UInterchangeMeshNode::GetAttributeCategory(const UE::Interchange::FAttri
 	{
 		return FString(TEXT("SkeletonDependencies"));
 	}
-	else if (NodeAttributeKeyString.StartsWith(UE::Interchange::FMeshNodeStaticData::GetShapeDependenciesKey().ToString()))
+	else if (NodeAttributeKeyString.StartsWith(UE::Interchange::FMeshNodeStaticData::GetMorphTargetDependenciesKey().ToString()))
 	{
-		return FString(TEXT("ShapeDependencies"));
+		return FString(TEXT("MorphTargetDependencies"));
 	}
 	else if (NodeAttributeKeyString.StartsWith(UE::Interchange::FMeshNodeStaticData::GetSceneInstancesUidsKey().ToString()))
 	{
@@ -222,14 +232,14 @@ bool UInterchangeMeshNode::SetSkinnedMesh(const bool bIsSkinnedMesh)
 	return false;
 }
 
-bool UInterchangeMeshNode::IsBlendShape() const
+bool UInterchangeMeshNode::IsMorphTarget() const
 {
-	if (!Attributes->ContainAttribute(UE::Interchange::FMeshNodeStaticData::IsBlendShapeKey()))
+	if (!Attributes->ContainAttribute(UE::Interchange::FMeshNodeStaticData::IsMorphTargetKey()))
 	{
 		return false;
 	}
 
-	UE::Interchange::FAttributeStorage::TAttributeHandle<bool> Handle = Attributes->GetAttributeHandle<bool>(UE::Interchange::FMeshNodeStaticData::IsBlendShapeKey());
+	UE::Interchange::FAttributeStorage::TAttributeHandle<bool> Handle = Attributes->GetAttributeHandle<bool>(UE::Interchange::FMeshNodeStaticData::IsMorphTargetKey());
 	if (Handle.IsValid())
 	{
 		bool bValue = false;
@@ -239,40 +249,40 @@ bool UInterchangeMeshNode::IsBlendShape() const
 	return false;
 }
 
-bool UInterchangeMeshNode::SetBlendShape(const bool bIsBlendShape)
+bool UInterchangeMeshNode::SetMorphTarget(const bool bIsMorphTarget)
 {
-	UE::Interchange::EAttributeStorageResult Result = Attributes->RegisterAttribute(UE::Interchange::FMeshNodeStaticData::IsBlendShapeKey(), bIsBlendShape);
+	UE::Interchange::EAttributeStorageResult Result = Attributes->RegisterAttribute(UE::Interchange::FMeshNodeStaticData::IsMorphTargetKey(), bIsMorphTarget);
 	if (IsAttributeStorageResultSuccess(Result))
 	{
-		UE::Interchange::FAttributeStorage::TAttributeHandle<bool> Handle = Attributes->GetAttributeHandle<bool>(UE::Interchange::FMeshNodeStaticData::IsBlendShapeKey());
+		UE::Interchange::FAttributeStorage::TAttributeHandle<bool> Handle = Attributes->GetAttributeHandle<bool>(UE::Interchange::FMeshNodeStaticData::IsMorphTargetKey());
 		return Handle.IsValid();
 	}
 	return false;
 }
 
-bool UInterchangeMeshNode::GetBlendShapeName(FString& OutBlendShapeName) const
+bool UInterchangeMeshNode::GetMorphTargetName(FString& OutMorphTargetName) const
 {
-	OutBlendShapeName.Empty();
-	if (!Attributes->ContainAttribute(UE::Interchange::FMeshNodeStaticData::BlendShapeNameKey()))
+	OutMorphTargetName.Empty();
+	if (!Attributes->ContainAttribute(UE::Interchange::FMeshNodeStaticData::MorphTargetNameKey()))
 	{
 		return false;
 	}
 
-	UE::Interchange::FAttributeStorage::TAttributeHandle<FString> Handle = Attributes->GetAttributeHandle<FString>(UE::Interchange::FMeshNodeStaticData::BlendShapeNameKey());
+	UE::Interchange::FAttributeStorage::TAttributeHandle<FString> Handle = Attributes->GetAttributeHandle<FString>(UE::Interchange::FMeshNodeStaticData::MorphTargetNameKey());
 	if (Handle.IsValid())
 	{
-		Handle.Get(OutBlendShapeName);
+		Handle.Get(OutMorphTargetName);
 		return true;
 	}
 	return false;
 }
 
-bool UInterchangeMeshNode::SetBlendShapeName(const FString& BlendShapeName)
+bool UInterchangeMeshNode::SetMorphTargetName(const FString& MorphTargetName)
 {
-	UE::Interchange::EAttributeStorageResult Result = Attributes->RegisterAttribute(UE::Interchange::FMeshNodeStaticData::BlendShapeNameKey(), BlendShapeName);
+	UE::Interchange::EAttributeStorageResult Result = Attributes->RegisterAttribute(UE::Interchange::FMeshNodeStaticData::MorphTargetNameKey(), MorphTargetName);
 	if (IsAttributeStorageResultSuccess(Result))
 	{
-		UE::Interchange::FAttributeStorage::TAttributeHandle<FString> Handle = Attributes->GetAttributeHandle<FString>(UE::Interchange::FMeshNodeStaticData::BlendShapeNameKey());
+		UE::Interchange::FAttributeStorage::TAttributeHandle<FString> Handle = Attributes->GetAttributeHandle<FString>(UE::Interchange::FMeshNodeStaticData::MorphTargetNameKey());
 		return Handle.IsValid();
 	}
 	return false;
@@ -280,11 +290,11 @@ bool UInterchangeMeshNode::SetBlendShapeName(const FString& BlendShapeName)
 
 const TOptional<FString> UInterchangeMeshNode::GetPayLoadKey() const
 {
-	if (!Attributes->ContainAttribute(UE::Interchange::FMeshNodeStaticData::PayloadSourceFileKey()))
+	if (!Attributes->ContainAttribute(UE::Interchange::FMeshNodeStaticData::PayloadSourceKey()))
 	{
 		return TOptional<FString>();
 	}
-	UE::Interchange::FAttributeStorage::TAttributeHandle<FString> AttributeHandle = Attributes->GetAttributeHandle<FString>(UE::Interchange::FMeshNodeStaticData::PayloadSourceFileKey());
+	UE::Interchange::FAttributeStorage::TAttributeHandle<FString> AttributeHandle = Attributes->GetAttributeHandle<FString>(UE::Interchange::FMeshNodeStaticData::PayloadSourceKey());
 	if (!AttributeHandle.IsValid())
 	{
 		return TOptional<FString>();
@@ -293,7 +303,7 @@ const TOptional<FString> UInterchangeMeshNode::GetPayLoadKey() const
 	UE::Interchange::EAttributeStorageResult Result = AttributeHandle.Get(PayloadKey);
 	if (!IsAttributeStorageResultSuccess(Result))
 	{
-		LogAttributeStorageErrors(Result, TEXT("UInterchangeMeshNode.GetPayLoadKey"), UE::Interchange::FMeshNodeStaticData::PayloadSourceFileKey());
+		LogAttributeStorageErrors(Result, TEXT("UInterchangeMeshNode.GetPayLoadKey"), UE::Interchange::FMeshNodeStaticData::PayloadSourceKey());
 		return TOptional<FString>();
 	}
 	return TOptional<FString>(PayloadKey);
@@ -301,13 +311,43 @@ const TOptional<FString> UInterchangeMeshNode::GetPayLoadKey() const
 
 void UInterchangeMeshNode::SetPayLoadKey(const FString& PayloadKey)
 {
-	UE::Interchange::EAttributeStorageResult Result = Attributes->RegisterAttribute(UE::Interchange::FMeshNodeStaticData::PayloadSourceFileKey(), PayloadKey);
+	UE::Interchange::EAttributeStorageResult Result = Attributes->RegisterAttribute(UE::Interchange::FMeshNodeStaticData::PayloadSourceKey(), PayloadKey);
+	if (!IsAttributeStorageResultSuccess(Result))
+	{
+		LogAttributeStorageErrors(Result, TEXT("UInterchangeMeshNode.SetPayLoadKey"), UE::Interchange::FMeshNodeStaticData::PayloadSourceKey());
+	}
+}
+
+const TOptional<FString> UInterchangeMeshNode::GetAnimationCurvePayLoadKey() const
+{
+	if (!Attributes->ContainAttribute(UE::Interchange::FMeshNodeStaticData::PayloadAnimationCurveKey()))
+	{
+		return TOptional<FString>();
+	}
+	UE::Interchange::FAttributeStorage::TAttributeHandle<FString> AttributeHandle = Attributes->GetAttributeHandle<FString>(UE::Interchange::FMeshNodeStaticData::PayloadAnimationCurveKey());
+	if (!AttributeHandle.IsValid())
+	{
+		return TOptional<FString>();
+	}
+	FString PayloadKey;
+	UE::Interchange::EAttributeStorageResult Result = AttributeHandle.Get(PayloadKey);
+	if (!IsAttributeStorageResultSuccess(Result))
+	{
+		LogAttributeStorageErrors(Result, TEXT("UInterchangeMeshNode.GetAnimationCurvePayLoadKey"), UE::Interchange::FMeshNodeStaticData::PayloadAnimationCurveKey());
+		return TOptional<FString>();
+	}
+	return TOptional<FString>(PayloadKey);
+}
+
+void UInterchangeMeshNode::SetAnimationCurvePayLoadKey(const FString& PayloadKey)
+{
+	UE::Interchange::EAttributeStorageResult Result = Attributes->RegisterAttribute(UE::Interchange::FMeshNodeStaticData::PayloadAnimationCurveKey(), PayloadKey);
 	if (!IsAttributeStorageResultSuccess(Result))
 	{
-		LogAttributeStorageErrors(Result, TEXT("UInterchangeMeshNode.SetPayLoadKey"), UE::Interchange::FMeshNodeStaticData::PayloadSourceFileKey());
+		LogAttributeStorageErrors(Result, TEXT("UInterchangeMeshNode.SetAnimationCurvePayLoadKey"), UE::Interchange::FMeshNodeStaticData::PayloadAnimationCurveKey());
 	}
 }
-	
+
 bool UInterchangeMeshNode::GetCustomVertexCount(int32& AttributeValue) const
 {
 	IMPLEMENT_NODE_ATTRIBUTE_GETTER(VertexCount, int32);
@@ -423,29 +463,29 @@ bool UInterchangeMeshNode::RemoveSkeletonDependencyUid(const FString& Dependency
 	return SkeletonDependencies.RemoveItem(DependencyUid);
 }
 
-int32 UInterchangeMeshNode::GetShapeDependeciesCount() const
+int32 UInterchangeMeshNode::GetMorphTargetDependeciesCount() const
 {
-	return ShapeDependencies.GetCount();
+	return MorphTargetDependencies.GetCount();
 }
 
-void UInterchangeMeshNode::GetShapeDependencies(TArray<FString>& OutDependencies) const
+void UInterchangeMeshNode::GetMorphTargetDependencies(TArray<FString>& OutDependencies) const
 {
-	ShapeDependencies.GetItems(OutDependencies);
+	MorphTargetDependencies.GetItems(OutDependencies);
 }
 
-void UInterchangeMeshNode::GetShapeDependency(const int32 Index, FString& OutDependency) const
+void UInterchangeMeshNode::GetMorphTargetDependency(const int32 Index, FString& OutDependency) const
 {
-	ShapeDependencies.GetItem(Index, OutDependency);
+	MorphTargetDependencies.GetItem(Index, OutDependency);
 }
 
-bool UInterchangeMeshNode::SetShapeDependencyUid(const FString& DependencyUid)
+bool UInterchangeMeshNode::SetMorphTargetDependencyUid(const FString& DependencyUid)
 {
-	return ShapeDependencies.AddItem(DependencyUid);
+	return MorphTargetDependencies.AddItem(DependencyUid);
 }
 
-bool UInterchangeMeshNode::RemoveShapeDependencyUid(const FString& DependencyUid)
+bool UInterchangeMeshNode::RemoveMorphTargetDependencyUid(const FString& DependencyUid)
 {
-	return ShapeDependencies.RemoveItem(DependencyUid);
+	return MorphTargetDependencies.RemoveItem(DependencyUid);
 }
 
 int32 UInterchangeMeshNode::GetSceneInstanceUidsCount() const
diff --git a/Engine/Plugins/Experimental/Interchange/Runtime/Source/Nodes/Private/InterchangeSceneNode.cpp b/Engine/Plugins/Experimental/Interchange/Runtime/Source/Nodes/Private/InterchangeSceneNode.cpp
--- a/Engine/Plugins/Experimental/Interchange/Runtime/Source/Nodes/Private/InterchangeSceneNode.cpp
+++ b/Engine/Plugins/Experimental/Interchange/Runtime/Source/Nodes/Private/InterchangeSceneNode.cpp
@@ -98,10 +98,6 @@ FString UInterchangeSceneNode::GetKeyDisplayName(const UE::Interchange::FAttribu
 		}
 		return KeyDisplayName;
 	}
-	else if (NodeAttributeKey == Macro_CustomTransformCurvePayloadKeyKey)
-	{
-		return FString(TEXT("Transform Curve Payload Key"));
-	}
 	return Super::GetKeyDisplayName(NodeAttributeKey);
 }
 
@@ -119,10 +115,6 @@ FString UInterchangeSceneNode::GetAttributeCategory(const UE::Interchange::FAttr
 	{
 		return FString(TEXT("Joint"));
 	}
-	else if (NodeAttributeKey == Macro_CustomTransformCurvePayloadKeyKey)
-	{
-		return FString(TEXT("Animation"));
-	}
 	else if (NodeAttributeKeyString.StartsWith(UE::Interchange::FSceneNodeStaticData::GetNodeSpecializeTypeBaseKey().ToString()))
 	{
 		return FString(TEXT("SpecializeType"));
@@ -263,56 +255,6 @@ bool UInterchangeSceneNode::SetCustomAssetInstanceUid(const FString& AttributeVa
 	IMPLEMENT_NODE_ATTRIBUTE_SETTER_NODELEGATE(AssetInstanceUid, FString);
 }
 
-bool UInterchangeSceneNode::GetCustomIsNodeTransformAnimated(bool& AttributeValue) const
-{
-	IMPLEMENT_NODE_ATTRIBUTE_GETTER(IsNodeTransformAnimated, bool);
-}
-
-bool UInterchangeSceneNode::SetCustomIsNodeTransformAnimated(const bool& AttributeValue)
-{
-	IMPLEMENT_NODE_ATTRIBUTE_SETTER_NODELEGATE(IsNodeTransformAnimated, bool);
-}
-
-bool UInterchangeSceneNode::GetCustomNodeTransformAnimationKeyCount(int32& AttributeValue) const
-{
-	IMPLEMENT_NODE_ATTRIBUTE_GETTER(NodeTransformAnimationKeyCount, int32);
-}
-
-bool UInterchangeSceneNode::SetCustomNodeTransformAnimationKeyCount(const int32& AttributeValue)
-{
-	IMPLEMENT_NODE_ATTRIBUTE_SETTER_NODELEGATE(NodeTransformAnimationKeyCount, int32);
-}
-
-bool UInterchangeSceneNode::GetCustomNodeTransformAnimationStartTime(double& AttributeValue) const
-{
-	IMPLEMENT_NODE_ATTRIBUTE_GETTER(NodeTransformAnimationStartTime, double);
-}
-
-bool UInterchangeSceneNode::SetCustomNodeTransformAnimationStartTime(const double& AttributeValue)
-{
-	IMPLEMENT_NODE_ATTRIBUTE_SETTER_NODELEGATE(NodeTransformAnimationStartTime, double);
-}
-
-bool UInterchangeSceneNode::GetCustomNodeTransformAnimationEndTime(double& AttributeValue) const
-{
-	IMPLEMENT_NODE_ATTRIBUTE_GETTER(NodeTransformAnimationEndTime, double);
-}
-
-bool UInterchangeSceneNode::SetCustomNodeTransformAnimationEndTime(const double& AttributeValue)
-{
-	IMPLEMENT_NODE_ATTRIBUTE_SETTER_NODELEGATE(NodeTransformAnimationEndTime, double);
-}
-
-bool UInterchangeSceneNode::GetCustomTransformCurvePayloadKey(FString& AttributeValue) const
-{
-	IMPLEMENT_NODE_ATTRIBUTE_GETTER(TransformCurvePayloadKey, FString);
-}
-
-bool UInterchangeSceneNode::SetCustomTransformCurvePayloadKey(const FString& AttributeValue)
-{
-	IMPLEMENT_NODE_ATTRIBUTE_SETTER_NODELEGATE(TransformCurvePayloadKey, FString);
-}
-
 void UInterchangeSceneNode::ResetAllGlobalTransformCaches(const UInterchangeBaseNodeContainer* BaseNodeContainer)
 {
 	BaseNodeContainer->IterateNodes([](const FString& NodeUid, UInterchangeBaseNode* Node)
diff --git a/Engine/Plugins/Experimental/Interchange/Runtime/Source/Nodes/Public/InterchangeMeshNode.h b/Engine/Plugins/Experimental/Interchange/Runtime/Source/Nodes/Public/InterchangeMeshNode.h
--- a/Engine/Plugins/Experimental/Interchange/Runtime/Source/Nodes/Public/InterchangeMeshNode.h
+++ b/Engine/Plugins/Experimental/Interchange/Runtime/Source/Nodes/Public/InterchangeMeshNode.h
@@ -15,12 +15,13 @@ namespace UE
 
 		struct INTERCHANGENODES_API FMeshNodeStaticData : public FBaseNodeStaticData
 		{
-			static const FAttributeKey& PayloadSourceFileKey();
+			static const FAttributeKey& PayloadSourceKey();
+			static const FAttributeKey& PayloadAnimationCurveKey();
 			static const FAttributeKey& IsSkinnedMeshKey();
-			static const FAttributeKey& IsBlendShapeKey();
-			static const FAttributeKey& BlendShapeNameKey();
+			static const FAttributeKey& IsMorphTargetKey();
+			static const FAttributeKey& MorphTargetNameKey();
 			static const FAttributeKey& GetSkeletonDependenciesKey();
-			static const FAttributeKey& GetShapeDependenciesKey();
+			static const FAttributeKey& GetMorphTargetDependenciesKey();
 			static const FAttributeKey& GetSceneInstancesUidsKey();
 			static const FAttributeKey& GetSlotMaterialDependenciesKey();
 		};
@@ -76,35 +77,40 @@ public:
 	bool SetSkinnedMesh(const bool bIsSkinnedMesh);
 
 	/**
-	 * Return true if this node represent a blend shape
+	 * Return true if this node represent a morph target
 	 */
 	UFUNCTION(BlueprintCallable, Category = "Interchange | Node | Mesh")
-	bool IsBlendShape() const;
+	bool IsMorphTarget() const;
 
 	/**
-	 * Set the IsBlendShape attribute to determine if this node represent a blend shape.
+	 * Set the IsMorphTarget attribute to determine if this node represent a morph target.
 	 */
 	UFUNCTION(BlueprintCallable, Category = "Interchange | Node | Mesh")
-	bool SetBlendShape(const bool bIsBlendShape);
+	bool SetMorphTarget(const bool bIsMorphTarget);
 
 	/**
-	 * Get the blend shape name.
-	 * Return true if we successfully query the BlendShapeName attribute
+	 * Get the morph target name.
+	 * Return true if we successfully query the MorphTargetName attribute
 	 */
 	UFUNCTION(BlueprintCallable, Category = "Interchange | Node | Mesh")
-	bool GetBlendShapeName(FString& OutBlendShapeName) const;
+	bool GetMorphTargetName(FString& OutMorphTargetName) const;
 
 	/**
-	 * Set the BlendShapeName attribute to determine if this node represent a blend shape.
+	 * Set the MorphTargetName attribute to determine if this node represent a morph target.
 	 */
 	UFUNCTION(BlueprintCallable, Category = "Interchange | Node | Mesh")
-	bool SetBlendShapeName(const FString& BlendShapeName);
+	bool SetMorphTargetName(const FString& MorphTargetName);
 
 	/** Mesh node Interface Begin */
 	virtual const TOptional<FString> GetPayLoadKey() const;
 
 	UFUNCTION(BlueprintCallable, Category = "Interchange | Node | Mesh")
 	virtual void SetPayLoadKey(const FString& PayloadKey);
+
+	virtual const TOptional<FString> GetAnimationCurvePayLoadKey() const;
+
+	UFUNCTION(BlueprintCallable, Category = "Interchange | Node | Mesh")
+	virtual void SetAnimationCurvePayLoadKey(const FString& PayloadKey);
 	
 	/** Query this mesh vertices count. Return false if the attribute was not set.*/
 	UFUNCTION(BlueprintCallable, Category = "Interchange | Node | Mesh")
@@ -209,34 +215,34 @@ public:
 	bool RemoveSkeletonDependencyUid(const FString& DependencyUid);
 
 	/**
-	 * This function allow to retrieve the number of Shape dependencies for this object.
+	 * This function allow to retrieve the number of morph target dependencies for this object.
 	 */
 	UFUNCTION(BlueprintCallable, Category = "Interchange | Node | Mesh")
-	int32 GetShapeDependeciesCount() const;
+	int32 GetMorphTargetDependeciesCount() const;
 
 	/**
-	 * This function allow to retrieve the Shape dependency for this object.
+	 * This function allow to retrieve the morph target dependency for this object.
 	 */
 	UFUNCTION(BlueprintCallable, Category = "Interchange | Node | Mesh")
-	void GetShapeDependencies(TArray<FString>& OutDependencies) const;
+	void GetMorphTargetDependencies(TArray<FString>& OutDependencies) const;
 
 	/**
-	 * This function allow to retrieve one Shape dependency for this object.
+	 * This function allow to retrieve one morph target dependency for this object.
 	 */
 	UFUNCTION(BlueprintCallable, Category = "Interchange | Node | Mesh")
-	void GetShapeDependency(const int32 Index, FString& OutDependency) const;
+	void GetMorphTargetDependency(const int32 Index, FString& OutDependency) const;
 
 	/**
-	 * Add one Shape dependency to this object.
+	 * Add one morph target dependency to this object.
 	 */
 	UFUNCTION(BlueprintCallable, Category = "Interchange | Node | Mesh")
-	bool SetShapeDependencyUid(const FString& DependencyUid);
+	bool SetMorphTargetDependencyUid(const FString& DependencyUid);
 
 	/**
-	 * Remove one Shape dependency from this object.
+	 * Remove one morph target dependency from this object.
 	 */
 	UFUNCTION(BlueprintCallable, Category = "Interchange | Node | Mesh")
-	bool RemoveShapeDependencyUid(const FString& DependencyUid);
+	bool RemoveMorphTargetDependencyUid(const FString& DependencyUid);
 
 	/**
 	 * This function allow to retrieve the number of scene node instancing this mesh.
@@ -245,25 +251,25 @@ public:
 	int32 GetSceneInstanceUidsCount() const;
 
 	/**
-	 * This function allow to retrieve the Shape dependency for this object.
+	 * This function allow to retrieve the asset instances this scene node is refering.
 	 */
 	UFUNCTION(BlueprintCallable, Category = "Interchange | Node | Mesh")
 	void GetSceneInstanceUids(TArray<FString>& OutDependencies) const;
 
 	/**
-	 * This function allow to retrieve one Shape dependency for this object.
+	 * This function allow to retrieve an asset instance this scene node is refering.
 	 */
 	UFUNCTION(BlueprintCallable, Category = "Interchange | Node | Mesh")
 	void GetSceneInstanceUid(const int32 Index, FString& OutDependency) const;
 
 	/**
-	 * Add one Shape dependency to this object.
+	 * Add one asset instance this scene node is refering.
 	 */
 	UFUNCTION(BlueprintCallable, Category = "Interchange | Node | Mesh")
 	bool SetSceneInstanceUid(const FString& DependencyUid);
 
 	/**
-	 * Remove one Shape dependency from this object.
+	 * Remove one asset instance this scene node is refering.
 	 */
 	UFUNCTION(BlueprintCallable, Category = "Interchange | Node | Mesh")
 	bool RemoveSceneInstanceUid(const FString& DependencyUid);
@@ -305,7 +311,7 @@ private:
 
 	UE::Interchange::TArrayAttributeHelper<FString> SkeletonDependencies;
 	UE::Interchange::TArrayAttributeHelper<FString> MaterialDependencies;
-	UE::Interchange::TArrayAttributeHelper<FString> ShapeDependencies;
+	UE::Interchange::TArrayAttributeHelper<FString> MorphTargetDependencies;
 	UE::Interchange::TArrayAttributeHelper<FString> SceneInstancesUids;
 	
 	UE::Interchange::TMapAttributeHelper<FString, FString> SlotMaterialDependencies;
diff --git a/Engine/Plugins/Experimental/Interchange/Runtime/Source/Nodes/Public/InterchangeSceneNode.h b/Engine/Plugins/Experimental/Interchange/Runtime/Source/Nodes/Public/InterchangeSceneNode.h
--- a/Engine/Plugins/Experimental/Interchange/Runtime/Source/Nodes/Public/InterchangeSceneNode.h
+++ b/Engine/Plugins/Experimental/Interchange/Runtime/Source/Nodes/Public/InterchangeSceneNode.h
@@ -28,6 +28,11 @@ namespace UE
 	}//ns Interchange
 }//ns UE
 
+/**
+ * The scene node represent a transform node in the scene
+ * Scene node can have animations: Use UInterchangeAnimationAPI to get\set animation datas
+ * Scene node can have user defined attribute. Use UInterchangeUserDefinedAttributesAPI to get\set user define attribute data
+ */
 UCLASS(BlueprintType, Experimental)
 class INTERCHANGENODES_API UInterchangeSceneNode : public UInterchangeBaseNode
 {
@@ -66,7 +71,7 @@ public:
 	UFUNCTION(BlueprintCallable, Category = "Interchange | Node | Scene")
 	bool IsSpecializedTypeContains(const FString& SpecializedType) const;
 
-	/** Specialized type are scene node special type like (Joint or LODGroup).*/
+	/** Get the Specialized type this scene node represent (Joint, LODGroup, ...).*/
 	UFUNCTION(BlueprintCallable, Category = "Interchange | Node | Scene")
 	int32 GetSpecializedTypeCount() const;
 
@@ -82,105 +87,84 @@ public:
 	UFUNCTION(BlueprintCallable, Category = "Interchange | Node | Scene")
 	bool RemoveSpecializedType(const FString& SpecializedType);
 
-	//Default transform is the transform we have in the node (no bind pose, no time evaluation).
+	/** Get which asset, if any, a scene node is instantiating. Return false if the Attribute was not set previously.*/
+	UFUNCTION(BlueprintCallable, Category = "Interchange | Node | Scene")
+	bool GetCustomAssetInstanceUid(FString& AttributeValue) const;
+
+	/** Add asset this scene node is instantiating */
+	UFUNCTION(BlueprintCallable, Category = "Interchange | Node | Scene")
+	bool SetCustomAssetInstanceUid(const FString& AttributeValue);
 
-	/** Return the default scene node local transform. */
+
+	/**
+	 * Get the default scene node local transform.
+	 * Default transform is the local transform we have in the node(no bind pose, no time evaluation).
+	 */
 	UFUNCTION(BlueprintCallable, Category = "Interchange | Node | Scene")
 	bool GetCustomLocalTransform(FTransform& AttributeValue) const;
 
-	/** Store the default scene node local transform. */
+	/**
+	 * Set the default scene node local transform.
+	 * Default transform is the local transform we have in the node(no bind pose, no time evaluation).
+	 */
 	UFUNCTION(BlueprintCallable, Category = "Interchange | Node | Scene")
 	bool SetCustomLocalTransform(const UInterchangeBaseNodeContainer* BaseNodeContainer, const FTransform& AttributeValue, bool bResetCache = true);
 
-	/** Return the default scene node global transform. This value is computed with all parent local transform. */
+	/** Get the default scene node global transform. This value is computed with all parent local transform. */
 	UFUNCTION(BlueprintCallable, Category = "Interchange | Node | Scene")
 	bool GetCustomGlobalTransform(const UInterchangeBaseNodeContainer* BaseNodeContainer, const FTransform& GlobalOffsetTransform, FTransform& AttributeValue, bool bForceRecache = false) const;
 
-	//Bind pose transform is the transform of the joint when the binding with the mesh was done.
-	//This attribute should be set only if we have a joint.
 
-	/** Return the bind pose scene node local transform. */
+	/** Get the geometric offset. Any mesh attach to this scene node will be offset using this transform. */
+	UFUNCTION(BlueprintCallable, Category = "Interchange | Node | Scene")
+	bool GetCustomGeometricTransform(FTransform& AttributeValue) const;
+
+	/** Set the geometric offset. Any mesh attach to this scene node will be offset using this transform. */
+	UFUNCTION(BlueprintCallable, Category = "Interchange | Node | Scene")
+	bool SetCustomGeometricTransform(const FTransform& AttributeValue);
+
+	/***********************************************************************************************
+	* Skeleton bind bone API Begin
+	* 
+	* Bind pose transform is the transform of the joint when the binding with the mesh was done.
+	* This attribute should be set only if this scene node represent a joint.
+	* 
+	* Time zero transform is the transform of the node at time zero.
+	* Pipeline often have the option to evaluate the joint at time zero to create the bind pose.
+	* Time zero bind pose is also use if the translator did not find any bind pose or if we import
+	* unskinned mesh has a skeletalmesh (rigid mesh)
+	*/
+
+	/** Get the bind pose scene node local transform. */
 	UFUNCTION(BlueprintCallable, Category = "Interchange | Node | Joint")
 	bool GetCustomBindPoseLocalTransform(FTransform& AttributeValue) const;
 
-	/** Store the bind pose scene node local transform. */
+	/** Set the bind pose scene node local transform. */
 	UFUNCTION(BlueprintCallable, Category = "Interchange | Node | Joint")
 	bool SetCustomBindPoseLocalTransform(const UInterchangeBaseNodeContainer* BaseNodeContainer, const FTransform& AttributeValue, bool bResetCache = true);
 
-	/** Return the bind pose scene node global transform. This value is computed with all parent bind pose local transform. */
+	/** Get the bind pose scene node global transform. This value is computed with all parent bind pose local transform. */
 	UFUNCTION(BlueprintCallable, Category = "Interchange | Node | Joint")
 	bool GetCustomBindPoseGlobalTransform(const UInterchangeBaseNodeContainer* BaseNodeContainer, const FTransform& GlobalOffsetTransform, FTransform& AttributeValue, bool bForceRecache = false) const;
 
 	//Time zero transform is the transform of the node at time zero.
 	//This is useful when there is no bind pose or when we import rigid mesh.
 
-	/** Return the time zero scene node local transform. */
+	/** Get the time zero scene node local transform. */
 	UFUNCTION(BlueprintCallable, Category = "Interchange | Node | Joint")
 	bool GetCustomTimeZeroLocalTransform(FTransform& AttributeValue) const;
 
-	/** Store the time zero scene node local transform. */
+	/** Set the time zero scene node local transform. */
 	UFUNCTION(BlueprintCallable, Category = "Interchange | Node | Joint")
 	bool SetCustomTimeZeroLocalTransform(const UInterchangeBaseNodeContainer* BaseNodeContainer, const FTransform& AttributeValue, bool bResetCache = true);
 
-	/** Return the time zero scene node global transform. This value is computed with all parent timezero local transform. */
+	/** Get the time zero scene node global transform. This value is computed with all parent timezero local transform. */
 	UFUNCTION(BlueprintCallable, Category = "Interchange | Node | Joint")
 	bool GetCustomTimeZeroGlobalTransform(const UInterchangeBaseNodeContainer* BaseNodeContainer, const FTransform& GlobalOffsetTransform, FTransform& AttributeValue, bool bForceRecache = false) const;
 
-	/** Return the geometric offset. Any mesh attach to this scene node will be offset using this transform. */
-	UFUNCTION(BlueprintCallable, Category = "Interchange | Node | Joint")
-	bool GetCustomGeometricTransform(FTransform& AttributeValue) const;
-
-	/** Store the geometric offset. Any mesh attach to this scene node will be offset using this transform. */
-	UFUNCTION(BlueprintCallable, Category = "Interchange | Node | Joint")
-	bool SetCustomGeometricTransform(const FTransform& AttributeValue);
-
-	/** Return false if the Attribute was not set previously.*/
-	UFUNCTION(BlueprintCallable, Category = "Interchange | Node | Scene")
-	bool GetCustomAssetInstanceUid(FString& AttributeValue) const;
-
-	/** Tells which asset, if any, a scene node is instantiating */
-	UFUNCTION(BlueprintCallable, Category = "Interchange | Node | Scene")
-	bool SetCustomAssetInstanceUid(const FString& AttributeValue);
-
-	/** Return true if the node transform have animation, false if the attribute is not set or there is no animation for the transform. */
-	UFUNCTION(BlueprintCallable, Category = "Interchange | Node | Scene | animation")
-	bool GetCustomIsNodeTransformAnimated(bool& AttributeValue) const;
-
-	/** Setter to be able to query if the node transform has animation. */
-	UFUNCTION(BlueprintCallable, Category = "Interchange | Node | Scene | animation")
-	bool SetCustomIsNodeTransformAnimated(const bool& AttributeValue);
-
-	/** Return false if the Attribute was not set previously. This attribute store the payload key needed to retrieve the animation from the source*/
-	UFUNCTION(BlueprintCallable, Category = "Interchange | Node | Scene | animation")
-	bool GetCustomNodeTransformAnimationKeyCount(int32& AttributeValue) const;
-
-	/** If the node transform is animated, this set the payload key needed to retrieve the animation from the source. */
-	UFUNCTION(BlueprintCallable, Category = "Interchange | Node | Scene | animation")
-	bool SetCustomNodeTransformAnimationKeyCount(const int32& AttributeValue);
-
-	/** If the node transform is animated, this should return the starting frame of the animation for this node transform*/
-	UFUNCTION(BlueprintCallable, Category = "Interchange | Node | Scene | animation")
-	bool GetCustomNodeTransformAnimationStartTime(double& AttributeValue) const;
-
-	/** If the node transform is animated, this set the starting frame of the animation for this node transform*/
-	UFUNCTION(BlueprintCallable, Category = "Interchange | Node | Scene | animation")
-	bool SetCustomNodeTransformAnimationStartTime(const double& AttributeValue);
-
-	/** If the node transform is animated, this should return the ending frame of the animation for this node transform*/
-	UFUNCTION(BlueprintCallable, Category = "Interchange | Node | Scene | animation")
-	bool GetCustomNodeTransformAnimationEndTime(double& AttributeValue) const;
-
-	/** If the node transform is animated, this set the ending frame of the animation for this node transform*/
-	UFUNCTION(BlueprintCallable, Category = "Interchange | Node | Scene | animation")
-	bool SetCustomNodeTransformAnimationEndTime(const double& AttributeValue);
-
-	/** Return false if the Attribute was not set previously. This attribute store the payload key needed to retrieve the animation from the source*/
-	UFUNCTION(BlueprintCallable, Category = "Interchange | Node | Scene | animation")
-	bool GetCustomTransformCurvePayloadKey(FString& AttributeValue) const;
-
-	/** If the node transform is animated, this is the where we store the payload key needed to retrieve the animation from the source. */
-	UFUNCTION(BlueprintCallable, Category = "Interchange | Node | Scene | animation")
-	bool SetCustomTransformCurvePayloadKey(const FString& AttributeValue);
+	/*
+	* Skeleton bone API End
+	***********************************************************************************************/
 
 	/** This static function make sure all the global transform caches are reset for all the UInterchangeSceneNode nodes in the UInterchangeBaseNodeContainer */
 	static void ResetAllGlobalTransformCaches(const UInterchangeBaseNodeContainer* BaseNodeContainer);
@@ -191,25 +175,25 @@ public:
 	/**
 	 * Allow to retrieve the correspondence table between slot names and assigned materials for this object.
 	 */
-	UFUNCTION(BlueprintCallable, Category = "Interchange | Node | StaticMesh")
+	UFUNCTION(BlueprintCallable, Category = "Interchange | Node | Meshes")
 	void GetSlotMaterialDependencies(TMap<FString, FString>& OutMaterialDependencies) const;
 
 	/**
 	 * Allow to retrieve one Material dependency for a given slot of this object.
 	 */
-	UFUNCTION(BlueprintCallable, Category = "Interchange | Node | StaticMesh")
+	UFUNCTION(BlueprintCallable, Category = "Interchange | Node | Meshes")
 	bool GetSlotMaterialDependencyUid(const FString& SlotName, FString& OutMaterialDependency) const;
 
 	/**
 	 * Add one Material dependency to a specific slot name of this object.
 	 */
-	UFUNCTION(BlueprintCallable, Category = "Interchange | Node | StaticMesh")
+	UFUNCTION(BlueprintCallable, Category = "Interchange | Node | Meshes")
 	bool SetSlotMaterialDependencyUid(const FString& SlotName, const FString& MaterialDependencyUid);
 
 	/**
 	 * Remove the Material dependency associated with the given slot name from this object.
 	 */
-	UFUNCTION(BlueprintCallable, Category = "Interchange | Node | StaticMesh")
+	UFUNCTION(BlueprintCallable, Category = "Interchange | Node | Meshes")
 	bool RemoveSlotMaterialDependencyUid(const FString& SlotName);
 
 private:
@@ -221,23 +205,28 @@ private:
 		, FTransform& AttributeValue
 		, bool bForceRecache) const;
 
-	//Scene Attribute Keys
+	//Scene node default local transforms.
 	const UE::Interchange::FAttributeKey Macro_CustomLocalTransformKey = UE::Interchange::FAttributeKey(TEXT("LocalTransform"));
+	
+	//Scene node Local bind pose transforms (the specialize type should be set to joint)
 	const UE::Interchange::FAttributeKey Macro_CustomBindPoseLocalTransformKey = UE::Interchange::FAttributeKey(TEXT("BindPoseLocalTransform"));
+	
+	//Scene node local transforms at time zero. This attribute is important for rigid mesh import or if the translator did not fill the bind pose.
 	const UE::Interchange::FAttributeKey Macro_CustomTimeZeroLocalTransformKey = UE::Interchange::FAttributeKey(TEXT("TimeZeroLocalTransform"));
+	
+	//A scene node can have a transform apply to the mesh it reference.
 	const UE::Interchange::FAttributeKey Macro_CustomGeometricTransformKey = UE::Interchange::FAttributeKey(TEXT("GeometricTransform"));
-	const UE::Interchange::FAttributeKey Macro_CustomAssetInstanceUidKey = UE::Interchange::FAttributeKey(TEXT("AssetInstanceUid"));
-	//Animation data
-	const UE::Interchange::FAttributeKey Macro_CustomIsNodeTransformAnimatedKey = UE::Interchange::FAttributeKey(TEXT("IsNodeTransformAnimated"));
-	const UE::Interchange::FAttributeKey Macro_CustomNodeTransformAnimationKeyCountKey = UE::Interchange::FAttributeKey(TEXT("NodeTransformAnimationKeyCount"));
-	const UE::Interchange::FAttributeKey Macro_CustomNodeTransformAnimationStartTimeKey = UE::Interchange::FAttributeKey(TEXT("NodeTransformAnimationStartTime"));
-	const UE::Interchange::FAttributeKey Macro_CustomNodeTransformAnimationEndTimeKey = UE::Interchange::FAttributeKey(TEXT("NodeTransformAnimationEndTime"));
-	const UE::Interchange::FAttributeKey Macro_CustomTransformCurvePayloadKeyKey = UE::Interchange::FAttributeKey(TEXT("TransformCurvePayloadKey"));
 
+	//A scene node can reference an asset. Asset can be Mesh, Light, camera...
+	const UE::Interchange::FAttributeKey Macro_CustomAssetInstanceUidKey = UE::Interchange::FAttributeKey(TEXT("AssetInstanceUid"));
+	
+	//A scene node can represent many special types
 	UE::Interchange::TArrayAttributeHelper<FString> NodeSpecializeTypes;
 
+	//A scene node can have is own set of materials for the mesh it reference.
 	UE::Interchange::TMapAttributeHelper<FString, FString> SlotMaterialDependencies;
 
+	//mutable caches for global transforms
 	mutable TOptional<FTransform> CacheGlobalTransform;
 	mutable TOptional<FTransform> CacheBindPoseGlobalTransform;
 	mutable TOptional<FTransform> CacheTimeZeroGlobalTransform;
diff --git a/Engine/Plugins/Experimental/Interchange/Runtime/Source/Parsers/CommonParser/Private/InterchangeCommonAnimationPayload.cpp b/Engine/Plugins/Experimental/Interchange/Runtime/Source/Parsers/CommonParser/Private/InterchangeCommonAnimationPayload.cpp
--- a/Engine/Plugins/Experimental/Interchange/Runtime/Source/Parsers/CommonParser/Private/InterchangeCommonAnimationPayload.cpp
+++ b/Engine/Plugins/Experimental/Interchange/Runtime/Source/Parsers/CommonParser/Private/InterchangeCommonAnimationPayload.cpp
@@ -6,7 +6,20 @@
 
 #if WITH_ENGINE
 #include "Curves/RichCurve.h"
+#endif //WITH_ENGINE
 
+namespace UE::Interchange
+{
+	void FAnimationBakeTransformPayloadData::Serialize(FArchive& Ar)
+	{
+		Ar << BakeFrequency;
+		Ar << RangeStartTime;
+		Ar << RangeEndTime;
+		Ar << Transforms;
+	}
+}
+
+#if WITH_ENGINE
 void FInterchangeCurveKey::ToRichCurveKey(FRichCurveKey& RichCurveKey) const
 {
 	RichCurveKey.Time = Time;
@@ -62,7 +75,22 @@ void FInterchangeCurveKey::ToRichCurveKey(FRichCurveKey& RichCurveKey) const
 	RichCurveKey.LeaveTangentWeight = LeaveTangentWeight;
 
 }
+#endif //WITH_ENGINE
+
+void FInterchangeCurveKey::Serialize(FArchive& Ar)
+{
+	Ar << InterpMode;
+	Ar << TangentMode;
+	Ar << TangentWeightMode;
+	Ar << Time;
+	Ar << Value;
+	Ar << ArriveTangent;
+	Ar << ArriveTangentWeight;
+	Ar << LeaveTangent;
+	Ar << LeaveTangentWeight;
+}
 
+#if WITH_ENGINE
 void FInterchangeCurve::ToRichCurve(FRichCurve& OutRichCurve) const
 {
 	OutRichCurve.Keys.Reserve(Keys.Num());
@@ -73,5 +101,78 @@ void FInterchangeCurve::ToRichCurve(FRichCurve& OutRichCurve) const
 	}
 	OutRichCurve.AutoSetTangents();
 }
+#endif //WITH_ENGINE
+
+void FInterchangeCurve::Serialize(FArchive& Ar)
+{
+	Ar << Keys;
+}
+
+void FInterchangeStepCurve::RemoveRedundantKeys(float Threshold)
+{
+	const int32 KeyCount = KeyTimes.Num();
+	if (KeyCount < 2)
+	{
+		//Nothing to optimize
+		return;
+	}
+
+	if (FloatKeyValues.IsSet())
+	{
+		InternalRemoveRedundantKey<float>(FloatKeyValues.GetValue(), [Threshold](const float& ValueA, const float& ValueB)
+			{
+				return FMath::IsNearlyEqual(ValueA, ValueB, Threshold);
+			});
+	}
+	else if (IntegerKeyValues.IsSet())
+	{
+		InternalRemoveRedundantKey<int32>(IntegerKeyValues.GetValue(), [](const int32& ValueA, const int32& ValueB)
+			{
+				return ValueA == ValueB;
+			});
+	}
+	else if (StringKeyValues.IsSet())
+	{
+		InternalRemoveRedundantKey<FString>(StringKeyValues.GetValue(), [](const FString& ValueA, const FString& ValueB)
+			{
+				return ValueA == ValueB;
+			});
+	}
+}
+
+void FInterchangeStepCurve::Serialize(FArchive& Ar)
+{
+	Ar << KeyTimes;
+	Ar << FloatKeyValues;
+	Ar << IntegerKeyValues;
+	Ar << StringKeyValues;
+}
+
+template<typename ValueType>
+void FInterchangeStepCurve::InternalRemoveRedundantKey(TArray<ValueType>& Values, TFunctionRef<bool(const ValueType& ValueA, const ValueType& ValueB)> CompareFunction)
+{
+	const int32 KeyCount = Values.Num();
+	TArray<float> NewKeyTimes;
+	NewKeyTimes.Reserve(KeyCount);
+	TArray<ValueType> NewValues;
+	NewValues.Reserve(KeyCount);
+	ValueType LastValue = Values[0];
+	for (int32 KeyIndex = 1; KeyIndex < KeyCount; ++KeyIndex)
+	{
+		//Only add the not equal keys
+		if (!CompareFunction(LastValue, Values[KeyIndex]))
+		{
+			LastValue = Values[KeyIndex];
+			NewKeyTimes.Add(KeyTimes[KeyIndex]);
+			NewValues.Add(Values[KeyIndex]);
+		}
+	}
+	NewKeyTimes.Shrink();
+	NewValues.Shrink();
+	KeyTimes = MoveTemp(NewKeyTimes);
+	Values = MoveTemp(NewValues);
+}
 
-#endif //WITH_ENGINE
\ No newline at end of file
+template void FInterchangeStepCurve::InternalRemoveRedundantKey<float>(TArray<float>& Values, TFunctionRef<bool(const float& ValueA, const float& ValueB)> CompareFunction);
+template void FInterchangeStepCurve::InternalRemoveRedundantKey<int32>(TArray<int32>& Values, TFunctionRef<bool(const int32& ValueA, const int32& ValueB)> CompareFunction);
+template void FInterchangeStepCurve::InternalRemoveRedundantKey<FString>(TArray<FString>& Values, TFunctionRef<bool(const FString& ValueA, const FString& ValueB)> CompareFunction);
\ No newline at end of file
diff --git a/Engine/Plugins/Experimental/Interchange/Runtime/Source/Parsers/CommonParser/Public/InterchangeCommonAnimationPayload.h b/Engine/Plugins/Experimental/Interchange/Runtime/Source/Parsers/CommonParser/Public/InterchangeCommonAnimationPayload.h
--- a/Engine/Plugins/Experimental/Interchange/Runtime/Source/Parsers/CommonParser/Public/InterchangeCommonAnimationPayload.h
+++ b/Engine/Plugins/Experimental/Interchange/Runtime/Source/Parsers/CommonParser/Public/InterchangeCommonAnimationPayload.h
@@ -16,7 +16,10 @@
 
 namespace UE::Interchange
 {
-
+	/**
+	 * This payload class is use to get a scene node bake transform payload
+	 * The translator should bake the scene node transform using the bake settings provide by the factory.
+	 */
 	struct INTERCHANGECOMMONPARSER_API FAnimationBakeTransformPayloadData
 	{
 		double BakeFrequency = 30.0;
@@ -24,32 +27,10 @@ namespace UE::Interchange
 		double RangeEndTime = 1.0 / BakeFrequency;
 		TArray<FTransform> Transforms;
 
-		void Serialize(FArchive& Ar)
-		{
-			Ar << BakeFrequency;
-			Ar << RangeStartTime;
-			Ar << RangeEndTime;
-			Ar << Transforms;
-		}
+		void Serialize(FArchive& Ar);
 	};
 }
 
-UENUM()
-enum class EInterchangeTransformCurveChannel
-{
-	TranslationX = 0,
-	TranslationY = 1,
-	TranslationZ = 2,
-	EulerX = 3,
-	EulerY = 4,
-	EulerZ = 5,
-	ScaleX = 6,
-	ScaleY = 7,
-	ScaleZ = 8,
-	TransformChannelCount = 9,
-	None
-};
-
 /** If using Cubic, this enum describes how the tangents should be controlled. */
 UENUM()
 enum class EInterchangeCurveInterpMode : uint8
@@ -142,38 +123,23 @@ struct INTERCHANGECOMMONPARSER_API FInterchangeCurveKey
 	void ToRichCurveKey(FRichCurveKey& OutKey) const;
 #endif
 
-	void Serialize(FArchive& Ar)
-	{
-		Ar << InterpMode;
-		Ar << TangentMode;
-		Ar << TangentWeightMode;
-		Ar << Time;
-		Ar << Value;
-		Ar << ArriveTangent;
-		Ar << ArriveTangentWeight;
-		Ar << LeaveTangent;
-		Ar << LeaveTangentWeight;
-	}
+	void Serialize(FArchive& Ar);
 
-	friend FArchive& operator<<(FArchive& Ar, FInterchangeCurveKey& P)
+	friend FArchive& operator<<(FArchive& Ar, FInterchangeCurveKey& InterchangeCurveKey)
 	{
-		P.Serialize(Ar);
+		InterchangeCurveKey.Serialize(Ar);
 		return Ar;
 	}
 };
 
 /**
-* This struct contains only the key data, this is only used to pass animation data from translators to factories.
-* You cannot evaluate a curve with this struct.
+* This struct contains only the key data, this is only used to pass animation data from interchange worker process translators to factories.
 */
 USTRUCT()
 struct INTERCHANGECOMMONPARSER_API FInterchangeCurve
 {
 	GENERATED_BODY()
 
-	UPROPERTY()
-	EInterchangeTransformCurveChannel TransformChannel = EInterchangeTransformCurveChannel::None;
-	
 	UPROPERTY()
 	TArray<FInterchangeCurveKey> Keys;
 
@@ -182,15 +148,40 @@ struct INTERCHANGECOMMONPARSER_API FInterchangeCurve
 	void ToRichCurve(FRichCurve& OutKey) const;
 #endif
 
-	void Serialize(FArchive& Ar)
+	void Serialize(FArchive& Ar);
+
+	friend FArchive& operator<<(FArchive& Ar, FInterchangeCurve& InterchangeCurve)
 	{
-		Ar << TransformChannel;
-		Ar << Keys;
+		InterchangeCurve.Serialize(Ar);
+		return Ar;
 	}
+};
 
-	friend FArchive& operator<<(FArchive& Ar, FInterchangeCurve& P)
+/**
+* This struct contains only the key data, this is only used to pass animation data from translators to factories.
+*/
+USTRUCT()
+struct INTERCHANGECOMMONPARSER_API FInterchangeStepCurve
+{
+	GENERATED_BODY()
+
+	UPROPERTY()
+	TArray<float> KeyTimes;
+	TOptional<TArray<float>> FloatKeyValues;
+	TOptional<TArray<int32>> IntegerKeyValues;
+	TOptional<TArray<FString>> StringKeyValues;
+
+	void RemoveRedundantKeys(float Threshold);
+
+	void Serialize(FArchive& Ar);
+
+	friend FArchive& operator<<(FArchive& Ar, FInterchangeStepCurve& InterchangeStepCurve)
 	{
-		P.Serialize(Ar);
+		InterchangeStepCurve.Serialize(Ar);
 		return Ar;
 	}
+
+private:
+	template<typename ValueType>
+	void InternalRemoveRedundantKey(TArray<ValueType>& Values, TFunctionRef<bool(const ValueType& ValueA, const ValueType& ValueB)> CompareFunction);
 };
diff --git a/Engine/Plugins/Experimental/Interchange/Runtime/Source/Parsers/Fbx/Private/FbxAPI.h b/Engine/Plugins/Experimental/Interchange/Runtime/Source/Parsers/Fbx/Private/FbxAPI.h
--- a/Engine/Plugins/Experimental/Interchange/Runtime/Source/Parsers/Fbx/Private/FbxAPI.h
+++ b/Engine/Plugins/Experimental/Interchange/Runtime/Source/Parsers/Fbx/Private/FbxAPI.h
@@ -32,7 +32,7 @@ namespace UE
 			class FFbxParser
 			{
 			public:
-				explicit FFbxParser(UInterchangeResultsContainer* InResultsContainer)
+				explicit FFbxParser(TWeakObjectPtr<UInterchangeResultsContainer> InResultsContainer)
 					: ResultsContainer(InResultsContainer)
 				{}
 
@@ -55,7 +55,7 @@ namespace UE
 				template <typename T>
 				T* AddMessage() const
 				{
-					check(ResultsContainer != nullptr);
+					check(ResultsContainer.IsValid());
 					T* Item = ResultsContainer->Add<T>();
 					Item->SourceAssetName = SourceFilename;
 					return Item;
@@ -64,7 +64,7 @@ namespace UE
 
 				void AddMessage(UInterchangeResult* Item) const
 				{
-					check(ResultsContainer != nullptr);
+					check(ResultsContainer.IsValid());
 					ResultsContainer->Add(Item);
 					Item->SourceAssetName = SourceFilename;
 				}
@@ -79,7 +79,7 @@ namespace UE
 
 				void CleanupFbxData();
 
-				UInterchangeResultsContainer* ResultsContainer;
+				TWeakObjectPtr<UInterchangeResultsContainer> ResultsContainer;
 				FbxManager* SDKManager = nullptr;
 				FbxScene* SDKScene = nullptr;
 				FbxImporter* SDKImporter = nullptr;
diff --git a/Engine/Plugins/Experimental/Interchange/Runtime/Source/Parsers/Fbx/Private/FbxAnimation.cpp b/Engine/Plugins/Experimental/Interchange/Runtime/Source/Parsers/Fbx/Private/FbxAnimation.cpp
--- a/Engine/Plugins/Experimental/Interchange/Runtime/Source/Parsers/Fbx/Private/FbxAnimation.cpp
+++ b/Engine/Plugins/Experimental/Interchange/Runtime/Source/Parsers/Fbx/Private/FbxAnimation.cpp
@@ -15,7 +15,7 @@
 #include "MeshDescription.h"
 #include "Misc/FileHelper.h"
 #include "Misc/ScopeLock.h"
-#include "Nodes/InterchangeBaseNodeContainer.h"
+#include "Nodes/InterchangeAnimationAPI.h"
 #include "Serialization/LargeMemoryWriter.h"
 #include "SkeletalMeshAttributes.h"
 #include "StaticMeshAttributes.h"
@@ -25,67 +25,8 @@
 
 namespace UE::Interchange::Private
 {
-	float GetTransformChannelValue(const EInterchangeTransformCurveChannel Channel, const FTransform& Transform)
-	{
-		switch (Channel)
-		{
-			case EInterchangeTransformCurveChannel::TranslationX:
-			{
-				return Transform.GetLocation().X;
-			}
-			break;
-			case EInterchangeTransformCurveChannel::TranslationY:
-			{
-				return Transform.GetLocation().Y;
-			}
-			break;
-			case EInterchangeTransformCurveChannel::TranslationZ:
-			{
-				return Transform.GetLocation().Z;
-			}
-			break;
-			case EInterchangeTransformCurveChannel::EulerX:
-			{
-				return Transform.GetRotation().Euler().X;
-			}
-			break;
-			case EInterchangeTransformCurveChannel::EulerY:
-			{
-				return Transform.GetRotation().Euler().Y;
-			}
-			break;
-			case EInterchangeTransformCurveChannel::EulerZ:
-			{
-				return Transform.GetRotation().Euler().Z;
-			}
-			break;
-			case EInterchangeTransformCurveChannel::ScaleX:
-			{
-				return Transform.GetScale3D().X;
-			}
-			break;
-			case EInterchangeTransformCurveChannel::ScaleY:
-			{
-				return Transform.GetScale3D().Y;
-			}
-			break;
-			case EInterchangeTransformCurveChannel::ScaleZ:
-			{
-				return Transform.GetScale3D().Z;
-			}
-			break;
-		}
-		return 0.0f;
-	}
-
-	bool ImportCurve(FbxNode* Node, const EInterchangeTransformCurveChannel TransformChannel, const FbxAnimCurve* SourceFloatCurves, const float ScaleValue, TArray<FInterchangeCurveKey>& DestinationFloatCurve)
+	bool ImportCurve(const FbxAnimCurve* SourceFloatCurves, const float ScaleValue, TArray<FInterchangeCurveKey>& DestinationFloatCurve)
 	{
-		bool bUseNodeTransform = false;
-		if (TransformChannel != EInterchangeTransformCurveChannel::None)
-		{
-			bUseNodeTransform = Node != nullptr;
-		}
-
 		if (!SourceFloatCurves)
 		{
 			return true;
@@ -99,20 +40,7 @@ namespace UE::Interchange::Private
 			FbxAnimCurveKey Key = SourceFloatCurves->KeyGet(KeyIndex);
 			FbxTime KeyTime = Key.GetTime();
 			const float KeyTimeValue = static_cast<float>(KeyTime.GetSecondDouble());
-			float ValueAtTime = Key.GetValue();
-			if (bUseNodeTransform)
-			{
-				FbxAMatrix NodeTransform = Node->EvaluateGlobalTransform(KeyTime);
-				FbxNode* ParentNode = Node->GetParent();
-				if (ParentNode)
-				{
-					FbxAMatrix ParentTransform = ParentNode->EvaluateGlobalTransform(KeyTime);
-					NodeTransform = ParentTransform.Inverse() * NodeTransform;
-				}
-				FTransform Transform = UE::Interchange::Private::FFbxConvert::ConvertTransform(NodeTransform);
-				ValueAtTime = GetTransformChannelValue(TransformChannel, Transform);
-			}
-			float Value = ValueAtTime * ScaleValue;
+			float Value = Key.GetValue() * ScaleValue;
 			FInterchangeCurveKey& InterchangeCurveKey = DestinationFloatCurve.AddDefaulted_GetRef();
 			InterchangeCurveKey.Time = KeyTimeValue;
 			InterchangeCurveKey.Value = Value;
@@ -345,68 +273,114 @@ namespace UE::Interchange::Private
 		return true;
 	}
 
-	bool ImportCurve(const FbxAnimCurve* SourceFloatCurves, const float ScaleValue, TArray<FInterchangeCurveKey>& DestinationFloatCurve)
+	template<typename AttributeType>
+	void FillStepCurveAttribute(TArray<float>& OutFrameTimes, TArray<AttributeType>& OutFrameValues, const FbxAnimCurve* FbxCurve, TFunctionRef<AttributeType(const FbxAnimCurveKey*, const FbxTime*)> EvaluationFunction)
 	{
-		return ImportCurve(nullptr, EInterchangeTransformCurveChannel::None, SourceFloatCurves, ScaleValue, DestinationFloatCurve);
-	}
+		const int32 KeyCount = FbxCurve ? FbxCurve->KeyGetCount() : 0;
 
-	bool ImportTranslationCurves(FbxNode* Node, TArray<FbxAnimCurve*>& SourceTransformChannelCurves, TArray<FInterchangeCurve>& TransformChannelCurves)
-	{
-		bool bResult = true;
-		//Translation X
-		bResult &= ImportCurve(Node, EInterchangeTransformCurveChannel::TranslationX, SourceTransformChannelCurves[0], 1.0f, TransformChannelCurves[0].Keys);
-		TransformChannelCurves[0].TransformChannel = EInterchangeTransformCurveChannel::TranslationX;
-		//Translation Y
-		bResult &= ImportCurve(Node, EInterchangeTransformCurveChannel::TranslationY, SourceTransformChannelCurves[1], 1.0f, TransformChannelCurves[1].Keys);
-		TransformChannelCurves[1].TransformChannel = EInterchangeTransformCurveChannel::TranslationY;
-		//Translation Z
-		bResult &= ImportCurve(Node, EInterchangeTransformCurveChannel::TranslationZ, SourceTransformChannelCurves[2], 1.0f, TransformChannelCurves[2].Keys);
-		TransformChannelCurves[2].TransformChannel = EInterchangeTransformCurveChannel::TranslationZ;
+		if (KeyCount > 0)
+		{
+			OutFrameTimes.Reserve(KeyCount);
+			OutFrameValues.Reserve(KeyCount);
+			const FbxTime StartTime = FbxCurve->KeyGet(0).GetTime();
 
-		return bResult;
+			for (int32 KeyIndex = 0; KeyIndex < KeyCount; ++KeyIndex)
+			{
+				FbxAnimCurveKey Key = FbxCurve->KeyGet(KeyIndex);
+				FbxTime KeyTime = Key.GetTime() - StartTime;
+
+				OutFrameTimes.Add(KeyTime.GetSecondDouble());
+				OutFrameValues.Add(EvaluationFunction(&Key, &KeyTime));
+			}
+		}
+		else
+		{
+			OutFrameTimes.Add(0);
+			OutFrameValues.Add(EvaluationFunction(nullptr, nullptr));
+		}
 	}
-	
-	bool ImportRotationCurves(FbxNode* Node, TArray<FbxAnimCurve*>& SourceTransformChannelCurves, TArray<FInterchangeCurve>& TransformChannelCurves)
+
+	void ImportFloatStepCurve(const FbxAnimCurve* SourceCurves, FbxProperty& Property, FInterchangeStepCurve& DestinationCurve)
 	{
-		bool bResult = true;
-		//Euler X
-		bResult &= ImportCurve(Node, EInterchangeTransformCurveChannel::EulerX, SourceTransformChannelCurves[3], 1.0f, TransformChannelCurves[3].Keys);
-		TransformChannelCurves[3].TransformChannel = EInterchangeTransformCurveChannel::EulerX;
-		//Euler Y
-		bResult &= ImportCurve(Node, EInterchangeTransformCurveChannel::EulerY, SourceTransformChannelCurves[4], 1.0f, TransformChannelCurves[4].Keys);
-		TransformChannelCurves[4].TransformChannel = EInterchangeTransformCurveChannel::EulerY;
-		//Euler Z
-		bResult &= ImportCurve(Node, EInterchangeTransformCurveChannel::EulerZ, SourceTransformChannelCurves[5], 1.0f, TransformChannelCurves[5].Keys);
-		TransformChannelCurves[5].TransformChannel = EInterchangeTransformCurveChannel::EulerZ;
+		TArray<float> StepCurveValues;
+		FillStepCurveAttribute<float>(DestinationCurve.KeyTimes, StepCurveValues, SourceCurves, [&Property](const FbxAnimCurveKey* Key, const FbxTime* KeyTime)
+			{
+				if (Key)
+				{
+					return Key->GetValue();
+				}
+				else
+				{
+					return Property.Get<float>();
+				}
+			});
+		DestinationCurve.FloatKeyValues = StepCurveValues;
 
-		return bResult;
 	}
-	
-	bool ImportScaleCurves(FbxNode* Node, TArray<FbxAnimCurve*>& SourceTransformChannelCurves, TArray<FInterchangeCurve>& TransformChannelCurves)
+
+	void ImportIntegerStepCurve(const FbxAnimCurve* SourceCurves, FbxProperty& Property, FInterchangeStepCurve& DestinationCurve)
 	{
-		bool bResult = true;
-		//Scale X
-		bResult &= ImportCurve(Node, EInterchangeTransformCurveChannel::ScaleX, SourceTransformChannelCurves[6], 1.0f, TransformChannelCurves[6].Keys);
-		TransformChannelCurves[6].TransformChannel = EInterchangeTransformCurveChannel::ScaleX;
-		//Scale Y
-		bResult &= ImportCurve(Node, EInterchangeTransformCurveChannel::ScaleY, SourceTransformChannelCurves[7], 1.0f, TransformChannelCurves[7].Keys);
-		TransformChannelCurves[7].TransformChannel = EInterchangeTransformCurveChannel::ScaleY;
-		//Scale Z
-		bResult &= ImportCurve(Node, EInterchangeTransformCurveChannel::ScaleZ, SourceTransformChannelCurves[8], 1.0f, TransformChannelCurves[8].Keys);
-		TransformChannelCurves[8].TransformChannel = EInterchangeTransformCurveChannel::ScaleZ;
+		TArray<int32> StepCurveValues;
+		FillStepCurveAttribute<int32>(DestinationCurve.KeyTimes, StepCurveValues, SourceCurves, [&Property](const FbxAnimCurveKey* Key, const FbxTime* KeyTime)
+			{
+				if (Key)
+				{
+					return static_cast<int32>(Key->GetValue());
+				}
+				else
+				{
+					return static_cast<int32>(Property.Get<int32>());
+				}
+			});
+		DestinationCurve.IntegerKeyValues = StepCurveValues;
+	}
 
-		return bResult;
+	void ImportStringStepCurve(const FbxAnimCurve* SourceCurves, FbxProperty& Property, FInterchangeStepCurve& DestinationCurve)
+	{
+		TArray<FString> StepCurveValues;
+		FillStepCurveAttribute<FString>(DestinationCurve.KeyTimes, StepCurveValues, SourceCurves, [&Property](const FbxAnimCurveKey* Key, const FbxTime* KeyTime)
+			{
+				if (KeyTime)
+				{
+					FbxPropertyValue& EvaluatedValue = Property.EvaluateValue(*KeyTime);
+					FbxString StringValue;
+					EvaluatedValue.Get(&StringValue, EFbxType::eFbxString);
+					return FString(UTF8_TO_TCHAR(StringValue));
+				}
+				else
+				{
+					return FString(UTF8_TO_TCHAR(Property.Get<FbxString>()));
+				}
+			});
+		DestinationCurve.StringKeyValues = StepCurveValues;
 	}
 
-	bool ImportTransformCurves(FbxNode* Node
-		, TArray<FbxAnimCurve*>& SourceTransformChannelCurves
-		, TArray<FInterchangeCurve>& TransformChannelCurves)
+	void ImportEnumStepCurve(const FbxAnimCurve* SourceCurves, FbxProperty& Property, FInterchangeStepCurve& DestinationCurve)
 	{
-		bool bResult = true;
-		bResult &= ImportTranslationCurves(Node, SourceTransformChannelCurves, TransformChannelCurves);
-		bResult &= ImportRotationCurves(Node, SourceTransformChannelCurves, TransformChannelCurves);
-		bResult &= ImportScaleCurves(Node, SourceTransformChannelCurves, TransformChannelCurves);
-		return bResult;
+		TArray<FString> StepCurveValues;
+		FillStepCurveAttribute<FString>(DestinationCurve.KeyTimes, StepCurveValues, SourceCurves, [&Property](const FbxAnimCurveKey* Key, const FbxTime* KeyTime)
+			{
+				int32 EnumIndex = -1;
+
+				if (KeyTime)
+				{
+					FbxPropertyValue& EvaluatedValue = Property.EvaluateValue(*KeyTime);
+					EvaluatedValue.Get(&EnumIndex, EFbxType::eFbxEnum);
+				}
+				else
+				{
+					EnumIndex = Property.Get<FbxEnum>();
+				}
+
+				if (EnumIndex < 0 || EnumIndex >= Property.GetEnumCount())
+				{
+					return FString();
+				}
+
+				const char* EnumValue = Property.GetEnumValue(EnumIndex);
+				return FString(UTF8_TO_TCHAR(EnumValue));
+			});
+		DestinationCurve.StringKeyValues = StepCurveValues;
 	}
 
 	bool ImportBakeTransforms(FbxNode* Node, FAnimationBakeTransformPayloadData& AnimationBakeTransformPayloadData)
@@ -429,8 +403,9 @@ namespace UE::Interchange::Private
 		FbxTime TimeStep = 0;
 		TimeStep.SetSecondDouble(TimeStepSecond);
 
-		const int32 NumFrame = FMath::RoundToInt32((AnimationBakeTransformPayloadData.RangeEndTime - AnimationBakeTransformPayloadData.RangeStartTime * AnimationBakeTransformPayloadData.BakeFrequency));
-		
+		const int32 NumFrame = FMath::RoundToInt32((AnimationBakeTransformPayloadData.RangeEndTime - AnimationBakeTransformPayloadData.RangeStartTime) * AnimationBakeTransformPayloadData.BakeFrequency);
+		check(NumFrame >= 0);
+
 		//Add a threshold when we compare if we have reach the end of the animation
 		const FbxTime TimeComparisonThreshold = (UE_DOUBLE_KINDA_SMALL_NUMBER * static_cast<double>(FBXSDK_TC_SECOND));
 		AnimationBakeTransformPayloadData.Transforms.Empty(NumFrame);
@@ -563,24 +538,31 @@ namespace UE::Interchange::Private
 
 	bool FAnimationPayloadContextTransform::FetchPayloadToFile(FFbxParser& Parser, const FString& PayloadFilepath)
 	{
-		//Will need this one to import animated curves
+		if (AttributeFetchPayloadData.IsSet())
+		{
+			return InternalFetchCurveNodePayloadToFile(Parser, PayloadFilepath);
+		}
+		else if (MorphTargetFetchPayloadData.IsSet())
+		{
+			return InternalFetchMorphTargetCurvePayloadToFile(Parser, PayloadFilepath);
+		}
 		return false;
 	}
 
 	bool FAnimationPayloadContextTransform::FetchAnimationBakeTransformPayloadToFile(FFbxParser& Parser, const double BakeFrequency, const double RangeStartTime, const double RangeEndTime, const FString& PayloadFilepath)
 	{
-		if (!ensure(SDKScene != nullptr))
+		if (!ensure(NodeTransformFetchPayloadData.IsSet()))
 		{
 			UInterchangeResultError_Generic* Message = Parser.AddMessage<UInterchangeResultError_Generic>();
-			Message->InterchangeKey = FFbxHelper::GetFbxNodeHierarchyName(Node);
-			Message->Text = LOCTEXT("FBXSceneNull_Animation", "Cannot fetch FBX animation transform payload because the FBX scene is null.");
+			Message->Text = LOCTEXT("NodeTransformFetchPayloadData_NotSet", "Cannot fetch FBX animation transform payload because the FBX FNodeTransformFetchPayloadData is not set.");
 			return false;
 		}
 
-		if (!ensure(Node != nullptr))
+		FNodeTransformFetchPayloadData& FetchPayloadData = NodeTransformFetchPayloadData.GetValue();
+		if (!ensure(FetchPayloadData.Node))
 		{
 			UInterchangeResultError_Generic* Message = Parser.AddMessage<UInterchangeResultError_Generic>();
-			Message->InterchangeKey = FFbxHelper::GetFbxNodeHierarchyName(Node);
+			Message->InterchangeKey = FFbxHelper::GetFbxNodeHierarchyName(FetchPayloadData.Node);
 			Message->Text = LOCTEXT("FBXNodeNull", "Cannot fetch FBX animation transform payload because the FBX node is null.");
 			return false;
 		}
@@ -591,7 +573,7 @@ namespace UE::Interchange::Private
 		AnimationBakeTransformPayloadData.RangeStartTime = RangeStartTime;
 		AnimationBakeTransformPayloadData.RangeEndTime = RangeEndTime;
 
-		ImportBakeTransforms(Node, AnimationBakeTransformPayloadData);
+		ImportBakeTransforms(FetchPayloadData.Node, AnimationBakeTransformPayloadData);
 		{
 			FLargeMemoryWriter Ar;
 			AnimationBakeTransformPayloadData.Serialize(Ar);
@@ -603,7 +585,167 @@ namespace UE::Interchange::Private
 		return true;
 	}
 
-	void FFbxAnimation::AddNodeTransformAnimation(FbxScene* SDKScene, FbxNode* JointNode, UInterchangeBaseNodeContainer& NodeContainer, UInterchangeSceneNode* SceneNode, TMap<FString, TSharedPtr<FPayloadContextBase>>& PayloadContexts)
+	bool FAnimationPayloadContextTransform::InternalFetchCurveNodePayloadToFile(FFbxParser& Parser, const FString& PayloadFilepath)
+	{
+		check(AttributeFetchPayloadData.IsSet());
+		FAttributeFetchPayloadData& FetchPayloadData = AttributeFetchPayloadData.GetValue();
+		if (!ensure(FetchPayloadData.Node != nullptr))
+		{
+			UInterchangeResultError_Generic* Message = Parser.AddMessage<UInterchangeResultError_Generic>();
+			Message->InterchangeKey = FFbxHelper::GetFbxNodeHierarchyName(FetchPayloadData.Node);
+			Message->Text = LOCTEXT("InternalFetchCurveNodePayloadToFile_FBXNodeNull", "Cannot fetch FBX animation curve payload because the FBX node is null.");
+			return false;
+		}
+
+		if (!ensure(FetchPayloadData.AnimCurves != nullptr))
+		{
+			UInterchangeResultError_Generic* Message = Parser.AddMessage<UInterchangeResultError_Generic>();
+			Message->InterchangeKey = FFbxHelper::GetFbxNodeHierarchyName(FetchPayloadData.Node);
+			Message->Text = LOCTEXT("InternalFetchCurveNodePayloadToFile_FBXCurveNull", "Cannot fetch FBX user attribute animation curve payload because the FBX anim curve node is null.");
+			return false;
+		}
+		if (FetchPayloadData.bAttributeTypeIsStepCurveAnimation)
+		{
+			//Fetch TArray<FInterchangeStepCurve> step curve, no interpolation
+			TArray<FInterchangeStepCurve> InterchangeStepCurves;
+			const uint32 ChannelCount = FetchPayloadData.AnimCurves->GetChannelsCount();
+			for (uint32 ChannelIndex = 0; ChannelIndex < ChannelCount; ++ChannelIndex)
+			{
+				const uint32 ChannelCurveCount = FetchPayloadData.AnimCurves->GetCurveCount(ChannelIndex);
+				for (uint32 CurveIndex = 0; CurveIndex < ChannelCurveCount; ++CurveIndex)
+				{
+					if (FbxAnimCurve* CurrentAnimCurve = FetchPayloadData.AnimCurves->GetCurve(ChannelIndex, CurveIndex))
+					{
+						switch (FetchPayloadData.PropertyType)
+						{
+							case EFbxType::eFbxBool:
+							case EFbxType::eFbxChar:
+							case EFbxType::eFbxUChar:
+							case EFbxType::eFbxShort:
+							case EFbxType::eFbxUShort:
+							case EFbxType::eFbxInt:
+							case EFbxType::eFbxUInt:
+							case EFbxType::eFbxLongLong:
+							case EFbxType::eFbxULongLong:
+								ImportIntegerStepCurve(CurrentAnimCurve, FetchPayloadData.Property, InterchangeStepCurves.AddDefaulted_GetRef());
+								break;
+							case EFbxType::eFbxHalfFloat:
+							case EFbxType::eFbxFloat:
+							case EFbxType::eFbxDouble:
+							case EFbxType::eFbxDouble2:
+							case EFbxType::eFbxDouble3:
+							case EFbxType::eFbxDouble4:
+								check(false); //Float curve payload should be extract as FInterchangeCurve since we can interpolate them
+								ImportFloatStepCurve(CurrentAnimCurve, FetchPayloadData.Property, InterchangeStepCurves.AddDefaulted_GetRef());
+								break;
+							case EFbxType::eFbxEnum:
+								ImportEnumStepCurve(CurrentAnimCurve, FetchPayloadData.Property, InterchangeStepCurves.AddDefaulted_GetRef());
+								break;
+							case EFbxType::eFbxString:
+								ImportStringStepCurve(CurrentAnimCurve, FetchPayloadData.Property, InterchangeStepCurves.AddDefaulted_GetRef());
+								break;
+						}
+					}
+				}
+			}
+			{
+				FLargeMemoryWriter Ar;
+				Ar << InterchangeStepCurves;
+				uint8* ArchiveData = Ar.GetData();
+				int64 ArchiveSize = Ar.TotalSize();
+				TArray64<uint8> Buffer(ArchiveData, ArchiveSize);
+				FFileHelper::SaveArrayToFile(Buffer, *PayloadFilepath);
+			}
+		}
+		else
+		{
+			//Fetch TArray<FInterchangeCurve> which are float curves with interpolation
+			TArray<FInterchangeCurve> InterchangeCurves;
+			const uint32 ChannelCount = FetchPayloadData.AnimCurves->GetChannelsCount();
+			for (uint32 ChannelIndex = 0; ChannelIndex < ChannelCount; ++ChannelIndex)
+			{
+				const uint32 ChannelCurveCount = FetchPayloadData.AnimCurves->GetCurveCount(ChannelIndex);
+				for (uint32 CurveIndex = 0; CurveIndex < ChannelCurveCount; ++CurveIndex)
+				{
+					if (FbxAnimCurve* CurrentAnimCurve = FetchPayloadData.AnimCurves->GetCurve(ChannelIndex, CurveIndex))
+					{
+						ImportCurve(CurrentAnimCurve, 1.0f, InterchangeCurves.AddDefaulted_GetRef().Keys);
+					}
+				}
+			}
+			{
+				FLargeMemoryWriter Ar;
+				Ar << InterchangeCurves;
+				uint8* ArchiveData = Ar.GetData();
+				int64 ArchiveSize = Ar.TotalSize();
+				TArray64<uint8> Buffer(ArchiveData, ArchiveSize);
+				FFileHelper::SaveArrayToFile(Buffer, *PayloadFilepath);
+			}
+		}
+		return true;
+	}
+
+	bool FAnimationPayloadContextTransform::InternalFetchMorphTargetCurvePayloadToFile(FFbxParser& Parser, const FString& PayloadFilepath)
+	{
+		check(MorphTargetFetchPayloadData.IsSet());
+		FMorphTargetFetchPayloadData& FetchPayloadData = MorphTargetFetchPayloadData.GetValue();
+
+		if(!FetchPayloadData.SDKScene)
+		{
+			UInterchangeResultError_Generic* Message = Parser.AddMessage<UInterchangeResultError_Generic>();
+			Message->Text = LOCTEXT("InternalFetchMorphTargetCurvePayloadToFile_FBXSDKSceneNull", "InternalFetchMorphTargetCurvePayloadToFile, fbx sdk is nullptr.");
+			return false;
+		}
+		
+		int32 GeometryCount = FetchPayloadData.SDKScene->GetGeometryCount();
+
+		FbxGeometry* Geometry = FetchPayloadData.SDKScene->GetGeometry(FetchPayloadData.GeometryIndex);
+		if (!Geometry || Geometry->GetAttributeType() != FbxNodeAttribute::eMesh)
+		{
+			UInterchangeResultError_Generic* Message = Parser.AddMessage<UInterchangeResultError_Generic>();
+			Message->Text = LOCTEXT("InternalFetchMorphTargetCurvePayloadToFile_FBXGeometryNull", "Cannot fetch FBX geometry from the scene.");
+			return false;
+		}
+
+		FbxAnimCurve* AnimCurve = nullptr;
+		int32 NumAnimations = FetchPayloadData.SDKScene->GetSrcObjectCount<FbxAnimStack>();
+		//Anim stack should be merge so we expect to have only one stack here
+		ensure(NumAnimations <= 1);
+		if (FbxAnimStack* AnimStack = NumAnimations == 1 ? (FbxAnimStack*)FetchPayloadData.SDKScene->GetSrcObject<FbxAnimStack>(0) : nullptr)
+		{
+			AnimCurve = Geometry->GetShapeChannel(FetchPayloadData.MorphTargetIndex, FetchPayloadData.ChannelIndex, (FbxAnimLayer*)AnimStack->GetMember(0));
+		}
+
+		if (!ensure(AnimCurve != nullptr))
+		{
+			UInterchangeResultError_Generic* Message = Parser.AddMessage<UInterchangeResultError_Generic>();
+			if (FbxNode* Node = Geometry->GetNode())
+			{
+				Message->InterchangeKey = FFbxHelper::GetFbxNodeHierarchyName(Node);
+			}
+			Message->Text = LOCTEXT("InternalFetchMorphTargetCurvePayloadToFile_FBXCurveNull", "Cannot fetch FBX morph target animation curve payload because the FBX anim curve node is null.");
+			return false;
+		}
+		TArray<FInterchangeCurve> InterchangeCurves;
+		//Morph target curve in fbx are between 0 and 100, in Unreal we are between 0 and 1, so we must scale
+		//The curve with 0.01
+		constexpr float ScaleCurve = 0.01f;
+		ImportCurve(AnimCurve, ScaleCurve, InterchangeCurves.AddDefaulted_GetRef().Keys);
+		{
+			FLargeMemoryWriter Ar;
+			Ar << InterchangeCurves;
+			uint8* ArchiveData = Ar.GetData();
+			int64 ArchiveSize = Ar.TotalSize();
+			TArray64<uint8> Buffer(ArchiveData, ArchiveSize);
+			FFileHelper::SaveArrayToFile(Buffer, *PayloadFilepath);
+		}
+		return true;
+	}
+
+	void FFbxAnimation::AddNodeTransformAnimation(FbxScene* SDKScene
+		, FbxNode* Node
+		, UInterchangeSceneNode* SceneNode
+		, TMap<FString, TSharedPtr<FPayloadContextBase>>& PayloadContexts)
 	{
 		int32 NumAnimations = SDKScene->GetSrcObjectCount<FbxAnimStack>();
 		//Anim stack should be merge so we expect to have only one stack here
@@ -611,28 +753,122 @@ namespace UE::Interchange::Private
 
 		double FrameRate = FbxTime::GetFrameRate(SDKScene->GetGlobalSettings().GetTimeMode());
 
-		FGetFbxTransformCurvesParameters Parameters(SDKScene, JointNode);
+		FGetFbxTransformCurvesParameters Parameters(SDKScene, Node);
 		GetFbxTransformCurves(Parameters);
 
 		if (Parameters.IsNodeAnimated)
 		{
-			FString PayLoadKey = FFbxHelper::GetFbxNodeHierarchyName(JointNode) + TEXT("_AnimationPayloadKey");
+			FString PayLoadKey = FFbxHelper::GetFbxNodeHierarchyName(Node) + TEXT("_AnimationPayloadKey");
 			if (ensure(!PayloadContexts.Contains(PayLoadKey)))
 			{
 				TSharedPtr<FAnimationPayloadContextTransform> AnimPayload = MakeShared<FAnimationPayloadContextTransform>();
-				AnimPayload->Node = JointNode;
-				AnimPayload->SDKScene = SDKScene;
-				
+				FNodeTransformFetchPayloadData FetchPayloadData;
+				FetchPayloadData.Node = Node;
+				AnimPayload->NodeTransformFetchPayloadData = FetchPayloadData;
 				PayloadContexts.Add(PayLoadKey, AnimPayload);
 			}
-			SceneNode->SetCustomTransformCurvePayloadKey(PayLoadKey);
-			SceneNode->SetCustomIsNodeTransformAnimated(Parameters.IsNodeAnimated);
-			SceneNode->SetCustomNodeTransformAnimationKeyCount(Parameters.KeyCount);
-			SceneNode->SetCustomNodeTransformAnimationStartTime(Parameters.StartTime.GetSecondDouble());
-			SceneNode->SetCustomNodeTransformAnimationEndTime(Parameters.EndTime.GetSecondDouble());
+			UInterchangeAnimationAPI::SetCustomNodeTransformPayloadKey(SceneNode, PayLoadKey);
+			UInterchangeAnimationAPI::SetCustomIsNodeTransformAnimated(SceneNode, Parameters.IsNodeAnimated);
+			UInterchangeAnimationAPI::SetCustomNodeTransformAnimationKeyCount(SceneNode, Parameters.KeyCount);
+			UInterchangeAnimationAPI::SetCustomNodeTransformAnimationStartTime(SceneNode, Parameters.StartTime.GetSecondDouble());
+			UInterchangeAnimationAPI::SetCustomNodeTransformAnimationEndTime(SceneNode, Parameters.EndTime.GetSecondDouble());
 		}
 	}
 
+	void FFbxAnimation::AddNodeAttributeCurvesAnimation(FbxNode* Node
+		, FbxProperty& Property
+		, FbxAnimCurveNode* AnimCurveNode
+		, UInterchangeSceneNode* SceneNode
+		, TMap<FString, TSharedPtr<FPayloadContextBase>>& PayloadContexts
+		, EFbxType PropertyType
+		, TOptional<FString>& OutPayloadKey)
+	{
+		const FString PropertyName = FFbxHelper::GetFbxPropertyName(Property);
+		const FString PayLoadKey = FFbxHelper::GetFbxNodeHierarchyName(Node) + PropertyName + TEXT("_AnimationPayloadKey");
+		if (ensure(!PayloadContexts.Contains(PayLoadKey)))
+		{
+			TSharedPtr<FAnimationPayloadContextTransform> AnimPayload = MakeShared<FAnimationPayloadContextTransform>();
+			FAttributeFetchPayloadData FetchPayloadData;
+			FetchPayloadData.Node = Node;
+			FetchPayloadData.AnimCurves = AnimCurveNode;
+
+			//Any property that is not decimal should import with step curve
+			FetchPayloadData.bAttributeTypeIsStepCurveAnimation = !IsFbxPropertyTypeDecimal(PropertyType);
+			FetchPayloadData.PropertyType = PropertyType;
+			FetchPayloadData.Property = Property;
+			AnimPayload->AttributeFetchPayloadData = FetchPayloadData;
+			PayloadContexts.Add(PayLoadKey, AnimPayload);
+			OutPayloadKey = PayLoadKey;
+		}
+	}
+
+	void FFbxAnimation::AddMorphTargetCurvesAnimation(UInterchangeMeshNode* MorphTargetNode
+		, FbxScene* SDKScene
+		, int32 GeometryIndex
+		, int32 MorphTargetIndex
+		, int32 ChannelIndex
+		, const FString& MorphTargetPayloadKey
+		, TMap<FString, TSharedPtr<FPayloadContextBase>>& PayloadContexts)
+	{
+		check(MorphTargetNode);
+		const FString PayLoadKey = MorphTargetPayloadKey + TEXT("_CurveAnimationPayloadKey");
+		if (ensure(!PayloadContexts.Contains(PayLoadKey)))
+		{
+			TSharedPtr<FAnimationPayloadContextTransform> AnimPayload = MakeShared<FAnimationPayloadContextTransform>();
+			FMorphTargetFetchPayloadData FetchPayloadData;
+			FetchPayloadData.SDKScene = SDKScene;
+			FetchPayloadData.GeometryIndex = GeometryIndex;
+			FetchPayloadData.MorphTargetIndex = MorphTargetIndex;
+			FetchPayloadData.ChannelIndex = ChannelIndex;
+			AnimPayload->MorphTargetFetchPayloadData = FetchPayloadData;
+
+			PayloadContexts.Add(PayLoadKey, AnimPayload);
+		}
+		MorphTargetNode->SetAnimationCurvePayLoadKey(PayLoadKey);
+	}
+
+	bool FFbxAnimation::IsFbxPropertyTypeSupported(EFbxType PropertyType)
+	{
+		switch (PropertyType)
+		{
+		case EFbxType::eFbxBool:
+		case EFbxType::eFbxChar:		//!< 8 bit signed integer.
+		case EFbxType::eFbxUChar:		//!< 8 bit unsigned integer.
+		case EFbxType::eFbxShort:		//!< 16 bit signed integer.
+		case EFbxType::eFbxUShort:		//!< 16 bit unsigned integer.
+		case EFbxType::eFbxInt:			//!< 32 bit signed integer.
+		case EFbxType::eFbxUInt:		//!< 32 bit unsigned integer.
+		case EFbxType::eFbxLongLong:	//!< 64 bit signed integer.
+		case EFbxType::eFbxULongLong:	//!< 64 bit unsigned integer.
+		case EFbxType::eFbxHalfFloat:	//!< 16 bit floating point.
+		case EFbxType::eFbxFloat:		//!< Floating point value.
+		case EFbxType::eFbxDouble:	//!< Double width floating point value.
+		case EFbxType::eFbxDouble2:	//!< Vector of two double values.
+		case EFbxType::eFbxDouble3:	//!< Vector of three double values.
+		case EFbxType::eFbxDouble4:	//!< Vector of four double values.
+		case EFbxType::eFbxEnum:		//!< Enumeration.
+		case EFbxType::eFbxString:	//!< String.
+			return true;
+		}
+
+		return false;
+	};
+
+	bool FFbxAnimation::IsFbxPropertyTypeDecimal(EFbxType PropertyType)
+	{
+		switch (PropertyType)
+		{
+		case EFbxType::eFbxHalfFloat:
+		case EFbxType::eFbxFloat:
+		case EFbxType::eFbxDouble:
+		case EFbxType::eFbxDouble2:
+		case EFbxType::eFbxDouble3:
+		case EFbxType::eFbxDouble4:
+			return true;
+		}
+
+		return false;
+	}
 }//ns UE::Interchange::Private
 
 #undef LOCTEXT_NAMESPACE
diff --git a/Engine/Plugins/Experimental/Interchange/Runtime/Source/Parsers/Fbx/Private/FbxAnimation.h b/Engine/Plugins/Experimental/Interchange/Runtime/Source/Parsers/Fbx/Private/FbxAnimation.h
--- a/Engine/Plugins/Experimental/Interchange/Runtime/Source/Parsers/Fbx/Private/FbxAnimation.h
+++ b/Engine/Plugins/Experimental/Interchange/Runtime/Source/Parsers/Fbx/Private/FbxAnimation.h
@@ -9,11 +9,37 @@
 
 /** Forward declarations */
 struct FMeshDescription;
-class UInterchangeBaseNodeContainer;
+class UInterchangeMeshNode;
 class UInterchangeSceneNode;
 
 namespace UE::Interchange::Private
 {
+	struct FNodeTransformFetchPayloadData
+	{
+		FbxNode* Node = nullptr;
+	};
+
+	struct FAttributeFetchPayloadData
+	{
+		FbxNode* Node = nullptr;
+		FbxAnimCurveNode* AnimCurves = nullptr;
+		
+		//If true the payload is TArray<FInterchangeStepCurve>, if false the payload is TArray<FInterchangeCurve>
+		bool bAttributeTypeIsStepCurveAnimation = false;
+		fbxsdk::EFbxType PropertyType = fbxsdk::EFbxType::eFbxUndefined;
+		FbxProperty Property;
+	};
+
+	struct FMorphTargetFetchPayloadData
+	{
+		FbxScene* SDKScene = nullptr;
+		int32 GeometryIndex = 0;
+		int32 MorphTargetIndex = 0;
+		int32 ChannelIndex = 0;
+	};
+
+	
+
 	class FAnimationPayloadContextTransform : public FPayloadContextBase
 	{
 	public:
@@ -21,14 +47,44 @@ namespace UE::Interchange::Private
 		virtual FString GetPayloadType() const override { return TEXT("TransformAnimation-PayloadContext"); }
 		virtual bool FetchPayloadToFile(FFbxParser& Parser, const FString& PayloadFilepath) override;
 		virtual bool FetchAnimationBakeTransformPayloadToFile(FFbxParser& Parser, const double BakeFrequency, const double RangeStartTime, const double RangeEndTime, const FString& PayloadFilepath) override;
-		FbxNode* Node = nullptr;
-		FbxScene* SDKScene = nullptr;
+		
+		TOptional<FNodeTransformFetchPayloadData> NodeTransformFetchPayloadData;
+		TOptional<FAttributeFetchPayloadData> AttributeFetchPayloadData;
+		TOptional<FMorphTargetFetchPayloadData> MorphTargetFetchPayloadData;
+	private:
+		bool InternalFetchCurveNodePayloadToFile(FFbxParser& Parser, const FString& PayloadFilepath);
+		bool InternalFetchMorphTargetCurvePayloadToFile(FFbxParser& Parser, const FString& PayloadFilepath);
 	};
 
 	class FFbxAnimation
 	{
 	public:
 		/** This function add the payload key if the scene node transform is animated. */
-		static void AddNodeTransformAnimation(FbxScene* SDKScene, FbxNode* JointNode, UInterchangeBaseNodeContainer& NodeContainer, UInterchangeSceneNode* UnrealNode, TMap<FString, TSharedPtr<FPayloadContextBase>>& PayloadContexts);
+		static void AddNodeTransformAnimation(FbxScene* SDKScene
+			, FbxNode* Node
+			, UInterchangeSceneNode* UnrealNode
+			, TMap<FString, TSharedPtr<FPayloadContextBase>>& PayloadContexts);
+
+		/** This function add the payload key for an animated node user attribute. */
+		static void AddNodeAttributeCurvesAnimation(FbxNode* Node
+			, FbxProperty& Property
+			, FbxAnimCurveNode* AnimCurveNode
+			, UInterchangeSceneNode* SceneNode
+			, TMap<FString, TSharedPtr<FPayloadContextBase>>& PayloadContexts
+			, EFbxType PropertyType
+			, TOptional<FString>& OutPayloadKey);
+
+		/** This function add the payload key for an animated curve for morph target. */
+		static void AddMorphTargetCurvesAnimation(UInterchangeMeshNode* ShapeNode
+			, FbxScene* SDKScene
+			, int32 GeometryIndex
+			, int32 MorphTargetIndex
+			, int32 ChannelIndex
+			, const FString& MorphTargetPayloadKey
+			, TMap<FString, TSharedPtr<FPayloadContextBase>>& PayloadContexts);
+
+		static bool IsFbxPropertyTypeSupported(EFbxType PropertyType);
+		static bool IsFbxPropertyTypeDecimal(EFbxType PropertyType);
+
 	};
 }//ns UE::Interchange::Private
diff --git a/Engine/Plugins/Experimental/Interchange/Runtime/Source/Parsers/Fbx/Private/FbxMesh.cpp b/Engine/Plugins/Experimental/Interchange/Runtime/Source/Parsers/Fbx/Private/FbxMesh.cpp
--- a/Engine/Plugins/Experimental/Interchange/Runtime/Source/Parsers/Fbx/Private/FbxMesh.cpp
+++ b/Engine/Plugins/Experimental/Interchange/Runtime/Source/Parsers/Fbx/Private/FbxMesh.cpp
@@ -3,6 +3,7 @@
 #include "FbxMesh.h"
 
 #include "CoreMinimal.h"
+#include "FbxAnimation.h"
 #include "FbxAPI.h"
 #include "FbxConvert.h"
 #include "FbxHelper.h"
@@ -1186,25 +1187,25 @@ bool FMeshPayloadContext::FetchPayloadToFile(FFbxParser& Parser, const FString&
 }
 
 //////////////////////////////////////////////////////////////////////////
-/// FShapePayloadContext implementation
+/// FMorphTargetPayloadContext implementation
 
-bool FShapePayloadContext::FetchPayloadToFile(FFbxParser& Parser, const FString& PayloadFilepath)
+bool FMorphTargetPayloadContext::FetchPayloadToFile(FFbxParser& Parser, const FString& PayloadFilepath)
 {
 	if (!ensure(Shape))
 	{
 		//Todo log an error
 		return false;
 	}
-	//Import the BlendShape
-	FMeshDescription ShapeMeshDescription;
-	FStaticMeshAttributes StaticMeshAttribute(ShapeMeshDescription);
+	//Import the MorphTarget
+	FMeshDescription MorphTargetMeshDescription;
+	FStaticMeshAttributes StaticMeshAttribute(MorphTargetMeshDescription);
 	StaticMeshAttribute.Register();
-	FMeshDescriptionImporter MeshDescriptionImporter(Parser, &ShapeMeshDescription, SDKScene, SDKGeometryConverter);
+	FMeshDescriptionImporter MeshDescriptionImporter(Parser, &MorphTargetMeshDescription, SDKScene, SDKGeometryConverter);
 	if (!MeshDescriptionImporter.FillMeshDescriptionFromFbxShape(Shape))
 	{
 		UInterchangeResultError_Generic* Message = Parser.AddMessage<UInterchangeResultError_Generic>();
 		Message->InterchangeKey = FFbxHelper::GetMeshUniqueID(Shape);
-		Message->Text = LOCTEXT("MeshDescriptionBlendShapeError", "Unable to create MeshDescription from FBX blend shape.");
+		Message->Text = LOCTEXT("MeshDescriptionMorphTargetError", "Unable to create MeshDescription from FBX morph target.");
 
 		return false;
 	}
@@ -1212,7 +1213,7 @@ bool FShapePayloadContext::FetchPayloadToFile(FFbxParser& Parser, const FString&
 	//Dump the MeshDescription to a file
 	{
 		FLargeMemoryWriter Ar;
-		ShapeMeshDescription.Serialize(Ar);
+		MorphTargetMeshDescription.Serialize(Ar);
 		uint8* ArchiveData = Ar.GetData();
 		int64 ArchiveSize = Ar.TotalSize();
 		TArray64<uint8> Buffer(ArchiveData, ArchiveSize);
@@ -1272,172 +1273,193 @@ void FFbxMesh::AddAllMeshes(FbxScene* SDKScene, FbxGeometryConverter* SDKGeometr
 		FString MeshUniqueID = FFbxHelper::GetMeshUniqueID(Mesh);
 		const UInterchangeMeshNode* ExistingMeshNode = Cast<UInterchangeMeshNode>(NodeContainer.GetNode(MeshUniqueID));
 		UInterchangeMeshNode* MeshNode = nullptr;
-		if (!ExistingMeshNode)
+		if (ExistingMeshNode)
 		{
-			MeshNode = CreateMeshNode(NodeContainer, MeshName, MeshUniqueID);
-			if (Geometry->GetDeformerCount(FbxDeformer::eSkin) > 0)
-			{
-				//Set the skinned mesh attribute
-				MeshNode->SetSkinnedMesh(true);
-				ExtractSkinnedMeshNodeJoints(SDKScene, NodeContainer, Mesh, MeshNode);
-			}
-			Mesh->ComputeBBox();
-			const int32 MeshVertexCount = Mesh->GetControlPointsCount();
-			MeshNode->SetCustomVertexCount(MeshVertexCount);
-			const int32 MeshPolygonCount = Mesh->GetPolygonCount();
-			MeshNode->SetCustomPolygonCount(MeshPolygonCount);
-			const FBox MeshBoundingBox = FBox(FFbxConvert::ConvertPos(Mesh->BBoxMin.Get()), FFbxConvert::ConvertPos(Mesh->BBoxMax.Get()));
-			MeshNode->SetCustomBoundingBox(MeshBoundingBox);
-			const bool bMeshHasVertexNormal = Mesh->GetElementNormalCount() > 0;
-			MeshNode->SetCustomHasVertexNormal(bMeshHasVertexNormal);
-			const bool bMeshHasVertexBinormal = Mesh->GetElementBinormalCount() > 0;
-			MeshNode->SetCustomHasVertexBinormal(bMeshHasVertexBinormal);
-			const bool bMeshHasVertexTangent = Mesh->GetElementTangentCount() > 0;
-			MeshNode->SetCustomHasVertexTangent(bMeshHasVertexTangent);
-			const bool bMeshHasSmoothGroup = Mesh->GetElementSmoothingCount() > 0;
-			MeshNode->SetCustomHasSmoothGroup(bMeshHasSmoothGroup);
-			const bool bMeshHasVertexColor = Mesh->GetElementVertexColorCount() > 0;
-			MeshNode->SetCustomHasVertexColor(bMeshHasVertexColor);
-			const int32 MeshUVCount = Mesh->GetElementUVCount();
-			MeshNode->SetCustomUVCount(MeshUVCount);
+			//This mesh node was already created
+			continue;
+		}
+
+		MeshNode = CreateMeshNode(NodeContainer, MeshName, MeshUniqueID);
+		if (Geometry->GetDeformerCount(FbxDeformer::eSkin) > 0)
+		{
+			//Set the skinned mesh attribute
+			MeshNode->SetSkinnedMesh(true);
+			ExtractSkinnedMeshNodeJoints(SDKScene, NodeContainer, Mesh, MeshNode);
+		}
+		Mesh->ComputeBBox();
+		const int32 MeshVertexCount = Mesh->GetControlPointsCount();
+		MeshNode->SetCustomVertexCount(MeshVertexCount);
+		const int32 MeshPolygonCount = Mesh->GetPolygonCount();
+		MeshNode->SetCustomPolygonCount(MeshPolygonCount);
+		const FBox MeshBoundingBox = FBox(FFbxConvert::ConvertPos(Mesh->BBoxMin.Get()), FFbxConvert::ConvertPos(Mesh->BBoxMax.Get()));
+		MeshNode->SetCustomBoundingBox(MeshBoundingBox);
+		const bool bMeshHasVertexNormal = Mesh->GetElementNormalCount() > 0;
+		MeshNode->SetCustomHasVertexNormal(bMeshHasVertexNormal);
+		const bool bMeshHasVertexBinormal = Mesh->GetElementBinormalCount() > 0;
+		MeshNode->SetCustomHasVertexBinormal(bMeshHasVertexBinormal);
+		const bool bMeshHasVertexTangent = Mesh->GetElementTangentCount() > 0;
+		MeshNode->SetCustomHasVertexTangent(bMeshHasVertexTangent);
+		const bool bMeshHasSmoothGroup = Mesh->GetElementSmoothingCount() > 0;
+		MeshNode->SetCustomHasSmoothGroup(bMeshHasSmoothGroup);
+		const bool bMeshHasVertexColor = Mesh->GetElementVertexColorCount() > 0;
+		MeshNode->SetCustomHasVertexColor(bMeshHasVertexColor);
+		const int32 MeshUVCount = Mesh->GetElementUVCount();
+		MeshNode->SetCustomUVCount(MeshUVCount);
 						
-			//Add Material dependencies, we use always the first fbx node that instance the fbx geometry to grab the fbx surface materials.
+		//Add Material dependencies, we use always the first fbx node that instance the fbx geometry to grab the fbx surface materials.
+		{
+			//Grab all Material indexes use by the mesh
+			TArray<int32> MaterialIndexes;
+			int32 PolygonCount = Mesh->GetPolygonCount();
+			if (FbxGeometryElementMaterial* GeometryElementMaterial = Mesh->GetElementMaterial())
 			{
-				//Grab all Material indexes use by the mesh
-				TArray<int32> MaterialIndexes;
-				int32 PolygonCount = Mesh->GetPolygonCount();
-				if (FbxGeometryElementMaterial* GeometryElementMaterial = Mesh->GetElementMaterial())
+				FbxLayerElementArrayTemplate<int32>& IndexArray = GeometryElementMaterial->GetIndexArray();
+				switch (GeometryElementMaterial->GetMappingMode())
 				{
-					FbxLayerElementArrayTemplate<int32>& IndexArray = GeometryElementMaterial->GetIndexArray();
-					switch (GeometryElementMaterial->GetMappingMode())
+					case FbxGeometryElement::eByPolygon:
 					{
-						case FbxGeometryElement::eByPolygon:
+						if (IndexArray.GetCount() == PolygonCount)
 						{
-							if (IndexArray.GetCount() == PolygonCount)
+							for (int32 PolygonIndex = 0; PolygonIndex < PolygonCount; ++PolygonIndex)
 							{
-								for (int i = 0; i < PolygonCount; ++i)
-								{
-									MaterialIndexes.AddUnique(IndexArray.GetAt(i));
-								}
+								MaterialIndexes.AddUnique(IndexArray.GetAt(PolygonIndex));
 							}
 						}
-						break;
+					}
+					break;
 
-						case FbxGeometryElement::eAllSame:
+					case FbxGeometryElement::eAllSame:
+					{
+						if (IndexArray.GetCount() > 0)
 						{
-							if (IndexArray.GetCount() > 0)
-							{
-								MaterialIndexes.AddUnique(IndexArray.GetAt(0));
-							}
+							MaterialIndexes.AddUnique(IndexArray.GetAt(0));
 						}
-						break;
 					}
+					break;
 				}
-				if (FbxNode* FbxMeshNode = Mesh->GetNode())
+			}
+			if (FbxNode* FbxMeshNode = Mesh->GetNode())
+			{
+				bool bAddAllNodeMaterials = (MaterialIndexes.Num() == 0);
+				const int32 MaterialCount = FbxMeshNode->GetMaterialCount();
+				for (int32 MaterialIndex = 0; MaterialIndex < MaterialCount; ++MaterialIndex)
 				{
-					bool bAddAllNodeMaterials = (MaterialIndexes.Num() == 0);
-					const int32 MaterialCount = FbxMeshNode->GetMaterialCount();
-					for (int32 MaterialIndex = 0; MaterialIndex < MaterialCount; ++MaterialIndex)
+					FbxSurfaceMaterial* FbxMaterial = FbxMeshNode->GetMaterial(MaterialIndex);
+					FString MaterialName = FFbxHelper::GetFbxObjectName(FbxMaterial);
+					FString MaterialUid = TEXT("\\Material\\") + MaterialName;
+					if (bAddAllNodeMaterials || MaterialIndexes.Contains(MaterialIndex))
 					{
-						FbxSurfaceMaterial* FbxMaterial = FbxMeshNode->GetMaterial(MaterialIndex);
-						FString MaterialName = FFbxHelper::GetFbxObjectName(FbxMaterial);
-						FString MaterialUid = TEXT("\\Material\\") + MaterialName;
-						if (bAddAllNodeMaterials || MaterialIndexes.Contains(MaterialIndex))
-						{
-							MeshNode->SetSlotMaterialDependencyUid(MaterialName, MaterialUid);
-						}
+						MeshNode->SetSlotMaterialDependencyUid(MaterialName, MaterialUid);
 					}
 				}
 			}
+		}
 
-			FString PayLoadKey = MeshUniqueID;
-			if (ensure(!PayloadContexts.Contains(PayLoadKey)))
-			{
-				TSharedPtr<FMeshPayloadContext> GeoPayload = MakeShared<FMeshPayloadContext>();
-				GeoPayload->Mesh = Mesh;
-				GeoPayload->SDKScene = SDKScene;
-				GeoPayload->SDKGeometryConverter = SDKGeometryConverter;
-				PayloadContexts.Add(PayLoadKey, GeoPayload);
-			}
-			MeshNode->SetPayLoadKey(PayLoadKey);
+		FString PayLoadKey = MeshUniqueID;
+		if (ensure(!PayloadContexts.Contains(PayLoadKey)))
+		{
+			TSharedPtr<FMeshPayloadContext> GeoPayload = MakeShared<FMeshPayloadContext>();
+			GeoPayload->Mesh = Mesh;
+			GeoPayload->SDKScene = SDKScene;
+			GeoPayload->SDKGeometryConverter = SDKGeometryConverter;
+			PayloadContexts.Add(PayLoadKey, GeoPayload);
+		}
+		MeshNode->SetPayLoadKey(PayLoadKey);
 
-			//Add all BlendShapes for this geometry node
+		int32 NumAnimations = SDKScene->GetSrcObjectCount<FbxAnimStack>();
+			
+		//Anim stack should be merge so we expect to have only one stack here
+		ensure(NumAnimations <= 1);
+		FbxAnimStack* AnimStack = NumAnimations == 1 ? (FbxAnimStack*)SDKScene->GetSrcObject<FbxAnimStack>(0) : nullptr;
+
+		//Add all MorphTargets for this geometry node
 
-			TMap<FString, FbxShape*> ShapeNameToFbxShape;
-			const int32 BlendShapeCount = Geometry->GetDeformerCount(FbxDeformer::eBlendShape);
+		TMap<FString, FbxShape*> ShapeNameToFbxShape;
+		const int32 MorphTargetCount = Geometry->GetDeformerCount(FbxDeformer::eBlendShape);
 
-			for (int32 BlendShapeIndex = 0; BlendShapeIndex < BlendShapeCount; ++BlendShapeIndex)
+		for (int32 MorphTargetIndex = 0; MorphTargetIndex < MorphTargetCount; ++MorphTargetIndex)
+		{
+			FbxBlendShape* MorphTarget = (FbxBlendShape*)Geometry->GetDeformer(MorphTargetIndex, FbxDeformer::eBlendShape);
+			const int32 MorphTargetChannelCount = MorphTarget->GetBlendShapeChannelCount();
+			FString MorphTargetName = FFbxHelper::GetFbxObjectName(MorphTarget);
+			// see below where this is used for explanation...
+			const bool bMightBeBadMAXFile = (MorphTargetName == FString("Morpher"));
+			for (int32 ChannelIndex = 0; ChannelIndex < MorphTargetChannelCount; ++ChannelIndex)
 			{
-				FbxBlendShape* BlendShape = (FbxBlendShape*)Geometry->GetDeformer(BlendShapeIndex, FbxDeformer::eBlendShape);
-				const int32 BlendShapeChannelCount = BlendShape->GetBlendShapeChannelCount();
-				FString BlendShapeName = FFbxHelper::GetFbxObjectName(BlendShape);
-				// see below where this is used for explanation...
-				const bool bMightBeBadMAXFile = (BlendShapeName == FString("Morpher"));
-				for (int32 ChannelIndex = 0; ChannelIndex < BlendShapeChannelCount; ++ChannelIndex)
+				FbxBlendShapeChannel* Channel = MorphTarget->GetBlendShapeChannel(ChannelIndex);
+				if (!Channel)
 				{
-					FbxBlendShapeChannel* Channel = BlendShape->GetBlendShapeChannel(ChannelIndex);
-					if (Channel)
+					continue;
+				}
+				//Find which morph target should we use according to the weight.
+				const int32 CurrentChannelMorphTargetCount = Channel->GetTargetShapeCount();
+				FString ChannelName = FFbxHelper::GetFbxObjectName(Channel);
+				// Maya adds the name of the MorphTarget and an underscore to the front of the channel name, so remove it
+				if (ChannelName.StartsWith(MorphTargetName))
+				{
+					ChannelName.RightInline(ChannelName.Len() - (MorphTargetName.Len() + 1), false);
+				}
+				for (int32 ChannelMorphTargetIndex = 0; ChannelMorphTargetIndex < CurrentChannelMorphTargetCount; ++ChannelMorphTargetIndex)
+				{
+					FbxShape* Shape = Channel->GetTargetShape(ChannelMorphTargetIndex);
+					FString ShapeName;
+					if (CurrentChannelMorphTargetCount > 1)
+					{
+						ShapeName = FFbxHelper::GetFbxObjectName(Shape);
+					}
+					else
 					{
-						//Find which shape should we use according to the weight.
-						const int32 CurrentChannelShapeCount = Channel->GetTargetShapeCount();
-						FString ChannelName = FFbxHelper::GetFbxObjectName(Channel);
-						// Maya adds the name of the blendshape and an underscore to the front of the channel name, so remove it
-						if (ChannelName.StartsWith(BlendShapeName))
+						if (bMightBeBadMAXFile)
 						{
-							ChannelName.RightInline(ChannelName.Len() - (BlendShapeName.Len() + 1), false);
+							ShapeName = FFbxHelper::GetFbxObjectName(Shape);
 						}
-						for (int32 ShapeIndex = 0; ShapeIndex < CurrentChannelShapeCount; ++ShapeIndex)
+						else
 						{
-							FbxShape* Shape = Channel->GetTargetShape(ShapeIndex);
-							FString ShapeName;
-							if (CurrentChannelShapeCount > 1)
-							{
-								ShapeName = FFbxHelper::GetFbxObjectName(Shape);
-							}
-							else
-							{
-								if (bMightBeBadMAXFile)
-								{
-									ShapeName = FFbxHelper::GetFbxObjectName(Shape);
-								}
-								else
-								{
-									// Maya concatenates the number of the shape to the end of its name, so instead use the name of the channel
-									ShapeName = ChannelName;
-								}
-							}
-							ensure(!ShapeNameToFbxShape.Contains(ShapeName));
-							ShapeNameToFbxShape.Add(ShapeName, Shape);
-							FString ShapeAttributeName = FFbxHelper::GetMeshName(Shape);
-							FString ShapeUniqueID = FFbxHelper::GetMeshUniqueID(Shape);
-							const UInterchangeMeshNode* ExistingShapeNode = Cast<const UInterchangeMeshNode>(NodeContainer.GetNode(ShapeUniqueID));
-							if (!ExistingShapeNode)
-							{
-								UInterchangeMeshNode* ShapeNode = CreateMeshNode(NodeContainer, ShapeAttributeName, ShapeUniqueID);
-								const bool bIsBlendShape = true;
-								ShapeNode->SetBlendShape(bIsBlendShape);
-								ShapeNode->SetBlendShapeName(ShapeName);
+							// Maya concatenates the number of the shape to the end of its name, so instead use the name of the channel
+							ShapeName = ChannelName;
+						}
+					}
+					ensure(!ShapeNameToFbxShape.Contains(ShapeName));
+					ShapeNameToFbxShape.Add(ShapeName, Shape);
+					FString MorphTargetAttributeName = FFbxHelper::GetMeshName(Shape);
+					FString MorphTargetUniqueID = FFbxHelper::GetMeshUniqueID(Shape);
+					const UInterchangeMeshNode* ExistingMorphTargetNode = Cast<const UInterchangeMeshNode>(NodeContainer.GetNode(MorphTargetUniqueID));
+					if (!ExistingMorphTargetNode)
+					{
+						UInterchangeMeshNode* MorphTargetNode = CreateMeshNode(NodeContainer, MorphTargetAttributeName, MorphTargetUniqueID);
+						const bool bIsMorphTarget = true;
+						MorphTargetNode->SetMorphTarget(bIsMorphTarget);
+						MorphTargetNode->SetMorphTargetName(ShapeName);
 
-								//Create a Mesh node dependency, so the mesh node can retrieve is associate blendshape
-								MeshNode->SetShapeDependencyUid(ShapeUniqueID);
+						//Create a Mesh node dependency, so the mesh node can retrieve is associate morph target
+						MeshNode->SetMorphTargetDependencyUid(MorphTargetUniqueID);
 
-								FString ShapePayLoadKey = ShapeUniqueID;
-								if (ensure(!PayloadContexts.Contains(ShapePayLoadKey)))
-								{
-									TSharedPtr<FShapePayloadContext> GeoPayload = MakeShared<FShapePayloadContext>();
-									GeoPayload->Shape = Shape;
-									GeoPayload->SDKScene = SDKScene;
-									GeoPayload->SDKGeometryConverter = SDKGeometryConverter;
-									PayloadContexts.Add(ShapePayLoadKey, GeoPayload);
-								}
-								ShapeNode->SetPayLoadKey(ShapePayLoadKey);
+						FString MorphTargetPayLoadKey = MorphTargetUniqueID;
+						if (ensure(!PayloadContexts.Contains(MorphTargetPayLoadKey)))
+						{
+							TSharedPtr<FMorphTargetPayloadContext> GeoPayload = MakeShared<FMorphTargetPayloadContext>();
+							GeoPayload->Shape = Shape;
+							GeoPayload->SDKScene = SDKScene;
+							GeoPayload->SDKGeometryConverter = SDKGeometryConverter;
+							PayloadContexts.Add(MorphTargetPayLoadKey, GeoPayload);
+						}
+						MorphTargetNode->SetPayLoadKey(MorphTargetPayLoadKey);
+
+						//Add the morph target animation float curve payload to the morph target.
+						if (AnimStack)
+						{
+							//If the morph target node is animated prepare a curve payload to retrieve the morph target animation
+							FbxAnimCurve* Curve = Geometry->GetShapeChannel(MorphTargetIndex, ChannelIndex, (FbxAnimLayer*)AnimStack->GetMember(0));
+							if ((Curve && Curve->KeyGetCount() > 0))
+							{
+								UE::Interchange::Private::FFbxAnimation::AddMorphTargetCurvesAnimation(MorphTargetNode, SDKScene, GeometryIndex, MorphTargetIndex, ChannelIndex, MorphTargetPayLoadKey, PayloadContexts);
 							}
 						}
 					}
-				}
-			}
-		}					
-	}
+				} // for CurrentChannelMorphTargetCount
+			} // for MorphTargetChannelCount
+		} // for MorphTargetCount
+	} // for GeometryCount
 }
 
 bool FFbxMesh::GetGlobalJointBindPoseTransform(FbxScene* SDKScene, FbxNode* Joint, FbxAMatrix& GlobalBindPoseJointMatrix)
diff --git a/Engine/Plugins/Experimental/Interchange/Runtime/Source/Parsers/Fbx/Private/FbxMesh.h b/Engine/Plugins/Experimental/Interchange/Runtime/Source/Parsers/Fbx/Private/FbxMesh.h
--- a/Engine/Plugins/Experimental/Interchange/Runtime/Source/Parsers/Fbx/Private/FbxMesh.h
+++ b/Engine/Plugins/Experimental/Interchange/Runtime/Source/Parsers/Fbx/Private/FbxMesh.h
@@ -94,11 +94,11 @@ namespace UE
 				FbxGeometryConverter* SDKGeometryConverter = nullptr;
 			};
 
-			class FShapePayloadContext : public FPayloadContextBase
+			class FMorphTargetPayloadContext : public FPayloadContextBase
 			{
 			public:
-				virtual ~FShapePayloadContext() {}
-				virtual FString GetPayloadType() const override { return TEXT("Shape-PayloadContext"); }
+				virtual ~FMorphTargetPayloadContext() {}
+				virtual FString GetPayloadType() const override { return TEXT("MorphTarget-PayloadContext"); }
 				virtual bool FetchPayloadToFile(FFbxParser& Parser, const FString& PayloadFilepath) override;
 				FbxShape* Shape = nullptr;
 				FbxScene* SDKScene = nullptr;
diff --git a/Engine/Plugins/Experimental/Interchange/Runtime/Source/Parsers/Fbx/Private/FbxScene.cpp b/Engine/Plugins/Experimental/Interchange/Runtime/Source/Parsers/Fbx/Private/FbxScene.cpp
--- a/Engine/Plugins/Experimental/Interchange/Runtime/Source/Parsers/Fbx/Private/FbxScene.cpp
+++ b/Engine/Plugins/Experimental/Interchange/Runtime/Source/Parsers/Fbx/Private/FbxScene.cpp
@@ -279,43 +279,15 @@ namespace UE
 						}
 					}
 				}
-				//Scene node transform can be animated, add the animation payload key.
-				FFbxAnimation::AddNodeTransformAnimation(SDKScene, Node, NodeContainer, UnrealNode, PayloadContexts);
+				//Scene node transform can be animated, add the transform animation payload key.
+				FFbxAnimation::AddNodeTransformAnimation(SDKScene, Node, UnrealNode, PayloadContexts);
 
 				//Add all custom Attributes for the node
-				auto IsSupportedExtraAttributeType = [](int32 InDataType)
-				{
-					switch (InDataType)
-					{
-					case EFbxType::eFbxBool:
-
-					case EFbxType::eFbxChar:		//!< 8 bit signed integer.
-					case EFbxType::eFbxUChar:		//!< 8 bit unsigned integer.
-					case EFbxType::eFbxShort:		//!< 16 bit signed integer.
-					case EFbxType::eFbxUShort:		//!< 16 bit unsigned integer.
-					case EFbxType::eFbxInt:			//!< 32 bit signed integer.
-					case EFbxType::eFbxUInt:		//!< 32 bit unsigned integer.
-					case EFbxType::eFbxLongLong:	//!< 64 bit signed integer.
-					case EFbxType::eFbxULongLong:	//!< 64 bit unsigned integer.
-					case EFbxType::eFbxHalfFloat:	//!< 16 bit floating point.
-					case EFbxType::eFbxFloat:		//!< Floating point value.
-					case EFbxType::eFbxDouble:	//!< Double width floating point value.
-					case EFbxType::eFbxDouble2:	//!< Vector of two double values.
-					case EFbxType::eFbxDouble3:	//!< Vector of three double values.
-					case EFbxType::eFbxDouble4:	//!< Vector of four double values.
-
-					case EFbxType::eFbxEnum:		//!< Enumeration.
-					case EFbxType::eFbxString:	//!< String.
-						return true;
-					}
-
-					return false;
-				};
-
 				FbxProperty Property = Node->GetFirstProperty();
 				while (Property.IsValid())
 				{
-					if (Property.GetFlag(FbxPropertyFlags::eUserDefined) && IsSupportedExtraAttributeType(Property.GetPropertyDataType().GetType()))
+					EFbxType PropertyType =  Property.GetPropertyDataType().GetType();
+					if (Property.GetFlag(FbxPropertyFlags::eUserDefined) && FFbxAnimation::IsFbxPropertyTypeSupported(PropertyType))
 					{
 						FString PropertyName = FFbxHelper::GetFbxPropertyName(Property);
 
@@ -323,8 +295,8 @@ namespace UE
 						TOptional<FString> PayloadKey;
 						if (CurveNode && CurveNode->IsAnimated())
 						{
-							//Set the optional payload key
-							//Create an attribute animation payload
+							//Attribute is animated, add the curves payload key that represent the attribute animation
+							FFbxAnimation::AddNodeAttributeCurvesAnimation(Node, Property, CurveNode, UnrealNode, PayloadContexts, PropertyType, PayloadKey);
 						}
 						switch (Property.GetPropertyDataType().GetType())
 						{
diff --git a/Engine/Plugins/Experimental/Interchange/Runtime/Source/Parsers/Fbx/Private/InterchangeFbxParser.cpp b/Engine/Plugins/Experimental/Interchange/Runtime/Source/Parsers/Fbx/Private/InterchangeFbxParser.cpp
--- a/Engine/Plugins/Experimental/Interchange/Runtime/Source/Parsers/Fbx/Private/InterchangeFbxParser.cpp
+++ b/Engine/Plugins/Experimental/Interchange/Runtime/Source/Parsers/Fbx/Private/InterchangeFbxParser.cpp
@@ -5,119 +5,140 @@
 #include "CoreMinimal.h"
 #include "FbxAPI.h"
 #include "HAL/CriticalSection.h"
+#include "InterchangeTextureNode.h"
 #include "Misc/ScopeLock.h"
 #include "Nodes/InterchangeBaseNodeContainer.h"
-#include "InterchangeTextureNode.h"
+#include "Serialization/MemoryWriter.h"
 #include "UObject/StrongObjectPtr.h"
 #include "UObject/UObjectGlobals.h"
 #include "UObject/Package.h"
 
 #define LOCTEXT_NAMESPACE "InterchangeFbxParser"
 
-namespace UE
+namespace UE::Interchange
 {
-	namespace Interchange
+	namespace Private
 	{
-		FInterchangeFbxParser::FInterchangeFbxParser()
+		FString HashString(const FString& String)
 		{
-			ResultsContainer = NewObject<UInterchangeResultsContainer>(GetTransientPackage());
-			ResultsContainer->AddToRoot();
-			FbxParserPrivate = MakeUnique<Private::FFbxParser>(ResultsContainer);
+			TArray<uint8> TempBytes;
+			TempBytes.Reserve(64);
+			//The archive is flagged as persistent so that machines of different endianness produce identical binary results.
+			FMemoryWriter Ar(TempBytes, /*bIsPersistent=*/ true);
+			//Hack because serialization do not support const
+			FString& NonConst = *const_cast<FString*>(&String);
+			Ar << NonConst;
+			FSHA1 Sha;
+			Sha.Update(TempBytes.GetData(), TempBytes.Num() * TempBytes.GetTypeSize());
+			Sha.Final();
+			// Retrieve the hash and use it to construct a pseudo-GUID.
+			uint32 Hash[5];
+			Sha.GetHash((uint8*)Hash);
+			FGuid Guid = FGuid(Hash[0] ^ Hash[4], Hash[1], Hash[2], Hash[3]);
+			return Guid.ToString(EGuidFormats::Base36Encoded);
 		}
+	}
+
+	FInterchangeFbxParser::FInterchangeFbxParser()
+	{
+		ResultsContainer.Reset(NewObject<UInterchangeResultsContainer>(GetTransientPackage()));
+		FbxParserPrivate = MakeUnique<Private::FFbxParser>(ResultsContainer.Get());
+	}
+	FInterchangeFbxParser::~FInterchangeFbxParser()
+	{
+		FbxParserPrivate = nullptr;
+	}
 
-		FInterchangeFbxParser::~FInterchangeFbxParser()
+	void FInterchangeFbxParser::LoadFbxFile(const FString& Filename, const FString& ResultFolder)
+	{
+		check(FbxParserPrivate.IsValid());
+		SourceFilename = Filename;
+		ResultsContainer->Empty();
+
+		if (!FbxParserPrivate->LoadFbxFile(Filename))
 		{
-			ResultsContainer->RemoveFromRoot();
-			FbxParserPrivate = nullptr;
+			UInterchangeResultError_Generic* Error = AddMessage<UInterchangeResultError_Generic>();
+			Error->SourceAssetName = SourceFilename;
+			Error->Text = LOCTEXT("CantLoadFbxFile", "Cannot load the FBX file.");
+			return;
 		}
 
-		void FInterchangeFbxParser::LoadFbxFile(const FString& Filename, const FString& ResultFolder)
+		ResultFilepath = ResultFolder + TEXT("/SceneDescription.itc");
+		//Since we are not in main thread we cannot use TStrongPtr, so we will add the object to the root and remove it when we are done
+		UInterchangeBaseNodeContainer* Container = NewObject<UInterchangeBaseNodeContainer>(GetTransientPackage(), NAME_None);
+		if (!ensure(Container != nullptr))
 		{
-			check(FbxParserPrivate.IsValid());
-			SourceFilename = Filename;
-			ResultsContainer->Empty();
-
-			if (!FbxParserPrivate->LoadFbxFile(Filename))
-			{
-				UInterchangeResultError_Generic* Error = AddMessage<UInterchangeResultError_Generic>();
-				Error->SourceAssetName = SourceFilename;
-				Error->Text = LOCTEXT("CantLoadFbxFile", "Cannot load the FBX file.");
-				return;
-			}
-
-			ResultFilepath = ResultFolder + TEXT("/SceneDescription.itc");
-			//Since we are not in main thread we cannot use TStrongPtr, so we will add the object to the root and remove it when we are done
-			UInterchangeBaseNodeContainer* Container = NewObject<UInterchangeBaseNodeContainer>(GetTransientPackage(), NAME_None);
-			if (!ensure(Container != nullptr))
-			{
-				UInterchangeResultError_Generic* Error = AddMessage<UInterchangeResultError_Generic>();
-				Error->SourceAssetName = SourceFilename;
-				Error->Text = LOCTEXT("CantAllocate", "Cannot allocate base node container to add FBX scene data.");
-				return;
-			}
-
-			Container->AddToRoot();
-			FbxParserPrivate->FillContainerWithFbxScene(*Container);
-			Container->SaveToFile(ResultFilepath);
-			Container->RemoveFromRoot();
+			UInterchangeResultError_Generic* Error = AddMessage<UInterchangeResultError_Generic>();
+			Error->SourceAssetName = SourceFilename;
+			Error->Text = LOCTEXT("CantAllocate", "Cannot allocate base node container to add FBX scene data.");
+			return;
 		}
 
-		void FInterchangeFbxParser::FetchPayload(const FString& PayloadKey, const FString& ResultFolder)
+		Container->AddToRoot();
+		FbxParserPrivate->FillContainerWithFbxScene(*Container);
+		Container->SaveToFile(ResultFilepath);
+		Container->RemoveFromRoot();
+	}
+
+	void FInterchangeFbxParser::FetchPayload(const FString& PayloadKey, const FString& ResultFolder)
+	{
+		check(FbxParserPrivate.IsValid());
+		ResultsContainer->Empty();
+		FString PayloadFilepathCopy;
 		{
-			check(FbxParserPrivate.IsValid());
-			ResultsContainer->Empty();
-			FString PayloadFilepathCopy;
-			{
-				FScopeLock Lock(&ResultPayloadsCriticalSection);
-				FString& PayloadFilepath = ResultPayloads.FindOrAdd(PayloadKey);
-				PayloadFilepath = ResultFolder + TEXT("/") + PayloadKey + TEXT(".payload");
-
-				//Copy the map filename key because we are multithreaded and the TMap can be reallocated
-				PayloadFilepathCopy = PayloadFilepath;
-			}
-			if (!FbxParserPrivate->FetchPayloadData(PayloadKey, PayloadFilepathCopy))
-			{
-				UInterchangeResultError_Generic* Error = AddMessage<UInterchangeResultError_Generic>();
-				Error->SourceAssetName = SourceFilename;
-				Error->Text = LOCTEXT("CantFetchPayload", "Cannot fetch FBX payload data.");
-				return;
-			}
+			FScopeLock Lock(&ResultPayloadsCriticalSection);
+			FString& PayloadFilepath = ResultPayloads.FindOrAdd(PayloadKey);
+			//To avoid file path with too many character, we hash the payloadKey so we have a deterministic length for the file path.
+			FString PayloadKeyHash = Private::HashString(PayloadKey);
+			PayloadFilepath = ResultFolder + TEXT("/") + PayloadKeyHash + TEXT(".payload");
+
+			//Copy the map filename key because we are multithreaded and the TMap can be reallocated
+			PayloadFilepathCopy = PayloadFilepath;
 		}
-
-		void FInterchangeFbxParser::FetchAnimationBakeTransformPayload(const FString& PayloadKey, const double BakeFrequency, const double RangeStartTime, const double RangeEndTime, const FString& ResultFolder)
+		if (!FbxParserPrivate->FetchPayloadData(PayloadKey, PayloadFilepathCopy))
 		{
-			check(FbxParserPrivate.IsValid());
-			ResultsContainer->Empty();
-			FString PayloadFilepathCopy;
-			{
-				FScopeLock Lock(&ResultPayloadsCriticalSection);
-				FString& PayloadFilepath = ResultPayloads.FindOrAdd(PayloadKey);
-				PayloadFilepath = ResultFolder + TEXT("/") + PayloadKey + TEXT(".payload");
-
-				//Copy the map filename key because we are multithreaded and the TMap can be reallocated
-				PayloadFilepathCopy = PayloadFilepath;
-			}
-			if (!FbxParserPrivate->FetchAnimationBakeTransformPayload(PayloadKey, BakeFrequency, RangeStartTime, RangeEndTime, PayloadFilepathCopy))
-			{
-				UInterchangeResultError_Generic* Error = AddMessage<UInterchangeResultError_Generic>();
-				Error->SourceAssetName = SourceFilename;
-				Error->Text = LOCTEXT("CantFetchPayload", "Cannot fetch FBX payload data.");
-				return;
-			}
+			UInterchangeResultError_Generic* Error = AddMessage<UInterchangeResultError_Generic>();
+			Error->SourceAssetName = SourceFilename;
+			Error->Text = LOCTEXT("CantFetchPayload", "Cannot fetch FBX payload data.");
+			return;
 		}
+	}
 
-		TArray<FString> FInterchangeFbxParser::GetJsonLoadMessages() const
+	void FInterchangeFbxParser::FetchAnimationBakeTransformPayload(const FString& PayloadKey, const double BakeFrequency, const double RangeStartTime, const double RangeEndTime, const FString& ResultFolder)
+	{
+		check(FbxParserPrivate.IsValid());
+		ResultsContainer->Empty();
+		FString PayloadFilepathCopy;
+		{
+			FScopeLock Lock(&ResultPayloadsCriticalSection);
+			FString& PayloadFilepath = ResultPayloads.FindOrAdd(PayloadKey);
+			//To avoid file path with too many character, we hash the payloadKey so we have a deterministic length for the file path.
+			FString PayloadKeyHash = Private::HashString(PayloadKey);
+			PayloadFilepath = ResultFolder + TEXT("/") + PayloadKeyHash + TEXT(".payload");
+
+			//Copy the map filename key because we are multithreaded and the TMap can be reallocated
+			PayloadFilepathCopy = PayloadFilepath;
+		}
+		if (!FbxParserPrivate->FetchAnimationBakeTransformPayload(PayloadKey, BakeFrequency, RangeStartTime, RangeEndTime, PayloadFilepathCopy))
 		{
-			TArray<FString> JsonResults;
-			for (UInterchangeResult* Result : ResultsContainer->GetResults())
-			{
-				JsonResults.Add(Result->ToJson());
-			}
+			UInterchangeResultError_Generic* Error = AddMessage<UInterchangeResultError_Generic>();
+			Error->SourceAssetName = SourceFilename;
+			Error->Text = LOCTEXT("CantFetchPayload", "Cannot fetch FBX payload data.");
+			return;
+		}
+	}
 
-			return JsonResults;
+	TArray<FString> FInterchangeFbxParser::GetJsonLoadMessages() const
+	{
+		TArray<FString> JsonResults;
+		for (UInterchangeResult* Result : ResultsContainer->GetResults())
+		{
+			JsonResults.Add(Result->ToJson());
 		}
 
-	}//ns Interchange
-}//ns UE
+		return JsonResults;
+	}
+
+}//ns UE::Interchange
 
 #undef LOCTEXT_NAMESPACE
diff --git a/Engine/Plugins/Experimental/Interchange/Runtime/Source/Parsers/Fbx/Public/InterchangeFbxParser.h b/Engine/Plugins/Experimental/Interchange/Runtime/Source/Parsers/Fbx/Public/InterchangeFbxParser.h
--- a/Engine/Plugins/Experimental/Interchange/Runtime/Source/Parsers/Fbx/Public/InterchangeFbxParser.h
+++ b/Engine/Plugins/Experimental/Interchange/Runtime/Source/Parsers/Fbx/Public/InterchangeFbxParser.h
@@ -23,6 +23,7 @@ namespace UE
 		public:
 			FInterchangeFbxParser();
 			~FInterchangeFbxParser();
+
 			/**
 			 * Parse a file support by the fbx sdk. It just extract all the fbx node and create a FBaseNodeContainer and dump it in a json file inside the ResultFolder
 			 * @param - Filename is the file that the fbx sdk will read (.fbx or .obj)
@@ -71,7 +72,7 @@ namespace UE
 
 
 		private:
-			TObjectPtr<UInterchangeResultsContainer> ResultsContainer = nullptr;
+			TStrongObjectPtr<UInterchangeResultsContainer> ResultsContainer = nullptr;
 			FString SourceFilename;
 			FString ResultFilepath;
 			mutable FCriticalSection ResultPayloadsCriticalSection;
diff --git a/Engine/Plugins/Experimental/Interchange/Runtime/Source/Pipelines/Private/InterchangeGenericAnimationPipeline.cpp b/Engine/Plugins/Experimental/Interchange/Runtime/Source/Pipelines/Private/InterchangeGenericAnimationPipeline.cpp
--- a/Engine/Plugins/Experimental/Interchange/Runtime/Source/Pipelines/Private/InterchangeGenericAnimationPipeline.cpp
+++ b/Engine/Plugins/Experimental/Interchange/Runtime/Source/Pipelines/Private/InterchangeGenericAnimationPipeline.cpp
@@ -1,17 +1,22 @@
 // Copyright Epic Games, Inc. All Rights Reserved. 
 #include "InterchangeGenericAnimationPipeline.h"
 
+#include "Animation/AnimationSettings.h"
 #include "CoreMinimal.h"
 #include "InterchangeAnimSequenceFactoryNode.h"
+#include "InterchangeMeshNode.h"
 #include "InterchangePipelineLog.h"
 #include "InterchangeSceneNode.h"
 #include "InterchangeSkeletalMeshFactoryNode.h"
+#include "InterchangeSkeletalMeshLodDataNode.h"
 #include "InterchangeSkeletonFactoryNode.h"
 #include "InterchangeSkeletonHelper.h"
 #include "InterchangeSourceData.h"
+#include "Nodes/InterchangeAnimationAPI.h"
 #include "Nodes/InterchangeBaseNode.h"
 #include "Nodes/InterchangeBaseNodeContainer.h"
 #include "Nodes/InterchangeSourceNode.h"
+#include "Nodes/InterchangeUserDefinedAttribute.h"
 
 
 namespace UE::Interchange::Private
@@ -24,7 +29,7 @@ namespace UE::Interchange::Private
 		if (const UInterchangeSceneNode* SceneNode = Cast<UInterchangeSceneNode>(BaseNodeContainer->GetNode(NodeUid)))
 		{
 			bool bIsAnimated = false;
-			if (SceneNode->GetCustomIsNodeTransformAnimated(bIsAnimated) && bIsAnimated)
+			if (UInterchangeAnimationAPI::GetCustomIsNodeTransformAnimated(SceneNode, bIsAnimated) && bIsAnimated)
 			{
 				return true;
 			}
@@ -65,6 +70,12 @@ void UInterchangeGenericAnimationPipeline::AdjustSettingsForReimportType(EInterc
 
 void UInterchangeGenericAnimationPipeline::ExecutePreImportPipeline(UInterchangeBaseNodeContainer* InBaseNodeContainer, const TArray<UInterchangeSourceData*>& InSourceDatas)
 {
+	if (!bImportAnimations)
+	{
+		//Nothing to import
+		return;
+	}
+
 	check(!CommonSkeletalMeshesAndAnimationsProperties.IsNull());
 	if (!InBaseNodeContainer)
 	{
@@ -152,9 +163,9 @@ void UInterchangeGenericAnimationPipeline::ExecutePreImportPipeline(UInterchange
 				{
 					double SceneNodeAnimStart;
 					double SceneNodeAnimStop;
-					if (SceneNode->GetCustomNodeTransformAnimationStartTime(SceneNodeAnimStart))
+					if (UInterchangeAnimationAPI::GetCustomNodeTransformAnimationStartTime(SceneNode, SceneNodeAnimStart))
 					{
-						if (SceneNode->GetCustomNodeTransformAnimationEndTime(SceneNodeAnimStop))
+						if (UInterchangeAnimationAPI::GetCustomNodeTransformAnimationEndTime(SceneNode, SceneNodeAnimStop))
 						{
 							if (RangeStart > SceneNodeAnimStart)
 							{
@@ -172,24 +183,132 @@ void UInterchangeGenericAnimationPipeline::ExecutePreImportPipeline(UInterchange
 		}
 	}
 
+	//Retrieve all skeletons and all morph target anim data
+	TMap<UInterchangeSkeletonFactoryNode*, TArray<const UInterchangeMeshNode*>> MorphTargetsPerSkeletons;
 
-	//Retrieve all skeletons
-	TArray<UInterchangeSkeletonFactoryNode*> Skeletons;
-	BaseNodeContainer->IterateNodes([&Skeletons, LocalNodeContainer = BaseNodeContainer](const FString& NodeUid, UInterchangeBaseNode* Node)
+	BaseNodeContainer->IterateNodes([&MorphTargetsPerSkeletons, LocalNodeContainer = BaseNodeContainer](const FString& NodeUid, UInterchangeBaseNode* Node)
 		{
 			if (UInterchangeSkeletonFactoryNode* SkeletonFactoryNode = Cast<UInterchangeSkeletonFactoryNode>(Node))
 			{
+				const UInterchangeSkeletalMeshFactoryNode* SkeletalMeshFactoryNode = nullptr;
+				//Iterate dependencies
+				TArray<FString> SkeletalMeshNodeUids;
+				LocalNodeContainer->GetNodes(UInterchangeSkeletalMeshFactoryNode::StaticClass(), SkeletalMeshNodeUids);
+				for (const FString& SkelMeshFactoryNodeUid : SkeletalMeshNodeUids)
+				{
+					if (const UInterchangeSkeletalMeshFactoryNode* CurrentSkeletalMeshFactoryNode = Cast<const UInterchangeSkeletalMeshFactoryNode>(LocalNodeContainer->GetFactoryNode(SkelMeshFactoryNodeUid)))
+					{
+						TArray<FString> SkeletalMeshDependencies;
+						CurrentSkeletalMeshFactoryNode->GetFactoryDependencies(SkeletalMeshDependencies);
+						for (const FString& SkeletalMeshDependencyUid : SkeletalMeshDependencies)
+						{
+							if (NodeUid.Equals(SkeletalMeshDependencyUid))
+							{
+								SkeletalMeshFactoryNode = CurrentSkeletalMeshFactoryNode;
+								break;
+							}
+						}
+					}
+				}
+
 				FString RootSceneNodeUid;
 				SkeletonFactoryNode->GetCustomRootJointUid(RootSceneNodeUid);
 				if(UE::Interchange::Private::IsSkeletonAnimatedRecursive(RootSceneNodeUid, LocalNodeContainer))
 				{
-					Skeletons.Add(SkeletonFactoryNode);
+					MorphTargetsPerSkeletons.FindOrAdd(SkeletonFactoryNode);
+				}
+				
+				if (SkeletalMeshFactoryNode)
+				{
+					//Find the skeletalmesh morph targets
+					FString SkeletalMeshFactoryNodeUid = SkeletalMeshFactoryNode->GetUniqueID();
+					//Get the LOD data factory node
+					TArray<FString> LodDataChildren;
+					SkeletalMeshFactoryNode->GetLodDataUniqueIds(LodDataChildren);
+					for (const FString& ChildUid : LodDataChildren)
+					{
+						if (const UInterchangeSkeletalMeshLodDataNode* LodData = Cast<UInterchangeSkeletalMeshLodDataNode>(LocalNodeContainer->GetNode(ChildUid)))
+						{
+							TArray<FString> MeshUids;
+							LodData->GetMeshUids(MeshUids);
+							for (const FString& MeshUid : MeshUids)
+							{
+								if (const UInterchangeBaseNode* BaseNode = LocalNodeContainer->GetNode(MeshUid))
+								{
+									const UInterchangeMeshNode* MeshNode = nullptr;
+									FString RealMeshUid;
+									if (const UInterchangeSceneNode* SceneNode = Cast<UInterchangeSceneNode>(BaseNode))
+									{
+										SceneNode->GetCustomAssetInstanceUid(RealMeshUid);
+									}
+									else
+									{
+										RealMeshUid = MeshUid;
+									}
+									MeshNode = Cast<UInterchangeMeshNode>(LocalNodeContainer->GetNode(RealMeshUid));
+									if (MeshNode)
+									{
+										TArray<FString> MorphTargetUids;
+										MeshNode->GetMorphTargetDependencies(MorphTargetUids);
+										bool bIsMorphTargetAnimated = false;
+										TArray<const UInterchangeMeshNode*> MorphTargets;
+										for (const FString& MorphTargetUid : MorphTargetUids)
+										{
+											if (const UInterchangeMeshNode* MorphTargetNode = Cast<UInterchangeMeshNode>(LocalNodeContainer->GetNode(MorphTargetUid)))
+											{
+												MorphTargets.AddUnique(MorphTargetNode);
+												TOptional<FString> MorphTargetAnimationPayloadKey = MorphTargetNode->GetAnimationCurvePayLoadKey();
+												if (MorphTargetAnimationPayloadKey.IsSet())
+												{
+													bIsMorphTargetAnimated = true;
+												}
+											}
+										}
+										if (bIsMorphTargetAnimated)
+										{
+											TArray<const UInterchangeMeshNode*>& MorphTargetNodes = MorphTargetsPerSkeletons.FindOrAdd(SkeletonFactoryNode);
+											MorphTargetNodes.Reset(MorphTargets.Num());
+											MorphTargetNodes.Append(MorphTargets);
+										}
+									}
+								}
+							}
+						}
+					}
+					
+					//If we do not already create an anim sequence for this skeleton, add the anim sequence if we have
+					//at least one animated user defined attribute.
+					if (!MorphTargetsPerSkeletons.Contains(SkeletonFactoryNode))
+					{
+						//Add anim sequence if we have at least one animated user defined attributes
+						LocalNodeContainer->BreakableIterateNodeChildren(RootSceneNodeUid, [&MorphTargetsPerSkeletons, SkeletonFactoryNode](const UInterchangeBaseNode* Node)
+							{
+								if (const UInterchangeSceneNode* SceneNode = Cast<UInterchangeSceneNode>(Node))
+								{
+									TArray<FInterchangeUserDefinedAttributeInfo> AttributeInfos;
+									UInterchangeUserDefinedAttributesAPI::GetUserDefinedAttributeInfos(SceneNode, AttributeInfos);
+									for (const FInterchangeUserDefinedAttributeInfo& AttributeInfo : AttributeInfos)
+									{
+										if (AttributeInfo.PayloadKey.IsSet())
+										{
+											MorphTargetsPerSkeletons.FindOrAdd(SkeletonFactoryNode);
+											return true;
+										}
+									}
+								}
+								return false;
+							});
+					}
 				}
 			}
 		});
+
+	const TArray<FString> CustomAttributeNamesToImport = UAnimationSettings::Get()->GetBoneCustomAttributeNamesToImport();
+
 	//for each animated skeleton create one anim sequence factory node
-	for (UInterchangeSkeletonFactoryNode* SkeletonFactoryNode : Skeletons)
+	for (const TPair<UInterchangeSkeletonFactoryNode*, TArray<const UInterchangeMeshNode*>>& SkeletonAndMorphTargets : MorphTargetsPerSkeletons)
 	{
+		UInterchangeSkeletonFactoryNode* SkeletonFactoryNode = SkeletonAndMorphTargets.Key;
 		const FString AnimSequenceUid = TEXT("\\AnimSequence") + SkeletonFactoryNode->GetUniqueID();
 		FString AnimSequenceName = SkeletonFactoryNode->GetDisplayLabel();
 		if (AnimSequenceName.EndsWith(TEXT("_Skeleton")))
@@ -231,6 +350,28 @@ void UInterchangeGenericAnimationPipeline::ExecutePreImportPipeline(UInterchange
 			AnimSequenceFactoryNode->SetCustomImportBoneTracksRangeStop(RangeStop);
 		}
 
+		AnimSequenceFactoryNode->SetCustomImportAttributeCurves(bImportCustomAttribute);
+		AnimSequenceFactoryNode->SetCustomDoNotImportCurveWithZero(bDoNotImportCurveWithZero);
+		AnimSequenceFactoryNode->SetCustomRemoveCurveRedundantKeys(bRemoveCurveRedundantKeys);
+		AnimSequenceFactoryNode->SetCustomDeleteExistingMorphTargetCurves(bDeleteExistingMorphTargetCurves);
+		AnimSequenceFactoryNode->SetCustomDeleteExistingCustomAttributeCurves(bDeleteExistingCustomAttributeCurves);
+		AnimSequenceFactoryNode->SetCustomDeleteExistingNonCurveCustomAttributes(bDeleteExistingNonCurveCustomAttributes);
+
+		AnimSequenceFactoryNode->SetCustomMaterialDriveParameterOnCustomAttribute(bSetMaterialDriveParameterOnCustomAttribute);
+		for (const FString& MaterialSuffixe : MaterialCurveSuffixes)
+		{
+			AnimSequenceFactoryNode->SetAnimatedMaterialCurveSuffixe(MaterialSuffixe);
+		}
+
+		//Add the animated morph targets uid so the factory can import them
+		if (SkeletonAndMorphTargets.Value.Num() > 0)
+		{
+			for (const UInterchangeMeshNode* MorphTargetNode : SkeletonAndMorphTargets.Value)
+			{
+				AnimSequenceFactoryNode->SetAnimatedMorphTargetDependencyUid(MorphTargetNode->GetUniqueID());
+			}
+		}
+
 		//USkeleton cannot be created without a valid skeletalmesh
 		const FString SkeletonUid = SkeletonFactoryNode->GetUniqueID();
 		AnimSequenceFactoryNode->AddFactoryDependencyUid(SkeletonUid);
@@ -242,6 +383,45 @@ void UInterchangeGenericAnimationPipeline::ExecutePreImportPipeline(UInterchange
 			//Iterate all joints to set the meta data value in the anim sequence factory node
 			UE::Interchange::Private::FSkeletonHelper::RecursiveAddSkeletonMetaDataValues(BaseNodeContainer, AnimSequenceFactoryNode, RootJointUid);
 #endif //WITH_EDITOR
+			BaseNodeContainer->IterateNodeChildren(RootJointUid, [&AnimSequenceFactoryNode, &CustomAttributeNamesToImport](const UInterchangeBaseNode* Node)
+			{
+				if (const UInterchangeSceneNode* SceneNode = Cast<UInterchangeSceneNode>(Node))
+				{
+					FString BoneName = SceneNode->GetDisplayLabel();
+					bool bImportAllAttributesOnBone = UAnimationSettings::Get()->BoneNamesWithCustomAttributes.Contains(BoneName);
+						
+					TArray<FInterchangeUserDefinedAttributeInfo> AttributeInfos;
+					UInterchangeUserDefinedAttributesAPI::GetUserDefinedAttributeInfos(SceneNode, AttributeInfos);
+					for (const FInterchangeUserDefinedAttributeInfo& AttributeInfo : AttributeInfos)
+					{
+						if (AttributeInfo.PayloadKey.IsSet())
+						{
+							bool bDecimalType = false;
+							switch (AttributeInfo.Type)
+							{
+							case UE::Interchange::EAttributeTypes::Float:
+							case UE::Interchange::EAttributeTypes::Float16:
+							case UE::Interchange::EAttributeTypes::Double:
+								{
+									bDecimalType = true;
+								}
+								break;
+							}
+
+							const bool bForceImportBoneCustomAttribute = CustomAttributeNamesToImport.Contains(AttributeInfo.Name);
+							//Material attribute curve
+							if (!bImportAllAttributesOnBone && bDecimalType && !bForceImportBoneCustomAttribute)
+							{
+								AnimSequenceFactoryNode->SetAnimatedAttributeCurveName(AttributeInfo.Name);
+							}
+							else if (bForceImportBoneCustomAttribute || bImportAllAttributesOnBone)
+							{
+								AnimSequenceFactoryNode->SetAnimatedAttributeStepCurveName(AttributeInfo.Name);
+							}
+						}
+					}
+				}
+			});
 		}
 
 		const bool bImportOnlyAnimation = CommonSkeletalMeshesAndAnimationsProperties->bImportOnlyAnimations;
diff --git a/Engine/Plugins/Experimental/Interchange/Runtime/Source/Pipelines/Private/InterchangePipelineMeshesUtilities.cpp b/Engine/Plugins/Experimental/Interchange/Runtime/Source/Pipelines/Private/InterchangePipelineMeshesUtilities.cpp
--- a/Engine/Plugins/Experimental/Interchange/Runtime/Source/Pipelines/Private/InterchangePipelineMeshesUtilities.cpp
+++ b/Engine/Plugins/Experimental/Interchange/Runtime/Source/Pipelines/Private/InterchangePipelineMeshesUtilities.cpp
@@ -121,7 +121,7 @@ UInterchangePipelineMeshesUtilities* UInterchangePipelineMeshesUtilities::Create
 									MeshGeometry.ReferencingMeshInstanceUids.Add(MeshInstance.MeshInstanceUid);
 									MeshInstance.ReferencingMeshGeometryUids.Add(MeshUid);
 									MeshInstance.bReferenceSkinnedMesh |= MeshGeometry.MeshNode->IsSkinnedMesh();
-									MeshInstance.bReferenceBlendShape |= MeshGeometry.MeshNode->IsBlendShape();
+									MeshInstance.bReferenceMorphTarget |= MeshGeometry.MeshNode->IsMorphTarget();
 								}
 								else
 								{
@@ -133,7 +133,7 @@ UInterchangePipelineMeshesUtilities* UInterchangePipelineMeshesUtilities::Create
 									MeshGeometry.ReferencingMeshInstanceUids.Add(MeshInstance.MeshInstanceUid);
 									MeshInstance.ReferencingMeshGeometryUids.Add(MeshUid);
 									MeshInstance.bReferenceSkinnedMesh |= MeshGeometry.MeshNode->IsSkinnedMesh();
-									MeshInstance.bReferenceBlendShape |= MeshGeometry.MeshNode->IsBlendShape();
+									MeshInstance.bReferenceMorphTarget |= MeshGeometry.MeshNode->IsMorphTarget();
 								}
 							}
 						}
@@ -231,7 +231,7 @@ void UInterchangePipelineMeshesUtilities::GetAllSkinnedMeshInstance(TArray<FStri
 	for (const TPair<FString, FInterchangeMeshInstance>& MeshInstanceUidAndMeshInstance : MeshInstancesPerMeshInstanceUid)
 	{
 		const FInterchangeMeshInstance& MeshInstance = MeshInstanceUidAndMeshInstance.Value;
-		if((bConvertStaticMeshToSkeletalMesh || MeshInstance.bReferenceSkinnedMesh) && !MeshInstance.bReferenceBlendShape)
+		if((bConvertStaticMeshToSkeletalMesh || MeshInstance.bReferenceSkinnedMesh) && !MeshInstance.bReferenceMorphTarget)
 		{
 			MeshInstanceUids.Add(MeshInstance.MeshInstanceUid);
 		}
@@ -243,7 +243,7 @@ void UInterchangePipelineMeshesUtilities::IterateAllSkinnedMeshInstance(TFunctio
 	for (const TPair<FString, FInterchangeMeshInstance>& MeshInstanceUidAndMeshInstance : MeshInstancesPerMeshInstanceUid)
 	{
 		const FInterchangeMeshInstance& MeshInstance = MeshInstanceUidAndMeshInstance.Value;
-		if ((bConvertStaticMeshToSkeletalMesh || MeshInstance.bReferenceSkinnedMesh) && !MeshInstance.bReferenceBlendShape)
+		if ((bConvertStaticMeshToSkeletalMesh || MeshInstance.bReferenceSkinnedMesh) && !MeshInstance.bReferenceMorphTarget)
 		{
 			IterationLambda(MeshInstance);
 		}
@@ -256,7 +256,7 @@ void UInterchangePipelineMeshesUtilities::GetAllStaticMeshInstance(TArray<FStrin
 	for (const TPair<FString, FInterchangeMeshInstance>& MeshInstanceUidAndMeshInstance : MeshInstancesPerMeshInstanceUid)
 	{
 		const FInterchangeMeshInstance& MeshInstance = MeshInstanceUidAndMeshInstance.Value;
-		if ((bConvertSkeletalMeshToStaticMesh || !MeshInstance.bReferenceSkinnedMesh) && !MeshInstance.bReferenceBlendShape)
+		if ((bConvertSkeletalMeshToStaticMesh || !MeshInstance.bReferenceSkinnedMesh) && !MeshInstance.bReferenceMorphTarget)
 		{
 			MeshInstanceUids.Add(MeshInstance.MeshInstanceUid);
 		}
@@ -268,7 +268,7 @@ void UInterchangePipelineMeshesUtilities::IterateAllStaticMeshInstance(TFunction
 	for (const TPair<FString, FInterchangeMeshInstance>& MeshInstanceUidAndMeshInstance : MeshInstancesPerMeshInstanceUid)
 	{
 		const FInterchangeMeshInstance& MeshInstance = MeshInstanceUidAndMeshInstance.Value;
-		if ((bConvertSkeletalMeshToStaticMesh || !MeshInstance.bReferenceSkinnedMesh) && !MeshInstance.bReferenceBlendShape)
+		if ((bConvertSkeletalMeshToStaticMesh || !MeshInstance.bReferenceSkinnedMesh) && !MeshInstance.bReferenceMorphTarget)
 		{
 			IterationLambda(MeshInstance);
 		}
@@ -320,7 +320,7 @@ void UInterchangePipelineMeshesUtilities::GetAllStaticMeshGeometry(TArray<FStrin
 	for (const TPair<FString, FInterchangeMeshGeometry>& MeshGeometryUidAndMeshGeometry : MeshGeometriesPerMeshUid)
 	{
 		const FInterchangeMeshGeometry& MeshGeometry = MeshGeometryUidAndMeshGeometry.Value;
-		if ((bConvertSkeletalMeshToStaticMesh || !MeshGeometry.MeshNode->IsSkinnedMesh()) && !MeshGeometry.MeshNode->IsBlendShape())
+		if ((bConvertSkeletalMeshToStaticMesh || !MeshGeometry.MeshNode->IsSkinnedMesh()) && !MeshGeometry.MeshNode->IsMorphTarget())
 		{
 			MeshGeometryUids.Add(MeshGeometry.MeshUid);
 		}
@@ -332,7 +332,7 @@ void UInterchangePipelineMeshesUtilities::IterateAllStaticMeshGeometry(TFunction
 	for (const TPair<FString, FInterchangeMeshGeometry>& MeshGeometryUidAndMeshGeometry : MeshGeometriesPerMeshUid)
 	{
 		const FInterchangeMeshGeometry& MeshGeometry = MeshGeometryUidAndMeshGeometry.Value;
-		if ((bConvertSkeletalMeshToStaticMesh || !MeshGeometry.MeshNode->IsSkinnedMesh()) && !MeshGeometry.MeshNode->IsBlendShape())
+		if ((bConvertSkeletalMeshToStaticMesh || !MeshGeometry.MeshNode->IsSkinnedMesh()) && !MeshGeometry.MeshNode->IsMorphTarget())
 		{
 			IterationLambda(MeshGeometry);
 		}
diff --git a/Engine/Plugins/Experimental/Interchange/Runtime/Source/Pipelines/Public/InterchangeGenericAnimationPipeline.h b/Engine/Plugins/Experimental/Interchange/Runtime/Source/Pipelines/Public/InterchangeGenericAnimationPipeline.h
--- a/Engine/Plugins/Experimental/Interchange/Runtime/Source/Pipelines/Public/InterchangeGenericAnimationPipeline.h
+++ b/Engine/Plugins/Experimental/Interchange/Runtime/Source/Pipelines/Public/InterchangeGenericAnimationPipeline.h
@@ -44,7 +44,7 @@ public:
 	UPROPERTY(EditAnywhere, BlueprintReadWrite, Category = "Animations")
 	bool bImportAnimations = true;
 
-	/** Import bone transform tracks. If false, this will discard any bone transform tracks. (useful for curves only animations)*/
+	/** Import bone transform tracks. If false, this will discard any bone transform tracks.*/
 	UPROPERTY(EditAnywhere, BlueprintReadWrite, Category = "Animations")
 	bool bImportBoneTracks = true;
 
@@ -78,7 +78,7 @@ public:
 
 	/** When importing custom attribute as curve, remove redundant keys */
 	UPROPERTY(EditAnywhere, BlueprintReadWrite, Category = "Animations", meta = (EditCondition = "bImportCustomAttribute", DisplayName = "Remove Redundant Keys"))
-	bool bRemoveRedundantKeys = false;
+	bool bRemoveCurveRedundantKeys = false;
 
 	/** When importing custom attribute or morphtarget as curve, do not import if it doens't have any value other than zero. This is to avoid adding extra curves to evaluate */
 	UPROPERTY(EditAnywhere, BlueprintReadWrite, Category = "Animations", meta = (EditCondition = "bImportCustomAttribute", DisplayName = "Do not import curves with only 0 values"))
@@ -92,7 +92,7 @@ public:
 	UPROPERTY(EditAnywhere, BlueprintReadWrite, Category = "Animations")
 	bool bDeleteExistingCustomAttributeCurves = false;
 
-	/** If enabled, this will delete this type of asset from the FBX */
+	/** If true, all previous morph target curves will be deleted when doing a re-import */
 	UPROPERTY(EditAnywhere, BlueprintReadWrite, Category = "Animations")
 	bool bDeleteExistingMorphTargetCurves = false;
 
diff --git a/Engine/Plugins/Experimental/Interchange/Runtime/Source/Pipelines/Public/InterchangePipelineMeshesUtilities.h b/Engine/Plugins/Experimental/Interchange/Runtime/Source/Pipelines/Public/InterchangePipelineMeshesUtilities.h
--- a/Engine/Plugins/Experimental/Interchange/Runtime/Source/Pipelines/Public/InterchangePipelineMeshesUtilities.h
+++ b/Engine/Plugins/Experimental/Interchange/Runtime/Source/Pipelines/Public/InterchangePipelineMeshesUtilities.h
@@ -45,7 +45,7 @@ struct FInterchangeMeshInstance
 	{
 		LodGroupNode = nullptr;
 		bReferenceSkinnedMesh = false;
-		bReferenceBlendShape = false;
+		bReferenceMorphTarget = false;
 	}
 	/**
 	 * This ID represent either 1: a lod group scene node uid or 2: a mesh scene node uid.
@@ -63,7 +63,7 @@ struct FInterchangeMeshInstance
 	bool bReferenceSkinnedMesh;
 
 	UPROPERTY(EditAnywhere, Category = "Interchange | Pipeline | MeshesInstance")
-	bool bReferenceBlendShape;
+	bool bReferenceMorphTarget;
 
 	/**
 	 * Each scene node here represent a mesh scene node. Only if we represent a lod group we can have more then 1 mesh scene node for a specific lod index.
diff --git a/Engine/Source/Programs/InterchangeWorker/Private/InterchangeWorker.cpp b/Engine/Source/Programs/InterchangeWorker/Private/InterchangeWorker.cpp
--- a/Engine/Source/Programs/InterchangeWorker/Private/InterchangeWorker.cpp
+++ b/Engine/Source/Programs/InterchangeWorker/Private/InterchangeWorker.cpp
@@ -78,20 +78,12 @@ int32 Main(int32 Argc, TCHAR * Argv[])
  		UE_LOG(LogInterchangeWorker, Error, TEXT("%s"), *WorkerVersionError);
  	}
 
-	if (!WorkerVersionError.IsEmpty())
+	//FInterchangeWorkerImpl Scope
 	{
-		UInterchangeResultError_Generic* Message = NewObject<UInterchangeResultError_Generic>(GetTransientPackage());
-		Message->SourceAssetName = TEXT("InterchangeWorkerProgram");
-		Message->Text = FText::FromString(WorkerVersionError);
-		WorkerVersionError = Message->ToJson();
-		Message->ClearFlags(RF_Standalone);
-		Message->ClearInternalFlags(EInternalObjectFlags::Async | EInternalObjectFlags::AsyncLoading);
-		//Should be garbage collectable now
+		FInterchangeWorkerImpl Worker(FCString::Atoi(*ServerPID), FCString::Atoi(*ServerPort), ResultFolder);
+		Worker.Run(WorkerVersionError);
 	}
 
-	FInterchangeWorkerImpl Worker(FCString::Atoi(*ServerPID), FCString::Atoi(*ServerPort), ResultFolder);
-	Worker.Run(WorkerVersionError);
-
 	if (!WorkerVersionError.IsEmpty())
 	{
 		return EXIT_FAILURE;
diff --git a/Engine/Source/Runtime/Experimental/Interchange/Core/Private/Nodes/InterchangeAnimationAPI.cpp b/Engine/Source/Runtime/Experimental/Interchange/Core/Private/Nodes/InterchangeAnimationAPI.cpp
new file mode 100644
--- /dev/null
+++ b/Engine/Source/Runtime/Experimental/Interchange/Core/Private/Nodes/InterchangeAnimationAPI.cpp
@@ -0,0 +1,61 @@
+// Copyright Epic Games, Inc. All Rights Reserved.
+
+#include "Nodes/InterchangeAnimationAPI.h"
+
+#include "Nodes/InterchangeBaseNode.h"
+
+const UE::Interchange::FAttributeKey UInterchangeAnimationAPI::Macro_CustomIsNodeTransformAnimatedKey = UE::Interchange::FAttributeKey(TEXT("IsNodeTransformAnimated"));
+const UE::Interchange::FAttributeKey UInterchangeAnimationAPI::Macro_CustomNodeTransformAnimationKeyCountKey = UE::Interchange::FAttributeKey(TEXT("NodeTransformAnimationKeyCount"));
+const UE::Interchange::FAttributeKey UInterchangeAnimationAPI::Macro_CustomNodeTransformAnimationStartTimeKey = UE::Interchange::FAttributeKey(TEXT("NodeTransformAnimationStartTime"));
+const UE::Interchange::FAttributeKey UInterchangeAnimationAPI::Macro_CustomNodeTransformAnimationEndTimeKey = UE::Interchange::FAttributeKey(TEXT("NodeTransformAnimationEndTime"));
+const UE::Interchange::FAttributeKey UInterchangeAnimationAPI::Macro_CustomNodeTransformPayloadKeyKey = UE::Interchange::FAttributeKey(TEXT("NodeTransformPayloadKey"));
+
+bool UInterchangeAnimationAPI::GetCustomIsNodeTransformAnimated(const UInterchangeBaseNode* InterchangeBaseNode, bool& AttributeValue)
+{
+	return InterchangeBaseNode->GetBooleanAttribute(Macro_CustomIsNodeTransformAnimatedKey.Key, AttributeValue);
+}
+
+bool UInterchangeAnimationAPI::SetCustomIsNodeTransformAnimated(UInterchangeBaseNode* InterchangeBaseNode, const bool& AttributeValue)
+{
+	return InterchangeBaseNode->AddBooleanAttribute(Macro_CustomIsNodeTransformAnimatedKey.Key, AttributeValue);
+}
+
+bool UInterchangeAnimationAPI::GetCustomNodeTransformAnimationKeyCount(const UInterchangeBaseNode* InterchangeBaseNode, int32& AttributeValue)
+{
+	return InterchangeBaseNode->GetInt32Attribute(Macro_CustomNodeTransformAnimationKeyCountKey.Key, AttributeValue);
+}
+
+bool UInterchangeAnimationAPI::SetCustomNodeTransformAnimationKeyCount(UInterchangeBaseNode* InterchangeBaseNode, const int32& AttributeValue)
+{
+	return InterchangeBaseNode->AddInt32Attribute(Macro_CustomNodeTransformAnimationKeyCountKey.Key, AttributeValue);
+}
+
+bool UInterchangeAnimationAPI::GetCustomNodeTransformAnimationStartTime(const UInterchangeBaseNode* InterchangeBaseNode, double& AttributeValue)
+{
+	return InterchangeBaseNode->GetDoubleAttribute(Macro_CustomNodeTransformAnimationStartTimeKey.Key, AttributeValue);
+}
+
+bool UInterchangeAnimationAPI::SetCustomNodeTransformAnimationStartTime(UInterchangeBaseNode* InterchangeBaseNode, const double& AttributeValue)
+{
+	return InterchangeBaseNode->AddDoubleAttribute(Macro_CustomNodeTransformAnimationStartTimeKey.Key, AttributeValue);
+}
+
+bool UInterchangeAnimationAPI::GetCustomNodeTransformAnimationEndTime(const UInterchangeBaseNode* InterchangeBaseNode, double& AttributeValue)
+{
+	return InterchangeBaseNode->GetDoubleAttribute(Macro_CustomNodeTransformAnimationEndTimeKey.Key, AttributeValue);
+}
+
+bool UInterchangeAnimationAPI::SetCustomNodeTransformAnimationEndTime(UInterchangeBaseNode* InterchangeBaseNode, const double& AttributeValue)
+{
+	return InterchangeBaseNode->AddDoubleAttribute(Macro_CustomNodeTransformAnimationEndTimeKey.Key, AttributeValue);
+}
+
+bool UInterchangeAnimationAPI::GetCustomNodeTransformPayloadKey(const UInterchangeBaseNode* InterchangeBaseNode, FString& AttributeValue)
+{
+	return InterchangeBaseNode->GetStringAttribute(Macro_CustomNodeTransformPayloadKeyKey.Key, AttributeValue);
+}
+
+bool UInterchangeAnimationAPI::SetCustomNodeTransformPayloadKey(UInterchangeBaseNode* InterchangeBaseNode, const FString& AttributeValue)
+{
+	return InterchangeBaseNode->AddStringAttribute(Macro_CustomNodeTransformPayloadKeyKey.Key, AttributeValue);
+}
diff --git a/Engine/Source/Runtime/Experimental/Interchange/Core/Private/Nodes/InterchangeBaseNode.cpp b/Engine/Source/Runtime/Experimental/Interchange/Core/Private/Nodes/InterchangeBaseNode.cpp
--- a/Engine/Source/Runtime/Experimental/Interchange/Core/Private/Nodes/InterchangeBaseNode.cpp
+++ b/Engine/Source/Runtime/Experimental/Interchange/Core/Private/Nodes/InterchangeBaseNode.cpp
@@ -228,6 +228,16 @@ bool UInterchangeBaseNode::GetFloatAttribute(const FName& NodeAttributeKey, floa
 	INTERCHANGE_BASE_NODE_GET_ATTRIBUTE(float);
 }
 
+bool UInterchangeBaseNode::AddDoubleAttribute(const FName& NodeAttributeKey, const double& Value)
+{
+	INTERCHANGE_BASE_NODE_ADD_ATTRIBUTE(double);
+}
+
+bool UInterchangeBaseNode::GetDoubleAttribute(const FName& NodeAttributeKey, double& OutValue) const
+{
+	INTERCHANGE_BASE_NODE_GET_ATTRIBUTE(double);
+}
+
 bool UInterchangeBaseNode::AddStringAttribute(const FName& NodeAttributeKey, const FString& Value)
 {
 	INTERCHANGE_BASE_NODE_ADD_ATTRIBUTE(FString);
diff --git a/Engine/Source/Runtime/Experimental/Interchange/Core/Private/Types/AttributeStorage.cpp b/Engine/Source/Runtime/Experimental/Interchange/Core/Private/Types/AttributeStorage.cpp
--- a/Engine/Source/Runtime/Experimental/Interchange/Core/Private/Types/AttributeStorage.cpp
+++ b/Engine/Source/Runtime/Experimental/Interchange/Core/Private/Types/AttributeStorage.cpp
@@ -85,6 +85,20 @@ namespace UE
 			return AttributeTypeString;
 		}
 
+		EAttributeTypes StringToAttributeType(const FString& AttributeTypeString)
+		{
+			const int32 MaxAttributetypes = static_cast<int32>(EAttributeTypes::Max);
+			for (int32 AttributeTypeEnumIndex = 0; AttributeTypeEnumIndex < MaxAttributetypes; ++AttributeTypeEnumIndex)
+			{
+				const EAttributeTypes Attributetype = static_cast<EAttributeTypes>(AttributeTypeEnumIndex);
+				if (AttributeTypeString.Equals(AttributeTypeToString(Attributetype)))
+				{
+					return Attributetype;
+				}
+			}
+			return EAttributeTypes::None;
+		}
+
 		void LogAttributeStorageErrors(const EAttributeStorageResult Result, const FString OperationName, const FAttributeKey AttributeKey)
 		{
 			//////////////////////////////////////////////////////////////////////////
diff --git a/Engine/Source/Runtime/Experimental/Interchange/Core/Public/Nodes/InterchangeAnimationAPI.h b/Engine/Source/Runtime/Experimental/Interchange/Core/Public/Nodes/InterchangeAnimationAPI.h
new file mode 100644
--- /dev/null
+++ b/Engine/Source/Runtime/Experimental/Interchange/Core/Public/Nodes/InterchangeAnimationAPI.h
@@ -0,0 +1,79 @@
+// Copyright Epic Games, Inc. All Rights Reserved.
+
+#pragma once
+
+#include "Types/AttributeStorage.h"
+#include "Misc/Optional.h"
+#include "Nodes/InterchangeBaseNode.h"
+#include "UObject/Class.h"
+#include "UObject/Object.h"
+#include "UObject/ObjectMacros.h"
+
+#include "InterchangeAnimationAPI.generated.h"
+
+/**
+ * UInterchangeAnimationAPI is used to store and retrieve animation data(i.e. DCC node attributes, pipelines will have access to those attributes)
+ */
+UCLASS(BlueprintType, Experimental)
+class INTERCHANGECORE_API UInterchangeAnimationAPI : public UObject
+{
+	GENERATED_BODY()
+ 
+public:
+
+	/** Get true if the node transform have animation. Return false if the Attribute is not set. */
+	UFUNCTION(BlueprintCallable, Category = "Interchange | Node | Scene | Animation")
+	static bool GetCustomIsNodeTransformAnimated(const UInterchangeBaseNode* InterchangeBaseNode, bool& AttributeValue);
+
+	/** Set true if the node transform has animation. */
+	UFUNCTION(BlueprintCallable, Category = "Interchange | Node | Scene | Animation")
+	static bool SetCustomIsNodeTransformAnimated(UInterchangeBaseNode* InterchangeBaseNode, const bool& AttributeValue);
+
+	/** Get the node transform animation key count. Return false if the Attribute is not set. */
+	UFUNCTION(BlueprintCallable, Category = "Interchange | Node | Scene | Animation")
+	static bool GetCustomNodeTransformAnimationKeyCount(const UInterchangeBaseNode* InterchangeBaseNode, int32& AttributeValue);
+
+	/** Set the node transform animation key count. */
+	UFUNCTION(BlueprintCallable, Category = "Interchange | Node | Scene | Animation")
+	static bool SetCustomNodeTransformAnimationKeyCount(UInterchangeBaseNode* InterchangeBaseNode, const int32& AttributeValue);
+
+	/**
+	 * Get the animation start time in second for this node transform.
+	 * Return false if the Attribute is not set.
+	 */
+	UFUNCTION(BlueprintCallable, Category = "Interchange | Node | Scene | Animation")
+	static bool GetCustomNodeTransformAnimationStartTime(const UInterchangeBaseNode* InterchangeBaseNode, double& AttributeValue);
+
+	/** Set the animation start time in second for this node transform. */
+	UFUNCTION(BlueprintCallable, Category = "Interchange | Node | Scene | Animation")
+	static bool SetCustomNodeTransformAnimationStartTime(UInterchangeBaseNode* InterchangeBaseNode, const double& AttributeValue);
+
+	/**
+	 * Get the animation end time in second for this node transform.
+	 * Return false if the Attribute is not set.
+	 */
+	UFUNCTION(BlueprintCallable, Category = "Interchange | Node | Scene | Animation")
+	static bool GetCustomNodeTransformAnimationEndTime(const UInterchangeBaseNode* InterchangeBaseNode, double& AttributeValue);
+
+	/** Set the animation end time in second for this node transform. */
+	UFUNCTION(BlueprintCallable, Category = "Interchange | Node | Scene | Animation")
+	static bool SetCustomNodeTransformAnimationEndTime(UInterchangeBaseNode* InterchangeBaseNode, const double& AttributeValue);
+
+	/**
+	 * Get the payload key needed to retrieve the animation from the source.
+	 * Return false if the Attribute is not set.
+	 */
+	UFUNCTION(BlueprintCallable, Category = "Interchange | Node | Scene | Animation")
+	static bool GetCustomNodeTransformPayloadKey(const UInterchangeBaseNode* InterchangeBaseNode, FString& AttributeValue);
+
+	/** Set the payload key needed to retrieve the animation from the source. */
+	UFUNCTION(BlueprintCallable, Category = "Interchange | Node | Scene | Animation")
+	static bool SetCustomNodeTransformPayloadKey(UInterchangeBaseNode* InterchangeBaseNode, const FString& AttributeValue);
+	
+private:
+	static const UE::Interchange::FAttributeKey Macro_CustomIsNodeTransformAnimatedKey;
+	static const UE::Interchange::FAttributeKey Macro_CustomNodeTransformAnimationKeyCountKey;
+	static const UE::Interchange::FAttributeKey Macro_CustomNodeTransformAnimationStartTimeKey;
+	static const UE::Interchange::FAttributeKey Macro_CustomNodeTransformAnimationEndTimeKey;
+	static const UE::Interchange::FAttributeKey Macro_CustomNodeTransformPayloadKeyKey;
+};
diff --git a/Engine/Source/Runtime/Experimental/Interchange/Core/Public/Nodes/InterchangeBaseNode.h b/Engine/Source/Runtime/Experimental/Interchange/Core/Public/Nodes/InterchangeBaseNode.h
--- a/Engine/Source/Runtime/Experimental/Interchange/Core/Public/Nodes/InterchangeBaseNode.h
+++ b/Engine/Source/Runtime/Experimental/Interchange/Core/Public/Nodes/InterchangeBaseNode.h
@@ -365,6 +365,18 @@ public:
 	UFUNCTION(BlueprintCallable, Category = "Interchange | Node")
 	bool GetFloatAttribute(const FName& NodeAttributeKey, float& OutValue) const;
 
+	/**
+	 * Add a double attribute to this node. Return false if the attribute do not exist or if we cannot add it
+	 */
+	UFUNCTION(BlueprintCallable, Category = "Interchange | Node")
+	bool AddDoubleAttribute(const FName& NodeAttributeKey, const double& Value);
+	
+	/**
+	 * Get a double attribute from this node. Return false if the attribute do not exist
+	 */
+	UFUNCTION(BlueprintCallable, Category = "Interchange | Node")
+	bool GetDoubleAttribute(const FName& NodeAttributeKey, double& OutValue) const;
+
 	/**
 	 * Add a string attribute to this node. Return false if the attribute do not exist or if we cannot add it
 	 */
diff --git a/Engine/Source/Runtime/Experimental/Interchange/Core/Public/Nodes/InterchangeBaseNodeContainer.h b/Engine/Source/Runtime/Experimental/Interchange/Core/Public/Nodes/InterchangeBaseNodeContainer.h
--- a/Engine/Source/Runtime/Experimental/Interchange/Core/Public/Nodes/InterchangeBaseNodeContainer.h
+++ b/Engine/Source/Runtime/Experimental/Interchange/Core/Public/Nodes/InterchangeBaseNodeContainer.h
@@ -55,6 +55,50 @@ public:
 		}
 	}
 
+	/**
+	 * Recursively traverse the hierarchy starting with the specified node unique ID.
+	 * If the iteration lambda return true the iteration will stop. If the iteration lambda return false the iteration will continue.
+	 */
+	void IterateNodeChildren(const FString& NodeUniqueID, TFunctionRef<void(const UInterchangeBaseNode*)> IterationLambda) const
+	{
+
+		if (const UInterchangeBaseNode* Node = GetNode(NodeUniqueID))
+		{
+			IterationLambda(Node);
+			const TArray<FString> ChildrenIds = GetNodeChildrenUids(NodeUniqueID);
+			for (int32 ChildIndex = 0; ChildIndex < ChildrenIds.Num(); ++ChildIndex)
+			{
+				IterateNodeChildren(ChildrenIds[ChildIndex], IterationLambda);
+			}
+		}
+	}
+
+	/**
+	 * Recursively traverse the hierarchy starting with the specified node unique ID.
+	 * If the iteration lambda return true the iteration will stop. If the iteration lambda return false the iteration will continue.
+	 * 
+	 * @return - Return true if the iteration was break, false otherwise
+	 */
+	bool BreakableIterateNodeChildren(const FString& NodeUniqueID, TFunctionRef<bool(const UInterchangeBaseNode*)> IterationLambda) const
+	{
+		if (const UInterchangeBaseNode* Node = GetNode(NodeUniqueID))
+		{
+			if (IterationLambda(Node))
+			{
+				return true;
+			}
+			const TArray<FString> ChildrenIds = GetNodeChildrenUids(NodeUniqueID);
+			for (int32 ChildIndex = 0; ChildIndex < ChildrenIds.Num(); ++ChildIndex)
+			{
+				if (BreakableIterateNodeChildren(ChildrenIds[ChildIndex], IterationLambda))
+				{
+					return true;
+				}
+			}
+		}
+		return false;
+	}
+
 	/** Unordered iteration of the all nodes, but I can be stop early by returning true */
 	void BreakableIterateNodes(TFunctionRef<bool(const FString&, UInterchangeBaseNode*)> IterationLambda) const;
 
diff --git a/Engine/Source/Runtime/Experimental/Interchange/Core/Public/Types/AttributeStorage.h b/Engine/Source/Runtime/Experimental/Interchange/Core/Public/Types/AttributeStorage.h
--- a/Engine/Source/Runtime/Experimental/Interchange/Core/Public/Types/AttributeStorage.h
+++ b/Engine/Source/Runtime/Experimental/Interchange/Core/Public/Types/AttributeStorage.h
@@ -230,7 +230,10 @@ namespace UE
 			BoxSphereBounds3f	= 54,
 			BoxSphereBounds3d	= 55,
 			Sphere3f			= 56,
-			Sphere3d			= 57
+			Sphere3d			= 57,
+
+			//Max should always be updated if we add a new supported type
+			Max                 = 58
 		};
 
 		/**
@@ -238,6 +241,12 @@ namespace UE
 		 */
 		INTERCHANGECORE_API FString AttributeTypeToString(EAttributeTypes AttributeType);
 
+		/**
+		 * Return the AttributeType for the specified FString, return EAttributeTypes::None if the string do not match any
+		 * supported attribute type.
+		 */
+		INTERCHANGECORE_API EAttributeTypes StringToAttributeType(const FString& AttributeTypeString);
+
 		/**
 		 * Stub for attribute type traits.
 		 *
