Fixing up short type names in property meta data

diff --git a/Engine/Plugins/Animation/LiveLink/Source/LiveLinkComponents/Public/LiveLinkControllerBase.h b/Engine/Plugins/Animation/LiveLink/Source/LiveLinkComponents/Public/LiveLinkControllerBase.h
--- a/Engine/Plugins/Animation/LiveLink/Source/LiveLinkComponents/Public/LiveLinkControllerBase.h
+++ b/Engine/Plugins/Animation/LiveLink/Source/LiveLinkComponents/Public/LiveLinkControllerBase.h
@@ -97,7 +97,7 @@ protected:
 	/**
 	 * A component reference (customized) that allows the user to specify a component that this controller should control
 	 */
-	UPROPERTY(EditInstanceOnly, Category = "ComponentToControl", meta = (UseComponentPicker, AllowedClasses = "ActorComponent", DisallowedClasses = "LiveLinkComponentController"))
+	UPROPERTY(EditInstanceOnly, Category = "ComponentToControl", meta = (UseComponentPicker, AllowedClasses = "/Script/Engine.ActorComponent", DisallowedClasses = "/Script/LiveLinkComponents.LiveLinkComponentController"))
 	FComponentReference ComponentPicker;
 
 protected:
diff --git a/Engine/Plugins/Experimental/Interchange/Runtime/Source/Pipelines/Public/InterchangeGenericMaterialPipeline.h b/Engine/Plugins/Experimental/Interchange/Runtime/Source/Pipelines/Public/InterchangeGenericMaterialPipeline.h
--- a/Engine/Plugins/Experimental/Interchange/Runtime/Source/Pipelines/Public/InterchangeGenericMaterialPipeline.h
+++ b/Engine/Plugins/Experimental/Interchange/Runtime/Source/Pipelines/Public/InterchangeGenericMaterialPipeline.h
@@ -39,7 +39,7 @@ public:
 	EInterchangeMaterialImportOption MaterialImport = EInterchangeMaterialImportOption::ImportAsMaterials;
 
 	/** Optional material used as the parent when importing materials as instances. If no parent material is specified, one will be automatically selected during the import process. */
-	UPROPERTY(EditAnywhere, BlueprintReadWrite, Category = "Materials", Meta= (EditCondition="MaterialImport==EInterchangeMaterialImportOption::ImportAsMaterialInstances", AllowedClasses="MaterialInterface"))
+	UPROPERTY(EditAnywhere, BlueprintReadWrite, Category = "Materials", Meta= (EditCondition="MaterialImport==EInterchangeMaterialImportOption::ImportAsMaterialInstances", AllowedClasses="/Script/Engine.MaterialInterface"))
 	FSoftObjectPath ParentMaterial;
 
 	virtual void AdjustSettingsForReimportType(EInterchangeReimportType ImportType, TObjectPtr<UObject> ReimportAsset) override;
diff --git a/Engine/Plugins/Runtime/GameplayAbilities/Source/GameplayAbilities/Public/GameplayCueNotifyTypes.h b/Engine/Plugins/Runtime/GameplayAbilities/Source/GameplayAbilities/Public/GameplayCueNotifyTypes.h
--- a/Engine/Plugins/Runtime/GameplayAbilities/Source/GameplayAbilities/Public/GameplayCueNotifyTypes.h
+++ b/Engine/Plugins/Runtime/GameplayAbilities/Source/GameplayAbilities/Public/GameplayCueNotifyTypes.h
@@ -592,7 +592,7 @@ public:
 	FGameplayCueNotify_PlacementInfo PlacementInfoOverride;
 
 	// Camera lens effect to play.
-	UPROPERTY(EditDefaultsOnly, BlueprintReadOnly, Category = GameplayCueNotify, Meta = (MustImplement = "CameraLensEffectInterface"))
+	UPROPERTY(EditDefaultsOnly, BlueprintReadOnly, Category = GameplayCueNotify, Meta = (MustImplement = "/Script/Engine.CameraLensEffectInterface"))
 	TSubclassOf<AActor> CameraLensEffect;
 
 	// If enabled, use the spawn condition override and not the default one.
diff --git a/Engine/Plugins/Runtime/StateTree/Source/StateTreeEditorModule/Public/StateTreeEditorData.h b/Engine/Plugins/Runtime/StateTree/Source/StateTreeEditorModule/Public/StateTreeEditorData.h
--- a/Engine/Plugins/Runtime/StateTree/Source/StateTreeEditorModule/Public/StateTreeEditorData.h
+++ b/Engine/Plugins/Runtime/StateTree/Source/StateTreeEditorModule/Public/StateTreeEditorData.h
@@ -130,7 +130,7 @@ public:
 	UPROPERTY(EditDefaultsOnly, Category = Common)
 	FStateTreeStateParameters RootParameters;
 
-	UPROPERTY(EditDefaultsOnly, Category = "Evaluators", meta = (BaseStruct = "StateTreeEvaluatorBase", BaseClass = "StateTreeEvaluatorBlueprintBase"))
+	UPROPERTY(EditDefaultsOnly, Category = "Evaluators", meta = (BaseStruct = "/Script/StateTreeModule.StateTreeEvaluatorBase", BaseClass = "/Script/StateTreeModule.StateTreeEvaluatorBlueprintBase"))
 	TArray<FStateTreeEditorNode> Evaluators;
 
 	UPROPERTY(meta = (ExcludeFromHash))
diff --git a/Engine/Source/Editor/UnrealEd/Classes/Editor/TemplateMapInfo.h b/Engine/Source/Editor/UnrealEd/Classes/Editor/TemplateMapInfo.h
--- a/Engine/Source/Editor/UnrealEd/Classes/Editor/TemplateMapInfo.h
+++ b/Engine/Source/Editor/UnrealEd/Classes/Editor/TemplateMapInfo.h
@@ -19,11 +19,11 @@ struct FTemplateMapInfo
 	TSoftObjectPtr<class UTexture2D> ThumbnailTexture;
 
 	/** The Texture associated with this map template */
-	UPROPERTY(config, meta = (AllowedClasses = "Texture2D"))
+	UPROPERTY(config, meta = (AllowedClasses = "/Script/Engine.Texture2D"))
 	FSoftObjectPath Thumbnail;
 
 	/** The object path to the template map */
-	UPROPERTY(config, meta = (AllowedClasses = "World"))
+	UPROPERTY(config, meta = (AllowedClasses = "/Script/Engine.World"))
 	FSoftObjectPath Map;
 
 	/** Optional display name override for this map template  */
diff --git a/Engine/Source/Runtime/CoreUObject/Public/UObject/NoExportTypes.h b/Engine/Source/Runtime/CoreUObject/Public/UObject/NoExportTypes.h
--- a/Engine/Source/Runtime/CoreUObject/Public/UObject/NoExportTypes.h
+++ b/Engine/Source/Runtime/CoreUObject/Public/UObject/NoExportTypes.h
@@ -471,7 +471,7 @@ struct FGuid
  * A point or direction FVector in 3d space.
  * @note The full C++ class is located here: Engine\Source\Runtime\Core\Public\Math\Vector.h
  */
-USTRUCT(immutable, noexport, BlueprintType, IsAlwaysAccessible, HasDefaults, HasNoOpConstructor, IsCoreType, meta = (HasNativeBreak = "Engine.KismetMathLibrary.BreakVector3f"))
+USTRUCT(immutable, noexport, BlueprintType, IsAlwaysAccessible, HasDefaults, HasNoOpConstructor, IsCoreType, meta = (HasNativeBreak = "/Script/Engine.KismetMathLibrary.BreakVector3f"))
 struct FVector3f
 {
 	UPROPERTY(EditAnywhere, Category = Vector, SaveGame)
diff --git a/Engine/Source/Runtime/Engine/Classes/Engine/Engine.h b/Engine/Source/Runtime/Engine/Classes/Engine/Engine.h
--- a/Engine/Source/Runtime/Engine/Classes/Engine/Engine.h
+++ b/Engine/Source/Runtime/Engine/Classes/Engine/Engine.h
@@ -1249,7 +1249,7 @@ public:
 	TObjectPtr<class UPhysicalMaterial> DefaultDestructiblePhysMaterial;
 
 	/** Path of the PhysicalMaterial to use if none is defined for a particular object. */
-	UPROPERTY(globalconfig, EditAnywhere, Category = DefaultMaterials, meta = (AllowedClasses = "PhysicalMaterial", DisplayName = "Destructible Physics Material"))
+	UPROPERTY(globalconfig, EditAnywhere, Category = DefaultMaterials, meta = (AllowedClasses = "/Script/PhysicsCore.PhysicalMaterial", DisplayName = "Destructible Physics Material"))
 	FSoftObjectPath DefaultDestructiblePhysMaterialName;
 
 	/** Deprecated rules for redirecting renamed objects, replaced by the CoreRedirects system*/
diff --git a/Engine/Source/Runtime/Engine/Classes/Engine/HitResult.h b/Engine/Source/Runtime/Engine/Classes/Engine/HitResult.h
--- a/Engine/Source/Runtime/Engine/Classes/Engine/HitResult.h
+++ b/Engine/Source/Runtime/Engine/Classes/Engine/HitResult.h
@@ -15,7 +15,7 @@ class UPrimitiveComponent;
 /**
  * Structure containing information about one hit of a trace, such as point of impact and surface normal at that point.
  */
-USTRUCT(BlueprintType, meta = (HasNativeBreak = "Engine.GameplayStatics.BreakHitResult", HasNativeMake = "Engine.GameplayStatics.MakeHitResult"))
+USTRUCT(BlueprintType, meta = (HasNativeBreak = "/Script/Engine.GameplayStatics.BreakHitResult", HasNativeMake = "/Script/Engine.GameplayStatics.MakeHitResult"))
 struct ENGINE_API FHitResult
 {
 	GENERATED_BODY()
diff --git a/Engine/Source/Runtime/EngineSettings/Classes/GameMapsSettings.h b/Engine/Source/Runtime/EngineSettings/Classes/GameMapsSettings.h
--- a/Engine/Source/Runtime/EngineSettings/Classes/GameMapsSettings.h
+++ b/Engine/Source/Runtime/EngineSettings/Classes/GameMapsSettings.h
@@ -74,11 +74,11 @@ struct FTemplateMapInfoOverride
 	GENERATED_USTRUCT_BODY()
 
 	/** The thumbnail to display in the new level dialog */
-	UPROPERTY(config, EditAnywhere, Category = DefaultMaps, meta = (AllowedClasses = "Texture2D"))
+	UPROPERTY(config, EditAnywhere, Category = DefaultMaps, meta = (AllowedClasses = "/Script/Engine.Texture2D"))
 	FSoftObjectPath Thumbnail;
 
 	/** The path to the template map */
-	UPROPERTY(config, EditAnywhere, Category = DefaultMaps, meta = (AllowedClasses = "World"))
+	UPROPERTY(config, EditAnywhere, Category = DefaultMaps, meta = (AllowedClasses = "/Script/Engine.World"))
 	FSoftObjectPath Map;
 
 	/** The display name of the map template in the level dialog. If this is empty the thumbnail name will be used */
