Interchange: Fix crash when not importing animation
#rnx

diff --git a/Engine/Plugins/Experimental/Interchange/Runtime/Source/Pipelines/Private/InterchangeGenericAnimationPipeline.cpp b/Engine/Plugins/Experimental/Interchange/Runtime/Source/Pipelines/Private/InterchangeGenericAnimationPipeline.cpp
--- a/Engine/Plugins/Experimental/Interchange/Runtime/Source/Pipelines/Private/InterchangeGenericAnimationPipeline.cpp
+++ b/Engine/Plugins/Experimental/Interchange/Runtime/Source/Pipelines/Private/InterchangeGenericAnimationPipeline.cpp
@@ -70,26 +70,27 @@ void UInterchangeGenericAnimationPipeline::AdjustSettingsForReimportType(EInterc
 
 void UInterchangeGenericAnimationPipeline::ExecutePreImportPipeline(UInterchangeBaseNodeContainer* InBaseNodeContainer, const TArray<UInterchangeSourceData*>& InSourceDatas)
 {
-	if (!bImportAnimations)
+	if (!InBaseNodeContainer)
 	{
-		//Nothing to import
+		UE_LOG(LogInterchangePipeline, Warning, TEXT("UInterchangeGenericAnimationPipeline: Cannot execute pre-import pipeline because InBaseNodeContrainer is null"));
 		return;
 	}
 
-	check(!CommonSkeletalMeshesAndAnimationsProperties.IsNull());
-	if (!InBaseNodeContainer)
+	BaseNodeContainer = InBaseNodeContainer;
+
+	if (!bImportAnimations)
 	{
-		UE_LOG(LogInterchangePipeline, Warning, TEXT("UInterchangeGenericAnimationPipeline: Cannot execute pre-import pipeline because InBaseNodeContrainer is null"));
+		//Nothing to import
 		return;
 	}
 
+	check(!CommonSkeletalMeshesAndAnimationsProperties.IsNull());
+
 	if (CommonSkeletalMeshesAndAnimationsProperties->bImportOnlyAnimations && CommonSkeletalMeshesAndAnimationsProperties->Skeleton.IsNull())
 	{
 		UE_LOG(LogInterchangePipeline, Warning, TEXT("UInterchangeGenericAnimationPipeline: Cannot execute pre-import pipeline because we cannot import animation only but not specify any valid skeleton"));
 		return;
 	}
-
-	BaseNodeContainer = InBaseNodeContainer;
 	SourceDatas.Empty(InSourceDatas.Num());
 	for (const UInterchangeSourceData* SourceData : InSourceDatas)
 	{
