Interchange runtime support
- Forced material instance creation at runtime.
- Enabled texture import only for 2D textures.
- Enabled static mesh import through the BuildFromMeshDescriptions API.

diff --git a/Engine/Plugins/Experimental/Interchange/Runtime/Source/Import/InterchangeImport.Build.cs b/Engine/Plugins/Experimental/Interchange/Runtime/Source/Import/InterchangeImport.Build.cs
--- a/Engine/Plugins/Experimental/Interchange/Runtime/Source/Import/InterchangeImport.Build.cs
+++ b/Engine/Plugins/Experimental/Interchange/Runtime/Source/Import/InterchangeImport.Build.cs
@@ -35,6 +35,7 @@ namespace UnrealBuildTool.Rules
 					"InterchangeDispatcher",
 					"InterchangeMessages",
 					"Json",
+					"RenderCore",
 					"RHI",
 					"TextureUtilitiesCommon",
 				}
diff --git a/Engine/Plugins/Experimental/Interchange/Runtime/Source/Import/Private/Material/InterchangeMaterialFactory.cpp b/Engine/Plugins/Experimental/Interchange/Runtime/Source/Import/Private/Material/InterchangeMaterialFactory.cpp
--- a/Engine/Plugins/Experimental/Interchange/Runtime/Source/Import/Private/Material/InterchangeMaterialFactory.cpp
+++ b/Engine/Plugins/Experimental/Interchange/Runtime/Source/Import/Private/Material/InterchangeMaterialFactory.cpp
@@ -342,7 +342,7 @@ UObject* UInterchangeMaterialFactory::CreateEmptyAsset(const FCreateAssetParams&
 				if (MaterialInstanceFactoryNode->GetCustomParent(ParentPath))
 				{
 					FSoftObjectPath ParentMaterial(ParentPath);
-					Material = UMaterialInstanceDynamic::Create(Cast<UMaterialInterface>(ParentMaterial.TryLoad()), Arguments.Parent);
+					Material = UMaterialInstanceDynamic::Create(Cast<UMaterialInterface>(ParentMaterial.TryLoad()), Arguments.Parent, *Arguments.AssetName);
 				}
 			}
 		}
@@ -457,13 +457,12 @@ void UInterchangeMaterialFactory::PreImportPreCompletedCallback(const FImportPre
 	check(IsInGameThread());
 	Super::PreImportPreCompletedCallback(Arguments);
 
-	//TODO make sure this work at runtime
-#if WITH_EDITORONLY_DATA
 	if (ensure(Arguments.ImportedObject && Arguments.SourceData))
 	{
 		//We must call the Update of the asset source file in the main thread because UAssetImportData::Update execute some delegate we do not control
 		UMaterialInterface* ImportedMaterialInterface = CastChecked<UMaterialInterface>(Arguments.ImportedObject);
 
+#if WITH_EDITOR
 		//Update the samplers type in case the textures were changed during their PreImportPreCompletedCallback
 		if (UMaterial* ImportedMaterial = Cast<UMaterial>(ImportedMaterialInterface))
 		{
@@ -475,12 +474,15 @@ void UInterchangeMaterialFactory::PreImportPreCompletedCallback(const FImportPre
 				}
 			}
 		}
-		else if (UMaterialInstance* MaterialInstance = Cast<UMaterialInstance>(ImportedMaterialInterface))
+		else
+#endif // WITH_EDITOR
+		if (UMaterialInstance* MaterialInstance = Cast<UMaterialInstance>(ImportedMaterialInterface))
 		{
 			// Material instances expect their parameters to only be updated from the game thread
 			SetupMaterialInstance(MaterialInstance, Arguments.NodeContainer, Cast<UInterchangeBaseMaterialFactoryNode>(Arguments.FactoryNode));
 		}
 
+#if WITH_EDITORONLY_DATA
 		UE::Interchange::FFactoryCommon::FUpdateImportAssetDataParameters UpdateImportAssetDataParameters(ImportedMaterialInterface
 																										  , ImportedMaterialInterface->AssetImportData
 																										  , Arguments.SourceData
@@ -489,8 +491,8 @@ void UInterchangeMaterialFactory::PreImportPreCompletedCallback(const FImportPre
 																										  , Arguments.Pipelines);
 
 		ImportedMaterialInterface->AssetImportData = UE::Interchange::FFactoryCommon::UpdateImportAssetData(UpdateImportAssetDataParameters);
+#endif // WITH_EDITORONLY_DATA
 	}
-#endif
 }
 
 bool UInterchangeMaterialFactory::GetSourceFilenames(const UObject* Object, TArray<FString>& OutSourceFilenames) const
diff --git a/Engine/Plugins/Experimental/Interchange/Runtime/Source/Import/Private/Mesh/InterchangeStaticMeshFactory.cpp b/Engine/Plugins/Experimental/Interchange/Runtime/Source/Import/Private/Mesh/InterchangeStaticMeshFactory.cpp
--- a/Engine/Plugins/Experimental/Interchange/Runtime/Source/Import/Private/Mesh/InterchangeStaticMeshFactory.cpp
+++ b/Engine/Plugins/Experimental/Interchange/Runtime/Source/Import/Private/Mesh/InterchangeStaticMeshFactory.cpp
@@ -49,12 +49,6 @@ UClass* UInterchangeStaticMeshFactory::GetFactoryClass() const
 
 UObject* UInterchangeStaticMeshFactory::CreateEmptyAsset(const FCreateAssetParams& Arguments)
 {
-#if !WITH_EDITOR || !WITH_EDITORONLY_DATA
-
-	UE_LOG(LogInterchangeImport, Error, TEXT("Cannot import StaticMesh asset in runtime, this is an editor only feature."));
-	return nullptr;
-
-#else
 	UStaticMesh* StaticMesh = nullptr;
 	if (!Arguments.AssetNode || !Arguments.AssetNode->GetObjectClass()->IsChildOf(GetFactoryClass()))
 	{
@@ -79,6 +73,16 @@ UObject* UInterchangeStaticMeshFactory::CreateEmptyAsset(const FCreateAssetParam
 	{
 		//This is a reimport, we are just re-updating the source data
 		StaticMesh = Cast<UStaticMesh>(ExistingAsset);
+
+		// Clear the render resources on the static mesh from the game thread so that we're ready to update it
+		if (StaticMesh && StaticMesh->AreRenderingResourcesInitialized())
+		{
+			const bool bInvalidateLighting = true;
+			const bool bRefreshBounds = true;
+			FStaticMeshComponentRecreateRenderStateContext RecreateRenderStateContext(StaticMesh, bInvalidateLighting, bRefreshBounds);
+			StaticMesh->ReleaseResources();
+			StaticMesh->ReleaseResourcesFence.Wait();
+		}
 	}
 	
 	if (!StaticMesh)
@@ -90,20 +94,15 @@ UObject* UInterchangeStaticMeshFactory::CreateEmptyAsset(const FCreateAssetParam
 	// create the BodySetup on the game thread
 	StaticMesh->CreateBodySetup();
 	
+#if WITH_EDITOR
 	StaticMesh->PreEditChange(nullptr);	
+#endif // WITH_EDITOR
+
 	return StaticMesh;
-#endif //else !WITH_EDITOR || !WITH_EDITORONLY_DATA
 }
 
 UObject* UInterchangeStaticMeshFactory::CreateAsset(const FCreateAssetParams& Arguments)
 {
-#if !WITH_EDITOR || !WITH_EDITORONLY_DATA
-
-	UE_LOG(LogInterchangeImport, Error, TEXT("Cannot import static mesh asset in runtime, this is an editor only feature."));
-	return nullptr;
-
-#else
-
 	if (!Arguments.AssetNode || !Arguments.AssetNode->GetObjectClass()->IsChildOf(GetFactoryClass()))
 	{
 		return nullptr;
@@ -149,22 +148,16 @@ UObject* UInterchangeStaticMeshFactory::CreateAsset(const FCreateAssetParams& Ar
 		UE_LOG(LogInterchangeImport, Error, TEXT("Could not create StaticMesh asset %s"), *Arguments.AssetName);
 		return nullptr;
 	}
-			
+
+	ensure(!StaticMesh->AreRenderingResourcesInitialized());
+
 	const int32 LodCount = StaticMeshFactoryNode->GetLodDataCount();
 	const int32 PrevLodCount = StaticMesh->GetNumLODs();
 	const int32 FinalLodCount = FMath::Max(PrevLodCount, LodCount);
 
+#if WITH_EDITOR
 	StaticMesh->SetNumSourceModels(FinalLodCount);
-
-	// Make sure that mesh descriptions for added LODs are kept as is when the mesh is built
-	for (int32 LodIndex = PrevLodCount; LodIndex < FinalLodCount; ++LodIndex)
-	{
-		FStaticMeshSourceModel& SrcModel = StaticMesh->GetSourceModel(LodIndex);
-
-		SrcModel.ReductionSettings.MaxDeviation = 0.0f;
-		SrcModel.ReductionSettings.PercentTriangles = 1.0f;
-		SrcModel.ReductionSettings.PercentVertices = 1.0f;
-	}
+#endif // WITH_EDITOR
 
 	// If we are reimporting, cache the existing vertex colors so they can be optionally reapplied after reimport
 	TMap<FVector3f, FColor> ExisitingVertexColorData;
@@ -180,7 +173,11 @@ UObject* UInterchangeStaticMeshFactory::CreateAsset(const FCreateAssetParams& Ar
 
 		if (MaterialSlotIndex == INDEX_NONE)
 		{
-			StaticMesh->GetStaticMaterials().Emplace(MaterialInterface, MaterialSlotName);
+			MaterialSlotIndex = StaticMesh->GetStaticMaterials().Emplace(MaterialInterface, MaterialSlotName);
+#if !WITH_EDITOR
+			// UV density is not supported to be generated at runtime for now. We fake that it has been initialized so that we don't trigger ensures.
+			StaticMesh->GetStaticMaterials()[MaterialSlotIndex].UVChannelData = FMeshUVChannelInfo(1.f);
+#endif
 		}
 		else
 		{
@@ -210,6 +207,9 @@ UObject* UInterchangeStaticMeshFactory::CreateAsset(const FCreateAssetParams& Ar
 	StaticMeshFactoryNode->GetLodDataUniqueIds(LodDataUniqueIds);
 	ensure(LodDataUniqueIds.Num() == LodCount);
 
+	TArray<FMeshDescription> LodMeshDescriptions;
+	LodMeshDescriptions.SetNum(LodCount);
+
 	bool bImportedCustomCollision = false;
 	int32 CurrentLodIndex = 0;
 	for (int32 LodIndex = 0; LodIndex < LodCount; ++LodIndex)
@@ -223,7 +223,7 @@ UObject* UInterchangeStaticMeshFactory::CreateAsset(const FCreateAssetParams& Ar
 		}
 
 		// Add the lod mesh data to the static mesh
-		FMeshDescription* LodMeshDescription = StaticMesh->CreateMeshDescription(CurrentLodIndex);
+		FMeshDescription& LodMeshDescription = LodMeshDescriptions[CurrentLodIndex];
 
 		FStaticMeshOperations::FAppendSettings AppendSettings;
 		for (int32 ChannelIdx = 0; ChannelIdx < FStaticMeshOperations::FAppendSettings::MAX_NUM_UV_CHANNELS; ++ChannelIdx)
@@ -252,20 +252,20 @@ UObject* UInterchangeStaticMeshFactory::CreateAsset(const FCreateAssetParams& Ar
 			if (!bFirstValidMoved)
 			{
 				FMeshDescription& MeshDescription = const_cast<FMeshDescription&>(LodMeshPayload->MeshDescription);
-				*LodMeshDescription = MoveTemp(MeshDescription);
+				LodMeshDescription = MoveTemp(MeshDescription);
 				bFirstValidMoved = true;
 
 				// Bake the payload mesh, with the provided transform
 				if (!MeshPayload.Transform.Equals(FTransform::Identity))
 				{
-					FStaticMeshOperations::ApplyTransform(*LodMeshDescription, MeshPayload.Transform);
+					FStaticMeshOperations::ApplyTransform(LodMeshDescription, MeshPayload.Transform);
 				}
 			}
 			else
 			{
 				// Bake the payload mesh, with the provided transform
 				AppendSettings.MeshTransform = MeshPayload.Transform;
-				FStaticMeshOperations::AppendMeshDescription(LodMeshPayload->MeshDescription, *LodMeshDescription, AppendSettings);
+				FStaticMeshOperations::AppendMeshDescription(LodMeshPayload->MeshDescription, LodMeshDescription, AppendSettings);
 			}
 		}
 
@@ -275,7 +275,7 @@ UObject* UInterchangeStaticMeshFactory::CreateAsset(const FCreateAssetParams& Ar
 		// Override -> replace the vertex color by the override color
 		// @todo: new mesh description attribute for painted vertex colors?
 		{
-			FStaticMeshAttributes Attributes(*LodMeshDescription);
+			FStaticMeshAttributes Attributes(LodMeshDescription);
 			TVertexInstanceAttributesRef<FVector4f> VertexInstanceColors = Attributes.GetVertexInstanceColors();
 			bool bReplaceVertexColor = false;
 			StaticMeshFactoryNode->GetCustomVertexColorReplace(bReplaceVertexColor);
@@ -285,12 +285,12 @@ UObject* UInterchangeStaticMeshFactory::CreateAsset(const FCreateAssetParams& Ar
 				StaticMeshFactoryNode->GetCustomVertexColorIgnore(bIgnoreVertexColor);
 				if (bIgnoreVertexColor)
 				{
-					for (const FVertexInstanceID& VertexInstanceID : LodMeshDescription->VertexInstances().GetElementIDs())
+					for (const FVertexInstanceID& VertexInstanceID : LodMeshDescription.VertexInstances().GetElementIDs())
 					{
 						//If we have old vertex color (reimport), we want to keep it if the option is ignore
 						if (ExisitingVertexColorData.Num() > 0)
 						{
-							const FVector3f& VertexPosition = LodMeshDescription->GetVertexPosition(LodMeshDescription->GetVertexInstanceVertex(VertexInstanceID));
+							const FVector3f& VertexPosition = LodMeshDescription.GetVertexPosition(LodMeshDescription.GetVertexInstanceVertex(VertexInstanceID));
 							const FColor* PaintedColor = ExisitingVertexColorData.Find(VertexPosition);
 							if (PaintedColor)
 							{
@@ -315,7 +315,7 @@ UObject* UInterchangeStaticMeshFactory::CreateAsset(const FCreateAssetParams& Ar
 					FColor OverrideVertexColor;
 					if (StaticMeshFactoryNode->GetCustomVertexColorOverride(OverrideVertexColor))
 					{
-						for (const FVertexInstanceID& VertexInstanceID : LodMeshDescription->VertexInstances().GetElementIDs())
+						for (const FVertexInstanceID& VertexInstanceID : LodMeshDescription.VertexInstances().GetElementIDs())
 						{
 							VertexInstanceColors[VertexInstanceID] = FVector4f(FLinearColor(OverrideVertexColor));
 						}
@@ -324,15 +324,11 @@ UObject* UInterchangeStaticMeshFactory::CreateAsset(const FCreateAssetParams& Ar
 			}
 		}
 
-		UStaticMesh::FCommitMeshDescriptionParams CommitMeshDescriptionParams;
-		CommitMeshDescriptionParams.bMarkPackageDirty = false; // Marking packages dirty isn't threadsafe
-		StaticMesh->CommitMeshDescription(CurrentLodIndex, CommitMeshDescriptionParams);
-
 		// Build section info map from materials
-		FStaticMeshConstAttributes StaticMeshAttributes(*LodMeshDescription);
+		FStaticMeshConstAttributes StaticMeshAttributes(LodMeshDescription);
 		TPolygonGroupAttributesRef<const FName> SlotNames = StaticMeshAttributes.GetPolygonGroupMaterialSlotNames();
 		int32 SectionIndex = 0;
-		for (FPolygonGroupID PolygonGroupID : LodMeshDescription->PolygonGroups().GetElementIDs())
+		for (FPolygonGroupID PolygonGroupID : LodMeshDescription.PolygonGroups().GetElementIDs())
 		{
 			int32 MaterialSlotIndex = StaticMesh->GetMaterialIndexFromImportedMaterialSlotName(SlotNames[PolygonGroupID]);
 					
@@ -340,12 +336,17 @@ UObject* UInterchangeStaticMeshFactory::CreateAsset(const FCreateAssetParams& Ar
 			if (MaterialSlotIndex == INDEX_NONE)
 			{
 				MaterialSlotIndex = StaticMesh->GetStaticMaterials().Emplace(UMaterial::GetDefaultMaterial(MD_Surface), SlotNames[PolygonGroupID]);
+#if !WITH_EDITOR
+				StaticMesh->GetStaticMaterials()[MaterialSlotIndex].UVChannelData = FMeshUVChannelInfo(1.f);
+#endif
 			}
 
+#if WITH_EDITOR
 			FMeshSectionInfo Info = StaticMesh->GetSectionInfoMap().Get(CurrentLodIndex, SectionIndex);
 			Info.MaterialIndex = MaterialSlotIndex;
 			StaticMesh->GetSectionInfoMap().Remove(CurrentLodIndex, SectionIndex);
 			StaticMesh->GetSectionInfoMap().Set(CurrentLodIndex, SectionIndex, Info);
+#endif
 
 			SectionIndex++;
 		}
@@ -359,23 +360,18 @@ UObject* UInterchangeStaticMeshFactory::CreateAsset(const FCreateAssetParams& Ar
 			bImportedCustomCollision |= ImportConvexCollision(Arguments, StaticMesh, LodDataNode);
 		}
 
-		if (!Arguments.ReimportObject)
-		{
-			FStaticMeshSourceModel& SrcModel = StaticMesh->GetSourceModel(CurrentLodIndex);
-
-			const int32 NumUVChannels = StaticMeshAttributes.GetVertexInstanceUVs().GetNumChannels();
-			const int32 FirstOpenUVChannel = NumUVChannels >= MAX_MESH_TEXTURE_COORDS_MD ? 1 : NumUVChannels;
-			SrcModel.BuildSettings.DstLightmapIndex = FirstOpenUVChannel;
-
-			if (CurrentLodIndex == 0)
-			{
-				StaticMesh->SetLightMapCoordinateIndex(FirstOpenUVChannel);
-			}
-		}
-
 		CurrentLodIndex++;
 	}
 
+#if WITH_EDITOR
+	{
+		const bool bIsAReimport = Arguments.ReimportObject != nullptr;
+		SetupSourceModelsSettings(*StaticMesh, LodMeshDescriptions, PrevLodCount, FinalLodCount, bIsAReimport);
+	}
+#endif // WITH_EDITOR
+
+	CommitMeshDescriptions(*StaticMesh, MoveTemp(LodMeshDescriptions));
+
 	ImportSockets(Arguments, StaticMesh, StaticMeshFactoryNode);
 
 	if (!Arguments.ReimportObject)
@@ -383,6 +379,7 @@ UObject* UInterchangeStaticMeshFactory::CreateAsset(const FCreateAssetParams& Ar
 		// Apply all StaticMeshFactoryNode custom attributes to the static mesh asset
 		StaticMeshFactoryNode->ApplyAllCustomAttributeToObject(StaticMesh);
 	}
+#if WITH_EDITOR
 	else
 	{
 		//Apply the re import strategy 
@@ -397,24 +394,91 @@ UObject* UInterchangeStaticMeshFactory::CreateAsset(const FCreateAssetParams& Ar
 		CurrentNode->FillAllCustomAttributeFromObject(StaticMesh);
 		UE::Interchange::FFactoryCommon::ApplyReimportStrategyToAsset(StaticMesh, PreviousNode, CurrentNode, StaticMeshFactoryNode);
 	}
-	
+#endif // WITH_EDITOR
+
 	if (!bImportedCustomCollision)
 	{
 		GenerateKDopCollision(Arguments, StaticMesh);
 	}
+#if WITH_EDITORONLY_DATA
 	else
 	{
 		StaticMesh->bCustomizedCollision = true;
 	}
+#endif // WITH_EDITORONLY_DATA
 
 	// Getting the file Hash will cache it into the source data
 	Arguments.SourceData->GetFileContentHash();
 
 	return StaticMeshObject;
+}
+
+void UInterchangeStaticMeshFactory::CommitMeshDescriptions(UStaticMesh& StaticMesh, TArray<FMeshDescription>&& LodMeshDescriptions)
+{
+#if WITH_EDITOR
+	for (int32 LodIndex = 0; LodIndex < LodMeshDescriptions.Num(); ++LodIndex)
+	{
+		FMeshDescription* StaticMeshDescription = StaticMesh.CreateMeshDescription(LodIndex);
+		check(StaticMeshDescription);
+		*StaticMeshDescription = MoveTemp(LodMeshDescriptions[LodIndex]);
+
+		UStaticMesh::FCommitMeshDescriptionParams CommitMeshDescriptionParams;
+		CommitMeshDescriptionParams.bMarkPackageDirty = false; // Marking packages dirty isn't threadsafe
+		StaticMesh.CommitMeshDescription(LodIndex, CommitMeshDescriptionParams);
+	}
+#else // WITH_EDITOR
+	TArray<const FMeshDescription*> MeshDescriptionPointers;
+	MeshDescriptionPointers.Reserve(LodMeshDescriptions.Num());
+
+	for (const FMeshDescription& MeshDescription : LodMeshDescriptions)
+	{
+		MeshDescriptionPointers.Add(&MeshDescription);
+	}
 
-#endif //else !WITH_EDITOR || !WITH_EDITORONLY_DATA
+	UStaticMesh::FBuildMeshDescriptionsParams BuildMeshDescriptionsParams;
+	BuildMeshDescriptionsParams.bUseHashAsGuid = true;
+	// Do not mark the package dirty since MarkPackageDirty is not thread safe
+	BuildMeshDescriptionsParams.bMarkPackageDirty = false;
+	BuildMeshDescriptionsParams.bBuildSimpleCollision = false;
+	// Do not commit since we only need the render data and commit is slow
+	BuildMeshDescriptionsParams.bCommitMeshDescription = false;
+	BuildMeshDescriptionsParams.bFastBuild = true;
+
+	StaticMesh.BuildFromMeshDescriptions(MeshDescriptionPointers, BuildMeshDescriptionsParams);
+#endif // !WITH_EDITOR
 }
 
+#if WITH_EDITORONLY_DATA
+void UInterchangeStaticMeshFactory::SetupSourceModelsSettings(UStaticMesh& StaticMesh, const TArray<FMeshDescription>& LodMeshDescriptions, int32 PreviousLodCount, int32 FinalLodCount, bool bIsAReimport)
+{
+	for (int32 LodIndex = 0; LodIndex < FinalLodCount; ++LodIndex)
+	{
+		FStaticMeshSourceModel& SrcModel = StaticMesh.GetSourceModel(LodIndex);
+
+		// Make sure that mesh descriptions for added LODs are kept as is when the mesh is built
+		if (LodIndex >= PreviousLodCount)
+		{
+			SrcModel.ReductionSettings.MaxDeviation = 0.0f;
+			SrcModel.ReductionSettings.PercentTriangles = 1.0f;
+			SrcModel.ReductionSettings.PercentVertices = 1.0f;
+		}
+
+		if (!bIsAReimport && LodMeshDescriptions.IsValidIndex(LodIndex))
+		{
+			FStaticMeshConstAttributes StaticMeshAttributes(LodMeshDescriptions[LodIndex]);
+			const int32 NumUVChannels = StaticMeshAttributes.GetVertexInstanceUVs().GetNumChannels();
+			const int32 FirstOpenUVChannel = NumUVChannels >= MAX_MESH_TEXTURE_COORDS_MD ? 1 : NumUVChannels;
+
+			SrcModel.BuildSettings.DstLightmapIndex = FirstOpenUVChannel;
+
+			if (LodIndex == 0)
+			{
+				StaticMesh.SetLightMapCoordinateIndex(FirstOpenUVChannel);
+			}
+		}
+	}
+}
+#endif // WITH_EDITORONLY_DATA
 
 /* This function is call in the completion task on the main thread, use it to call main thread post creation step for your assets */
 void UInterchangeStaticMeshFactory::PreImportPreCompletedCallback(const FImportPreCompletedCallbackParams& Arguments)
diff --git a/Engine/Plugins/Experimental/Interchange/Runtime/Source/Import/Private/Texture/InterchangeTextureFactory.cpp b/Engine/Plugins/Experimental/Interchange/Runtime/Source/Import/Private/Texture/InterchangeTextureFactory.cpp
--- a/Engine/Plugins/Experimental/Interchange/Runtime/Source/Import/Private/Texture/InterchangeTextureFactory.cpp
+++ b/Engine/Plugins/Experimental/Interchange/Runtime/Source/Import/Private/Texture/InterchangeTextureFactory.cpp
@@ -12,6 +12,8 @@
 #include "Engine/TextureDefines.h"
 #include "Engine/TextureLightProfile.h"
 #include "HAL/FileManagerGeneric.h"
+#include "ImageCoreUtils.h"
+#include "ImageUtils.h"
 #include "InterchangeAssetImportData.h"
 #include "InterchangeImportCommon.h"
 #include "InterchangeImportLog.h"
@@ -29,6 +31,7 @@
 #include "InterchangeTranslatorBase.h"
 #include "Nodes/InterchangeBaseNode.h"
 #include "ProfilingDebugging/CpuProfilerTrace.h"
+#include "Serialization/EditorBulkData.h"
 #include "Texture/InterchangeBlockedTexturePayloadInterface.h"
 #include "Texture/InterchangeSlicedTexturePayloadInterface.h"
 #include "Texture/InterchangeTextureLightProfilePayloadInterface.h"
@@ -521,8 +524,10 @@
 		return {};
 	}
 
-#if WITH_EDITORONLY_DATA
-	void SetupTextureSourceData(UTexture* Texture, const FImportImage& Image, UE::Serialization::FEditorBulkData::FSharedBufferWithID&& BufferAndId)
+	void SetupTextureSourceData(UTexture* Texture, const FImportImage& Image, UE::Serialization::FEditorBulkData::FSharedBufferWithID&& BufferAndId);
+
+#if WITH_EDITOR
+	void SetupTextureSourceDataEditor(UTexture* Texture, const FImportImage& Image, UE::Serialization::FEditorBulkData::FSharedBufferWithID&& BufferAndId)
 	{
 		Texture->Source.InitWithCompressedSourceData(
 			Image.SizeX,
@@ -544,7 +549,7 @@
 		}
 	}
 
-	void SetupTexture2DSourceData(UTexture2D* Texture2D, const FImportBlockedImage& BlockedImage, UE::Serialization::FEditorBulkData::FSharedBufferWithID&& BufferAndId)
+	void SetupTexture2DSourceDataEditor(UTexture2D* Texture2D, const FImportBlockedImage& BlockedImage, UE::Serialization::FEditorBulkData::FSharedBufferWithID&& BufferAndId)
 	{
 		if (BlockedImage.BlocksData.Num() > 1)
 		{
@@ -580,11 +585,11 @@
 			Image.SizeY = Block.SizeY;
 			Image.NumMips = Block.NumMips;
 
-			SetupTextureSourceData(Texture2D, Image, MoveTemp(BufferAndId));
+			SetupTextureSourceDataEditor(Texture2D, Image, MoveTemp(BufferAndId));
 		}
 	}
 
-	void SetupTextureSourceData(UTexture* Texture, const FImportSlicedImage& SlicedImage, UE::Serialization::FEditorBulkData::FSharedBufferWithID&& BufferAndId)
+	void SetupTextureSourceDataEditor(UTexture* Texture, const FImportSlicedImage& SlicedImage, UE::Serialization::FEditorBulkData::FSharedBufferWithID&& BufferAndId)
 	{
 		Texture->Source.InitLayered(
 			SlicedImage.SizeX,
@@ -606,108 +611,6 @@
 		}
 	}
 
-	void SetupTextureSourceData(UTextureLightProfile* TextureLightProfile, const FImportLightProfile& LightProfile, UE::Serialization::FEditorBulkData::FSharedBufferWithID&& BufferAndId)
-	{
-		const FImportImage& ImportImage = LightProfile;
-		SetupTextureSourceData(TextureLightProfile, ImportImage, MoveTemp(BufferAndId));
-
-		TextureLightProfile->Brightness = LightProfile.Brightness;
-		TextureLightProfile->TextureMultiplier = LightProfile.TextureMultiplier;
-	}
-
-	bool CanSetupTexture2DSourceData(FTexturePayloadVariant& TexturePayload)
-	{
-		static_assert(TVariantSize<FTexturePayloadVariant>::Value == 5, "Please update the code below and this assert to reflect the change to the variant type.");
-
-		if (TOptional<FImportBlockedImage>* BlockedImage = TexturePayload.TryGet<TOptional<FImportBlockedImage>>())
-		{
-			if (BlockedImage->IsSet())
-			{
-				return (*BlockedImage)->IsValid();
-			}
-		}
-		else if (TOptional<FImportImage>* Image = TexturePayload.TryGet<TOptional<FImportImage>>())
-		{
-			if (Image->IsSet())
-			{
-				return (*Image)->IsValid();
-			}
-		}
-		else if (TOptional<FImportLightProfile>* LightProfile = TexturePayload.TryGet<TOptional<FImportLightProfile>>())
-		{
-			if (LightProfile->IsSet())
-			{
-				return (*LightProfile)->IsValid();
-			}
-		}
-
-		return false;
-	}
-
-	void SetupTexture2DSourceData(UTexture2D* Texture2D, FProcessedPayload& ProcessedPayload)
-	{
-		if (TOptional<FImportBlockedImage>* BlockedImage = ProcessedPayload.SettingsFromPayload.TryGet<TOptional<FImportBlockedImage>>())
-		{
-			if (BlockedImage->IsSet())
-			{
-				SetupTexture2DSourceData(Texture2D, BlockedImage->GetValue(), MoveTemp(ProcessedPayload.PayloadAndId));
-			}
-		}
-		else if (TOptional<FImportImage>* Image = ProcessedPayload.SettingsFromPayload.TryGet<TOptional<FImportImage>>())
-		{
-			if (Image->IsSet())
-			{
-				SetupTextureSourceData(Texture2D, Image->GetValue(), MoveTemp(ProcessedPayload.PayloadAndId));
-			}
-		}
-		else if (TOptional<FImportLightProfile>* OptionalLightProfile = ProcessedPayload.SettingsFromPayload.TryGet<TOptional<FImportLightProfile>>())
-		{
-			if (OptionalLightProfile->IsSet())
-			{
-				FImportLightProfile& LightProfile = OptionalLightProfile->GetValue();
-
-				if (UTextureLightProfile* TextureLightProfile = Cast<UTextureLightProfile>(Texture2D))
-				{
-					SetupTextureSourceData(TextureLightProfile, LightProfile, MoveTemp(ProcessedPayload.PayloadAndId));
-				}
-				else
-				{
-					SetupTextureSourceData(Texture2D, LightProfile, MoveTemp(ProcessedPayload.PayloadAndId));
-				}
-			}
-		}
-		else
-		{
-			// The TexturePayload should be validated before calling this function
-			checkNoEntry();
-		}
-
-		// The texture has been imported and has no editor specific changes applied so we clear the painted flag.
-		Texture2D->bHasBeenPaintedInEditor = false;
-
-		// If the texture is larger than a certain threshold make it VT. This is explicitly done after the
-		// application of the existing settings above, so if a texture gets reimported at a larger size it will
-		// still be properly flagged as a VT (note: What about reimporting at a lower resolution?)
-		static const TConsoleVariableData<int32>* CVarVirtualTexturesEnabled = IConsoleManager::Get().FindTConsoleVariableDataInt(TEXT("r.VirtualTextures"));
-		check(CVarVirtualTexturesEnabled);
-
-		if (CVarVirtualTexturesEnabled->GetValueOnGameThread())
-		{
-			const int VirtualTextureAutoEnableThreshold = GetDefault<UTextureImportSettings>()->AutoVTSize;
-			const int VirtualTextureAutoEnableThresholdPixels = VirtualTextureAutoEnableThreshold * VirtualTextureAutoEnableThreshold;
-
-			// We do this in pixels so a 8192 x 128 texture won't get VT enabled 
-			// We use the Source size instead of simple Texture2D->GetSizeX() as this uses the size of the platform data
-			// however for a new texture platform data may not be generated yet, and for an reimport of a texture this is the size of the
-			// old texture. 
-			// Using source size gives one small caveat. It looks at the size before mipmap power of two padding adjustment.
-			if (Texture2D->Source.GetSizeX() * Texture2D->Source.GetSizeY() >= VirtualTextureAutoEnableThresholdPixels)
-			{
-				Texture2D->VirtualTextureStreaming = true;
-			}
-		}
-	}
-
 	bool CanSetupTextureCubeSourceData(FTexturePayloadVariant& TexturePayload)
 	{
 		static_assert(TVariantSize<FTexturePayloadVariant>::Value == 5, "Please update the code below and this assert to reflect the change to the variant type.");
@@ -748,7 +651,7 @@
 				// Cube texture always have six slice
 				if (SlicedImage.NumSlice == 6)
 				{
-					SetupTextureSourceData(TextureCube, SlicedImage, MoveTemp(ProcessedPayload.PayloadAndId));
+					SetupTextureSourceDataEditor(TextureCube, SlicedImage, MoveTemp(ProcessedPayload.PayloadAndId));
 				}
 			}
 		}
@@ -838,84 +741,6 @@
 		}
 	}
 
-	void LogErrorInvalidPayload(const FString& TextureClass, const FString& ObjectName)
-	{
-		UE_LOG(LogInterchangeImport, Error, TEXT("UInterchangeTextureFactory: The Payload was invalid for a %s. (%s)"), *TextureClass, *ObjectName);
-	}
-
-	FSharedBuffer MoveRawDataToSharedBuffer(FTexturePayloadVariant& TexturePayload)
-	{
-		static_assert(TVariantSize<FTexturePayloadVariant>::Value == 5, "Please update the code below and this assert to reflect the change to the variant type.");
-
-		if (TOptional<FImportBlockedImage>* BlockedImage = TexturePayload.TryGet<TOptional<FImportBlockedImage>>())
-		{
-			return (*BlockedImage)->RawData.MoveToShared();
-		}
-		else if (TOptional<FImportImage>* Image = TexturePayload.TryGet<TOptional<FImportImage>>())
-		{
-			return (*Image)->RawData.MoveToShared();
-		}
-		else if (TOptional<FImportSlicedImage>* SlicedImage = TexturePayload.TryGet<TOptional<FImportSlicedImage>>())
-		{
-			return (*SlicedImage)->RawData.MoveToShared();
-		}
-		else if (TOptional<FImportLightProfile>* LightProfile = TexturePayload.TryGet<TOptional<FImportLightProfile>>())
-		{
-			return (*LightProfile)->RawData.MoveToShared();
-		}
-
-		// The TexturePayload should be validated before calling this function
-		checkNoEntry();
-		return {};
-	}
-
-	FProcessedPayload& FProcessedPayload::operator=(UE::Interchange::Private::InterchangeTextureFactory::FTexturePayloadVariant&& InPayloadVariant)
-	{
-		SettingsFromPayload = MoveTemp(InPayloadVariant);
-		PayloadAndId = MoveRawDataToSharedBuffer(SettingsFromPayload);
-
-		return *this;
-	}
-
-	bool FProcessedPayload::IsValid() const
-	{
-		if (SettingsFromPayload.IsType<FEmptyVariantState>())
-		{
-			return false;
-		}
-
-		return true;
-	}
-
-
-	TArray<FString> GetFilesToHash(const FTextureFactoryNodeVariant& TextureFactoryNodeVariant, const FTexturePayloadVariant& TexturePayload)
-	{
-		TArray<FString> FilesToHash;
-		// Standard texture 2D payload
-		if (const UInterchangeTexture2DFactoryNode* const* TextureNode = TextureFactoryNodeVariant.TryGet<UInterchangeTexture2DFactoryNode*>())
-		{
-			using namespace UE::Interchange;
-			if (const TOptional<FImportBlockedImage>* OptionalBlockedPayload = TexturePayload.TryGet<TOptional<FImportBlockedImage>>())
-			{
-				if (OptionalBlockedPayload->IsSet())
-				{
-					const FImportBlockedImage& BlockImage = OptionalBlockedPayload->GetValue();
-					TMap<int32, FString> BlockAndFiles = (*TextureNode)->GetSourceBlocks();
-					FilesToHash.Reserve(BlockAndFiles.Num());
-					for (const FTextureSourceBlock& BlockData : BlockImage.BlocksData)
-					{
-						if (FString* FilePath = BlockAndFiles.Find(UE::TextureUtilitiesCommon::GetUDIMIndex(BlockData.BlockX, BlockData.BlockY)))
-						{
-							FilesToHash.Add(*FilePath);
-						}
-					}
-				}
-
-			}
-		}
-		return FilesToHash;
-	}
-
 	FGraphEventArray GenerateHashSourceFilesTasks(const UInterchangeSourceData* SourceData, TArray<FString>&& FilesToHash, TArray<FAssetImportInfo::FSourceFile>& OutSourceFiles)
 	{
 		struct FHashSourceTaskBase
@@ -1020,8 +845,255 @@
 
 		return TasksToDo;
 	}
+#else // WITH_EDITOR
+	void SetupTexture2DSourceDataRuntime(UTexture2D* Texture2D, const FImportImage& ImportImage, UE::Serialization::FEditorBulkData::FSharedBufferWithID&& BufferAndId)
+	{
+		Texture2D->SRGB = ImportImage.bSRGB;
+		ERawImageFormat::Type PixelFormatRawFormat;
+		const EPixelFormat PixelFormat = FImageCoreUtils::GetPixelFormatForRawImageFormat(FImageCoreUtils::ConvertToRawImageFormat(ImportImage.Format), &PixelFormatRawFormat);
 
+		UE::Serialization::FEditorBulkData BulkData;
+		BulkData.UpdatePayload(MoveTemp(BufferAndId));
+		TFuture<FSharedBuffer> Payload = BulkData.GetPayload();
+
+		const EGammaSpace GammaSpace = ImportImage.bSRGB ? EGammaSpace::sRGB : EGammaSpace::Linear;
+		constexpr int32 NumSlices = 1;
+		FImageView SourceImageView(const_cast<void*>(Payload.Get().GetData()), ImportImage.SizeX, ImportImage.SizeY, NumSlices, PixelFormatRawFormat, GammaSpace);
+
+		FImage DecompressedSourceImage;
+		if (ImportImage.RawDataCompressionFormat != TSCF_None)
+		{
+			if (FImageUtils::DecompressImage(Payload.Get().GetData(), Payload.Get().GetSize(), DecompressedSourceImage))
+			{
+				SourceImageView = DecompressedSourceImage;
+			}
+		}
+
+		{
+			Texture2D->SetPlatformData(new FTexturePlatformData());
+			Texture2D->GetPlatformData()->SizeX = ImportImage.SizeX;
+			Texture2D->GetPlatformData()->SizeY = ImportImage.SizeY;
+			Texture2D->GetPlatformData()->PixelFormat = PixelFormat;
+
+			// Allocate first mipmap.
+			int32 NumBlocksX = ImportImage.SizeX / GPixelFormats[PixelFormat].BlockSizeX;
+			int32 NumBlocksY = ImportImage.SizeY / GPixelFormats[PixelFormat].BlockSizeY;
+			FTexture2DMipMap* Mip = new FTexture2DMipMap();
+			Texture2D->GetPlatformData()->Mips.Add(Mip);
+			Mip->SizeX = ImportImage.SizeX;
+			Mip->SizeY = ImportImage.SizeY;
+			Mip->BulkData.Lock(LOCK_READ_WRITE);
+			Mip->BulkData.Realloc(NumBlocksX * NumBlocksY * GPixelFormats[PixelFormat].BlockBytes);
+			Mip->BulkData.Unlock();
+		}
+
+		{
+
+			Texture2D->bNotOfflineProcessed = true;
+
+			uint8* MipData = static_cast<uint8*>(Texture2D->GetPlatformData()->Mips[0].BulkData.Lock(LOCK_READ_WRITE));
+			check(MipData != nullptr);
+			int64 MipDataSize = Texture2D->GetPlatformData()->Mips[0].BulkData.GetBulkDataSize();
+
+			FImageView MipImageView(MipData, ImportImage.SizeX, ImportImage.SizeY, NumSlices, PixelFormatRawFormat, GammaSpace);
+
+			// copy into texture and convert if necessary :
+			FImageCore::CopyImage(SourceImageView, MipImageView);
+
+			Texture2D->GetPlatformData()->Mips[0].BulkData.Unlock();
+
+			Texture2D->UpdateResource();
+		}
+	}
+#endif // !WITH_EDITOR
+
+	FSharedBuffer MoveRawDataToSharedBuffer(FTexturePayloadVariant& TexturePayload)
+	{
+		static_assert(TVariantSize<FTexturePayloadVariant>::Value == 5, "Please update the code below and this assert to reflect the change to the variant type.");
+
+		if (TOptional<FImportBlockedImage>* BlockedImage = TexturePayload.TryGet<TOptional<FImportBlockedImage>>())
+		{
+			return (*BlockedImage)->RawData.MoveToShared();
+		}
+		else if (TOptional<FImportImage>* Image = TexturePayload.TryGet<TOptional<FImportImage>>())
+		{
+			return (*Image)->RawData.MoveToShared();
+		}
+		else if (TOptional<FImportSlicedImage>* SlicedImage = TexturePayload.TryGet<TOptional<FImportSlicedImage>>())
+		{
+			return (*SlicedImage)->RawData.MoveToShared();
+		}
+		else if (TOptional<FImportLightProfile>* LightProfile = TexturePayload.TryGet<TOptional<FImportLightProfile>>())
+		{
+			return (*LightProfile)->RawData.MoveToShared();
+		}
+
+		// The TexturePayload should be validated before calling this function
+		checkNoEntry();
+		return {};
+	}
+
+	FProcessedPayload& FProcessedPayload::operator=(UE::Interchange::Private::InterchangeTextureFactory::FTexturePayloadVariant&& InPayloadVariant)
+	{
+		SettingsFromPayload = MoveTemp(InPayloadVariant);
+		PayloadAndId = MoveRawDataToSharedBuffer(SettingsFromPayload);
+
+		return *this;
+	}
+
+	bool FProcessedPayload::IsValid() const
+	{
+		if (SettingsFromPayload.IsType<FEmptyVariantState>())
+		{
+			return false;
+		}
+
+		return true;
+	}
+
+
+	TArray<FString> GetFilesToHash(const FTextureFactoryNodeVariant& TextureFactoryNodeVariant, const FTexturePayloadVariant& TexturePayload)
+	{
+		TArray<FString> FilesToHash;
+		// Standard texture 2D payload
+		if (const UInterchangeTexture2DFactoryNode* const* TextureNode = TextureFactoryNodeVariant.TryGet<UInterchangeTexture2DFactoryNode*>())
+		{
+			using namespace UE::Interchange;
+			if (const TOptional<FImportBlockedImage>* OptionalBlockedPayload = TexturePayload.TryGet<TOptional<FImportBlockedImage>>())
+			{
+				if (OptionalBlockedPayload->IsSet())
+				{
+					const FImportBlockedImage& BlockImage = OptionalBlockedPayload->GetValue();
+					TMap<int32, FString> BlockAndFiles = (*TextureNode)->GetSourceBlocks();
+					FilesToHash.Reserve(BlockAndFiles.Num());
+					for (const FTextureSourceBlock& BlockData : BlockImage.BlocksData)
+					{
+						if (FString* FilePath = BlockAndFiles.Find(UE::TextureUtilitiesCommon::GetUDIMIndex(BlockData.BlockX, BlockData.BlockY)))
+						{
+							FilesToHash.Add(*FilePath);
+						}
+					}
+				}
+
+			}
+		}
+		return FilesToHash;
+	}
+
+	void LogErrorInvalidPayload(const FString& TextureClass, const FString& ObjectName)
+	{
+		UE_LOG(LogInterchangeImport, Error, TEXT("UInterchangeTextureFactory: The Payload was invalid for a %s. (%s)"), *TextureClass, *ObjectName);
+	}
+
+	bool CanSetupTexture2DSourceData(FTexturePayloadVariant& TexturePayload)
+	{
+		static_assert(TVariantSize<FTexturePayloadVariant>::Value == 5, "Please update the code below and this assert to reflect the change to the variant type.");
+
+#if WITH_EDITOR
+		if (TOptional<FImportBlockedImage>* BlockedImage = TexturePayload.TryGet<TOptional<FImportBlockedImage>>())
+		{
+			if (BlockedImage->IsSet())
+			{
+				return (*BlockedImage)->IsValid();
+			}
+		}
+		else
+#endif // WITH_EDITOR
+		if (TOptional<FImportImage>* Image = TexturePayload.TryGet<TOptional<FImportImage>>())
+		{
+			if (Image->IsSet())
+			{
+				return (*Image)->IsValid();
+			}
+		}
+		else if (TOptional<FImportLightProfile>* LightProfile = TexturePayload.TryGet<TOptional<FImportLightProfile>>())
+		{
+			if (LightProfile->IsSet())
+			{
+				return (*LightProfile)->IsValid();
+			}
+		}
+
+		return false;
+	}
+
+	void SetupTextureSourceData(UTexture* Texture, const FImportImage& Image, UE::Serialization::FEditorBulkData::FSharedBufferWithID&& BufferAndId)
+	{
+#if WITH_EDITOR
+		SetupTextureSourceDataEditor(Texture, Image, MoveTemp(BufferAndId));
+#else
+		SetupTexture2DSourceDataRuntime(Cast<UTexture2D>(Texture), Image, MoveTemp(BufferAndId));
+#endif
+	}
+
+	void SetupTexture2DSourceData(UTexture2D* Texture2D, FProcessedPayload& ProcessedPayload)
+	{
+#if WITH_EDITOR
+		if (TOptional<FImportBlockedImage>* BlockedImage = ProcessedPayload.SettingsFromPayload.TryGet<TOptional<FImportBlockedImage>>())
+		{
+			if (BlockedImage->IsSet())
+			{
+				SetupTexture2DSourceDataEditor(Texture2D, BlockedImage->GetValue(), MoveTemp(ProcessedPayload.PayloadAndId));
+			}
+		}
+		else
+#endif // WITH_EDITOR
+		if (TOptional<FImportImage>* Image = ProcessedPayload.SettingsFromPayload.TryGet<TOptional<FImportImage>>())
+		{
+			if (Image->IsSet())
+			{
+				SetupTextureSourceData(Texture2D, Image->GetValue(), MoveTemp(ProcessedPayload.PayloadAndId));
+			}
+		}
+		else if (TOptional<FImportLightProfile>* OptionalLightProfile = ProcessedPayload.SettingsFromPayload.TryGet<TOptional<FImportLightProfile>>())
+		{
+			if (OptionalLightProfile->IsSet())
+			{
+				FImportLightProfile& LightProfile = OptionalLightProfile->GetValue();
+
+				if (UTextureLightProfile* TextureLightProfile = Cast<UTextureLightProfile>(Texture2D))
+				{
+					SetupTextureSourceData(TextureLightProfile, LightProfile, MoveTemp(ProcessedPayload.PayloadAndId));
+				}
+				else
+				{
+					SetupTextureSourceData(Texture2D, LightProfile, MoveTemp(ProcessedPayload.PayloadAndId));
+				}
+			}
+		}
+		else
+		{
+			// The TexturePayload should be validated before calling this function
+			checkNoEntry();
+		}
+
+#if WITH_EDITORONLY_DATA
+		// The texture has been imported and has no editor specific changes applied so we clear the painted flag.
+		Texture2D->bHasBeenPaintedInEditor = false;
+
+		// If the texture is larger than a certain threshold make it VT. This is explicitly done after the
+		// application of the existing settings above, so if a texture gets reimported at a larger size it will
+		// still be properly flagged as a VT (note: What about reimporting at a lower resolution?)
+		static const TConsoleVariableData<int32>* CVarVirtualTexturesEnabled = IConsoleManager::Get().FindTConsoleVariableDataInt(TEXT("r.VirtualTextures"));
+		check(CVarVirtualTexturesEnabled);
+
+		if (CVarVirtualTexturesEnabled->GetValueOnGameThread())
+		{
+			const int VirtualTextureAutoEnableThreshold = GetDefault<UTextureImportSettings>()->AutoVTSize;
+			const int VirtualTextureAutoEnableThresholdPixels = VirtualTextureAutoEnableThreshold * VirtualTextureAutoEnableThreshold;
+
+			// We do this in pixels so a 8192 x 128 texture won't get VT enabled 
+			// We use the Source size instead of simple Texture2D->GetSizeX() as this uses the size of the platform data
+			// however for a new texture platform data may not be generated yet, and for an reimport of a texture this is the size of the
+			// old texture. 
+			// Using source size gives one small caveat. It looks at the size before mipmap power of two padding adjustment.
+			if (Texture2D->Source.GetSizeX() * Texture2D->Source.GetSizeY() >= VirtualTextureAutoEnableThresholdPixels)
+			{
+				Texture2D->VirtualTextureStreaming = true;
+			}
+		}
 #endif // WITH_EDITORONLY_DATA
+	}
  }
 
 
@@ -1036,7 +1108,6 @@
 
 	UObject* Texture = nullptr;
 
-#if WITH_EDITORONLY_DATA
 	if (!Arguments.AssetNode)
 	{
 		UE_LOG(LogInterchangeImport, Error, TEXT("UInterchangeTextureFactory: Asset node parameter is null."));
@@ -1086,7 +1157,9 @@
 		{
 			LightProfile->AddressX = TA_Clamp;
 			LightProfile->AddressY = TA_Clamp;
+#if WITH_EDITORONLY_DATA
 			LightProfile->MipGenSettings = TMGS_NoMipmaps;
+#endif
 			LightProfile->LODGroup = TEXTUREGROUP_IESLightProfile;
 		}
 	}
@@ -1102,20 +1175,12 @@
 		return nullptr;
 	}
 
-#endif //WITH_EDITORONLY_DATA
 	return Texture;
 }
 
 // The payload fetching and the heavy operations are done here
 UObject* UInterchangeTextureFactory::CreateAsset(const FCreateAssetParams& Arguments)
 {
-#if !WITH_EDITORONLY_DATA
-
-	UE_LOG(LogInterchangeImport, Error, TEXT("Cannot import texture asset in runtime, this is an editor only feature."));
-	return nullptr;
-
-#else
-
 	TRACE_CPUPROFILER_EVENT_SCOPE(UInterchangeTextureFactory::CreateAsset);
 
 	using namespace  UE::Interchange::Private::InterchangeTextureFactory;
@@ -1198,6 +1263,7 @@
 	{
 		bCanSetup = CanSetupTexture2DSourceData(TexturePayload);
 	}
+#if WITH_EDITOR
 	else if (UTextureCube* TextureCube = Cast<UTextureCube>(Texture))
 	{
 		bCanSetup = CanSetupTextureCubeSourceData(TexturePayload);
@@ -1206,6 +1272,7 @@
 	{
 		bCanSetup = CanSetupTexture2DArraySourceData(TexturePayload);
 	}
+#endif // WITH_EDITOR
 
 	if (!bCanSetup)
 	{
@@ -1213,7 +1280,11 @@
 		return Texture;
 	}
 
-	FGraphEventArray TasksToDo = GenerateHashSourceFilesTasks(Arguments.SourceData, GetFilesToHash(TextureFactoryNodeVariant, TexturePayload), SourceFiles);
+	FGraphEventArray TasksToDo;
 
+#if WITH_EDITORONLY_DATA
+	TasksToDo = GenerateHashSourceFilesTasks(Arguments.SourceData, GetFilesToHash(TextureFactoryNodeVariant, TexturePayload), SourceFiles);
+#endif
+
 	// Hash the payload while we hash the source files
 
@@ -1227,7 +1298,6 @@
 	//The interchange completion task (call in the GameThread after the factories pass), will call PostEditChange which will trig another asynchronous system that will build all texture in parallel
 
 	return Texture;
-#endif
 }
 
 /* This function is call in the completion task on the main thread, use it to call main thread post creation step for your assets*/
@@ -1239,13 +1309,12 @@
 
 	UTexture* Texture = Cast<UTexture>(Arguments.ImportedObject);
 
-#if WITH_EDITOR
-
 	// Finish the import on the game thread by doing the setup on the texture here
 	if (Texture && ProcessedPayload.IsValid())
 	{
+#if WITH_EDITOR
 		Texture->PreEditChange(nullptr);
-
+#endif
 		using namespace UE::Interchange::Private::InterchangeTextureFactory;
 
 		// Setup source data
@@ -1253,6 +1322,7 @@
 		{
 			SetupTexture2DSourceData(Texture2D, ProcessedPayload);
 		}
+#if WITH_EDITOR
 		else if (UTextureCube* TextureCube = Cast<UTextureCube>(Texture))
 		{
 			SetupTextureCubeSourceData(TextureCube, ProcessedPayload);
@@ -1289,6 +1359,7 @@
 			//Apply reimport strategy
 			UE::Interchange::FFactoryCommon::ApplyReimportStrategyToAsset(Texture, PreviousNode, CurrentNode, TextureFactoryNode);
 		}
+#endif // WITH_EDITOR
 	}
 	else
 	{
@@ -1300,7 +1371,6 @@
 			Texture->MarkAsGarbage();
 		}
 	}
-#endif //WITH_EDITOR
 
 	Super::PreImportPreCompletedCallback(Arguments);
 
diff --git a/Engine/Plugins/Experimental/Interchange/Runtime/Source/Import/Private/Texture/InterchangeTexturePayloadData.cpp b/Engine/Plugins/Experimental/Interchange/Runtime/Source/Import/Private/Texture/InterchangeTexturePayloadData.cpp
--- a/Engine/Plugins/Experimental/Interchange/Runtime/Source/Import/Private/Texture/InterchangeTexturePayloadData.cpp
+++ b/Engine/Plugins/Experimental/Interchange/Runtime/Source/Import/Private/Texture/InterchangeTexturePayloadData.cpp
@@ -87,15 +87,11 @@ void UE::Interchange::FImportImage::Init2DWithMips(int32 InSizeX, int32 InSizeY,
 
 int64 UE::Interchange::FImportImage::GetMipSize(int32 InMipIndex) const
 {
-#if WITH_EDITOR
 	check(InMipIndex >= 0);
 	check(InMipIndex < NumMips);
 	const int32 MipSizeX = FMath::Max(SizeX >> InMipIndex, 1);
 	const int32 MipSizeY = FMath::Max(SizeY >> InMipIndex, 1);
 	return (int64)MipSizeX * MipSizeY * FTextureSource::GetBytesPerPixel(Format);
-#else
-	return 0;
-#endif
 }
 
 int64 UE::Interchange::FImportImage::ComputeBufferSize() const
diff --git a/Engine/Plugins/Experimental/Interchange/Runtime/Source/Import/Public/Mesh/InterchangeStaticMeshFactory.h b/Engine/Plugins/Experimental/Interchange/Runtime/Source/Import/Public/Mesh/InterchangeStaticMeshFactory.h
--- a/Engine/Plugins/Experimental/Interchange/Runtime/Source/Import/Public/Mesh/InterchangeStaticMeshFactory.h
+++ b/Engine/Plugins/Experimental/Interchange/Runtime/Source/Import/Public/Mesh/InterchangeStaticMeshFactory.h
@@ -61,4 +61,10 @@ private:
 	bool ImportConvexCollision(const FCreateAssetParams& Arguments, UStaticMesh* StaticMesh, const UInterchangeStaticMeshLodDataNode* LodDataNode);
 	bool GenerateKDopCollision(const FCreateAssetParams& Arguments, UStaticMesh* StaticMesh);
 	bool ImportSockets(const FCreateAssetParams& Arguments, UStaticMesh* StaticMesh, const UInterchangeStaticMeshFactoryNode* FactoryNode);
+
+	void CommitMeshDescriptions(UStaticMesh& StaticMesh, TArray<FMeshDescription>&& LodMeshDescriptions);
+
+#if WITH_EDITORONLY_DATA
+	void SetupSourceModelsSettings(UStaticMesh& StaticMesh, const TArray<FMeshDescription>& LodMeshDescriptions, int32 PreviousLodCount, int32 FinalLodCount, bool bIsAReimport);
+#endif
 };
diff --git a/Engine/Plugins/Experimental/Interchange/Runtime/Source/Pipelines/Private/InterchangeGenericAssetsPipeline.cpp b/Engine/Plugins/Experimental/Interchange/Runtime/Source/Pipelines/Private/InterchangeGenericAssetsPipeline.cpp
--- a/Engine/Plugins/Experimental/Interchange/Runtime/Source/Pipelines/Private/InterchangeGenericAssetsPipeline.cpp
+++ b/Engine/Plugins/Experimental/Interchange/Runtime/Source/Pipelines/Private/InterchangeGenericAssetsPipeline.cpp
@@ -259,228 +259,239 @@ void UInterchangeGenericAssetsPipeline::ExecutePostImportPipeline(const UInterch
 		AnimationPipeline->ScriptedExecutePostImportPipeline(InBaseNodeContainer, NodeKey, CreatedAsset, bIsAReimport);
 	}
 
+#if WITH_EDITORONLY_DATA
+	AddPackageMetaData(CreatedAsset, InBaseNodeContainer->GetNode(NodeKey));
+#endif
+}
+
+#if WITH_EDITORONLY_DATA
+void UInterchangeGenericAssetsPipeline::AddPackageMetaData(UObject* CreatedAsset, const UInterchangeBaseNode* Node)
+{
+	if (!CreatedAsset || !Node)
+	{
+		return;
+	}
+
 	const FString InterchangeMetaDataPrefix = TEXT("INTERCHANGE.");
 
-	if (const UInterchangeBaseNode* Node = InBaseNodeContainer->GetNode(NodeKey))
+	//Add UObject package meta data
+	if (UMetaData* MetaData = CreatedAsset->GetOutermost()->GetMetaData())
 	{
-		//Add UObject package meta data
-		if (UMetaData* MetaData = CreatedAsset->GetOutermost()->GetMetaData())
+		//Cleanup existing INTERCHANGE_ prefix metadata name for this object (in casse we re-import)
 		{
-			//Cleanup existing INTERCHANGE_ prefix metadata name for this object (in casse we re-import)
+			TArray<FName> InterchangeMetaDataKeys;
+			if(TMap<FName, FString>* MetaDataMapPtr = MetaData->GetMapForObject(CreatedAsset))
 			{
-				TArray<FName> InterchangeMetaDataKeys;
-				if(TMap<FName, FString>* MetaDataMapPtr = MetaData->GetMapForObject(CreatedAsset))
+				for (const TPair<FName, FString>& ObjectMetadata : *MetaDataMapPtr)
 				{
-					for (const TPair<FName, FString>& ObjectMetadata : *MetaDataMapPtr)
+					if (ObjectMetadata.Key.ToString().StartsWith(InterchangeMetaDataPrefix))
 					{
-						if (ObjectMetadata.Key.ToString().StartsWith(InterchangeMetaDataPrefix))
-						{
-							InterchangeMetaDataKeys.Add(ObjectMetadata.Key);
-						}
-					}
-					for (const FName& MetaDataKey : InterchangeMetaDataKeys)
-					{
-						MetaData->RemoveValue(CreatedAsset, MetaDataKey);
+						InterchangeMetaDataKeys.Add(ObjectMetadata.Key);
 					}
 				}
-			}
-			TArray<FInterchangeUserDefinedAttributeInfo> UserAttributeInfos;
-			UInterchangeUserDefinedAttributesAPI::GetUserDefinedAttributeInfos(Node, UserAttributeInfos);
-			//We must convert all different type to String since meta data only support string
-			for (const FInterchangeUserDefinedAttributeInfo& UserAttributeInfo : UserAttributeInfos)
-			{
-				if (UserAttributeInfo.PayloadKey.IsSet())
+				for (const FName& MetaDataKey : InterchangeMetaDataKeys)
 				{
-					//Skip animated attributes
-					continue;
+					MetaData->RemoveValue(CreatedAsset, MetaDataKey);
 				}
-				TOptional<FString> MetaDataValue;
-				TOptional<FString> PayloadKey;
-				switch (UserAttributeInfo.Type)
+			}
+		}
+		TArray<FInterchangeUserDefinedAttributeInfo> UserAttributeInfos;
+		UInterchangeUserDefinedAttributesAPI::GetUserDefinedAttributeInfos(Node, UserAttributeInfos);
+		//We must convert all different type to String since meta data only support string
+		for (const FInterchangeUserDefinedAttributeInfo& UserAttributeInfo : UserAttributeInfos)
+		{
+			if (UserAttributeInfo.PayloadKey.IsSet())
+			{
+				//Skip animated attributes
+				continue;
+			}
+			TOptional<FString> MetaDataValue;
+			TOptional<FString> PayloadKey;
+			switch (UserAttributeInfo.Type)
+			{
+				case UE::Interchange::EAttributeTypes::Bool:
 				{
-					case UE::Interchange::EAttributeTypes::Bool:
-					{
-						bool Value = false;
-						if(UInterchangeUserDefinedAttributesAPI::GetUserDefinedAttribute(Node, UserAttributeInfo.Name, Value, PayloadKey))
-						{
-							MetaDataValue = UE::Interchange::AttributeValueToString(Value);
-						}
-					}
-					break;
-					case UE::Interchange::EAttributeTypes::Int8:
+					bool Value = false;
+					if(UInterchangeUserDefinedAttributesAPI::GetUserDefinedAttribute(Node, UserAttributeInfo.Name, Value, PayloadKey))
 					{
-						int8 Value = 0;
-						if (UInterchangeUserDefinedAttributesAPI::GetUserDefinedAttribute(Node, UserAttributeInfo.Name, Value, PayloadKey))
-						{
-							MetaDataValue = UE::Interchange::AttributeValueToString(Value);
-						}
+						MetaDataValue = UE::Interchange::AttributeValueToString(Value);
 					}
-					break;
-					case UE::Interchange::EAttributeTypes::Int16:
+				}
+				break;
+				case UE::Interchange::EAttributeTypes::Int8:
+				{
+					int8 Value = 0;
+					if (UInterchangeUserDefinedAttributesAPI::GetUserDefinedAttribute(Node, UserAttributeInfo.Name, Value, PayloadKey))
 					{
-						int16 Value = 0;
-						if (UInterchangeUserDefinedAttributesAPI::GetUserDefinedAttribute(Node, UserAttributeInfo.Name, Value, PayloadKey))
-						{
-							MetaDataValue = UE::Interchange::AttributeValueToString(Value);
-						}
+						MetaDataValue = UE::Interchange::AttributeValueToString(Value);
 					}
-					break;
-					case UE::Interchange::EAttributeTypes::Int32:
+				}
+				break;
+				case UE::Interchange::EAttributeTypes::Int16:
+				{
+					int16 Value = 0;
+					if (UInterchangeUserDefinedAttributesAPI::GetUserDefinedAttribute(Node, UserAttributeInfo.Name, Value, PayloadKey))
 					{
-						int32 Value = 0;
-						if (UInterchangeUserDefinedAttributesAPI::GetUserDefinedAttribute(Node, UserAttributeInfo.Name, Value, PayloadKey))
-						{
-							MetaDataValue = UE::Interchange::AttributeValueToString(Value);
-						}
+						MetaDataValue = UE::Interchange::AttributeValueToString(Value);
 					}
-					break;
-					case UE::Interchange::EAttributeTypes::Int64:
+				}
+				break;
+				case UE::Interchange::EAttributeTypes::Int32:
+				{
+					int32 Value = 0;
+					if (UInterchangeUserDefinedAttributesAPI::GetUserDefinedAttribute(Node, UserAttributeInfo.Name, Value, PayloadKey))
 					{
-						int64 Value = 0;
-						if (UInterchangeUserDefinedAttributesAPI::GetUserDefinedAttribute(Node, UserAttributeInfo.Name, Value, PayloadKey))
-						{
-							MetaDataValue = UE::Interchange::AttributeValueToString(Value);
-						}
+						MetaDataValue = UE::Interchange::AttributeValueToString(Value);
 					}
-					break;
-					case UE::Interchange::EAttributeTypes::UInt8:
+				}
+				break;
+				case UE::Interchange::EAttributeTypes::Int64:
+				{
+					int64 Value = 0;
+					if (UInterchangeUserDefinedAttributesAPI::GetUserDefinedAttribute(Node, UserAttributeInfo.Name, Value, PayloadKey))
 					{
-						uint8 Value = 0;
-						if (UInterchangeUserDefinedAttributesAPI::GetUserDefinedAttribute(Node, UserAttributeInfo.Name, Value, PayloadKey))
-						{
-							MetaDataValue = UE::Interchange::AttributeValueToString(Value);
-						}
+						MetaDataValue = UE::Interchange::AttributeValueToString(Value);
 					}
-					break;
-					case UE::Interchange::EAttributeTypes::UInt16:
+				}
+				break;
+				case UE::Interchange::EAttributeTypes::UInt8:
+				{
+					uint8 Value = 0;
+					if (UInterchangeUserDefinedAttributesAPI::GetUserDefinedAttribute(Node, UserAttributeInfo.Name, Value, PayloadKey))
 					{
-						uint16 Value = 0;
-						if (UInterchangeUserDefinedAttributesAPI::GetUserDefinedAttribute(Node, UserAttributeInfo.Name, Value, PayloadKey))
-						{
-							MetaDataValue = UE::Interchange::AttributeValueToString(Value);
-						}
+						MetaDataValue = UE::Interchange::AttributeValueToString(Value);
 					}
-					break;
-					case UE::Interchange::EAttributeTypes::UInt32:
+				}
+				break;
+				case UE::Interchange::EAttributeTypes::UInt16:
+				{
+					uint16 Value = 0;
+					if (UInterchangeUserDefinedAttributesAPI::GetUserDefinedAttribute(Node, UserAttributeInfo.Name, Value, PayloadKey))
 					{
-						uint32 Value = 0;
-						if (UInterchangeUserDefinedAttributesAPI::GetUserDefinedAttribute(Node, UserAttributeInfo.Name, Value, PayloadKey))
-						{
-							MetaDataValue = UE::Interchange::AttributeValueToString(Value);
-						}
+						MetaDataValue = UE::Interchange::AttributeValueToString(Value);
 					}
-					break;
-					case UE::Interchange::EAttributeTypes::UInt64:
+				}
+				break;
+				case UE::Interchange::EAttributeTypes::UInt32:
+				{
+					uint32 Value = 0;
+					if (UInterchangeUserDefinedAttributesAPI::GetUserDefinedAttribute(Node, UserAttributeInfo.Name, Value, PayloadKey))
 					{
-						uint64 Value = 0;
-						if (UInterchangeUserDefinedAttributesAPI::GetUserDefinedAttribute(Node, UserAttributeInfo.Name, Value, PayloadKey))
-						{
-							MetaDataValue = UE::Interchange::AttributeValueToString(Value);
-						}
+						MetaDataValue = UE::Interchange::AttributeValueToString(Value);
 					}
-					break;
-					case UE::Interchange::EAttributeTypes::Float:
+				}
+				break;
+				case UE::Interchange::EAttributeTypes::UInt64:
+				{
+					uint64 Value = 0;
+					if (UInterchangeUserDefinedAttributesAPI::GetUserDefinedAttribute(Node, UserAttributeInfo.Name, Value, PayloadKey))
 					{
-						float Value = 0.0f;
-						if (UInterchangeUserDefinedAttributesAPI::GetUserDefinedAttribute(Node, UserAttributeInfo.Name, Value, PayloadKey))
-						{
-							MetaDataValue = UE::Interchange::AttributeValueToString(Value);
-						}
+						MetaDataValue = UE::Interchange::AttributeValueToString(Value);
 					}
-					break;
-					case UE::Interchange::EAttributeTypes::Float16:
+				}
+				break;
+				case UE::Interchange::EAttributeTypes::Float:
+				{
+					float Value = 0.0f;
+					if (UInterchangeUserDefinedAttributesAPI::GetUserDefinedAttribute(Node, UserAttributeInfo.Name, Value, PayloadKey))
 					{
-						FFloat16 Value = 0.0f;
-						if (UInterchangeUserDefinedAttributesAPI::GetUserDefinedAttribute(Node, UserAttributeInfo.Name, Value, PayloadKey))
-						{
-							MetaDataValue = UE::Interchange::AttributeValueToString(Value);
-						}
+						MetaDataValue = UE::Interchange::AttributeValueToString(Value);
 					}
-					break;
-					case UE::Interchange::EAttributeTypes::Vector2f:
+				}
+				break;
+				case UE::Interchange::EAttributeTypes::Float16:
+				{
+					FFloat16 Value = 0.0f;
+					if (UInterchangeUserDefinedAttributesAPI::GetUserDefinedAttribute(Node, UserAttributeInfo.Name, Value, PayloadKey))
 					{
-						FVector2f Value(0.0f);
-						if (UInterchangeUserDefinedAttributesAPI::GetUserDefinedAttribute(Node, UserAttributeInfo.Name, Value, PayloadKey))
-						{
-							MetaDataValue = UE::Interchange::AttributeValueToString(Value);
-						}
+						MetaDataValue = UE::Interchange::AttributeValueToString(Value);
 					}
-					break;
-					case UE::Interchange::EAttributeTypes::Vector3f:
+				}
+				break;
+				case UE::Interchange::EAttributeTypes::Vector2f:
+				{
+					FVector2f Value(0.0f);
+					if (UInterchangeUserDefinedAttributesAPI::GetUserDefinedAttribute(Node, UserAttributeInfo.Name, Value, PayloadKey))
 					{
-						FVector3f Value(0.0f);
-						if (UInterchangeUserDefinedAttributesAPI::GetUserDefinedAttribute(Node, UserAttributeInfo.Name, Value, PayloadKey))
-						{
-							MetaDataValue = UE::Interchange::AttributeValueToString(Value);
-						}
+						MetaDataValue = UE::Interchange::AttributeValueToString(Value);
 					}
-					break;
-					case UE::Interchange::EAttributeTypes::Vector4f:
+				}
+				break;
+				case UE::Interchange::EAttributeTypes::Vector3f:
+				{
+					FVector3f Value(0.0f);
+					if (UInterchangeUserDefinedAttributesAPI::GetUserDefinedAttribute(Node, UserAttributeInfo.Name, Value, PayloadKey))
 					{
-						FVector4f Value(0.0f);
-						if (UInterchangeUserDefinedAttributesAPI::GetUserDefinedAttribute(Node, UserAttributeInfo.Name, Value, PayloadKey))
-						{
-							MetaDataValue = UE::Interchange::AttributeValueToString(Value);
-						}
+						MetaDataValue = UE::Interchange::AttributeValueToString(Value);
 					}
-					break;
-					case UE::Interchange::EAttributeTypes::Double:
+				}
+				break;
+				case UE::Interchange::EAttributeTypes::Vector4f:
+				{
+					FVector4f Value(0.0f);
+					if (UInterchangeUserDefinedAttributesAPI::GetUserDefinedAttribute(Node, UserAttributeInfo.Name, Value, PayloadKey))
 					{
-						double Value = 0.0;
-						if (UInterchangeUserDefinedAttributesAPI::GetUserDefinedAttribute(Node, UserAttributeInfo.Name, Value, PayloadKey))
-						{
-							MetaDataValue = UE::Interchange::AttributeValueToString(Value);
-						}
+						MetaDataValue = UE::Interchange::AttributeValueToString(Value);
 					}
-					break;
-					case UE::Interchange::EAttributeTypes::Vector2d:
+				}
+				break;
+				case UE::Interchange::EAttributeTypes::Double:
+				{
+					double Value = 0.0;
+					if (UInterchangeUserDefinedAttributesAPI::GetUserDefinedAttribute(Node, UserAttributeInfo.Name, Value, PayloadKey))
 					{
-						FVector2D Value(0.0);
-						if (UInterchangeUserDefinedAttributesAPI::GetUserDefinedAttribute(Node, UserAttributeInfo.Name, Value, PayloadKey))
-						{
-							MetaDataValue = UE::Interchange::AttributeValueToString(Value);
-						}
+						MetaDataValue = UE::Interchange::AttributeValueToString(Value);
 					}
-					break;
-					case UE::Interchange::EAttributeTypes::Vector3d:
+				}
+				break;
+				case UE::Interchange::EAttributeTypes::Vector2d:
+				{
+					FVector2D Value(0.0);
+					if (UInterchangeUserDefinedAttributesAPI::GetUserDefinedAttribute(Node, UserAttributeInfo.Name, Value, PayloadKey))
 					{
-						FVector3d Value(0.0);
-						if (UInterchangeUserDefinedAttributesAPI::GetUserDefinedAttribute(Node, UserAttributeInfo.Name, Value, PayloadKey))
-						{
-							MetaDataValue = UE::Interchange::AttributeValueToString(Value);
-						}
+						MetaDataValue = UE::Interchange::AttributeValueToString(Value);
 					}
-					break;
-					case UE::Interchange::EAttributeTypes::Vector4d:
+				}
+				break;
+				case UE::Interchange::EAttributeTypes::Vector3d:
+				{
+					FVector3d Value(0.0);
+					if (UInterchangeUserDefinedAttributesAPI::GetUserDefinedAttribute(Node, UserAttributeInfo.Name, Value, PayloadKey))
 					{
-						FVector4d Value(0.0);
-						if (UInterchangeUserDefinedAttributesAPI::GetUserDefinedAttribute(Node, UserAttributeInfo.Name, Value, PayloadKey))
-						{
-							MetaDataValue = UE::Interchange::AttributeValueToString(Value);
-						}
+						MetaDataValue = UE::Interchange::AttributeValueToString(Value);
 					}
-					break;
-					case UE::Interchange::EAttributeTypes::String:
+				}
+				break;
+				case UE::Interchange::EAttributeTypes::Vector4d:
+				{
+					FVector4d Value(0.0);
+					if (UInterchangeUserDefinedAttributesAPI::GetUserDefinedAttribute(Node, UserAttributeInfo.Name, Value, PayloadKey))
 					{
-						FString Value;
-						if (UInterchangeUserDefinedAttributesAPI::GetUserDefinedAttribute(Node, UserAttributeInfo.Name, Value, PayloadKey))
-						{
-							MetaDataValue = UE::Interchange::AttributeValueToString(Value);
-						}
+						MetaDataValue = UE::Interchange::AttributeValueToString(Value);
 					}
-					break;
 				}
-				if (MetaDataValue.IsSet())
+				break;
+				case UE::Interchange::EAttributeTypes::String:
 				{
-					const FString& MetaDataStringValue = MetaDataValue.GetValue();
-					const FName& MetaDataKey = FName(InterchangeMetaDataPrefix + UserAttributeInfo.Name);
-					//SetValue either add the key or set the new value
-					MetaData->SetValue(CreatedAsset, MetaDataKey, *MetaDataStringValue);
+					FString Value;
+					if (UInterchangeUserDefinedAttributesAPI::GetUserDefinedAttribute(Node, UserAttributeInfo.Name, Value, PayloadKey))
+					{
+						MetaDataValue = UE::Interchange::AttributeValueToString(Value);
+					}
 				}
+				break;
+			}
+			if (MetaDataValue.IsSet())
+			{
+				const FString& MetaDataStringValue = MetaDataValue.GetValue();
+				const FName& MetaDataKey = FName(InterchangeMetaDataPrefix + UserAttributeInfo.Name);
+				//SetValue either add the key or set the new value
+				MetaData->SetValue(CreatedAsset, MetaDataKey, *MetaDataStringValue);
 			}
 		}
 	}
 }
+#endif // WITH_EDITORONLY_DATA
 
 void UInterchangeGenericAssetsPipeline::SetReimportSourceIndex(UClass* ReimportObjectClass, const int32 SourceFileIndex)
 {
diff --git a/Engine/Plugins/Experimental/Interchange/Runtime/Source/Pipelines/Private/InterchangeGenericMaterialPipeline.cpp b/Engine/Plugins/Experimental/Interchange/Runtime/Source/Pipelines/Private/InterchangeGenericMaterialPipeline.cpp
--- a/Engine/Plugins/Experimental/Interchange/Runtime/Source/Pipelines/Private/InterchangeGenericMaterialPipeline.cpp
+++ b/Engine/Plugins/Experimental/Interchange/Runtime/Source/Pipelines/Private/InterchangeGenericMaterialPipeline.cpp
@@ -94,6 +94,14 @@ void UInterchangeGenericMaterialPipeline::ExecutePreImportPipeline(UInterchangeB
 		SourceNode->GetCustomImportUnusedMaterial(bImportUnusedMaterial);
 	}
 
+#if !WITH_EDITOR
+	// Can't import materials at runtime, fallback to instances
+	if (MaterialImport == EInterchangeMaterialImportOption::ImportAsMaterials)
+	{
+		MaterialImport = EInterchangeMaterialImportOption::ImportAsMaterialInstances;
+	}
+#endif // !WITH_EDITOR
+
 	if (MaterialImport == EInterchangeMaterialImportOption::ImportAsMaterials)
 	{
 		for (const UInterchangeShaderGraphNode* ShaderGraphNode : MaterialNodes)
diff --git a/Engine/Plugins/Experimental/Interchange/Runtime/Source/Pipelines/Public/InterchangeGenericAssetsPipeline.h b/Engine/Plugins/Experimental/Interchange/Runtime/Source/Pipelines/Public/InterchangeGenericAssetsPipeline.h
--- a/Engine/Plugins/Experimental/Interchange/Runtime/Source/Pipelines/Public/InterchangeGenericAssetsPipeline.h
+++ b/Engine/Plugins/Experimental/Interchange/Runtime/Source/Pipelines/Public/InterchangeGenericAssetsPipeline.h
@@ -104,6 +104,13 @@ private:
 	 */
 	void ImplementUseSourceNameForAssetOption();
 
+#if WITH_EDITORONLY_DATA
+	/**
+	 * Adds the user defined attributes (UInterchangeUserDefinedAttributesAPI) to the package meta data (UMetaData)
+	 */
+	void AddPackageMetaData(UObject* CreatedAsset, const UInterchangeBaseNode* Node);
+#endif // WITH_EDITORONLY_DATA
+
 	UInterchangeBaseNodeContainer* BaseNodeContainer;
 	TArray<const UInterchangeSourceData*> SourceDatas;
 	
diff --git a/Engine/Plugins/Tests/InterchangeTests/Source/InterchangeTests/Private/InterchangeImportTest.cpp b/Engine/Plugins/Tests/InterchangeTests/Source/InterchangeTests/Private/InterchangeImportTest.cpp
--- a/Engine/Plugins/Tests/InterchangeTests/Source/InterchangeTests/Private/InterchangeImportTest.cpp
+++ b/Engine/Plugins/Tests/InterchangeTests/Source/InterchangeTests/Private/InterchangeImportTest.cpp
@@ -239,19 +239,22 @@ bool FInterchangeImportTest::RunTest(const FString& Path)
 
 						Data.TestPlan->Steps[StepIndex]->FinishStep(Data, ExecutionInfo);
 
-						// Populate the automation test execution info with the interchange import results
-						for (UInterchangeResult* Result : Data.InterchangeResults->GetResults())
+						if (Data.InterchangeResults)
 						{
-							switch (Result->GetResultType())
+							// Populate the automation test execution info with the interchange import results
+							for (UInterchangeResult* Result : Data.InterchangeResults->GetResults())
 							{
-							case EInterchangeResultType::Error:
-								ExecutionInfo.AddError(Result->GetText().ToString());
-								bSuccess = false;
-								break;
-
-							case EInterchangeResultType::Warning:
-								ExecutionInfo.AddWarning(Result->GetText().ToString());
-								break;
+								switch (Result->GetResultType())
+								{
+								case EInterchangeResultType::Error:
+									ExecutionInfo.AddError(Result->GetText().ToString());
+									bSuccess = false;
+									break;
+
+								case EInterchangeResultType::Warning:
+									ExecutionInfo.AddWarning(Result->GetText().ToString());
+									break;
+								}
 							}
 						}
 
diff --git a/Engine/Plugins/Tests/InterchangeTests/Source/InterchangeTests/Private/InterchangeTestFunction.cpp b/Engine/Plugins/Tests/InterchangeTests/Source/InterchangeTests/Private/InterchangeTestFunction.cpp
--- a/Engine/Plugins/Tests/InterchangeTests/Source/InterchangeTests/Private/InterchangeTestFunction.cpp
+++ b/Engine/Plugins/Tests/InterchangeTests/Source/InterchangeTests/Private/InterchangeTestFunction.cpp
@@ -217,7 +217,7 @@ FInterchangeTestFunctionResult FInterchangeTestFunction::Invoke(const TArray<UOb
 	UClass* ExpectedClass = AssetClass;
 	for (UObject* Asset : AssetsToTest)
 	{
-		if (Asset->GetClass()->IsChildOf(ExpectedClass))
+		if (Asset && Asset->GetClass()->IsChildOf(ExpectedClass))
 		{
 			// If the asset is of the expected class, add it if the name matches the supplied optional asset name,
 			// or if we are not checking a specific name.
