Fix crash when loading a fbx file containing invalid UVs

#rnx

diff --git a/Engine/Plugins/Experimental/Interchange/Runtime/Source/Dispatcher/Public/InterchangeCommands.h b/Engine/Plugins/Experimental/Interchange/Runtime/Source/Dispatcher/Public/InterchangeCommands.h
--- a/Engine/Plugins/Experimental/Interchange/Runtime/Source/Dispatcher/Public/InterchangeCommands.h
+++ b/Engine/Plugins/Experimental/Interchange/Runtime/Source/Dispatcher/Public/InterchangeCommands.h
@@ -17,7 +17,7 @@ namespace UE
 		{
 		public:
 			//Major version should be updated when there is an existing API that has been change.
-			static int32 GetMajor() { return 9; }
+			static int32 GetMajor() { return 10; }
 			//Minor version must be updated when there is an addition to the API.
 			static int32 GetMinor() { return 0; }
 			//Patch version should be update if there is some bug fixes in the private code.
diff --git a/Engine/Plugins/Experimental/Interchange/Runtime/Source/Parsers/Fbx/Private/FbxMesh.cpp b/Engine/Plugins/Experimental/Interchange/Runtime/Source/Parsers/Fbx/Private/FbxMesh.cpp
--- a/Engine/Plugins/Experimental/Interchange/Runtime/Source/Parsers/Fbx/Private/FbxMesh.cpp
+++ b/Engine/Plugins/Experimental/Interchange/Runtime/Source/Parsers/Fbx/Private/FbxMesh.cpp
@@ -580,6 +580,12 @@ bool FMeshDescriptionImporter::FillMeshDescriptionFromFbxMesh(FbxMesh* Mesh, TAr
 			if (FBXUVs.LayerElementUV[UVLayerIndex] != nullptr)
 			{
 				int32 UVCount = FBXUVs.LayerElementUV[UVLayerIndex]->GetDirectArray().GetCount();
+				if (UVCount == 0)
+				{
+					UInterchangeResultMeshWarning_Generic* Message = AddMessage<UInterchangeResultMeshWarning_Generic>(Mesh);
+					Message->Text = LOCTEXT("CreateUVs_UVCorrupted", "Found invalid UVs value when importing mesh '{MeshName}'.");
+				}
+
 				TUVAttributesRef<FVector2f> UVCoordinates = MeshDescription->UVAttributes(UVLayerIndex).GetAttributesRef<FVector2f>(MeshAttribute::UV::UVCoordinate);
 				MeshDescription->ReserveNewUVs(UVCount, UVLayerIndex);
 				for (int32 UVIndex = 0; UVIndex < UVCount; UVIndex++)
@@ -957,6 +963,14 @@ bool FMeshDescriptionImporter::FillMeshDescriptionFromFbxMesh(FbxMesh* Mesh, TAr
 								? UVMapIndex
 								: FBXUVs.LayerElementUV[UVLayerIndex]->GetIndexArray().GetAt(UVMapIndex);
 
+							//When we have a valid UV channel but empty or not containing enough value, we do not
+							//set any UVIDs value for this corner
+							if (UVIndex >= FBXUVs.LayerElementUV[UVLayerIndex]->GetDirectArray().GetCount())
+							{
+								//The warning was already emit at this time that UVs are invalid for this mesh
+								UVIndex = -1;
+							}
+
 							if (UVIndex != -1)
 							{
 								UVIDs[VertexIndex] = UVIndex + UVOffsets[UVLayerIndex];
@@ -964,11 +978,6 @@ bool FMeshDescriptionImporter::FillMeshDescriptionFromFbxMesh(FbxMesh* Mesh, TAr
 								check(MeshDescription->VertexInstanceAttributes().GetAttribute<FVector2f>(CornerInstanceIDs[VertexIndex], MeshAttribute::VertexInstance::TextureCoordinate, UVLayerIndex) ==
 										MeshDescription->UVAttributes(UVLayerIndex).GetAttribute<FVector2f>(UVIndex + UVOffsets[UVLayerIndex], MeshAttribute::UV::UVCoordinate));
 							}
-							else
-							{
-								// TODO: what does it mean to have a UV index of -1?
-								// Investigate this case more carefully and handle it properly.
-							}
 						}
 					}
 
@@ -1146,12 +1155,9 @@ bool FMeshPayloadContext::FetchPayloadToFile(FFbxParser& Parser, const FString&
 		return false;
 	}
 
-	bool bFetchSkinnedData = (Mesh->GetDeformerCount(FbxDeformer::eSkin) > 0);
-				
 	TArray<FString> JointUniqueNames;
-
 	FMeshDescription MeshDescription;
-	if (bFetchSkinnedData)
+	if (bIsSkinnedMesh)
 	{
 		FSkeletalMeshAttributes SkeletalMeshAttribute(MeshDescription);
 		SkeletalMeshAttribute.Register();
@@ -1182,8 +1188,8 @@ bool FMeshPayloadContext::FetchPayloadToFile(FFbxParser& Parser, const FString&
 		FLargeMemoryWriter Ar;
 		MeshDescription.Serialize(Ar);
 
-		Ar << bFetchSkinnedData;
-		if (bFetchSkinnedData)
+		Ar << bIsSkinnedMesh;
+		if (bIsSkinnedMesh)
 		{
 			//When passing a skinned MeshDescription, We want to pass the joint Node ID so we can know what the influence bone index refer to
 			Ar << JointUniqueNames;
@@ -1293,9 +1299,12 @@ void FFbxMesh::AddAllMeshes(FbxScene* SDKScene, FbxGeometryConverter* SDKGeometr
 		MeshNode = CreateMeshNode(NodeContainer, MeshName, MeshUniqueID);
 		if (Geometry->GetDeformerCount(FbxDeformer::eSkin) > 0)
 		{
-			//Set the skinned mesh attribute
-			MeshNode->SetSkinnedMesh(true);
-			ExtractSkinnedMeshNodeJoints(SDKScene, NodeContainer, Mesh, MeshNode);
+			if (ExtractSkinnedMeshNodeJoints(SDKScene, NodeContainer, Mesh, MeshNode))
+			{
+				//Set the skinned mesh attribute
+				MeshNode->SetSkinnedMesh(true);
+			}
+			
 		}
 		Mesh->ComputeBBox();
 		const int32 MeshVertexCount = Mesh->GetControlPointsCount();
@@ -1373,6 +1382,7 @@ void FFbxMesh::AddAllMeshes(FbxScene* SDKScene, FbxGeometryConverter* SDKGeometr
 			GeoPayload->Mesh = Mesh;
 			GeoPayload->SDKScene = SDKScene;
 			GeoPayload->SDKGeometryConverter = SDKGeometryConverter;
+			GeoPayload->bIsSkinnedMesh = MeshNode->IsSkinnedMesh();
 			PayloadContexts.Add(PayLoadKey, GeoPayload);
 		}
 		MeshNode->SetPayLoadKey(PayLoadKey);
@@ -1515,8 +1525,9 @@ bool FFbxMesh::GetGlobalJointBindPoseTransform(FbxScene* SDKScene, FbxNode* Join
 	return false;
 }
 
-void FFbxMesh::ExtractSkinnedMeshNodeJoints(FbxScene* SDKScene, UInterchangeBaseNodeContainer& NodeContainer, FbxMesh* Mesh, UInterchangeMeshNode* MeshNode)
+bool FFbxMesh::ExtractSkinnedMeshNodeJoints(FbxScene* SDKScene, UInterchangeBaseNodeContainer& NodeContainer, FbxMesh* Mesh, UInterchangeMeshNode* MeshNode)
 {
+	bool bFoundValidJoint = false;
 	TArray<FString> JointNodeUniqueIDs;
 
 	const int32 SkinDeformerCount = Mesh->GetDeformerCount(FbxDeformer::eSkin);
@@ -1540,15 +1551,23 @@ void FFbxMesh::ExtractSkinnedMeshNodeJoints(FbxScene* SDKScene, UInterchangeBase
 				continue;
 			}
 			FbxNode* Link = Cluster->GetLink();
-			FString JointNodeUniqueID = FFbxHelper::GetFbxNodeHierarchyName(Link);
-			// find the bone index
-			if (!JointNodeUniqueIDs.Contains(JointNodeUniqueID))
+			//Any valid link should be under the scene root node
+			if (Link->GetParent() != nullptr)
 			{
-				JointNodeUniqueIDs.Add(JointNodeUniqueID);
-				MeshNode->SetSkeletonDependencyUid(JointNodeUniqueID);
+				bFoundValidJoint = true;
+
+				FString JointNodeUniqueID = FFbxHelper::GetFbxNodeHierarchyName(Link);
+
+				// find the bone index
+				if (!JointNodeUniqueIDs.Contains(JointNodeUniqueID))
+				{
+					JointNodeUniqueIDs.Add(JointNodeUniqueID);
+					MeshNode->SetSkeletonDependencyUid(JointNodeUniqueID);
+				}
 			}
 		}
 	}
+	return bFoundValidJoint;
 }
 
 UInterchangeMeshNode* FFbxMesh::CreateMeshNode(UInterchangeBaseNodeContainer& NodeContainer, const FString& NodeName, const FString& NodeUniqueID)
diff --git a/Engine/Plugins/Experimental/Interchange/Runtime/Source/Parsers/Fbx/Private/FbxMesh.h b/Engine/Plugins/Experimental/Interchange/Runtime/Source/Parsers/Fbx/Private/FbxMesh.h
--- a/Engine/Plugins/Experimental/Interchange/Runtime/Source/Parsers/Fbx/Private/FbxMesh.h
+++ b/Engine/Plugins/Experimental/Interchange/Runtime/Source/Parsers/Fbx/Private/FbxMesh.h
@@ -89,6 +89,7 @@ namespace UE
 				virtual ~FMeshPayloadContext() {}
 				virtual FString GetPayloadType() const override { return TEXT("Mesh-PayloadContext"); }
 				virtual bool FetchPayloadToFile(FFbxParser& Parser, const FString& PayloadFilepath) override;
+				bool bIsSkinnedMesh = false;
 				FbxMesh* Mesh = nullptr;
 				FbxScene* SDKScene = nullptr;
 				FbxGeometryConverter* SDKGeometryConverter = nullptr;
@@ -115,7 +116,8 @@ namespace UE
 				void AddAllMeshes(FbxScene* SDKScene, FbxGeometryConverter* SDKGeometryConverter, UInterchangeBaseNodeContainer& NodeContainer, TMap<FString, TSharedPtr<FPayloadContextBase>>& PayloadContexts);
 				static bool GetGlobalJointBindPoseTransform(FbxScene* SDKScene, FbxNode* Joint, FbxAMatrix& GlobalBindPoseJointMatrix);
 			protected:
-				void ExtractSkinnedMeshNodeJoints(FbxScene* SDKScene, UInterchangeBaseNodeContainer& NodeContainer, FbxMesh* Mesh, UInterchangeMeshNode* MeshNode);
+				/** Add joint to the interchange mesh node joint dependencies. Return false if there is no valid joint (not a valid skinned mesh) */
+				bool ExtractSkinnedMeshNodeJoints(FbxScene* SDKScene, UInterchangeBaseNodeContainer& NodeContainer, FbxMesh* Mesh, UInterchangeMeshNode* MeshNode);
 				UInterchangeMeshNode* CreateMeshNode(UInterchangeBaseNodeContainer& NodeContainer, const FString& NodeName, const FString& NodeUniqueID);
 
 			private:
diff --git a/Engine/Source/Programs/InterchangeWorker/Private/InterchangeWorker.cpp b/Engine/Source/Programs/InterchangeWorker/Private/InterchangeWorker.cpp
--- a/Engine/Source/Programs/InterchangeWorker/Private/InterchangeWorker.cpp
+++ b/Engine/Source/Programs/InterchangeWorker/Private/InterchangeWorker.cpp
@@ -61,20 +61,22 @@ int32 Main(int32 Argc, TCHAR * Argv[])
 	if (!UE::Interchange::DispatcherCommandVersion::FromString(InterchangeDispatcherVersion, Major, Minor, Patch, bLWCDisabled))
 	{
 		WorkerVersionError = TEXT("Incompatible interchange dispatcher version string command argument.");
-		WorkerVersionError += TEXT(" Editor Version: ");
+		WorkerVersionError += TEXT(" Editor Version: [");
 		WorkerVersionError += InterchangeDispatcherVersion;
-		WorkerVersionError += TEXT(" Worker Version: ");
+		WorkerVersionError += TEXT("] Worker Version: [");
 		WorkerVersionError += UE::Interchange::DispatcherCommandVersion::ToString();
+		WorkerVersionError += TEXT("]");
 		UE_LOG(LogInterchangeWorker, Error, TEXT("%s"), *WorkerVersionError);
 	}
 
  	if (!UE::Interchange::DispatcherCommandVersion::IsAPICompatible(Major, Minor, Patch, bLWCDisabled))
  	{
 		WorkerVersionError = TEXT("Incompatible interchange dispatcher API version. Please recompile InterchangeWorker target.");
-		WorkerVersionError += TEXT("\n\tEditor Version: ");
+		WorkerVersionError += TEXT(" Editor Version: [");
 		WorkerVersionError += InterchangeDispatcherVersion;
-		WorkerVersionError += TEXT("\n\tWorker Version: ");
+		WorkerVersionError += TEXT("] Worker Version: [");
 		WorkerVersionError += UE::Interchange::DispatcherCommandVersion::ToString();
+		WorkerVersionError += TEXT("]");
  		UE_LOG(LogInterchangeWorker, Error, TEXT("%s"), *WorkerVersionError);
  	}
 
diff --git a/Engine/Source/Programs/InterchangeWorker/Private/InterchangeWorkerImpl.cpp b/Engine/Source/Programs/InterchangeWorker/Private/InterchangeWorkerImpl.cpp
--- a/Engine/Source/Programs/InterchangeWorker/Private/InterchangeWorkerImpl.cpp
+++ b/Engine/Source/Programs/InterchangeWorker/Private/InterchangeWorkerImpl.cpp
@@ -49,7 +49,11 @@ bool FInterchangeWorkerImpl::Run(const FString& WorkerVersionError)
 	if (bVersionError)
 	{
 		FErrorCommand ErrorCmd;
-		ErrorCmd.ErrorMessage = WorkerVersionError;
+		
+		UInterchangeResultError_Generic* Message = FbxParser.AddMessage<UInterchangeResultError_Generic>();
+		Message->Text = FText::FromString(WorkerVersionError);
+
+		ErrorCmd.ErrorMessage = Message->ToJson();
 		CommandIO.SendCommand(ErrorCmd, Config::SendCommandTimeout_s);
 		//We want to time out after maximum of 5 seconds
 		double TimeOut = 5.0;
