Interchange: Add the build nanite static mesh option to the import settings. The settings is added under category "Static Meshes|BuildSettings".

#rnx

diff --git a/Engine/Plugins/Experimental/Interchange/Runtime/Source/FactoryNodes/Private/InterchangeStaticMeshFactoryNode.cpp b/Engine/Plugins/Experimental/Interchange/Runtime/Source/FactoryNodes/Private/InterchangeStaticMeshFactoryNode.cpp
--- a/Engine/Plugins/Experimental/Interchange/Runtime/Source/FactoryNodes/Private/InterchangeStaticMeshFactoryNode.cpp
+++ b/Engine/Plugins/Experimental/Interchange/Runtime/Source/FactoryNodes/Private/InterchangeStaticMeshFactoryNode.cpp
@@ -53,6 +53,16 @@ UClass* UInterchangeStaticMeshFactoryNode::GetObjectClass() const
 #endif
 }
 
+bool UInterchangeStaticMeshFactoryNode::GetCustomBuildNanite(bool& AttributeValue) const
+{
+	IMPLEMENT_NODE_ATTRIBUTE_GETTER(BuildNanite, bool)
+}
+
+bool UInterchangeStaticMeshFactoryNode::SetCustomBuildNanite(const bool& AttributeValue, bool bAddApplyDelegate /*= true*/)
+{
+	IMPLEMENT_NODE_ATTRIBUTE_SETTER_WITH_CUSTOM_DELEGATE(UInterchangeStaticMeshFactoryNode, BuildNanite, bool, UStaticMesh);
+}
+
 int32 UInterchangeStaticMeshFactoryNode::GetSocketUidCount() const
 {
 	return SocketUids.GetCount();
diff --git a/Engine/Plugins/Experimental/Interchange/Runtime/Source/FactoryNodes/Public/InterchangeStaticMeshFactoryNode.h b/Engine/Plugins/Experimental/Interchange/Runtime/Source/FactoryNodes/Public/InterchangeStaticMeshFactoryNode.h
--- a/Engine/Plugins/Experimental/Interchange/Runtime/Source/FactoryNodes/Public/InterchangeStaticMeshFactoryNode.h
+++ b/Engine/Plugins/Experimental/Interchange/Runtime/Source/FactoryNodes/Public/InterchangeStaticMeshFactoryNode.h
@@ -52,6 +52,14 @@ public:
 	virtual class UClass* GetObjectClass() const override;
 
 public:
+	/** Get weather the static mesh factory should set the nanite build settings. Return false if the attribute was not set.*/
+	UFUNCTION(BlueprintCallable, Category = "Interchange | Node | StaticMesh")
+	bool GetCustomBuildNanite(bool& AttributeValue) const;
+
+	/** Set weather the static mesh factory should set the nanite build settings. Return false if the attribute was not set.*/
+	UFUNCTION(BlueprintCallable, Category = "Interchange | Node | StaticMesh")
+	bool SetCustomBuildNanite(const bool& AttributeValue, bool bAddApplyDelegate = true);
+
 	/** Return The number of socket UIDs this static mesh has.*/
 	UFUNCTION(BlueprintCallable, Category = "Interchange | Node | StaticMesh")
 	int32 GetSocketUidCount() const;
@@ -73,9 +81,14 @@ private:
 	virtual void FillAssetClassFromAttribute() override;
 	virtual bool SetNodeClassFromClassAttribute() override;
 
+	const UE::Interchange::FAttributeKey Macro_CustomBuildNaniteKey = UE::Interchange::FAttributeKey(TEXT("BuildNanite"));
+
 	UE::Interchange::TArrayAttributeHelper<FString> SocketUids;
 
 protected:
+	
+	IMPLEMENT_NODE_ATTRIBUTE_DELEGATE_BY_PROPERTYNAME(BuildNanite, bool, UStaticMesh, TEXT("NaniteSettings.bEnabled"));
+
 #if WITH_ENGINE
 	TSubclassOf<UStaticMesh> AssetClass = nullptr;
 #endif
diff --git a/Engine/Plugins/Experimental/Interchange/Runtime/Source/Pipelines/Private/InterchangeGenericStaticMeshPipeline.cpp b/Engine/Plugins/Experimental/Interchange/Runtime/Source/Pipelines/Private/InterchangeGenericStaticMeshPipeline.cpp
--- a/Engine/Plugins/Experimental/Interchange/Runtime/Source/Pipelines/Private/InterchangeGenericStaticMeshPipeline.cpp
+++ b/Engine/Plugins/Experimental/Interchange/Runtime/Source/Pipelines/Private/InterchangeGenericStaticMeshPipeline.cpp
@@ -399,22 +399,28 @@ UInterchangeStaticMeshFactoryNode* UInterchangeGenericMeshPipeline::CreateStatic
 
 	switch (CommonMeshesProperties->VertexColorImportOption)
 	{
-	case EInterchangeVertexColorImportOption::IVCIO_Replace:
-	{
-		StaticMeshFactoryNode->SetCustomVertexColorReplace(true);
-	}
-	break;
-	case EInterchangeVertexColorImportOption::IVCIO_Ignore:
-	{
-		StaticMeshFactoryNode->SetCustomVertexColorIgnore(true);
+		case EInterchangeVertexColorImportOption::IVCIO_Replace:
+		{
+			StaticMeshFactoryNode->SetCustomVertexColorReplace(true);
+		}
+		break;
+		case EInterchangeVertexColorImportOption::IVCIO_Ignore:
+		{
+			StaticMeshFactoryNode->SetCustomVertexColorIgnore(true);
+		}
+		break;
+		case EInterchangeVertexColorImportOption::IVCIO_Override:
+		{
+			StaticMeshFactoryNode->SetCustomVertexColorOverride(CommonMeshesProperties->VertexOverrideColor);
+		}
+		break;
 	}
-	break;
-	case EInterchangeVertexColorImportOption::IVCIO_Override:
+
+	if (bBuildNanite)
 	{
-		StaticMeshFactoryNode->SetCustomVertexColorOverride(CommonMeshesProperties->VertexOverrideColor);
-	}
-	break;
+		StaticMeshFactoryNode->SetCustomBuildNanite(bBuildNanite);
 	}
+
 	return StaticMeshFactoryNode;
 }
 
diff --git a/Engine/Plugins/Experimental/Interchange/Runtime/Source/Pipelines/Public/InterchangeGenericMeshPipeline.h b/Engine/Plugins/Experimental/Interchange/Runtime/Source/Pipelines/Public/InterchangeGenericMeshPipeline.h
--- a/Engine/Plugins/Experimental/Interchange/Runtime/Source/Pipelines/Public/InterchangeGenericMeshPipeline.h
+++ b/Engine/Plugins/Experimental/Interchange/Runtime/Source/Pipelines/Public/InterchangeGenericMeshPipeline.h
@@ -66,6 +66,11 @@ public:
 	UPROPERTY(EditAnywhere, BlueprintReadWrite, Category = "Static Meshes")
 	bool bOneConvexHullPerUCX = true;
 
+	//////	Static Meshes Build settings Properties //////
+
+	/** If enabled this option will allow you to use Nanite rendering at runtime. Can only be used with simple opaque materials. */
+	UPROPERTY(EditAnywhere, BlueprintReadWrite, Category = "Static Meshes|Build Settings")
+	bool bBuildNanite = false;
 	
 	//////	SKELETAL_MESHES_CATEGORY Properties //////
 
