[Interchange]
Added all necessary classes, translated node, factory node and factory, to  import data to create a LevelSequence asset
Added a method in UInterchangeGenericAnimationPipeline to create a UInterchangeAnimationTrackSetFactoryNode from a UInterchangeAnimationTrackSetNode
Added a static method, BuildFactoryNodeUid, in UInterchangeFactoryBaseNode to enforce a unique pattern for unique id of factory nodes
Modified all pipelines to use UInterchangeFactoryBaseNode::BuildFactoryNodeUid where appropriate.
Reworked FTaskParsing::DoTask to keep tasks to create assets and actors in the same array. Required to create LevelSequence assets.

#test All 30 Interchange tests passed with a pre-existing warning for SkinBasic about a missing texture.

diff --git a/Engine/Plugins/Experimental/Interchange/Runtime/Source/FactoryNodes/InterchangeFactoryNodes.Build.cs b/Engine/Plugins/Experimental/Interchange/Runtime/Source/FactoryNodes/InterchangeFactoryNodes.Build.cs
--- a/Engine/Plugins/Experimental/Interchange/Runtime/Source/FactoryNodes/InterchangeFactoryNodes.Build.cs
+++ b/Engine/Plugins/Experimental/Interchange/Runtime/Source/FactoryNodes/InterchangeFactoryNodes.Build.cs
@@ -18,6 +18,13 @@ namespace UnrealBuildTool.Rules
 					"TextureUtilitiesCommon",
 				}
 			);
+
+			PrivateDependencyModuleNames.AddRange(
+				new string[]
+				{
+					"LevelSequence",
+				}
+			);
 		}
 	}
 }
diff --git a/Engine/Plugins/Experimental/Interchange/Runtime/Source/FactoryNodes/Private/InterchangeAnimationTrackSetFactoryNode.cpp b/Engine/Plugins/Experimental/Interchange/Runtime/Source/FactoryNodes/Private/InterchangeAnimationTrackSetFactoryNode.cpp
new file mode 100644
--- /dev/null
+++ b/Engine/Plugins/Experimental/Interchange/Runtime/Source/FactoryNodes/Private/InterchangeAnimationTrackSetFactoryNode.cpp
@@ -0,0 +1,50 @@
+// Copyright Epic Games, Inc. All Rights Reserved.
+
+#include "InterchangeAnimationTrackSetFactoryNode.h"
+
+#include "LevelSequence.h"
+
+UInterchangeAnimationTrackSetFactoryNode::UInterchangeAnimationTrackSetFactoryNode()
+{
+	CustomAnimationTrackUids.Initialize(Attributes, TEXT("AnimationTrackUids"));
+}
+
+UClass* UInterchangeAnimationTrackSetFactoryNode::GetObjectClass() const
+{
+	return ULevelSequence::StaticClass();
+}
+
+int32 UInterchangeAnimationTrackSetFactoryNode::GetCustomAnimationTrackUidCount() const
+{
+	return CustomAnimationTrackUids.GetCount();
+}
+
+void UInterchangeAnimationTrackSetFactoryNode::GetCustomAnimationTrackUids(TArray<FString>& OutAnimationTrackUids) const
+{
+	CustomAnimationTrackUids.GetItems(OutAnimationTrackUids);
+}
+
+void UInterchangeAnimationTrackSetFactoryNode::GetCustomAnimationTrackUid(const int32 Index, FString& OutAnimationTrackUid) const
+{
+	CustomAnimationTrackUids.GetItem(Index, OutAnimationTrackUid);
+}
+
+bool UInterchangeAnimationTrackSetFactoryNode::AddCustomAnimationTrackUid(const FString& AnimationTrackUid)
+{
+	return CustomAnimationTrackUids.AddItem(AnimationTrackUid);
+}
+
+bool UInterchangeAnimationTrackSetFactoryNode::RemoveCustomAnimationTrackUid(const FString& AnimationTrackUid)
+{
+	return CustomAnimationTrackUids.RemoveItem(AnimationTrackUid);
+}
+
+bool UInterchangeAnimationTrackSetFactoryNode::GetCustomFrameRate(float& AttributeValue) const
+{
+	IMPLEMENT_NODE_ATTRIBUTE_GETTER(FrameRate, float);
+}
+
+bool UInterchangeAnimationTrackSetFactoryNode::SetCustomFrameRate(const float& AttributeValue)
+{
+	IMPLEMENT_NODE_ATTRIBUTE_SETTER_NODELEGATE(FrameRate, float);
+}
diff --git a/Engine/Plugins/Experimental/Interchange/Runtime/Source/FactoryNodes/Public/InterchangeAnimationTrackSetFactoryNode.h b/Engine/Plugins/Experimental/Interchange/Runtime/Source/FactoryNodes/Public/InterchangeAnimationTrackSetFactoryNode.h
new file mode 100644
--- /dev/null
+++ b/Engine/Plugins/Experimental/Interchange/Runtime/Source/FactoryNodes/Public/InterchangeAnimationTrackSetFactoryNode.h
@@ -0,0 +1,82 @@
+// Copyright Epic Games, Inc. All Rights Reserved.
+
+#pragma once
+
+#include "CoreMinimal.h"
+#include "Nodes/InterchangeFactoryBaseNode.h"
+
+#include "InterchangeAnimationTrackSetFactoryNode.generated.h"
+
+UCLASS(BlueprintType, Experimental)
+class INTERCHANGEFACTORYNODES_API UInterchangeAnimationTrackSetFactoryNode : public UInterchangeFactoryBaseNode
+{
+	GENERATED_BODY()
+
+	UInterchangeAnimationTrackSetFactoryNode();
+
+public:
+	static FStringView StaticAssetTypeName()
+	{
+		return TEXT("AnimationTrackSetFactory");
+	}
+
+	/**
+	 * Return the node type name of the class, we use this when reporting errors
+	 */
+	virtual FString GetTypeName() const override
+	{
+		const FString TypeName = TEXT("AnimationTrackSetFactoryNode");
+		return TypeName;
+	}
+
+	/** Get the class this node want to create */
+	UFUNCTION(BlueprintCallable, Category = "Interchange | Node | AnimationTrackSetFactory")
+	virtual class UClass* GetObjectClass() const override;
+
+	/**
+	 * This function allow to retrieve the number of track dependencies for this object.
+	 */
+	UFUNCTION(BlueprintCallable, Category = "Interchange | Node | AnimationTrackSetFactory")
+	int32 GetCustomAnimationTrackUidCount() const;
+
+	/**
+	 * This function allow to retrieve the track dependency for this object.
+	 */
+	UFUNCTION(BlueprintCallable, Category = "Interchange | Node | AnimationTrackSetFactory")
+	void GetCustomAnimationTrackUids(TArray<FString>& OutAnimationTrackUids) const;
+
+	/**
+	 * This function allow to retrieve one track dependency for this object.
+	 */
+	UFUNCTION(BlueprintCallable, Category = "Interchange | Node | AnimationTrackSetFactory")
+	void GetCustomAnimationTrackUid(const int32 Index, FString& OutAnimationTrackUid) const;
+
+	/**
+	 * Add one track dependency to this object.
+	 */
+	UFUNCTION(BlueprintCallable, Category = "Interchange | Node | AnimationTrackSetFactory")
+	bool AddCustomAnimationTrackUid(const FString& AnimationTrackUid);
+
+	/**
+	 * Remove one track dependency from this object.
+	 */
+	UFUNCTION(BlueprintCallable, Category = "Interchange | Node | AnimationTrackSetFactory")
+	bool RemoveCustomAnimationTrackUid(const FString& AnimationTrackUid);
+
+	/**
+	 * Set the frame rate for the animations in the level sequence.
+	 */
+	UFUNCTION(BlueprintCallable, Category = "Interchange | Node | AnimationTrackSetFactory")
+	bool SetCustomFrameRate(const float& AttributeValue);
+
+	/**
+	 * Get the frame rate for the animations in the level sequence.
+	 */
+	UFUNCTION(BlueprintCallable, Category = "Interchange | Node | AnimationTrackSetFactory")
+	bool GetCustomFrameRate(float& AttributeValue) const;
+
+private:
+	const UE::Interchange::FAttributeKey Macro_CustomFrameRateKey = UE::Interchange::FAttributeKey(TEXT("FrameRate"));
+
+	UE::Interchange::TArrayAttributeHelper<FString> CustomAnimationTrackUids;
+};
diff --git a/Engine/Plugins/Experimental/Interchange/Runtime/Source/Import/InterchangeImport.Build.cs b/Engine/Plugins/Experimental/Interchange/Runtime/Source/Import/InterchangeImport.Build.cs
--- a/Engine/Plugins/Experimental/Interchange/Runtime/Source/Import/InterchangeImport.Build.cs
+++ b/Engine/Plugins/Experimental/Interchange/Runtime/Source/Import/InterchangeImport.Build.cs
@@ -17,7 +17,10 @@ namespace UnrealBuildTool.Rules
 					"InterchangeEngine",
 					"InterchangeFactoryNodes",
 					"InterchangeNodes",
+					"LevelSequence",
 					"MeshDescription",
+					"MovieScene",
+					"MovieSceneTracks",
 					"StaticMeshDescription",
 					"SkeletalMeshDescription",
 				}
diff --git a/Engine/Plugins/Experimental/Interchange/Runtime/Source/Import/Private/Animation/InterchangeAnimationTrackSetFactory.cpp b/Engine/Plugins/Experimental/Interchange/Runtime/Source/Import/Private/Animation/InterchangeAnimationTrackSetFactory.cpp
new file mode 100644
--- /dev/null
+++ b/Engine/Plugins/Experimental/Interchange/Runtime/Source/Import/Private/Animation/InterchangeAnimationTrackSetFactory.cpp
@@ -0,0 +1,699 @@
+// Copyright Epic Games, Inc. All Rights Reserved. 
+
+#include "Animation/InterchangeAnimationTrackSetFactory.h"
+
+#include "Animation/InterchangeAnimationPayload.h"
+#include "Animation/InterchangeAnimationPayloadInterface.h"
+#include "InterchangeAnimationTrackSetFactoryNode.h"
+#include "InterchangeAnimationTrackSetNode.h"
+#include "InterchangeAnimSequenceFactoryNode.h"
+#include "InterchangeImportLog.h"
+#include "InterchangeResult.h"
+#include "InterchangeSourceData.h"
+#include "InterchangeTranslatorBase.h"
+#include "Nodes/InterchangeBaseNode.h"
+#include "Nodes/InterchangeBaseNodeContainer.h"
+
+//#include "Algo/Transform.h"
+#include "Channels/MovieSceneBoolChannel.h"
+#include "Channels/MovieSceneChannelProxy.h"
+#include "Channels/MovieSceneDoubleChannel.h"
+#include "LevelSequence.h"
+#include "MovieScene.h"
+#include "MovieSceneSection.h"
+#include "MovieSceneSequence.h"
+#include "Sections/MovieScene3DTransformSection.h"
+#include "Sections/MovieSceneBoolSection.h"
+#include "Sections/MovieSceneSubSection.h"
+#include "Tracks/MovieScene3DTransformTrack.h"
+#include "Tracks/MovieSceneSubTrack.h"
+#include "Tracks/MovieSceneVisibilityTrack.h"
+
+#if WITH_EDITORONLY_DATA
+
+#include "EditorFramework/AssetImportData.h"
+
+#endif //WITH_EDITORONLY_DATA
+
+#define LOCTEXT_NAMESPACE "InterchangeAnimationTrackSetFactory"
+
+#if WITH_EDITOR
+namespace UE::Interchange::Private
+{
+	class FAnimationTrackSetHelper
+	{
+	public:
+		FAnimationTrackSetHelper(ULevelSequence& InLevelSequence, UInterchangeAnimationTrackSetFactoryNode& InFactoryNode, const UInterchangeBaseNodeContainer& InNodeContainer, const IInterchangeAnimationPayloadInterface& InPayloadInterface)
+			: LevelSequence(InLevelSequence)
+			, MovieScene(InLevelSequence.MovieScene)
+			, FactoryNode(InFactoryNode)
+			, NodeContainer(InNodeContainer)
+			, PayloadInterface(InPayloadInterface)
+		{
+		}
+
+		void PopulateLevelSequence();
+
+	private:
+		void PopulateTransformTrack(const UInterchangeTransformAnimationTrackNode& TransformTrackNode, int32 TrackIndex);
+		void PopulateSubsequenceTrack(const UInterchangeAnimationTrackSetInstanceNode& InstanceNode);
+		void PopulateAnimationTrack(const UInterchangeAnimationTrackNode& AnimationTrackNode, int32 TrackIndex);
+
+		AActor* GetActor(const UInterchangeAnimationTrackNode& TrackNode);
+		FGuid BindActorToLevelSequence(AActor& Actor, int32 TrackIndex);
+		void UpdateTransformChannels(TArrayView<FMovieSceneDoubleChannel*>& Channels, int32 IndexOffset, const TArray<FRichCurve>& Curves);
+		
+		template<typename ChannelType, typename ValueType> 
+		void UpdateStepChannel(ChannelType& Channel, const TArray<float>& KeyTimes, const TArray<ValueType>& Values)
+		{
+			const FFrameRate& FrameRate = MovieScene->GetTickResolution();
+
+			TMovieSceneChannelData<ValueType> Data = Channel.GetData();
+
+			Data.Reset();
+
+			for (int32 KeyIndex = 0; KeyIndex < KeyTimes.Num(); ++KeyIndex)
+			{
+				FFrameNumber FrameNumber = FrameRate.AsFrameNumber(KeyTimes[KeyIndex]);
+
+				if (FrameNumber < this->MinFrameNumber)
+				{
+					this->MinFrameNumber = FrameNumber;
+				}
+
+				if (FrameNumber > this->MaxFrameNumber)
+				{
+					this->MaxFrameNumber = FrameNumber;
+				}
+
+				Data.AddKey(FrameNumber, Values[KeyIndex]);
+			}
+		}
+
+	private:
+		ULevelSequence& LevelSequence;
+		UMovieScene* MovieScene = nullptr;
+		UInterchangeAnimationTrackSetFactoryNode& FactoryNode;
+		const UInterchangeBaseNodeContainer& NodeContainer;
+		const IInterchangeAnimationPayloadInterface& PayloadInterface;
+
+		FFrameNumber MinFrameNumber;
+		FFrameNumber MaxFrameNumber;
+		bool ClearSubsequenceMasterTrack = true;
+	};
+
+	void FAnimationTrackSetHelper::PopulateLevelSequence()
+	{
+		if (!MovieScene || FactoryNode.GetCustomAnimationTrackUidCount() == 0)
+		{
+			return;
+		}
+
+		float FrameRate;
+		if (FactoryNode.GetCustomFrameRate(FrameRate))
+		{
+			if (FrameRate <= 0.f)
+			{
+				FrameRate = 30.f;
+			}
+
+			MovieScene->SetDisplayRate(UE::Interchange::Animation::ConvertSampleRatetoFrameRate(FrameRate));
+		}
+		// Set 30 FPS as the default frame rate
+		else
+		{
+			MovieScene->SetDisplayRate(UE::Interchange::Animation::ConvertSampleRatetoFrameRate(30.f));
+		}
+
+		TArray<FString> AnimationTrackUids;
+		FactoryNode.GetCustomAnimationTrackUids(AnimationTrackUids);
+
+		int32 TrackIndex = 0;
+		for (const FString& AnimationTrackUid : AnimationTrackUids)
+		{
+			if (const UInterchangeBaseNode* TranslatedNode = NodeContainer.GetNode(AnimationTrackUid))
+			{
+				if (const UInterchangeTransformAnimationTrackNode* TransformTrackNode = Cast<UInterchangeTransformAnimationTrackNode>(TranslatedNode))
+				{
+					PopulateTransformTrack(*TransformTrackNode, TrackIndex);
+				}
+				else if (const UInterchangeAnimationTrackSetInstanceNode* InstanceTrackNode = Cast<UInterchangeAnimationTrackSetInstanceNode>(TranslatedNode))
+				{
+					PopulateSubsequenceTrack(*InstanceTrackNode);
+				}
+				else if (const UInterchangeAnimationTrackNode* TrackNode = Cast<UInterchangeAnimationTrackNode>(TranslatedNode))
+				{
+					PopulateAnimationTrack(*TrackNode, TrackIndex);
+				}
+
+				++TrackIndex;
+			}
+		}
+
+		LevelSequence.MovieScene->SetPlaybackRange(TRange<FFrameNumber>::Inclusive(MinFrameNumber, MaxFrameNumber));
+		LevelSequence.MovieScene->SetEvaluationType(EMovieSceneEvaluationType::FrameLocked);
+	}
+
+	void FAnimationTrackSetHelper::PopulateTransformTrack(const UInterchangeTransformAnimationTrackNode& TransformTrackNode, int32 TrackIndex)
+	{
+		// Get targeted actor exists
+		AActor* Actor = GetActor(TransformTrackNode);
+		if (!Actor)
+		{
+			UE_LOG(LogInterchangeImport, Warning, TEXT("Cannot find actor for animation track %s"), *TransformTrackNode.GetDisplayLabel());
+			return;
+		}
+
+		// Get payload
+		FString PayloadKey;
+		if (!TransformTrackNode.GetCustomAnimationPayloadKey(PayloadKey))
+		{
+			UE_LOG(LogInterchangeImport, Warning, TEXT("No payload key for animation track %s on actor %s"), *TransformTrackNode.GetDisplayLabel(), *Actor->GetActorLabel());
+			return;
+		}
+
+		TFuture<TOptional<UE::Interchange::FAnimationCurvePayloadData>> Result = PayloadInterface.GetAnimationCurvePayloadData(PayloadKey);
+		const TOptional<UE::Interchange::FAnimationCurvePayloadData>& PayloadData = Result.Get();
+		if (!PayloadData.IsSet() || PayloadData->Curves.Num() != 9)
+		{
+			UE_LOG(LogInterchangeImport, Warning, TEXT("No payload for animation track %s on actor %s"), *TransformTrackNode.GetDisplayLabel(), *Actor->GetActorLabel());
+			return;
+		}
+
+		FGuid ObjectBinding = BindActorToLevelSequence(*Actor, TrackIndex);
+
+		UMovieScene3DTransformTrack* TransformTrack = MovieScene->FindTrack<UMovieScene3DTransformTrack>(ObjectBinding);
+		if (!TransformTrack)
+		{
+			TransformTrack = MovieScene->AddTrack<UMovieScene3DTransformTrack>(ObjectBinding);
+		}
+		else
+		{
+			TransformTrack->RemoveAllAnimationData();
+		}
+
+		if (!TransformTrack)
+		{
+			return;
+		}
+
+		bool bSectionAdded = false;
+		UMovieScene3DTransformSection* TransformSection = Cast<UMovieScene3DTransformSection>(TransformTrack->FindOrAddSection(0, bSectionAdded));
+		if (!TransformSection)
+		{
+			return;
+		}
+
+		if (bSectionAdded)
+		{
+			int32 CompletionMode;
+			if (TransformTrackNode.GetCustomCompletionMode(CompletionMode))
+			{
+				// Make sure EMovieSceneCompletionMode enum value are still between 0 and 2
+				static_assert((uint8)EInterchangeAimationCompletionMode::KeepState == (uint8)EMovieSceneCompletionMode::KeepState, "ENUM_VALUE_HAS_CHANGED");
+				static_assert((uint8)EInterchangeAimationCompletionMode::RestoreState == (uint8)EMovieSceneCompletionMode::RestoreState, "ENUM_VALUE_HAS_CHANGED");
+				static_assert((uint8)EInterchangeAimationCompletionMode::ProjectDefault == (uint8)EMovieSceneCompletionMode::ProjectDefault, "ENUM_VALUE_HAS_CHANGED");
+
+				TransformSection->EvalOptions.CompletionMode = (EMovieSceneCompletionMode)CompletionMode;
+			}
+			// By default the completion mode is EMovieSceneCompletionMode::ProjectDefault
+			else
+			{
+				TransformSection->EvalOptions.CompletionMode = EMovieSceneCompletionMode::ProjectDefault;
+			}
+
+			TransformSection->SetRange(TRange<FFrameNumber>::All());
+		}
+
+		const FFrameRate FrameRate = MovieScene->GetDisplayRate();
+		const FFrameRate TickResolution = MovieScene->GetTickResolution();
+
+		TArrayView<FMovieSceneDoubleChannel*> Channels = TransformSection->GetChannelProxy().GetChannels<FMovieSceneDoubleChannel>();
+
+		UpdateTransformChannels(Channels, 0, PayloadData->Curves); // Translation
+		UpdateTransformChannels(Channels, 3, PayloadData->Curves); // Rotation
+		UpdateTransformChannels(Channels, 6, PayloadData->Curves); // Scaling
+
+		// Remove unnecessary keys
+		FKeyDataOptimizationParams OptimParams;
+		OptimParams.DisplayRate = MovieScene->GetDisplayRate();
+		for (FMovieSceneDoubleChannel* Channel : Channels)
+		{
+			Channel->Optimize(OptimParams);
+		}
+
+
+		int32 EnabledTransformChannels;
+		if (TransformTrackNode.GetCustomUsedChannels(EnabledTransformChannels))
+		{
+			TransformSection->SetMask((EMovieSceneTransformChannel)EnabledTransformChannels);
+		}
+		// By default all channels are enabled
+		else
+		{
+			TransformSection->SetMask(EMovieSceneTransformChannel::AllTransform);
+		}
+
+		if (USceneComponent* SceneComp = Actor->GetRootComponent())
+		{
+			SceneComp->SetMobility(EComponentMobility::Movable);
+		}
+	}
+
+	void FAnimationTrackSetHelper::PopulateSubsequenceTrack(const UInterchangeAnimationTrackSetInstanceNode& InstanceNode)
+	{
+		FString TrackSetNodeUid;
+		if (!InstanceNode.GetCustomTrackSetDependencyUid(TrackSetNodeUid))
+		{
+			UE_LOG(LogInterchangeImport, Error, TEXT("No unique id specified for the animation track set referenced by animation track %s."), *InstanceNode.GetDisplayLabel());
+			return;
+		}
+
+		const FString TrackSetFactoryNodeUid = TEXT("Factory_") + TrackSetNodeUid;
+		const UInterchangeAnimationTrackSetFactoryNode* InstanceFactoryNode = Cast<UInterchangeAnimationTrackSetFactoryNode>(NodeContainer.GetNode(TrackSetFactoryNodeUid));
+		if (!InstanceFactoryNode || !InstanceFactoryNode->ReferenceObject.IsValid())
+		{
+			UE_LOG(LogInterchangeImport, Error, TEXT("Cannot find factory of animation track set referenced by animation track %s."), *InstanceNode.GetDisplayLabel());
+			return;
+		}
+
+		UMovieSceneSequence* TargetMovieSceneSequence = CastChecked<UMovieSceneSequence>(InstanceFactoryNode->ReferenceObject.TryLoad());
+
+		// Create SubTrack
+		UMovieSceneSubTrack* SubTrack = MovieScene->FindMasterTrack<UMovieSceneSubTrack>();
+		if (!SubTrack)
+		{
+			SubTrack = MovieScene->AddMasterTrack<UMovieSceneSubTrack>();
+		}
+		else if (SubTrack && ClearSubsequenceMasterTrack)
+		{
+			SubTrack->RemoveAllAnimationData();
+		}
+
+		if (!SubTrack)
+		{
+			return;
+		}
+
+		FFrameRate DisplayRate = MovieScene->GetDisplayRate();
+
+		float SrcLowerBound = 0.f;
+		float SrcUpperBound = 0.f;
+		int32 StartFrame = 0;
+		if (InstanceNode.GetCustomStartFrame(StartFrame))
+		{
+			SrcLowerBound = DisplayRate.AsSeconds(StartFrame);
+		}
+
+		int32 Duration = 0;
+		if (InstanceNode.GetCustomDuration(Duration))
+		{
+			SrcUpperBound = DisplayRate.AsSeconds(StartFrame + Duration);
+		}
+
+		const FFrameRate TickResolution = MovieScene->GetTickResolution();
+		const FFrameNumber DstLowerBound = TickResolution.AsFrameNumber(SrcLowerBound);
+		const FFrameNumber DstUpperBound = TickResolution.AsFrameNumber(SrcUpperBound);
+
+		// Internally AddSequenceOnRow will automatically bump overlapping subsequences, so we can just add where it's ideal for us
+		UMovieSceneSubSection* NewSection = SubTrack->AddSequenceOnRow(TargetMovieSceneSequence, DstLowerBound, DstUpperBound.Value - DstLowerBound.Value, INDEX_NONE);
+		
+		NewSection->Parameters.TimeScale = 1.f;
+		InstanceNode.GetCustomTimeScale(NewSection->Parameters.TimeScale);
+		
+		int32 CompletionMode;
+		if (InstanceNode.GetCustomCompletionMode(CompletionMode))
+		{
+			NewSection->EvalOptions.CompletionMode = (EMovieSceneCompletionMode)CompletionMode;
+		}
+		// By default the completion mode is EMovieSceneCompletionMode::ProjectDefault
+		else
+		{
+			NewSection->EvalOptions.CompletionMode = EMovieSceneCompletionMode::ProjectDefault;
+		}
+
+		MinFrameNumber = FMath::Min(MinFrameNumber, NewSection->GetRange().GetLowerBoundValue());
+		MaxFrameNumber = FMath::Max(MaxFrameNumber, NewSection->GetRange().GetUpperBoundValue());
+
+		ClearSubsequenceMasterTrack = false;
+	}
+
+	void FAnimationTrackSetHelper::PopulateAnimationTrack(const UInterchangeAnimationTrackNode& AnimationTrackNode, int32 TrackIndex)
+	{
+		int32 TargetedProperty;
+		if (!AnimationTrackNode.GetCustomTargetedProperty(TargetedProperty))
+		{
+			return;
+		}
+
+		// Only visibility is supported for the time being
+		if (TargetedProperty != (int32)EInterchangeAnimatedProperty::Visibility)
+		{
+			return;
+		}
+
+		// Get targeted actor exists
+		AActor* Actor = GetActor(AnimationTrackNode);
+		if (!Actor)
+		{
+			UE_LOG(LogInterchangeImport, Warning, TEXT("Cannot find actor for animation track %s"), *AnimationTrackNode.GetDisplayLabel());
+			return;
+		}
+
+		// Get payload
+		FString PayloadKey;
+		if (!AnimationTrackNode.GetCustomAnimationPayloadKey(PayloadKey))
+		{
+			UE_LOG(LogInterchangeImport, Warning, TEXT("No payload key for animation track %s on actor %s"), *AnimationTrackNode.GetDisplayLabel(), *Actor->GetActorLabel());
+			return;
+		}
+
+		TFuture<TOptional<UE::Interchange::FAnimationStepCurvePayloadData>> Result = PayloadInterface.GetAnimationStepCurvePayloadData(PayloadKey);
+		const TOptional<UE::Interchange::FAnimationStepCurvePayloadData>& PayloadData = Result.Get();
+		if (!PayloadData.IsSet() || PayloadData->StepCurves.Num() != 1)
+		{
+			UE_LOG(LogInterchangeImport, Warning, TEXT("No payload for animation track %s on actor %s"), *AnimationTrackNode.GetDisplayLabel(), *Actor->GetActorLabel());
+			return;
+		}
+
+		FGuid ObjectBinding = BindActorToLevelSequence(*Actor, TrackIndex);
+
+		UMovieSceneVisibilityTrack* VisibilityTrack = MovieScene->FindTrack<UMovieSceneVisibilityTrack>(ObjectBinding);
+		if (!VisibilityTrack)
+		{
+			VisibilityTrack = MovieScene->AddTrack<UMovieSceneVisibilityTrack>(ObjectBinding);
+		}
+		else
+		{
+			VisibilityTrack->RemoveAllAnimationData();
+		}
+
+		if (!VisibilityTrack)
+		{
+			return;
+		}
+
+		bool bSectionAdded = false;
+		UMovieSceneBoolSection* BoolSection = Cast<UMovieSceneBoolSection>(VisibilityTrack->FindOrAddSection(0, bSectionAdded));
+		if (!BoolSection)
+		{
+			return;
+		}
+
+		if (bSectionAdded)
+		{
+			int32 CompletionMode;
+			if (AnimationTrackNode.GetCustomCompletionMode(CompletionMode))
+			{
+				// Make sure EMovieSceneCompletionMode enum value are still between 0 and 2
+				static_assert(0 == (uint32)EMovieSceneCompletionMode::KeepState, "ENUM_VALUE_HAS_CHANGED");
+				static_assert(1 == (uint32)EMovieSceneCompletionMode::RestoreState, "ENUM_VALUE_HAS_CHANGED");
+				static_assert(2 == (uint32)EMovieSceneCompletionMode::ProjectDefault, "ENUM_VALUE_HAS_CHANGED");
+
+				BoolSection->EvalOptions.CompletionMode = (EMovieSceneCompletionMode)CompletionMode;
+			}
+			// By default the completion mode is EMovieSceneCompletionMode::ProjectDefault
+			else
+			{
+				BoolSection->EvalOptions.CompletionMode = EMovieSceneCompletionMode::ProjectDefault;
+			}
+
+			BoolSection->SetRange(TRange<FFrameNumber>::All());
+		}
+
+		TArrayView<FMovieSceneBoolChannel*> Channels = BoolSection->GetChannelProxy().GetChannels<FMovieSceneBoolChannel>();
+		if (ensure(Channels[0]))
+		{
+			const FInterchangeStepCurve& Curve = PayloadData->StepCurves[0];
+			const TArray<float>& KeyTimes = Curve.KeyTimes;
+			const TArray<bool>& Values = Curve.BooleanKeyValues.GetValue();
+
+			UpdateStepChannel(*(Channels[0]), KeyTimes, Values);
+		}
+	}
+
+	AActor* FAnimationTrackSetHelper::GetActor(const UInterchangeAnimationTrackNode& TrackNode)
+	{
+		AActor* Actor = nullptr;
+
+		FString ActorNodeUid;
+		if (TrackNode.GetCustomActorDependencyUid(ActorNodeUid))
+		{
+			const FString ActorFactoryNodeUid = TEXT("Factory_") + ActorNodeUid;
+			const UInterchangeFactoryBaseNode* ActorFactoryNode = Cast<UInterchangeFactoryBaseNode>(NodeContainer.GetNode(ActorFactoryNodeUid));
+
+			if (ActorFactoryNode && ActorFactoryNode->ReferenceObject.IsValid())
+			{
+				Actor = Cast<AActor>(ActorFactoryNode->ReferenceObject.TryLoad());
+			}
+		}
+
+		return Actor;
+	}
+
+	FGuid FAnimationTrackSetHelper::BindActorToLevelSequence(AActor& Actor, int32 TrackIndex)
+	{
+		// Bind the actor to the level sequence
+		// But first, check if there's already a possessable at the current index
+		FMovieScenePossessable* OldPossessable = nullptr;
+		if (TrackIndex < MovieScene->GetPossessableCount())
+		{
+			OldPossessable = &MovieScene->GetPossessable(TrackIndex);
+		}
+
+		FString ActorLabel = Actor.GetActorLabel();
+		FGuid ObjectBinding;
+		if (OldPossessable)
+		{
+			// If so, remove any track associated with it and unbind it from the level sequence
+			FGuid OldGuid = OldPossessable->GetGuid();
+			UMovieScene3DTransformTrack* TransformTrack = MovieScene->FindTrack<UMovieScene3DTransformTrack>(OldGuid);
+			if (TransformTrack)
+			{
+				MovieScene->RemoveTrack(*TransformTrack);
+			}
+			LevelSequence.UnbindPossessableObjects(OldGuid);
+
+			// Replace the old possessable with the new one in the MovieScene
+			FMovieScenePossessable NewPossessable(ActorLabel, Actor.GetClass());
+			MovieScene->ReplacePossessable(OldGuid, NewPossessable);
+			ObjectBinding = NewPossessable.GetGuid();
+		}
+		else
+		{
+			// No previous possessable, so simply add it to the MovieScene
+			ObjectBinding = MovieScene->AddPossessable(ActorLabel, Actor.GetClass());
+		}
+
+		// Finally, associate the guid with the actor
+		LevelSequence.BindPossessableObject(ObjectBinding, Actor, Actor.GetWorld());
+
+		return ObjectBinding;
+	}
+
+	void FAnimationTrackSetHelper::UpdateTransformChannels(TArrayView<FMovieSceneDoubleChannel*>& Channels, int32 IndexOffset, const TArray<FRichCurve>& Curves)
+	{
+		auto CopyToChannel = [this](FMovieSceneDoubleChannel* Channel, const FRichCurve& Curve)
+		{
+			const FFrameRate& FrameRate = this->MovieScene->GetTickResolution();
+
+			const TArray<FRichCurveKey>& CurveKeys = Curve.GetConstRefOfKeys();
+			
+			TArray<FFrameNumber> FrameNumbers;
+			FrameNumbers.Reserve(CurveKeys.Num());
+			
+			TArray<FMovieSceneDoubleValue> MovieSceneDoubleValues;
+			MovieSceneDoubleValues.Reserve(CurveKeys.Num());
+
+			for (int32 KeyIndex = 0; KeyIndex < CurveKeys.Num(); ++KeyIndex)
+			{
+				const FRichCurveKey& CurveKey = CurveKeys[KeyIndex];
+
+				FFrameNumber& FrameNumber = FrameNumbers.Add_GetRef(FrameRate.AsFrameNumber(CurveKey.Time));
+
+				if (FrameNumber < this->MinFrameNumber)
+				{
+					this->MinFrameNumber = FrameNumber;
+				}
+
+				if (FrameNumber > this->MaxFrameNumber)
+				{
+					this->MaxFrameNumber = FrameNumber;
+				}
+
+				FMovieSceneDoubleValue& SceneValue = MovieSceneDoubleValues.AddDefaulted_GetRef();
+
+				SceneValue.InterpMode = CurveKey.InterpMode;
+				SceneValue.Value = CurveKey.Value;
+			}
+
+			Channel->Set(FrameNumbers, MovieSceneDoubleValues);
+		};
+
+		CopyToChannel(Channels[IndexOffset + 0], Curves[IndexOffset + 0]);
+		CopyToChannel(Channels[IndexOffset + 1], Curves[IndexOffset + 1]);
+		CopyToChannel(Channels[IndexOffset + 2], Curves[IndexOffset + 2]);
+	}
+} //namespace UE::Interchange::Private
+#endif //WITH_EDITOR
+
+UClass* UInterchangeAnimationTrackSetFactory::GetFactoryClass() const
+{
+	return ULevelSequence::StaticClass();
+}
+
+UObject* UInterchangeAnimationTrackSetFactory::CreateEmptyAsset(const FCreateAssetParams& Arguments)
+{
+#if !WITH_EDITOR || !WITH_EDITORONLY_DATA
+
+	UE_LOG(LogInterchangeImport, Error, TEXT("Cannot import levelsequence asset in runtime, this is an editor only feature."));
+	return nullptr;
+#else
+	if (Arguments.ReimportObject)
+	{
+		UInterchangeResultError_Generic* Message = AddMessage<UInterchangeResultError_Generic>();
+		Message->SourceAssetName = Arguments.SourceData->GetFilename();
+		Message->DestinationAssetName = Arguments.AssetName;
+		Message->AssetType = ULevelSequence::StaticClass();
+		Message->Text = LOCTEXT("CreateEmptyAssetUnsupportedReimport", "Re-import of ULevelSequence not supported yet.");
+
+		return nullptr;
+	}
+
+	ULevelSequence* LevelSequence = nullptr;
+	if (!Arguments.AssetNode || !Arguments.AssetNode->GetObjectClass()->IsChildOf(GetFactoryClass()))
+	{
+		return nullptr;
+	}
+
+	const UInterchangeAnimationTrackSetFactoryNode* FactoryNode = Cast<UInterchangeAnimationTrackSetFactoryNode>(Arguments.AssetNode);
+	if (FactoryNode == nullptr)
+	{
+		return nullptr;
+	}
+
+	// create an asset if it doesn't exist
+	UObject* ExistingAsset = StaticFindObject(nullptr, Arguments.Parent, *Arguments.AssetName);
+
+	// create a new material or overwrite existing asset, if possible
+	if (!ExistingAsset)
+	{
+		LevelSequence = NewObject<ULevelSequence>(Arguments.Parent, *Arguments.AssetName, RF_Public | RF_Standalone);
+	}
+	else if (ExistingAsset->GetClass()->IsChildOf(ULevelSequence::StaticClass()))
+	{
+		// This is a reimport, we are just re-updating the source data
+		LevelSequence = Cast<ULevelSequence>(ExistingAsset);
+	}
+
+	if (!LevelSequence)
+	{
+		UE_LOG(LogInterchangeImport, Warning, TEXT("Could not create LevelSequence asset %s"), *Arguments.AssetName);
+		return nullptr;
+	}
+
+	LevelSequence->PreEditChange(nullptr);
+
+	return LevelSequence;
+#endif //else !WITH_EDITOR || !WITH_EDITORONLY_DATA
+}
+
+UObject* UInterchangeAnimationTrackSetFactory::CreateAsset(const FCreateAssetParams& Arguments)
+{
+#if !WITH_EDITOR || !WITH_EDITORONLY_DATA
+	// TODO: Can we import ULevelSequence at runtime
+	UE_LOG(LogInterchangeImport, Error, TEXT("Cannot import LevelSequence asset in runtime, this is an editor only feature."));
+	return nullptr;
+
+#else
+	using namespace UE::Interchange;
+
+	// Re-import is not supported yet
+	// Need to talk to the Sequencer team about adding an AssetImportData to ULevelSequence
+	if (Arguments.ReimportObject)
+	{
+		UInterchangeResultError_Generic* Message = AddMessage<UInterchangeResultError_Generic>();
+		Message->SourceAssetName = Arguments.SourceData->GetFilename();
+		Message->DestinationAssetName = Arguments.AssetName;
+		Message->AssetType = ULevelSequence::StaticClass();
+		Message->Text = LOCTEXT("CreateAssetUnsupportedReimport", "Re-import of ULevelSequence not supported yet.");
+
+		return nullptr;
+	}
+
+	if (!Arguments.NodeContainer || !Arguments.AssetNode || !Arguments.AssetNode->GetObjectClass()->IsChildOf(GetFactoryClass()))
+	{
+		return nullptr;
+	}
+
+	UInterchangeAnimationTrackSetFactoryNode* FactoryNode = Cast<UInterchangeAnimationTrackSetFactoryNode>(Arguments.AssetNode);
+	if (!FactoryNode)
+	{
+		return nullptr;
+	}
+
+	Translator = Arguments.Translator;
+	const IInterchangeAnimationPayloadInterface* AnimSequenceTranslatorPayloadInterface = Cast<IInterchangeAnimationPayloadInterface>(Arguments.Translator);
+	if (!AnimSequenceTranslatorPayloadInterface)
+	{
+		UE_LOG(LogInterchangeImport, Error, TEXT("Cannot import LevelSequence, the translator do not implement the IInterchangeAnimationPayloadInterface."));
+		return nullptr;
+	}
+
+	const UClass* LevelSequenceClass = FactoryNode->GetObjectClass();
+	check(LevelSequenceClass && LevelSequenceClass->IsChildOf(GetFactoryClass()));
+
+	// create an asset if it doesn't exist
+	UObject* ExistingAsset = StaticFindObject(nullptr, Arguments.Parent, *Arguments.AssetName);
+
+	ULevelSequence* LevelSequence = nullptr;
+	// create a new level sequence or overwrite existing asset, if possible
+	if (!ExistingAsset)
+	{
+		//NewObject is not thread safe, the asset registry directory watcher tick on the main thread can trig before we finish initializing the UObject and will crash
+		//The UObject should have been create by calling CreateEmptyAsset on the main thread.
+		check(IsInGameThread());
+		LevelSequence = NewObject<ULevelSequence>(Arguments.Parent, LevelSequenceClass, *Arguments.AssetName, RF_Public | RF_Standalone);
+	}
+	else if (ExistingAsset->GetClass()->IsChildOf(LevelSequenceClass))
+	{
+		//This is a reimport, we are just re-updating the source data
+		LevelSequence =  Cast<ULevelSequence>(ExistingAsset);
+	}
+
+	if (!ensure(LevelSequence))
+	{
+		UInterchangeResultError_Generic* Message = AddMessage<UInterchangeResultError_Generic>();
+		Message->SourceAssetName = Arguments.SourceData->GetFilename();
+		Message->DestinationAssetName = Arguments.AssetName;
+		Message->AssetType = ULevelSequence::StaticClass();
+		Message->Text = FText::Format(LOCTEXT("CreateAssetUnsupportedReimport", "Could not create nor find LevelSequence asset {0}."), FText::FromString(Arguments.AssetName));
+		return nullptr;
+	}
+
+	LevelSequence->Initialize();
+
+	Private::FAnimationTrackSetHelper Helper(*LevelSequence, *FactoryNode, *Arguments.NodeContainer, *AnimSequenceTranslatorPayloadInterface);
+	Helper.PopulateLevelSequence();
+
+	/** Apply all FactoryNode custom attributes to the level sequence asset */
+	FactoryNode->ApplyAllCustomAttributeToObject(LevelSequence);
+
+	//Getting the file Hash will cache it into the source data
+	Arguments.SourceData->GetFileContentHash();
+
+	//The interchange completion task (call in the GameThread after the factories pass), will call PostEditChange which will trig another asynchronous system that will build all material in parallel
+	return LevelSequence;
+#endif //else !WITH_EDITOR || !WITH_EDITORONLY_DATA
+}
+
+/* This function is call in the completion task on the main thread, use it to call main thread post creation step for your assets*/
+void UInterchangeAnimationTrackSetFactory::PreImportPreCompletedCallback(const FImportPreCompletedCallbackParams& Arguments)
+{
+	check(IsInGameThread());
+	Super::PreImportPreCompletedCallback(Arguments);
+
+	// TODO: Talk with sequence team about adding AssetImportData to ULevelSequence for re-import 
+}
+
+#undef LOCTEXT_NAMESPACE
\ No newline at end of file
diff --git a/Engine/Plugins/Experimental/Interchange/Runtime/Source/Import/Private/InterchangeImportModule.cpp b/Engine/Plugins/Experimental/Interchange/Runtime/Source/Import/Private/InterchangeImportModule.cpp
--- a/Engine/Plugins/Experimental/Interchange/Runtime/Source/Import/Private/InterchangeImportModule.cpp
+++ b/Engine/Plugins/Experimental/Interchange/Runtime/Source/Import/Private/InterchangeImportModule.cpp
@@ -3,6 +3,7 @@
 #include "InterchangeImportModule.h"
 
 #include "Animation/InterchangeAnimSequenceFactory.h"
+#include "Animation/InterchangeAnimationTrackSetFactory.h"
 #include "CoreMinimal.h"
 #include "Fbx/InterchangeFbxTranslator.h"
 #include "Gltf/InterchangeGltfTranslator.h"
@@ -80,6 +81,7 @@ void FInterchangeImportModule::StartupModule()
 		InterchangeManager.RegisterFactory(UInterchangeStaticMeshFactory::StaticClass());
 		InterchangeManager.RegisterFactory(UInterchangePhysicsAssetFactory::StaticClass());
 		InterchangeManager.RegisterFactory(UInterchangeActorFactory::StaticClass());
+		InterchangeManager.RegisterFactory(UInterchangeAnimationTrackSetFactory::StaticClass());
 		InterchangeManager.RegisterFactory(UInterchangeAnimSequenceFactory::StaticClass());
 		InterchangeManager.RegisterFactory(UInterchangeCineCameraActorFactory::StaticClass());
 		InterchangeManager.RegisterFactory(UInterchangeStaticMeshActorFactory::StaticClass());
diff --git a/Engine/Plugins/Experimental/Interchange/Runtime/Source/Import/Public/Animation/InterchangeAnimationTrackSetFactory.h b/Engine/Plugins/Experimental/Interchange/Runtime/Source/Import/Public/Animation/InterchangeAnimationTrackSetFactory.h
new file mode 100644
--- /dev/null
+++ b/Engine/Plugins/Experimental/Interchange/Runtime/Source/Import/Public/Animation/InterchangeAnimationTrackSetFactory.h
@@ -0,0 +1,38 @@
+// Copyright Epic Games, Inc. All Rights Reserved.
+
+#pragma once
+
+#include "CoreMinimal.h"
+#include "InterchangeFactoryBase.h"
+#include "UObject/Object.h"
+#include "UObject/ObjectMacros.h"
+
+#include "InterchangeAnimationTrackSetFactory.generated.h"
+
+UCLASS(BlueprintType, Experimental)
+class INTERCHANGEIMPORT_API UInterchangeAnimationTrackSetFactory : public UInterchangeFactoryBase
+{
+	GENERATED_BODY()
+public:
+
+	//////////////////////////////////////////////////////////////////////////
+	// Interchange factory base interface begin
+
+	virtual UClass* GetFactoryClass() const override;
+	virtual UObject* CreateEmptyAsset(const FCreateAssetParams& Arguments) override;
+	virtual UObject* CreateAsset(const FCreateAssetParams& Arguments) override;
+	virtual void PreImportPreCompletedCallback(const FImportPreCompletedCallbackParams& Arguments) override;
+	virtual bool CanExecuteOnAnyThread() const override
+	{
+		//Currently we cannot use the anim sequence controller outside of the game thread
+		return false;
+	}
+
+	// Interchange factory base interface end
+	//////////////////////////////////////////////////////////////////////////
+
+private:
+	const UInterchangeTranslatorBase* Translator = nullptr;
+};
+
+
diff --git a/Engine/Plugins/Experimental/Interchange/Runtime/Source/Nodes/Private/InterchangeAnimationTrackSetNode.cpp b/Engine/Plugins/Experimental/Interchange/Runtime/Source/Nodes/Private/InterchangeAnimationTrackSetNode.cpp
new file mode 100644
--- /dev/null
+++ b/Engine/Plugins/Experimental/Interchange/Runtime/Source/Nodes/Private/InterchangeAnimationTrackSetNode.cpp
@@ -0,0 +1,154 @@
+// Copyright Epic Games, Inc. All Rights Reserved.
+
+#include "InterchangeAnimationTrackSetNode.h"
+
+// UInterchangeAnimationTrackSetNode
+
+UInterchangeAnimationTrackSetNode::UInterchangeAnimationTrackSetNode()
+{
+	CustomAnimationTrackUids.Initialize(Attributes, TEXT("AnimationTrackUids"));
+}
+
+int32 UInterchangeAnimationTrackSetNode::GetCustomAnimationTrackUidCount() const
+{
+	return CustomAnimationTrackUids.GetCount();
+}
+
+void UInterchangeAnimationTrackSetNode::GetCustomAnimationTrackUids(TArray<FString>& OutAnimationTrackUids) const
+{
+	CustomAnimationTrackUids.GetItems(OutAnimationTrackUids);
+}
+
+void UInterchangeAnimationTrackSetNode::GetCustomAnimationTrackUid(const int32 Index, FString& OutAnimationTrackUid) const
+{
+	CustomAnimationTrackUids.GetItem(Index, OutAnimationTrackUid);
+}
+
+bool UInterchangeAnimationTrackSetNode::AddCustomAnimationTrackUid(const FString& AnimationTrackUid)
+{
+	return CustomAnimationTrackUids.AddItem(AnimationTrackUid);
+}
+
+bool UInterchangeAnimationTrackSetNode::RemoveCustomAnimationTrackUid(const FString& AnimationTrackUid)
+{
+	return CustomAnimationTrackUids.RemoveItem(AnimationTrackUid);
+}
+
+bool UInterchangeAnimationTrackSetNode::GetCustomFrameRate(float& AttributeValue) const
+{
+	IMPLEMENT_NODE_ATTRIBUTE_GETTER(FrameRate, float);
+}
+
+bool UInterchangeAnimationTrackSetNode::SetCustomFrameRate(const float& AttributeValue)
+{
+	IMPLEMENT_NODE_ATTRIBUTE_SETTER_NODELEGATE(FrameRate, float);
+}
+
+// UInterchangeAnimationTrackBaseNode
+
+bool UInterchangeAnimationTrackBaseNode::GetCustomCompletionMode(int32& AttributeValue) const
+{
+	IMPLEMENT_NODE_ATTRIBUTE_GETTER(CompletionMode, int32);
+	AttributeValue = FMath::Clamp(AttributeValue, (int32)EInterchangeAimationCompletionMode::KeepState, (int32)EInterchangeAimationCompletionMode::ProjectDefault);
+}
+
+bool UInterchangeAnimationTrackBaseNode::SetCustomCompletionMode(const int32& AttributeValue)
+{
+	IMPLEMENT_NODE_ATTRIBUTE_SETTER_NODELEGATE(CompletionMode, int32);
+}
+
+// UInterchangeAnimationTrackNode
+
+bool UInterchangeAnimationTrackNode::GetCustomActorDependencyUid(FString& AttributeValue) const
+{
+	IMPLEMENT_NODE_ATTRIBUTE_GETTER(ActorDependency, FString);
+}
+
+bool UInterchangeAnimationTrackNode::SetCustomActorDependencyUid(const FString& AttributeValue)
+{
+	IMPLEMENT_NODE_ATTRIBUTE_SETTER_NODELEGATE(ActorDependency, FString);
+}
+
+bool UInterchangeAnimationTrackNode::GetCustomAnimationPayloadKey(FString& AttributeValue) const
+{
+	IMPLEMENT_NODE_ATTRIBUTE_GETTER(AnimationPayload, FString);
+}
+
+bool UInterchangeAnimationTrackNode::SetCustomAnimationPayloadKey(const FString& AttributeValue)
+{
+	IMPLEMENT_NODE_ATTRIBUTE_SETTER_NODELEGATE(AnimationPayload, FString);
+}
+
+bool UInterchangeAnimationTrackNode::GetCustomFrameCount(int32& AttributeValue) const
+{
+	IMPLEMENT_NODE_ATTRIBUTE_GETTER(FrameCount, int32);
+}
+
+bool UInterchangeAnimationTrackNode::SetCustomFrameCount(const int32& AttributeValue)
+{
+	IMPLEMENT_NODE_ATTRIBUTE_SETTER_NODELEGATE(FrameCount, int32);
+}
+
+bool UInterchangeAnimationTrackNode::GetCustomTargetedProperty(int32& AttributeValue) const
+{
+	IMPLEMENT_NODE_ATTRIBUTE_GETTER(TargetedProperty, int32);
+}
+
+bool UInterchangeAnimationTrackNode::SetCustomTargetedProperty(const int32& AttributeValue)
+{
+	IMPLEMENT_NODE_ATTRIBUTE_SETTER_NODELEGATE(TargetedProperty, int32);
+}
+
+// UInterchangeAnimationTrackSetInstanceNode
+
+bool UInterchangeAnimationTrackSetInstanceNode::SetCustomTimeScale(const float& AttributeValue)
+{
+	IMPLEMENT_NODE_ATTRIBUTE_SETTER_NODELEGATE(TimeScale, float);
+}
+
+bool UInterchangeAnimationTrackSetInstanceNode::GetCustomTimeScale(float& AttributeValue) const
+{
+	IMPLEMENT_NODE_ATTRIBUTE_GETTER(TimeScale, float);
+}
+
+bool UInterchangeAnimationTrackSetInstanceNode::SetCustomDuration(const int32& AttributeValue)
+{
+	IMPLEMENT_NODE_ATTRIBUTE_SETTER_NODELEGATE(Duration, int32);
+}
+
+bool UInterchangeAnimationTrackSetInstanceNode::GetCustomDuration(int32& AttributeValue) const
+{
+	IMPLEMENT_NODE_ATTRIBUTE_GETTER(Duration, int32);
+}
+
+bool UInterchangeAnimationTrackSetInstanceNode::SetCustomStartFrame(const int32& AttributeValue)
+{
+	IMPLEMENT_NODE_ATTRIBUTE_SETTER_NODELEGATE(StartFrame, int32);
+}
+
+bool UInterchangeAnimationTrackSetInstanceNode::GetCustomStartFrame(int32& AttributeValue) const
+{
+	IMPLEMENT_NODE_ATTRIBUTE_GETTER(StartFrame, int32);
+}
+
+bool UInterchangeAnimationTrackSetInstanceNode::SetCustomTrackSetDependencyUid(const FString& AttributeValue)
+{
+	IMPLEMENT_NODE_ATTRIBUTE_SETTER_NODELEGATE(TrackSetDependencyUid, FString);
+}
+
+bool UInterchangeAnimationTrackSetInstanceNode::GetCustomTrackSetDependencyUid(FString& AttributeValue) const
+{
+	IMPLEMENT_NODE_ATTRIBUTE_GETTER(TrackSetDependencyUid, FString);
+}
+
+// UInterchangeTransformAnimationTrackNode
+
+bool UInterchangeTransformAnimationTrackNode::SetCustomUsedChannels(const int32& AttributeValue)
+{
+	IMPLEMENT_NODE_ATTRIBUTE_SETTER_NODELEGATE(UsedChannels, int32);
+}
+
+bool UInterchangeTransformAnimationTrackNode::GetCustomUsedChannels(int32& AttributeValue) const
+{
+	IMPLEMENT_NODE_ATTRIBUTE_GETTER(UsedChannels, int32);
+}
diff --git a/Engine/Plugins/Experimental/Interchange/Runtime/Source/Nodes/Public/InterchangeAnimationTrackSetNode.h b/Engine/Plugins/Experimental/Interchange/Runtime/Source/Nodes/Public/InterchangeAnimationTrackSetNode.h
new file mode 100644
--- /dev/null
+++ b/Engine/Plugins/Experimental/Interchange/Runtime/Source/Nodes/Public/InterchangeAnimationTrackSetNode.h
@@ -0,0 +1,355 @@
+// Copyright Epic Games, Inc. All Rights Reserved.
+
+#pragma once
+
+#include "CoreMinimal.h"
+#include "Nodes/InterchangeBaseNode.h"
+
+#include "InterchangeAnimationTrackSetNode.generated.h"
+
+/**
+ * Enumeration specifying which properties of a camera, light or scene node can be animated besides transform.
+ */
+UENUM(BlueprintType)
+enum class EInterchangeAnimatedProperty : uint8
+{
+	None UMETA(DisplayName = "No property.", ToolTip = "The associated animation track will be ignored."),
+	Visibility UMETA(DisplayName = "Visibility property.", ToolTip = "The associated animation track is applied to the visibility property of the actor"),
+	MAX,
+};
+/**
+ * Enumeration specifying how to handle the state of the animated property at the end of an animation track
+ */
+enum class EInterchangeAimationCompletionMode : uint8
+{
+	KeepState UMETA(DisplayName = "Keep State", ToolTip = "Keep the animated property at the state set at the end of the animation track."),
+	RestoreState UMETA(DisplayName = "Restore State", ToolTip = "Restore the animated property to its state before the start of the animation track."),
+	ProjectDefault UMETA(DisplayName = "Project Default", ToolTip = "Restore the animated property to the state set in the project for such property."),
+};
+
+/**
+ * Class to represent a set of animation track nodes sharing the same frame rate
+ */
+UCLASS(BlueprintType, Experimental)
+class INTERCHANGENODES_API UInterchangeAnimationTrackSetNode : public UInterchangeBaseNode
+{
+	GENERATED_BODY()
+
+	UInterchangeAnimationTrackSetNode();
+
+public:
+	static FStringView StaticAssetTypeName()
+	{
+		return TEXT("AnimationTrackSet");
+	}
+
+	/**
+	 * Return the node type name of the class, we use this when reporting errors
+	 */
+	virtual FString GetTypeName() const override
+	{
+		const FString TypeName = TEXT("AnimationTrackSetNode");
+		return TypeName;
+	}
+
+	/**
+	 * This function allow to retrieve the number of track dependencies for this object.
+	 */
+	UFUNCTION(BlueprintCallable, Category = "Interchange | Node | AnimationTrackSet")
+	int32 GetCustomAnimationTrackUidCount() const;
+
+	/**
+	 * This function allow to retrieve the track dependency for this object.
+	 */
+	UFUNCTION(BlueprintCallable, Category = "Interchange | Node | AnimationTrackSet")
+	void GetCustomAnimationTrackUids(TArray<FString>& OutAnimationTrackUids) const;
+
+	/**
+	 * This function allow to retrieve one track dependency for this object.
+	 */
+	UFUNCTION(BlueprintCallable, Category = "Interchange | Node | AnimationTrackSet")
+	void GetCustomAnimationTrackUid(const int32 Index, FString& OutAnimationTrackUid) const;
+
+	/**
+	 * Add one track dependency to this object.
+	 */
+	UFUNCTION(BlueprintCallable, Category = "Interchange | Node | AnimationTrackSet")
+	bool AddCustomAnimationTrackUid(const FString& AnimationTrackUid);
+
+	/**
+	 * Remove one track dependency from this object.
+	 */
+	UFUNCTION(BlueprintCallable, Category = "Interchange | Node | AnimationTrackSet")
+	bool RemoveCustomAnimationTrackUid(const FString& AnimationTrackUid);
+
+	/**
+	 * Set the frame rate for the animations in the level sequence.
+	 */
+	UFUNCTION(BlueprintCallable, Category = "Interchange | Node | AnimationTrackSet")
+	bool SetCustomFrameRate(const float& AttributeValue);
+
+	/**
+	 * Get the frame rate for the animations in the level sequence.
+	 */
+	UFUNCTION(BlueprintCallable, Category = "Interchange | Node | AnimationTrackSet")
+	bool GetCustomFrameRate(float& AttributeValue) const;
+
+private:
+	const UE::Interchange::FAttributeKey Macro_CustomFrameRateKey = UE::Interchange::FAttributeKey(TEXT("FrameRate"));
+
+	UE::Interchange::TArrayAttributeHelper<FString> CustomAnimationTrackUids;
+};
+
+/**
+ * Abstract class providing the minimal services required for an animation track node
+ */
+UCLASS(Abstract, Experimental)
+class INTERCHANGENODES_API UInterchangeAnimationTrackBaseNode : public UInterchangeBaseNode
+{
+	GENERATED_BODY()
+
+public:
+	static FStringView StaticAssetTypeName()
+	{
+		return TEXT("AnimationTrackBaseNode");
+	}
+
+	/**
+	* Return the node type name of the class, we use this when reporting errors
+	*/
+	virtual FString GetTypeName() const override
+	{
+		const FString TypeName = TEXT("AnimationTrackBaseNode");
+		return TypeName;
+	}
+
+	/**
+	 * Set how the actor's animated property should behave once its animation completes.
+	 */
+	UFUNCTION(BlueprintCallable, Category = "Interchange | Node | AnimationTrack")
+	bool SetCustomCompletionMode(const int32& AttributeValue);
+
+	/**
+	 * Get how the actor's animated property behaves once this animation is complete.
+	 * The output value will be clamped to the range of values defined in EInterchangeAimationCompletionMode
+	 */
+	UFUNCTION(BlueprintCallable, Category = "Interchange | Node | AnimationTrack")
+	bool GetCustomCompletionMode(int32& AttributeValue) const;
+
+private:
+	const UE::Interchange::FAttributeKey Macro_CustomCompletionModeKey = UE::Interchange::FAttributeKey(TEXT("CompletionMode"));
+};
+
+/**
+ * Class to represent an animation which instances another animation track set node
+ */
+UCLASS(BlueprintType, Experimental)
+class INTERCHANGENODES_API UInterchangeAnimationTrackSetInstanceNode : public UInterchangeAnimationTrackBaseNode
+{
+	GENERATED_BODY()
+
+public:
+	static FStringView StaticAssetTypeName()
+	{
+		return TEXT("AnimationTrackSetInstanceNode");
+	}
+
+	/**
+	 * Return the node type name of the class, we use this when reporting errors
+	 */
+	virtual FString GetTypeName() const override
+	{
+		const FString TypeName = TEXT("AnimationTrackSetInstanceNode");
+		return TypeName;
+	}
+
+	/**
+	 * Set the time scale used for the level sequence instance.
+	 */
+	UFUNCTION(BlueprintCallable, Category = "Interchange | Node | AnimationTrackSetInstance")
+	bool SetCustomTimeScale(const float& AttributeValue);
+
+	/**
+	 * Get the time scale used for the level sequence instance.
+	 */
+	UFUNCTION(BlueprintCallable, Category = "Interchange | Node | AnimationTrackSetInstance")
+	bool GetCustomTimeScale(float& AttributeValue) const;
+
+	/**
+	 * Set the level sequence instance duration in number of frames.
+	 */
+	UFUNCTION(BlueprintCallable, Category = "Interchange | Node | AnimationTrackSetInstance")
+	bool SetCustomDuration(const int32& AttributeValue);
+
+	/**
+	 * Get the level sequence instance duration in number of frames.
+	 */
+	UFUNCTION(BlueprintCallable, Category = "Interchange | Node | AnimationTrackSetInstance")
+	bool GetCustomDuration(int32& AttributeValue) const;
+
+	/**
+	 * Set the frame where the level sequence instance starts.
+	 */
+	UFUNCTION(BlueprintCallable, Category = "Interchange | Node | AnimationTrackSetInstance")
+	bool SetCustomStartFrame(const int32& AttributeValue);
+
+	/**
+	 * Get the frame where the level sequence instance starts.
+	 */
+	UFUNCTION(BlueprintCallable, Category = "Interchange | Node | AnimationTrackSetInstance")
+	bool GetCustomStartFrame(int32& AttributeValue) const;
+
+	/**
+	 * Set the unique id of the level sequence this instance references.
+	 */
+	UFUNCTION(BlueprintCallable, Category = "Interchange | Node | AnimationTrackSetInstance")
+	bool SetCustomTrackSetDependencyUid(const FString& AttributeValue);
+
+	/**
+	 * Get the unique id of the level sequence this instance references.
+	 */
+	UFUNCTION(BlueprintCallable, Category = "Interchange | Node | AnimationTrackSetInstance")
+	bool GetCustomTrackSetDependencyUid(FString& AttributeValue) const;
+
+private:
+	const UE::Interchange::FAttributeKey Macro_CustomStartFrameKey = UE::Interchange::FAttributeKey(TEXT("StartFrame"));
+	const UE::Interchange::FAttributeKey Macro_CustomDurationKey = UE::Interchange::FAttributeKey(TEXT("Duration"));
+	const UE::Interchange::FAttributeKey Macro_CustomTimeScaleKey = UE::Interchange::FAttributeKey(TEXT("TimeScale"));
+	const UE::Interchange::FAttributeKey Macro_CustomTrackSetDependencyUidKey = UE::Interchange::FAttributeKey(TEXT("SequenceDependencyUid"));
+};
+
+/**
+ * Class to represent an animation on the property of a camera, light or scene node
+ * The list of supported properties is enumerated in EInterchangeAnimatedProperty
+ */
+UCLASS(BlueprintType, Experimental)
+class INTERCHANGENODES_API UInterchangeAnimationTrackNode : public UInterchangeAnimationTrackBaseNode
+{
+	GENERATED_BODY()
+
+public:
+	static FStringView StaticAssetTypeName()
+	{
+		return TEXT("AnimationTrack");
+	}
+
+	/**
+	 * Return the node type name of the class, we use this when reporting errors
+	 */
+	virtual FString GetTypeName() const override
+	{
+		const FString TypeName = TEXT("AnimationTrackNode");
+		return TypeName;
+	}
+
+	/**
+	 * Set the actor dependency to this object.
+	 */
+	UFUNCTION(BlueprintCallable, Category = "Interchange | Node | AnimationTrack")
+	bool SetCustomActorDependencyUid(const FString& DependencyUid);
+
+	/**
+	 * Get the actor dependency to this object.
+	 */
+	UFUNCTION(BlueprintCallable, Category = "Interchange | Node | AnimationTrack")
+	bool GetCustomActorDependencyUid(FString& DependencyUid) const;
+
+	/**
+	 * Set the property animated by this track.
+	 */
+	UFUNCTION(BlueprintCallable, Category = "Interchange | Node | AnimationTrack")
+	bool SetCustomTargetedProperty(const int32& TargetedProperty);
+
+	/**
+	 * Get the property animated by this track.
+	 */
+	UFUNCTION(BlueprintCallable, Category = "Interchange | Node | AnimationTrack")
+	bool GetCustomTargetedProperty(int32& TargetedProperty) const;
+
+	/**
+	 * Set the payload key needed to retrieve the animation for this track.
+	 */
+	UFUNCTION(BlueprintCallable, Category = "Interchange | Node | AnimationTrack")
+	bool SetCustomAnimationPayloadKey(const FString& PayloadKey);
+
+	/**
+	 * Get the payload key needed to retrieve the animation for this track.
+	 */
+	UFUNCTION(BlueprintCallable, Category = "Interchange | Node | AnimationTrack")
+	bool GetCustomAnimationPayloadKey(FString& PayloadKey) const;
+
+	/**
+	 * Set the number of frames for the animation of this track.
+	 */
+	UFUNCTION(BlueprintCallable, Category = "Interchange | Node | AnimationTrack")
+	bool SetCustomFrameCount(const int32& AttributeValue);
+
+	/**
+	 * Get the number of frames for the animation of this track.
+	 */
+	UFUNCTION(BlueprintCallable, Category = "Interchange | Node | AnimationTrack")
+	bool GetCustomFrameCount(int32& AttributeValue) const;
+
+private:
+	const UE::Interchange::FAttributeKey Macro_CustomActorDependencyKey = UE::Interchange::FAttributeKey(TEXT("ActorDependency"));
+	const UE::Interchange::FAttributeKey Macro_CustomAnimationPayloadKey = UE::Interchange::FAttributeKey(TEXT("AnimationPayload"));
+	const UE::Interchange::FAttributeKey Macro_CustomFrameCountKey = UE::Interchange::FAttributeKey(TEXT("FrameCount"));
+	const UE::Interchange::FAttributeKey Macro_CustomTargetedPropertyKey = UE::Interchange::FAttributeKey(TEXT("TargetedProperty"));
+};
+
+/**
+ * Class to represent an animation on the transform of a camera, light or scene node
+ */
+UCLASS(BlueprintType, Experimental)
+class INTERCHANGENODES_API UInterchangeTransformAnimationTrackNode : public UInterchangeAnimationTrackNode
+{
+	GENERATED_BODY()
+
+public:
+	static FStringView StaticAssetTypeName()
+	{
+		return TEXT("TransformAnimationTrack");
+	}
+
+	/**
+	 * Return the node type name of the class, we use this when reporting errors
+	 */
+	virtual FString GetTypeName() const override
+	{
+		const FString TypeName = TEXT("TransformAnimationTrackNode");
+		return TypeName;
+	}
+
+	/**
+	 * Set which channels of this animation should be used. This is a bitwise mask.
+	 * Bits are interpreted as follow:
+	 *    None          = 0x000,
+	 *    TranslationX  = 0x001,
+	 *    TranslationY  = 0x002,
+	 *    TranslationZ  = 0x004,
+	 *    Translation   = TranslationX | TranslationY | TranslationZ,
+	 *    RotationX     = 0x008,
+	 *    RotationY     = 0x010,
+	 *    RotationZ     = 0x020,
+	 *    Rotation      = RotationX | RotationY | RotationZ,
+	 *    ScaleX        = 0x040,
+	 *    ScaleY        = 0x080,
+	 *    ScaleZ        = 0x100,
+	 *    Scale         = ScaleX | ScaleY | ScaleZ,
+	 *    AllTransform  = Translation | Rotation | Scale,
+	 *    Weight        = 0x200,
+	 *    All           = Translation | Rotation | Scale | Weight,
+	 */
+	UFUNCTION(BlueprintCallable, Category = "Interchange | Node | TransformAnimationTrack")
+	bool SetCustomUsedChannels(const int32& AttributeValue);
+
+	/**
+	 * Get which channels of this animation should be used. This is a bitmask.
+	 * See SetCustomUsedChannels for description of bitmask
+	 */
+	UFUNCTION(BlueprintCallable, Category = "Interchange | Node | TransformAnimationTrack")
+	bool GetCustomUsedChannels(int32& AttributeValue) const;
+
+private:
+	const UE::Interchange::FAttributeKey Macro_CustomUsedChannelsKey = UE::Interchange::FAttributeKey(TEXT("UsedChannels"));
+};
\ No newline at end of file
diff --git a/Engine/Plugins/Experimental/Interchange/Runtime/Source/Parsers/CommonParser/Public/InterchangeCommonAnimationPayload.h b/Engine/Plugins/Experimental/Interchange/Runtime/Source/Parsers/CommonParser/Public/InterchangeCommonAnimationPayload.h
--- a/Engine/Plugins/Experimental/Interchange/Runtime/Source/Parsers/CommonParser/Public/InterchangeCommonAnimationPayload.h
+++ b/Engine/Plugins/Experimental/Interchange/Runtime/Source/Parsers/CommonParser/Public/InterchangeCommonAnimationPayload.h
@@ -168,6 +168,7 @@ struct INTERCHANGECOMMONPARSER_API FInterchangeStepCurve
 	UPROPERTY()
 	TArray<float> KeyTimes;
 	TOptional<TArray<float>> FloatKeyValues;
+	TOptional<TArray<bool>> BooleanKeyValues;
 	TOptional<TArray<int32>> IntegerKeyValues;
 	TOptional<TArray<FString>> StringKeyValues;
 
diff --git a/Engine/Plugins/Experimental/Interchange/Runtime/Source/Pipelines/Private/InterchangeGenericAnimationPipeline.cpp b/Engine/Plugins/Experimental/Interchange/Runtime/Source/Pipelines/Private/InterchangeGenericAnimationPipeline.cpp
--- a/Engine/Plugins/Experimental/Interchange/Runtime/Source/Pipelines/Private/InterchangeGenericAnimationPipeline.cpp
+++ b/Engine/Plugins/Experimental/Interchange/Runtime/Source/Pipelines/Private/InterchangeGenericAnimationPipeline.cpp
@@ -3,6 +3,8 @@
 
 #include "Animation/AnimationSettings.h"
 #include "CoreMinimal.h"
+#include "InterchangeAnimationTrackSetFactoryNode.h"
+#include "InterchangeAnimationTrackSetNode.h"
 #include "InterchangeAnimSequenceFactoryNode.h"
 #include "InterchangeMeshNode.h"
 #include "InterchangePipelineLog.h"
@@ -18,7 +20,6 @@
 #include "Nodes/InterchangeSourceNode.h"
 #include "Nodes/InterchangeUserDefinedAttribute.h"
 
-
 namespace UE::Interchange::Private
 {
 	/**
@@ -84,7 +85,24 @@ void UInterchangeGenericAnimationPipeline::ExecutePreImportPipeline(UInterchange
 		return;
 	}
 
-	check(!CommonSkeletalMeshesAndAnimationsProperties.IsNull());
+	TArray<UInterchangeAnimationTrackSetNode*> TrackSetNodes;
+	BaseNodeContainer->IterateNodesOfType<UInterchangeAnimationTrackSetNode>([&](const FString& NodeUid, UInterchangeAnimationTrackSetNode* Node)
+		{
+			TrackSetNodes.Add(Node);
+		});
+
+	for (UInterchangeAnimationTrackSetNode* TrackSetNode : TrackSetNodes)
+	{
+		if (TrackSetNode)
+		{
+			CreateAnimationTrackSetFactoryNode(*TrackSetNode);
+		}
+	}
+
+	if (CommonSkeletalMeshesAndAnimationsProperties.IsNull())
+	{
+		return;
+	}
 
 	if (CommonSkeletalMeshesAndAnimationsProperties->bImportOnlyAnimations && CommonSkeletalMeshesAndAnimationsProperties->Skeleton.IsNull())
 	{
@@ -474,21 +492,58 @@ void UInterchangeGenericAnimationPipeline::ExecutePreImportPipeline(UInterchange
 	}
 }
 
-void UInterchangeGenericAnimationPipeline::ExecutePostImportPipeline(const UInterchangeBaseNodeContainer* InBaseNodeContainer, const FString& NodeKey, UObject* CreatedAsset, bool bIsAReimport)
+void UInterchangeGenericAnimationPipeline::CreateAnimationTrackSetFactoryNode(UInterchangeAnimationTrackSetNode& TranslatedNode)
 {
-	//We do not use the provided base container since ExecutePreImportPipeline cache it
-	//We just make sure the same one is pass in parameter
-	if (!InBaseNodeContainer || !ensure(BaseNodeContainer == InBaseNodeContainer) || !CreatedAsset)
+	const FString FactoryNodeUid = UInterchangeFactoryBaseNode::BuildFactoryNodeUid(TranslatedNode.GetUniqueID());
+
+	UInterchangeAnimationTrackSetFactoryNode* FactoryNode = NewObject<UInterchangeAnimationTrackSetFactoryNode>(BaseNodeContainer, NAME_None);
+
+	FactoryNode->InitializeNode(FactoryNodeUid, TranslatedNode.GetDisplayLabel(), EInterchangeNodeContainerType::FactoryData);
+	FactoryNode->SetEnabled(true);
+
+	TArray<FString> AnimationTrackUids;
+	TranslatedNode.GetCustomAnimationTrackUids(AnimationTrackUids);
+
+	for (const FString& AnimationTrackUid : AnimationTrackUids)
 	{
-		return;
+		FactoryNode->AddCustomAnimationTrackUid(AnimationTrackUid);
+
+		// Update factory's dependencies
+		if (const UInterchangeAnimationTrackBaseNode* TrackNode = Cast<UInterchangeAnimationTrackBaseNode>(BaseNodeContainer->GetNode(AnimationTrackUid)))
+		{
+			if (const UInterchangeTransformAnimationTrackNode* TransformTrackNode = Cast<UInterchangeTransformAnimationTrackNode>(TrackNode))
+			{
+				FString ActorNodeUid;
+				if (TransformTrackNode->GetCustomActorDependencyUid(ActorNodeUid))
+				{
+					const FString ActorFactoryNodeUid = TEXT("Factory_") + ActorNodeUid;
+					FactoryNode->AddFactoryDependencyUid(ActorFactoryNodeUid);
+				}
+			}
+			else if (const UInterchangeAnimationTrackSetInstanceNode* InstanceTrackNode = Cast<UInterchangeAnimationTrackSetInstanceNode>(TrackNode))
+			{
+				FString TrackSetNodeUid;
+				if (InstanceTrackNode->GetCustomTrackSetDependencyUid(TrackSetNodeUid))
+				{
+					const FString TrackSetFactoryNodeUid = TEXT("Factory_") + TrackSetNodeUid;
+					FactoryNode->AddFactoryDependencyUid(TrackSetFactoryNodeUid);
+				}
+
+			}
+		}
 	}
 
-	const UInterchangeBaseNode* Node = BaseNodeContainer->GetNode(NodeKey);
-	if (!Node)
+	float FrameRate;
+	if (TranslatedNode.GetCustomFrameRate(FrameRate))
 	{
-		return;
+		FactoryNode->SetCustomFrameRate(FrameRate);
 	}
-}
 
+	UInterchangeUserDefinedAttributesAPI::DuplicateAllUserDefinedAttribute(&TranslatedNode, FactoryNode, false);
 
+	FactoryNode->AddTargetNodeUid(TranslatedNode.GetUniqueID());
+	TranslatedNode.AddTargetNodeUid(FactoryNode->GetUniqueID());
+
+	BaseNodeContainer->AddNode(FactoryNode);
+}
 
diff --git a/Engine/Plugins/Experimental/Interchange/Runtime/Source/Pipelines/Private/InterchangeGenericMaterialPipeline.cpp b/Engine/Plugins/Experimental/Interchange/Runtime/Source/Pipelines/Private/InterchangeGenericMaterialPipeline.cpp
--- a/Engine/Plugins/Experimental/Interchange/Runtime/Source/Pipelines/Private/InterchangeGenericMaterialPipeline.cpp
+++ b/Engine/Plugins/Experimental/Interchange/Runtime/Source/Pipelines/Private/InterchangeGenericMaterialPipeline.cpp
@@ -139,7 +139,7 @@ void UInterchangeGenericMaterialPipeline::ExecutePreImportPipeline(UInterchangeB
 UInterchangeBaseMaterialFactoryNode* UInterchangeGenericMaterialPipeline::CreateBaseMaterialFactoryNode(const UInterchangeBaseNode* MaterialNode, TSubclassOf<UInterchangeBaseMaterialFactoryNode> NodeType)
 {
 	FString DisplayLabel = MaterialNode->GetDisplayLabel();
-	const FString NodeUid = UInterchangeMaterialFactoryNode::GetMaterialFactoryNodeUidFromMaterialNodeUid(MaterialNode->GetUniqueID());
+	const FString NodeUid = UInterchangeFactoryBaseNode::BuildFactoryNodeUid(MaterialNode->GetUniqueID());
 	UInterchangeBaseMaterialFactoryNode* MaterialFactoryNode = nullptr;
 	if (BaseNodeContainer->IsNodeUidValid(NodeUid))
 	{
@@ -1242,7 +1242,7 @@ UInterchangeMaterialExpressionFactoryNode* UInterchangeGenericMaterialPipeline::
 	// If we don't recognize the shader node type
 	// - Create material expression by trying to match the node type to a material expression class name
 
-	const FString MaterialExpressionUid = TEXT("Factory_") + ShaderNode->GetUniqueID();
+	const FString MaterialExpressionUid = UInterchangeFactoryBaseNode::BuildFactoryNodeUid(ShaderNode->GetUniqueID());
 
 	UInterchangeMaterialExpressionFactoryNode* MaterialExpression = Cast<UInterchangeMaterialExpressionFactoryNode>(BaseNodeContainer->GetFactoryNode(MaterialExpressionUid));
 	if (MaterialExpression != nullptr)
@@ -1257,7 +1257,7 @@ UInterchangeMaterialExpressionFactoryNode* UInterchangeGenericMaterialPipeline::
 		FString MaterialFunctionUid;
 		if (FunctionCallShaderNode->GetCustomMaterialFunction(MaterialFunctionUid))
 		{
-			const FString MaterialFunctionFactoryNodeUid = UInterchangeMaterialFactoryNode::GetMaterialFactoryNodeUidFromMaterialNodeUid(MaterialFunctionUid);
+			const FString MaterialFunctionFactoryNodeUid = UInterchangeFactoryBaseNode::BuildFactoryNodeUid(MaterialFunctionUid);
 			MaterialFunctionCallExpression->SetCustomMaterialFunctionDependency(MaterialFunctionFactoryNodeUid);
 			MaterialFunctionCallExpression->AddFactoryDependencyUid(MaterialFunctionFactoryNodeUid);
 		}
diff --git a/Engine/Plugins/Experimental/Interchange/Runtime/Source/Pipelines/Private/InterchangeGenericScenesPipeline.cpp b/Engine/Plugins/Experimental/Interchange/Runtime/Source/Pipelines/Private/InterchangeGenericScenesPipeline.cpp
--- a/Engine/Plugins/Experimental/Interchange/Runtime/Source/Pipelines/Private/InterchangeGenericScenesPipeline.cpp
+++ b/Engine/Plugins/Experimental/Interchange/Runtime/Source/Pipelines/Private/InterchangeGenericScenesPipeline.cpp
@@ -94,7 +94,7 @@ void UInterchangeGenericLevelPipeline::ExecuteSceneNodePreImport(UInterchangeBas
 		return;
 	}
 
-	ActorFactoryNode->InitializeNode(TEXT("Factory_") + SceneNode->GetUniqueID(), SceneNode->GetDisplayLabel(), EInterchangeNodeContainerType::FactoryData);
+	ActorFactoryNode->InitializeNode(UInterchangeFactoryBaseNode::BuildFactoryNodeUid(SceneNode->GetUniqueID()), SceneNode->GetDisplayLabel(), EInterchangeNodeContainerType::FactoryData);
 	const FString ActorFactoryNodeUid = FactoryNodeContainer->AddNode(ActorFactoryNode);
 	if (!SceneNode->GetParentUid().IsEmpty())
 	{
@@ -162,6 +162,8 @@ void UInterchangeGenericLevelPipeline::SetUpFactoryNode(UInterchangeActorFactory
 			SceneNode->GetSlotMaterialDependencies(SlotMaterialDependencies);
 
 			UE::Interchange::MeshesUtilities::ApplySlotMaterialDependencies(*MeshActorFactoryNode, SlotMaterialDependencies, *FactoryNodeContainer);
+
+			MeshActorFactoryNode->AddFactoryDependencyUid(UInterchangeFactoryBaseNode::BuildFactoryNodeUid(MeshNode->GetUniqueID()));
 		}
 	}
 	else if (const UInterchangeLightNode* LightNode = Cast<UInterchangeLightNode>(TranslatedAssetNode))
diff --git a/Engine/Plugins/Experimental/Interchange/Runtime/Source/Pipelines/Private/InterchangeGenericSkeletalMeshPipeline.cpp b/Engine/Plugins/Experimental/Interchange/Runtime/Source/Pipelines/Private/InterchangeGenericSkeletalMeshPipeline.cpp
--- a/Engine/Plugins/Experimental/Interchange/Runtime/Source/Pipelines/Private/InterchangeGenericSkeletalMeshPipeline.cpp
+++ b/Engine/Plugins/Experimental/Interchange/Runtime/Source/Pipelines/Private/InterchangeGenericSkeletalMeshPipeline.cpp
@@ -257,7 +257,7 @@ UInterchangeSkeletonFactoryNode* UInterchangeGenericMeshPipeline::CreateSkeleton
 	}
 
 	FString DisplayLabel = RootJointNode->GetDisplayLabel() + TEXT("_Skeleton");
-	FString SkeletonUid = TEXT("\\Skeleton\\") + RootJointNode->GetUniqueID();
+	FString SkeletonUid = UInterchangeFactoryBaseNode::BuildFactoryNodeUid(RootJointNode->GetUniqueID());
 
 	UInterchangeSkeletonFactoryNode* SkeletonFactoryNode = nullptr;
 	if (BaseNodeContainer->IsNodeUidValid(SkeletonUid))
@@ -314,7 +314,7 @@ UInterchangeSkeletalMeshFactoryNode* UInterchangeGenericMeshPipeline::CreateSkel
 	{
 		return nullptr;
 	}
-	const FString SkeletonUid = TEXT("\\Skeleton\\") + RootJointNode->GetUniqueID();
+	const FString SkeletonUid = UInterchangeFactoryBaseNode::BuildFactoryNodeUid(RootJointNode->GetUniqueID());
 	UInterchangeSkeletonFactoryNode* SkeletonFactoryNode = Cast<UInterchangeSkeletonFactoryNode>(BaseNodeContainer->GetFactoryNode(SkeletonUid));
 	if (!ensure(SkeletonFactoryNode))
 	{
@@ -393,7 +393,7 @@ UInterchangeSkeletalMeshFactoryNode* UInterchangeGenericMeshPipeline::CreateSkel
 		//Use the first scene node uid this skeletalmesh reference, add backslash since this uid is not asset typed (\\Mesh\\) like FirstMeshNodeUid
 		SkeletalMeshUid_MeshNamePart = TEXT("\\") + FirstSceneNode->GetUniqueID();
 	}
-	const FString SkeletalMeshUid = TEXT("\\SkeletalMesh") + SkeletalMeshUid_MeshNamePart + SkeletonUid;
+	const FString SkeletalMeshUid = UInterchangeFactoryBaseNode::BuildFactoryNodeUid(SkeletalMeshUid_MeshNamePart + SkeletonUid);
 	UInterchangeSkeletalMeshFactoryNode* SkeletalMeshFactoryNode = NewObject<UInterchangeSkeletalMeshFactoryNode>(BaseNodeContainer, NAME_None);
 	if (!ensure(SkeletalMeshFactoryNode))
 	{
@@ -438,7 +438,7 @@ UInterchangeSkeletalMeshFactoryNode* UInterchangeGenericMeshPipeline::CreateSkel
 		UInterchangePhysicsAssetFactoryNode* PhysicsAssetFactoryNode = NewObject<UInterchangePhysicsAssetFactoryNode>(BaseNodeContainer, NAME_None);
 		if (ensure(SkeletalMeshFactoryNode))
 		{
-			const FString PhysicsAssetUid = TEXT("\\PhysicsAsset") + SkeletalMeshUid_MeshNamePart + SkeletonUid;
+			const FString PhysicsAssetUid = UInterchangeFactoryBaseNode::BuildFactoryNodeUid(SkeletalMeshUid_MeshNamePart + SkeletonUid);
 			const FString PhysicsAssetDisplayLabel = DisplayLabel + TEXT("_PhysicsAsset");
 			PhysicsAssetFactoryNode->InitializePhysicsAssetNode(PhysicsAssetUid, PhysicsAssetDisplayLabel, UPhysicsAsset::StaticClass()->GetName());
 			PhysicsAssetFactoryNode->SetCustomSkeletalMeshUid(SkeletalMeshUid);
diff --git a/Engine/Plugins/Experimental/Interchange/Runtime/Source/Pipelines/Private/InterchangeGenericStaticMeshPipeline.cpp b/Engine/Plugins/Experimental/Interchange/Runtime/Source/Pipelines/Private/InterchangeGenericStaticMeshPipeline.cpp
--- a/Engine/Plugins/Experimental/Interchange/Runtime/Source/Pipelines/Private/InterchangeGenericStaticMeshPipeline.cpp
+++ b/Engine/Plugins/Experimental/Interchange/Runtime/Source/Pipelines/Private/InterchangeGenericStaticMeshPipeline.cpp
@@ -385,7 +385,7 @@ UInterchangeStaticMeshFactoryNode* UInterchangeGenericMeshPipeline::CreateStatic
 		return nullptr;
 	}
 
-	const FString StaticMeshUid = TEXT("\\StaticMesh") + StaticMeshUid_MeshNamePart;
+	const FString StaticMeshUid = UInterchangeFactoryBaseNode::BuildFactoryNodeUid(StaticMeshUid_MeshNamePart);
 	UInterchangeStaticMeshFactoryNode* StaticMeshFactoryNode = NewObject<UInterchangeStaticMeshFactoryNode>(BaseNodeContainer, NAME_None);
 	if (!ensure(StaticMeshFactoryNode))
 	{
diff --git a/Engine/Plugins/Experimental/Interchange/Runtime/Source/Pipelines/Private/InterchangeGenericTexturePipeline.cpp b/Engine/Plugins/Experimental/Interchange/Runtime/Source/Pipelines/Private/InterchangeGenericTexturePipeline.cpp
--- a/Engine/Plugins/Experimental/Interchange/Runtime/Source/Pipelines/Private/InterchangeGenericTexturePipeline.cpp
+++ b/Engine/Plugins/Experimental/Interchange/Runtime/Source/Pipelines/Private/InterchangeGenericTexturePipeline.cpp
@@ -232,7 +232,7 @@ UInterchangeTextureFactoryNode* UInterchangeGenericTexturePipeline::HandleCreati
 UInterchangeTextureFactoryNode* UInterchangeGenericTexturePipeline::CreateTextureFactoryNode(const UInterchangeTextureNode* TextureNode, const TSubclassOf<UInterchangeTextureFactoryNode>& FactorySubclass)
 {
 	FString DisplayLabel = TextureNode->GetDisplayLabel();
-	FString NodeUid = UInterchangeTextureFactoryNode::GetTextureFactoryNodeUidFromTextureNodeUid(TextureNode->GetUniqueID());
+	FString NodeUid = UInterchangeFactoryBaseNode::BuildFactoryNodeUid(TextureNode->GetUniqueID());
 	UInterchangeTextureFactoryNode* TextureFactoryNode = nullptr;
 	if (BaseNodeContainer->IsNodeUidValid(NodeUid))
 	{
diff --git a/Engine/Plugins/Experimental/Interchange/Runtime/Source/Pipelines/Public/InterchangeGenericAnimationPipeline.h b/Engine/Plugins/Experimental/Interchange/Runtime/Source/Pipelines/Public/InterchangeGenericAnimationPipeline.h
--- a/Engine/Plugins/Experimental/Interchange/Runtime/Source/Pipelines/Public/InterchangeGenericAnimationPipeline.h
+++ b/Engine/Plugins/Experimental/Interchange/Runtime/Source/Pipelines/Public/InterchangeGenericAnimationPipeline.h
@@ -12,6 +12,7 @@
 
 #include "InterchangeGenericAnimationPipeline.generated.h"
 
+class UInterchangeAnimationTrackSetNode;
 class UInterchangeMeshNode;
 
 /** Animation length type when importing */
@@ -106,8 +107,6 @@ protected:
 
 	virtual void ExecutePreImportPipeline(UInterchangeBaseNodeContainer* InBaseNodeContainer, const TArray<UInterchangeSourceData*>& InSourceDatas) override;
 
-	virtual void ExecutePostImportPipeline(const UInterchangeBaseNodeContainer* InBaseNodeContainer, const FString& NodeKey, UObject* CreatedAsset, bool bIsAReimport) override;
-
 	virtual bool CanExecuteOnAnyThread(EInterchangePipelineTask PipelineTask) override
 	{
 		return true;
@@ -115,6 +114,8 @@ protected:
 
 private:
 
+	void CreateAnimationTrackSetFactoryNode(UInterchangeAnimationTrackSetNode& Node);
+
 	UInterchangeBaseNodeContainer* BaseNodeContainer = nullptr;
 	TArray<const UInterchangeSourceData*> SourceDatas;
 
diff --git a/Engine/Source/Runtime/Experimental/Interchange/Core/Private/Nodes/InterchangeFactoryBaseNode.cpp b/Engine/Source/Runtime/Experimental/Interchange/Core/Private/Nodes/InterchangeFactoryBaseNode.cpp
--- a/Engine/Source/Runtime/Experimental/Interchange/Core/Private/Nodes/InterchangeFactoryBaseNode.cpp
+++ b/Engine/Source/Runtime/Experimental/Interchange/Core/Private/Nodes/InterchangeFactoryBaseNode.cpp
@@ -112,3 +112,7 @@ bool UInterchangeFactoryBaseNode::RemoveFactoryDependencyUid(const FString& Depe
 	return FactoryDependencies.RemoveItem(DependencyUid);
 }
 
+FString UInterchangeFactoryBaseNode::BuildFactoryNodeUid(const FString& TranslatedNodeUid)
+{
+	return TEXT("Factory_") + TranslatedNodeUid;
+}
diff --git a/Engine/Source/Runtime/Experimental/Interchange/Core/Public/Nodes/InterchangeFactoryBaseNode.h b/Engine/Source/Runtime/Experimental/Interchange/Core/Public/Nodes/InterchangeFactoryBaseNode.h
--- a/Engine/Source/Runtime/Experimental/Interchange/Core/Public/Nodes/InterchangeFactoryBaseNode.h
+++ b/Engine/Source/Runtime/Experimental/Interchange/Core/Public/Nodes/InterchangeFactoryBaseNode.h
@@ -93,6 +93,8 @@ public:
 	UPROPERTY()
 	mutable FSoftObjectPath ReferenceObject;
 
+	static FString BuildFactoryNodeUid(const FString& TranslatedNodeUid);
+
 protected:
 	/**
 	 * Those dependencies are use by the interchange parsing task to make sure the asset are created in the correct order.
diff --git a/Engine/Source/Runtime/Experimental/Interchange/Engine/Private/Tasks/InterchangeTaskParsing.cpp b/Engine/Source/Runtime/Experimental/Interchange/Engine/Private/Tasks/InterchangeTaskParsing.cpp
--- a/Engine/Source/Runtime/Experimental/Interchange/Engine/Private/Tasks/InterchangeTaskParsing.cpp
+++ b/Engine/Source/Runtime/Experimental/Interchange/Engine/Private/Tasks/InterchangeTaskParsing.cpp
@@ -101,6 +101,7 @@ void UE::Interchange::FTaskParsing::DoTask(ENamedThreads::Type CurrentThread, co
 	{
 		FString UniqueID;
 		int32 SourceIndex = INDEX_NONE;
+		bool bIsSceneNode = false;
 		TArray<FString> Dependencies;
 		FGraphEventRef GraphEventRef;
 		FGraphEventArray Prerequisites;
@@ -109,26 +110,26 @@ void UE::Interchange::FTaskParsing::DoTask(ENamedThreads::Type CurrentThread, co
 		TArray<UInterchangeFactoryBaseNode*, TInlineAllocator<1>> Nodes; // For scenes, we can group multiple nodes into a single task as they are usually very light
 	};
 
-	TArray<FTaskData> AssetTaskDatas;
-	TArray<FTaskData> SceneTaskDatas;
+	TArray<FTaskData> TaskDatas;
 
 	//Avoid creating asset if the asynchronous import is canceled, just create the completion task
 	if (!AsyncHelper->bCancel)
 	{
 		for (int32 SourceIndex = 0; SourceIndex < AsyncHelper->SourceDatas.Num(); ++SourceIndex)
 		{
-			TArray<FTaskData> SourceAssetTaskDatas;
-			TArray<FTaskData> SourceSceneTaskDatas;
+			TArray<FTaskData> SourceTaskDatas;
 
 			if (!AsyncHelper->BaseNodeContainers.IsValidIndex(SourceIndex))
 			{
 				continue;
 			}
+
 			UInterchangeBaseNodeContainer* BaseNodeContainer = AsyncHelper->BaseNodeContainers[SourceIndex].Get();
 			if (!BaseNodeContainer)
 			{
 				continue;
 			}
+
 			//Translation and pipelines are not executed, compute the children cache for translated and factory nodes
 			BaseNodeContainer->ComputeChildrenCache();
 
@@ -146,17 +147,18 @@ void UE::Interchange::FTaskParsing::DoTask(ENamedThreads::Type CurrentThread, co
 				{
 					const UClass* RegisteredFactoryClass = InterchangeManager->GetRegisteredFactoryClass(ObjectClass);
 
-					const bool bIsAsset = !(FactoryNode->GetObjectClass()->IsChildOf<AActor>() || FactoryNode->GetObjectClass()->IsChildOf<UActorComponent>());
+					const bool bIsSceneNode = FactoryNode->GetObjectClass()->IsChildOf<AActor>() || FactoryNode->GetObjectClass()->IsChildOf<UActorComponent>();
 
-					if (!RegisteredFactoryClass || (!bIsAsset && !bCanImportSceneNode))
+					if (!RegisteredFactoryClass || (bIsSceneNode && !bCanImportSceneNode))
 					{
 						//nothing we can import from this element
 						return;
 					}
 
-					FTaskData& NodeTaskData = bIsAsset ? SourceAssetTaskDatas.AddDefaulted_GetRef() : SourceSceneTaskDatas.AddDefaulted_GetRef();
+					FTaskData& NodeTaskData = SourceTaskDatas.AddDefaulted_GetRef();
 					NodeTaskData.UniqueID = FactoryNode->GetUniqueID();
 					NodeTaskData.SourceIndex = SourceIndex;
+					NodeTaskData.bIsSceneNode = bIsSceneNode;
 					NodeTaskData.Nodes.Add(FactoryNode);
 					FactoryNode->GetFactoryDependencies(NodeTaskData.Dependencies);
 					NodeTaskData.FactoryClass = RegisteredFactoryClass;
@@ -189,12 +191,10 @@ void UE::Interchange::FTaskParsing::DoTask(ENamedThreads::Type CurrentThread, co
 				};
 
 				// Nodes cannot depend on a node from another source, so it's faster to sort the dependencies per-source and then append those to the TaskData arrays.
-				SourceAssetTaskDatas.Sort(SortByDependencies);
-				SourceSceneTaskDatas.Sort(SortByDependencies);
+				SourceTaskDatas.Sort(SortByDependencies);
 			}
 
-			AssetTaskDatas.Append(MoveTemp(SourceAssetTaskDatas));
-			SceneTaskDatas.Append(MoveTemp(SourceSceneTaskDatas));
+			TaskDatas.Append(MoveTemp(SourceTaskDatas));
 		}
 	}
 
@@ -226,26 +226,31 @@ void UE::Interchange::FTaskParsing::DoTask(ENamedThreads::Type CurrentThread, co
 		return GraphEvents;
 	};
 
-	//Assets
-	FGraphEventArray AssetsCompletionPrerequistes;
+	TSet<FString> CreatedTasksAssetNames; // Tracks for which asset name we have created a task so that we don't have 2 tasks for the same asset name
+	TFunction<FGraphEventRef(FTaskData&)> CreateTasksFromData = [this, &AsyncHelper, &CreatedTasksAssetNames](FTaskData& TaskData)
 	{
-		TSet<FString> CreatedTasksAssetNames; // Tracks for which asset name we have created a task so that we don't have 2 tasks for the same asset name
-		TFunction<FGraphEventRef(FTaskData&)> CreateTasksForAsset = [this, &AsyncHelper, &CreatedTasksAssetNames](FTaskData& TaskData)
-		{
-			check(TaskData.Nodes.Num() == 1); //We expect 1 node per asset task
+		check(TaskData.Nodes.Num() == 1); //We expect 1 node per asset task
 
-			const int32 SourceIndex = TaskData.SourceIndex;
-			const UClass* const FactoryClass = TaskData.FactoryClass;
-			UInterchangeFactoryBaseNode* const FactoryNode = TaskData.Nodes[0];
-			const bool bFactoryCanRunOnAnyThread = FactoryClass->GetDefaultObject<UInterchangeFactoryBase>()->CanExecuteOnAnyThread();
+		const int32 SourceIndex = TaskData.SourceIndex;
+		const UClass* const FactoryClass = TaskData.FactoryClass;
+		UInterchangeFactoryBaseNode* const FactoryNode = TaskData.Nodes[0];
+		const bool bFactoryCanRunOnAnyThread = FactoryClass->GetDefaultObject<UInterchangeFactoryBase>()->CanExecuteOnAnyThread();
 
+		if (TaskData.bIsSceneNode)
+		{
+			return AsyncHelper->SceneTasks.Add_GetRef(
+				TGraphTask<FTaskCreateSceneObjects>::CreateTask(&(TaskData.Prerequisites))
+				.ConstructAndDispatchWhenReady(PackageBasePath, SourceIndex, WeakAsyncHelper, TaskData.Nodes, FactoryClass));
+		}
+		else
+		{
 			FString PackageSubPath;
 			FactoryNode->GetCustomSubPath(PackageSubPath);
 
 			const FString AssetFullPath = FPaths::Combine(PackageBasePath, PackageSubPath, FactoryNode->GetAssetName());
 
-			if ( ensureMsgf(!CreatedTasksAssetNames.Contains(AssetFullPath),
-				TEXT("Found multiple task data with the same asset name (%s). Only one will be executed."), *AssetFullPath) )
+			if (ensureMsgf(!CreatedTasksAssetNames.Contains(AssetFullPath),
+				TEXT("Found multiple task data with the same asset name (%s). Only one will be executed."), *AssetFullPath))
 			{
 				//Add create package task has a prerequisite of FTaskCreateAsset. Create package task is a game thread task
 				FGraphEventArray CreatePackagePrerequistes;
@@ -254,13 +259,12 @@ void UE::Interchange::FTaskParsing::DoTask(ENamedThreads::Type CurrentThread, co
 				);
 				CreatePackagePrerequistes.Add(AsyncHelper->CreatePackageTasks[CreatePackageTaskIndex]);
 
-				FGraphEventArray CreateAssetPrerequistes;
 				int32 CreateTaskIndex = AsyncHelper->CreateAssetTasks.Add(
 					TGraphTask<FTaskCreateAsset>::CreateTask(&(CreatePackagePrerequistes)).ConstructAndDispatchWhenReady(PackageBasePath, SourceIndex, WeakAsyncHelper, FactoryNode, bFactoryCanRunOnAnyThread)
 				);
-				CreateAssetPrerequistes.Add(AsyncHelper->CreateAssetTasks[CreateTaskIndex]);
 
-					CreatedTasksAssetNames.Add(AssetFullPath);
+				CreatedTasksAssetNames.Add(AssetFullPath);
+
 				return AsyncHelper->CreateAssetTasks[CreateTaskIndex];
 			}
 			else
@@ -269,32 +273,31 @@ void UE::Interchange::FTaskParsing::DoTask(ENamedThreads::Type CurrentThread, co
 				EmptyGraphEvent->DispatchSubsequents();
 				return EmptyGraphEvent;
 			}
-		};
-
-		AssetsCompletionPrerequistes = CreateTasksForEachTaskData(AssetTaskDatas, CreateTasksForAsset);
-	}
+		}
+	};
 
-	//Scenes
-	//Note: Scene tasks are delayed until all asset tasks are completed
-	FGraphEventArray ScenesCompletionPrerequistes;
+	FGraphEventArray CompletionPrerequistes;
+	for (int32 TaskIndex = 0; TaskIndex < TaskDatas.Num(); ++TaskIndex)
 	{
-		TFunction<FGraphEventRef(FTaskData&)> CreateTasksForSceneObject = [this, &AsyncHelper, &AssetsCompletionPrerequistes](FTaskData& TaskData)
-		{
-			const int32 SourceIndex = TaskData.SourceIndex;
-			const UClass* const FactoryClass = TaskData.FactoryClass;
+		FTaskData& TaskData = TaskDatas[TaskIndex];
 
-			return AsyncHelper->SceneTasks.Add_GetRef(
-				TGraphTask<FTaskCreateSceneObjects>::CreateTask(&AssetsCompletionPrerequistes)
-				.ConstructAndDispatchWhenReady(PackageBasePath, SourceIndex, WeakAsyncHelper, TaskData.Nodes, FactoryClass));
-		};
+		if (TaskData.Dependencies.Num() > 0)
+		{
+			//Search the previous node to find the dependence
+			for (int32 DepTaskIndex = 0; DepTaskIndex < TaskIndex; ++DepTaskIndex)
+			{
+				if (TaskData.Dependencies.Contains(TaskDatas[DepTaskIndex].UniqueID))
+				{
+					//Add has prerequisite
+					TaskData.Prerequisites.Add(TaskDatas[DepTaskIndex].GraphEventRef);
+				}
+			}
+		}
 
-		ScenesCompletionPrerequistes = CreateTasksForEachTaskData(SceneTaskDatas, CreateTasksForSceneObject);
+		TaskData.GraphEventRef = CreateTasksFromData(TaskData);
+		CompletionPrerequistes.Add(TaskData.GraphEventRef);
 	}
 
-	FGraphEventArray CompletionPrerequistes;
-	CompletionPrerequistes.Append(AssetsCompletionPrerequistes);
-	CompletionPrerequistes.Append(ScenesCompletionPrerequistes);
-
 	//Add an async task for pre completion
 	
 	FGraphEventArray PreCompletionPrerequistes;
