Fixed localization warning

diff --git a/Engine/Plugins/Experimental/Interchange/Runtime/Source/Import/Private/Animation/InterchangeAnimationTrackSetFactory.cpp b/Engine/Plugins/Experimental/Interchange/Runtime/Source/Import/Private/Animation/InterchangeAnimationTrackSetFactory.cpp
--- a/Engine/Plugins/Experimental/Interchange/Runtime/Source/Import/Private/Animation/InterchangeAnimationTrackSetFactory.cpp
+++ b/Engine/Plugins/Experimental/Interchange/Runtime/Source/Import/Private/Animation/InterchangeAnimationTrackSetFactory.cpp
@@ -667,7 +667,7 @@ UObject* UInterchangeAnimationTrackSetFactory::CreateAsset(const FCreateAssetPar
 		Message->SourceAssetName = Arguments.SourceData->GetFilename();
 		Message->DestinationAssetName = Arguments.AssetName;
 		Message->AssetType = ULevelSequence::StaticClass();
-		Message->Text = FText::Format(LOCTEXT("CreateAssetUnsupportedReimport", "Could not create nor find LevelSequence asset {0}."), FText::FromString(Arguments.AssetName));
+		Message->Text = FText::Format(LOCTEXT("CreateAssetFailed", "Could not create nor find LevelSequence asset {0}."), FText::FromString(Arguments.AssetName));
 		return nullptr;
 	}
 
