Adding the possibility to set the SRGB colorspace on an UInterchangeTextureNode

diff --git a/Engine/Plugins/Experimental/Interchange/Runtime/Source/Nodes/Public/InterchangeTextureNode.h b/Engine/Plugins/Experimental/Interchange/Runtime/Source/Nodes/Public/InterchangeTextureNode.h
--- a/Engine/Plugins/Experimental/Interchange/Runtime/Source/Nodes/Public/InterchangeTextureNode.h
+++ b/Engine/Plugins/Experimental/Interchange/Runtime/Source/Nodes/Public/InterchangeTextureNode.h
@@ -100,4 +100,20 @@ public:
 			LogAttributeStorageErrors(Result, TEXT("UInterchangeTextureNode.SetPayLoadKey"), UE::Interchange::FTextureNodeStaticData::PayloadSourceFileKey());
 		}
 	}
+
+	UFUNCTION(BlueprintCallable, Category = "Interchange | Node | Texture")
+	bool GetCustomSRGB(bool& AttributeValue) const
+	{
+		IMPLEMENT_NODE_ATTRIBUTE_GETTER(SRGB, bool);
+	}
+
+	UFUNCTION(BlueprintCallable, Category = "Interchange | Node | Texture")
+	bool SetCustomSRGB(const bool& AttributeValue)
+	{
+		IMPLEMENT_NODE_ATTRIBUTE_SETTER_NODELEGATE(SRGB, bool)
+	}
+
+private:
+
+	IMPLEMENT_NODE_ATTRIBUTE_KEY(SRGB)
 };
diff --git a/Engine/Plugins/Experimental/Interchange/Runtime/Source/Pipelines/Private/InterchangeGenericMaterialPipeline.cpp b/Engine/Plugins/Experimental/Interchange/Runtime/Source/Pipelines/Private/InterchangeGenericMaterialPipeline.cpp
--- a/Engine/Plugins/Experimental/Interchange/Runtime/Source/Pipelines/Private/InterchangeGenericMaterialPipeline.cpp
+++ b/Engine/Plugins/Experimental/Interchange/Runtime/Source/Pipelines/Private/InterchangeGenericMaterialPipeline.cpp
@@ -998,6 +998,14 @@ void UInterchangeGenericMaterialPipeline::HandleTextureSampleNode(const UInterch
 				TextureFactoryNode->SetCustomSRGB(false);
 			}
 		}
+
+		if(bool bSRGB; TextureNode->GetCustomSRGB(bSRGB))
+		{
+			if(UInterchangeTextureFactoryNode* TextureFactoryNode = Cast<UInterchangeTextureFactoryNode>(BaseNodeContainer->GetFactoryNode(TextureFactoryUid)))
+			{
+				TextureFactoryNode->SetCustomSRGB(bSRGB);
+			}
+		}
 	}
 	else
 	{
