- Adding the support of Lights through Getters/Setters in both Nodes and FactoryNodes
- Reading the lights in MaterialX
- Adding file prefix in MaterialX
- Reading Colorspace in MaterialX, only SRGB is supported this one needs CL 20691421
- Catch exception from MaterialX SDK, avoids UE to crash if the input file has errors in it

diff --git a/Engine/Plugins/Experimental/Interchange/Runtime/Source/FactoryNodes/Private/InterchangeLightFactoryNode.cpp b/Engine/Plugins/Experimental/Interchange/Runtime/Source/FactoryNodes/Private/InterchangeLightFactoryNode.cpp
new file mode 100644
--- /dev/null
+++ b/Engine/Plugins/Experimental/Interchange/Runtime/Source/FactoryNodes/Private/InterchangeLightFactoryNode.cpp
@@ -0,0 +1,147 @@
+#include "InterchangeLightFactoryNode.h"
+
+#include "Components/RectLightComponent.h"
+#include "Components/SpotLightComponent.h"
+
+bool UInterchangeBaseLightFactoryNode::GetCustomLightColor(FColor& AttributeValue) const
+{
+	IMPLEMENT_NODE_ATTRIBUTE_GETTER(LightColor, FColor)
+}
+
+bool UInterchangeBaseLightFactoryNode::SetCustomLightColor(const FColor& AttributeValue, bool bAddApplyDelegate)
+{
+	IMPLEMENT_NODE_ATTRIBUTE_SETTER(UInterchangeBaseLightFactoryNode, LightColor, FColor, ULightComponent)
+}
+
+bool UInterchangeBaseLightFactoryNode::GetCustomIntensity(float& AttributeValue) const
+{
+	IMPLEMENT_NODE_ATTRIBUTE_GETTER(Intensity, float)
+}
+
+bool UInterchangeBaseLightFactoryNode::SetCustomIntensity(float AttributeValue, bool bAddApplyDelegate)
+{
+	IMPLEMENT_NODE_ATTRIBUTE_SETTER(UInterchangeBaseLightFactoryNode, Intensity, float, ULightComponent)
+}
+
+bool UInterchangeBaseLightFactoryNode::GetCustomTemperature(float& AttributeValue) const
+{
+	IMPLEMENT_NODE_ATTRIBUTE_GETTER(Temperature, float)
+}
+
+bool UInterchangeBaseLightFactoryNode::SetCustomTemperature(float AttributeValue, bool bAddApplyDelegate)
+{
+	IMPLEMENT_NODE_ATTRIBUTE_SETTER(UInterchangeBaseLightFactoryNode, Temperature, float, ULightComponent)
+}
+
+bool UInterchangeBaseLightFactoryNode::GetCustomUseTemperature(bool& AttributeValue) const
+{
+	IMPLEMENT_NODE_ATTRIBUTE_GETTER(bUseTemperature, bool)
+}
+
+bool UInterchangeBaseLightFactoryNode::SetCustomUseTemperature(bool AttributeValue, bool bAddApplyDelegate)
+{
+	IMPLEMENT_NODE_ATTRIBUTE_SETTER(UInterchangeBaseLightFactoryNode, bUseTemperature, bool, ULightComponent)
+}
+
+bool UInterchangeLightFactoryNode::GetCustomIntensityUnits(ELightUnits& AttributeValue) const
+{
+	IMPLEMENT_NODE_ATTRIBUTE_GETTER(IntensityUnits, ELightUnits)
+}
+
+bool UInterchangeLightFactoryNode::SetCustomIntensityUnits(ELightUnits AttributeValue, bool bAddApplyDelegate)
+{
+	IMPLEMENT_NODE_ATTRIBUTE_SETTER(UInterchangeLightFactoryNode, IntensityUnits, ELightUnits, ULocalLightComponent)
+}
+
+bool UInterchangeLightFactoryNode::GetCustomAttenuationRadius(float& AttributeValue) const
+{
+	IMPLEMENT_NODE_ATTRIBUTE_GETTER(AttenuationRadius, float)
+}
+
+bool UInterchangeLightFactoryNode::SetCustomAttenuationRadius(float AttributeValue, bool bAddApplyDelegate)
+{
+	IMPLEMENT_NODE_ATTRIBUTE_SETTER(UInterchangeLightFactoryNode, AttenuationRadius, float, ULocalLightComponent)
+}
+
+bool UInterchangePointLightFactoryNode::GetCustomUseInverseSquaredFalloff(bool& AttributeValue) const
+{
+	IMPLEMENT_NODE_ATTRIBUTE_GETTER(bUseInverseSquaredFalloff, bool)
+}
+
+bool UInterchangePointLightFactoryNode::SetCustomUseInverseSquaredFalloff(bool AttributeValue, bool bAddApplyDelegate)
+{
+	IMPLEMENT_NODE_ATTRIBUTE_SETTER(UInterchangePointLightFactoryNode, bUseInverseSquaredFalloff, bool, UPointLightComponent)
+}
+
+bool UInterchangePointLightFactoryNode::GetCustomLightFalloffExponent(float& AttributeValue) const
+{
+	IMPLEMENT_NODE_ATTRIBUTE_GETTER(LightFalloffExponent, float)
+}
+
+bool UInterchangePointLightFactoryNode::SetCustomLightFalloffExponent(float AttributeValue, bool bAddApplyDelegate)
+{
+	IMPLEMENT_NODE_ATTRIBUTE_SETTER(UInterchangePointLightFactoryNode, LightFalloffExponent, float, UPointLightComponent)
+}
+
+bool UInterchangeSpotLightFactoryNode::GetCustomInnerConeAngle(float& AttributeValue) const
+{
+	IMPLEMENT_NODE_ATTRIBUTE_GETTER(InnerConeAngle, float)
+}
+
+bool UInterchangeSpotLightFactoryNode::SetCustomInnerConeAngle(float AttributeValue, bool bAddApplyDelegate)
+{
+	IMPLEMENT_NODE_ATTRIBUTE_SETTER(UInterchangeSpotLightFactoryNode, InnerConeAngle, float, USpotLightComponent)
+}
+
+bool UInterchangeSpotLightFactoryNode::GetCustomOuterConeAngle(float& AttributeValue) const
+{
+	IMPLEMENT_NODE_ATTRIBUTE_GETTER(OuterConeAngle, float)
+}
+
+bool UInterchangeSpotLightFactoryNode::SetCustomOuterConeAngle(float AttributeValue, bool bAddApplyDelegate)
+{
+	IMPLEMENT_NODE_ATTRIBUTE_SETTER(UInterchangeSpotLightFactoryNode, OuterConeAngle, float, USpotLightComponent)
+}
+
+//For the moment n/a, need more discussion specially about barn door's angle
+/*
+bool UInterchangeRectLightFactoryNode::GetCustomSourceWidth(float& AttributeValue) const
+{
+	IMPLEMENT_NODE_ATTRIBUTE_GETTER(SourceWidth, float)
+}
+
+bool UInterchangeRectLightFactoryNode::SetCustomSourceWidth(float AttributeValue, bool bAddApplyDelegate)
+{
+	IMPLEMENT_NODE_ATTRIBUTE_SETTER(UInterchangeRectLightFactoryNode, SourceWidth, float, URectLightComponent)
+}
+
+bool UInterchangeRectLightFactoryNode::GetCustomSourceHeight(float& AttributeValue) const
+{
+	IMPLEMENT_NODE_ATTRIBUTE_GETTER(SourceHeight, float)
+}
+
+bool UInterchangeRectLightFactoryNode::SetCustomSourceHeight(float AttributeValue, bool bAddApplyDelegate)
+{
+	IMPLEMENT_NODE_ATTRIBUTE_SETTER(UInterchangeRectLightFactoryNode, SourceHeight, float, URectLightComponent)
+}
+
+bool UInterchangeRectLightFactoryNode::GetCustomBarnDoorAngle(float& AttributeValue) const
+{
+	IMPLEMENT_NODE_ATTRIBUTE_GETTER(BarnDoorAngle, float)
+}
+
+bool UInterchangeRectLightFactoryNode::SetCustomBarnDoorAngle(float AttributeValue, bool bAddApplyDelegate)
+{
+	IMPLEMENT_NODE_ATTRIBUTE_SETTER(UInterchangeRectLightFactoryNode, BarnDoorAngle, float, URectLightComponent)
+}
+
+bool UInterchangeRectLightFactoryNode::GetCustomBarnDoorLength(float& AttributeValue) const
+{
+	IMPLEMENT_NODE_ATTRIBUTE_GETTER(BarnDoorLength, float)
+}
+
+bool UInterchangeRectLightFactoryNode::SetCustomBarnDoorLength(float AttributeValue, bool bAddApplyDelegate)
+{
+	IMPLEMENT_NODE_ATTRIBUTE_SETTER(UInterchangeRectLightFactoryNode, BarnDoorLength, float, URectLightComponent)
+}
+*/
\ No newline at end of file
diff --git a/Engine/Plugins/Experimental/Interchange/Runtime/Source/FactoryNodes/Public/InterchangeLightFactoryNode.h b/Engine/Plugins/Experimental/Interchange/Runtime/Source/FactoryNodes/Public/InterchangeLightFactoryNode.h
new file mode 100644
--- /dev/null
+++ b/Engine/Plugins/Experimental/Interchange/Runtime/Source/FactoryNodes/Public/InterchangeLightFactoryNode.h
@@ -0,0 +1,169 @@
+// Copyright Epic Games, Inc. All Rights Reserved.
+
+#pragma once
+
+#include "CoreMinimal.h"
+#include "InterchangeActorFactoryNode.h"
+
+#include "InterchangeLightFactoryNode.generated.h"
+
+UCLASS(BlueprintType)
+class INTERCHANGEFACTORYNODES_API UInterchangeBaseLightFactoryNode : public UInterchangeActorFactoryNode
+{
+	GENERATED_BODY()
+
+public:
+
+	UFUNCTION(BlueprintCallable, Category = "Interchange | Node | LightFactory")
+	bool GetCustomLightColor(FColor& AttributeValue) const;
+
+	UFUNCTION(BlueprintCallable, Category = "Interchange | Node | LightFactory")
+	bool SetCustomLightColor(const FColor& AttributeValue, bool bAddApplyDelegate = true);
+
+	UFUNCTION(BlueprintCallable, Category = "Interchange | Node | LightFactory")
+	bool GetCustomIntensity(float& AttributeValue) const;
+
+	UFUNCTION(BlueprintCallable, Category = "Interchange | Node | LightFactory")
+	bool SetCustomIntensity(float AttributeValue, bool bAddApplyDelegate = true);
+	
+	UFUNCTION(BlueprintCallable, Category = "Interchange | Node | LightFactory")
+	bool GetCustomTemperature(float& AttributeValue) const;
+
+	UFUNCTION(BlueprintCallable, Category = "Interchange | Node | LightFactory")
+	bool SetCustomTemperature(float AttributeValue, bool bAddApplyDelegate = true);
+
+	UFUNCTION(BlueprintCallable, Category = "Interchange | Node | LightFactory")
+	bool GetCustomUseTemperature(bool& AttributeValue) const;
+
+	UFUNCTION(BlueprintCallable, Category = "Interchange | Node | LightFactory")
+	bool SetCustomUseTemperature(bool AttributeValue, bool bAddApplyDelegate = true);
+
+private:
+
+	IMPLEMENT_NODE_ATTRIBUTE_KEY(LightColor)
+	IMPLEMENT_NODE_ATTRIBUTE_KEY(Intensity)
+	IMPLEMENT_NODE_ATTRIBUTE_KEY(Temperature)
+	IMPLEMENT_NODE_ATTRIBUTE_KEY(bUseTemperature)
+};
+
+UCLASS(BlueprintType)
+class INTERCHANGEFACTORYNODES_API UInterchangeDirectionalLightFactoryNode : public UInterchangeBaseLightFactoryNode
+{
+	GENERATED_BODY()
+};
+
+UCLASS(BlueprintType)
+class INTERCHANGEFACTORYNODES_API UInterchangeLightFactoryNode : public UInterchangeBaseLightFactoryNode
+{
+	GENERATED_BODY()
+
+public:
+
+	UFUNCTION(BlueprintCallable, Category = "Interchange | Node | LightFactory")
+	bool GetCustomIntensityUnits(ELightUnits& AttributeValue) const;
+
+	UFUNCTION(BlueprintCallable, Category = "Interchange | Node | LightFactory")
+	bool SetCustomIntensityUnits(ELightUnits AttributeValue, bool bAddApplyDelegate = true);
+
+	UFUNCTION(BlueprintCallable, Category = "Interchange | Node | LightFactory")
+	bool GetCustomAttenuationRadius(float& AttributeValue) const;
+
+	UFUNCTION(BlueprintCallable, Category = "Interchange | Node | LightFactory")
+	bool SetCustomAttenuationRadius(float AttributeValue, bool bAddApplyDelegate = true);
+
+private:
+
+	IMPLEMENT_NODE_ATTRIBUTE_KEY(IntensityUnits)
+	IMPLEMENT_NODE_ATTRIBUTE_KEY(AttenuationRadius)
+	IMPLEMENT_NODE_ATTRIBUTE_KEY(IESTexture)
+};
+
+UCLASS(BlueprintType)
+class INTERCHANGEFACTORYNODES_API UInterchangeRectLightFactoryNode : public UInterchangeLightFactoryNode
+{
+	GENERATED_BODY()
+
+public:
+	//For the moment n/a, need more discussion specially about barn door's angle
+/*
+	UFUNCTION(BlueprintCallable, Category = "Interchange | Node | RectLightFactory")
+	bool GetCustomSourceWidth(float& AttributeValue) const;
+
+	UFUNCTION(BlueprintCallable, Category = "Interchange | Node | RectLightFactory")
+	bool SetCustomSourceWidth(float AttributeValue, bool bAddApplyDelegate = true);
+
+	UFUNCTION(BlueprintCallable, Category = "Interchange | Node | RectLightFactory")
+	bool GetCustomSourceHeight(float& AttributeValue) const;
+
+	UFUNCTION(BlueprintCallable, Category = "Interchange | Node | RectLightFactory")
+	bool SetCustomSourceHeight(float AttributeValue, bool bAddApplyDelegate = true);
+
+	UFUNCTION(BlueprintCallable, Category = "Interchange | Node | RectLightFactory")
+	bool GetCustomBarnDoorAngle(float& AttributeValue) const;
+
+	UFUNCTION(BlueprintCallable, Category = "Interchange | Node | RectLightFactory")
+	bool SetCustomBarnDoorAngle(float AttributeValue, bool bAddApplyDelegate = true);
+
+	UFUNCTION(BlueprintCallable, Category = "Interchange | Node | RectLightFactory")
+	bool GetCustomBarnDoorLength(float& AttributeValue) const;
+
+	UFUNCTION(BlueprintCallable, Category = "Interchange | Node | RectLightFactory")
+	bool SetCustomBarnDoorLength(float AttributeValue, bool bAddApplyDelegate = true);
+
+private:
+	IMPLEMENT_NODE_ATTRIBUTE_KEY(SourceWidth)
+	IMPLEMENT_NODE_ATTRIBUTE_KEY(SourceHeight)
+	IMPLEMENT_NODE_ATTRIBUTE_KEY(BarnDoorAngle)
+	IMPLEMENT_NODE_ATTRIBUTE_KEY(BarnDoorLength)
+*/
+};
+
+
+UCLASS(BlueprintType)
+class INTERCHANGEFACTORYNODES_API UInterchangePointLightFactoryNode : public UInterchangeLightFactoryNode
+{
+	GENERATED_BODY()
+
+public:
+	UFUNCTION(BlueprintCallable, Category = "Interchange | Node | PointLightFactory")
+	bool GetCustomUseInverseSquaredFalloff(bool& AttributeValue) const;
+
+	UFUNCTION(BlueprintCallable, Category = "Interchange | Node | PointLightFactory")
+	bool SetCustomUseInverseSquaredFalloff(bool AttributeValue, bool bAddApplyDelegate = true);
+
+	UFUNCTION(BlueprintCallable, Category = "Interchange | Node | PointLightFactory")
+	bool GetCustomLightFalloffExponent(float& AttributeValue) const;
+
+	UFUNCTION(BlueprintCallable, Category = "Interchange | Node | PointLightFactory")
+	bool SetCustomLightFalloffExponent(float AttributeValue, bool bAddApplyDelegate = true);
+
+private:
+	IMPLEMENT_NODE_ATTRIBUTE_KEY(bUseInverseSquaredFalloff)
+	IMPLEMENT_NODE_ATTRIBUTE_KEY(LightFalloffExponent)
+};
+
+UCLASS(BlueprintType)
+class INTERCHANGEFACTORYNODES_API UInterchangeSpotLightFactoryNode : public UInterchangePointLightFactoryNode
+{
+	GENERATED_BODY()
+
+public:
+
+	UFUNCTION(BlueprintCallable, Category = "Interchange | Node | SpotLightFactory")
+	bool GetCustomInnerConeAngle(float& AttributeValue) const;
+
+	UFUNCTION(BlueprintCallable, Category = "Interchange | Node | SpotLightFactory")
+	bool SetCustomInnerConeAngle(float AttributeValue, bool bAddApplyDelegate = true);
+
+	UFUNCTION(BlueprintCallable, Category = "Interchange | Node | SpotLightFactory")
+	bool GetCustomOuterConeAngle(float& AttributeValue) const;
+
+	UFUNCTION(BlueprintCallable, Category = "Interchange | Node | SpotLightFactory")
+	bool SetCustomOuterConeAngle(float AttributeValue, bool bAddApplyDelegate = true);
+
+private:
+
+	IMPLEMENT_NODE_ATTRIBUTE_KEY(InnerConeAngle)
+	IMPLEMENT_NODE_ATTRIBUTE_KEY(OuterConeAngle)
+};
+
diff --git a/Engine/Plugins/Experimental/Interchange/Runtime/Source/Import/Private/MaterialX/InterchangeMaterialXDefinitions.h b/Engine/Plugins/Experimental/Interchange/Runtime/Source/Import/Private/MaterialX/InterchangeMaterialXDefinitions.h
--- a/Engine/Plugins/Experimental/Interchange/Runtime/Source/Import/Private/MaterialX/InterchangeMaterialXDefinitions.h
+++ b/Engine/Plugins/Experimental/Interchange/Runtime/Source/Import/Private/MaterialX/InterchangeMaterialXDefinitions.h
@@ -101,6 +101,46 @@ namespace MaterialX
 		}
 	}
 
+	namespace Lights
+	{
+		//There's no input per se in a Light, but we can find some common inputs among those lights
+		namespace Input
+		{
+			static constexpr const char* Color = "color";
+			static constexpr const char* Intensity = "intensity";
+		}
+
+		namespace PointLight
+		{
+			namespace Input
+			{
+				using namespace Lights::Input;
+				static constexpr const char* Position = "position";
+				static constexpr const char* DecayRate = "decay_rate";
+			}
+		}
+
+		namespace DirectionalLight
+		{			
+			namespace Input
+			{
+				using namespace Lights::Input;
+				static constexpr const char* Direction = "direction";
+			}
+		}
+
+		namespace SpotLight
+		{
+			namespace Input
+			{
+				using namespace PointLight::Input;
+				using namespace DirectionalLight::Input;
+				static constexpr const char* InnerAngle = "inner_angle";
+				static constexpr const char* OuterAngle = "outer_angle";
+			}
+		}
+	}
+
 	namespace Attributes
 	{
 		static constexpr const char* IsVisited = "UE:IsVisited";
@@ -132,11 +172,15 @@ namespace MaterialX
 		static constexpr const char* PointLight = "point_light";
 		static constexpr const char* DirectionalLight = "directional_light";
 		static constexpr const char* SpotLight = "spot_light";
+		static constexpr const char* StandardSurface = "standard_surface";
 	}
 
 	namespace NodeDefinition
 	{
 		static constexpr const char* StandardSurface = "ND_standard_surface_surfaceshader";
+		static constexpr const char* PointLight = "ND_point_light";
+		static constexpr const char* DirectionalLight = "ND_directional_light";
+		static constexpr const char* SpotLight = "ND_spot_light";
 	}
 
 	namespace Library
diff --git a/Engine/Plugins/Experimental/Interchange/Runtime/Source/Import/Private/MaterialX/InterchangeMaterialXTranslator.cpp b/Engine/Plugins/Experimental/Interchange/Runtime/Source/Import/Private/MaterialX/InterchangeMaterialXTranslator.cpp
--- a/Engine/Plugins/Experimental/Interchange/Runtime/Source/Import/Private/MaterialX/InterchangeMaterialXTranslator.cpp
+++ b/Engine/Plugins/Experimental/Interchange/Runtime/Source/Import/Private/MaterialX/InterchangeMaterialXTranslator.cpp
@@ -4,6 +4,7 @@
 #if WITH_EDITOR
 #include "MaterialXFormat/Util.h"
 #endif
+#include "Engine/Scene.h"
 #include "InterchangeImportLog.h"
 #include "InterchangeLightNode.h"
 #include "InterchangeManager.h"
@@ -25,10 +26,11 @@ namespace UE::Interchange::MaterialX
 	{
 		static const bool bStandardSurfacePackageLoaded = []() -> bool
 		{
-			if (FString FunctionPath = FPackageName::ExportTextPathToObjectPath(TEXT("MaterialFunction'/Interchange/Functions/MX_StandardSurface.MX_StandardSurface'"));
-				FPackageName::DoesPackageExist(FunctionPath))
+			const FString TextPath{ TEXT("MaterialFunction'/Interchange/Functions/MX_StandardSurface.MX_StandardSurface'") };
+			const FString FunctionPath{ FPackageName::ExportTextPathToObjectPath(TextPath) };
+			if(FPackageName::DoesPackageExist(FunctionPath))
 			{
-				if (FSoftObjectPath(FunctionPath).TryLoad())
+				if(FSoftObjectPath(FunctionPath).TryLoad())
 				{
 					return true;
 				}
@@ -52,18 +54,17 @@ namespace UE::Interchange::MaterialX
 UInterchangeMaterialXTranslator::UInterchangeMaterialXTranslator()
 #if WITH_EDITOR
 	:
-	InputNamesMaterialX2UE{
-		{TEXT("in"),        TEXT("Input")},
-		{TEXT("in1"),       TEXT("A")},
-		{TEXT("in2"),       TEXT("B")},
-		{TEXT("fg"),        TEXT("A")},
-		{TEXT("bg"),        TEXT("B")},
-		{TEXT("mix"),       TEXT("Factor")},
-		{TEXT("low"),       TEXT("Min")},
-		{TEXT("high"),      TEXT("Max")},
-		{TEXT("texcoord"),  UE::Interchange::Materials::Standard::Nodes::TextureSample::Inputs::Coordinates.ToString()} },
-		NodeNamesMaterialX2UE
-{
+InputNamesMaterialX2UE{
+	{TEXT("in"),        TEXT("Input")},
+	{TEXT("in1"),       TEXT("A")},
+	{TEXT("in2"),       TEXT("B")},
+	{TEXT("fg"),        TEXT("A")},
+	{TEXT("bg"),        TEXT("B")},
+	{TEXT("mix"),       TEXT("Factor")},
+	{TEXT("low"),       TEXT("Min")},
+	{TEXT("high"),      TEXT("Max")},
+	{TEXT("texcoord"),  UE::Interchange::Materials::Standard::Nodes::TextureSample::Inputs::Coordinates.ToString()} },
+NodeNamesMaterialX2UE{
 	{MaterialX::Category::Add,      TEXT("Add")},
 	{MaterialX::Category::Sub,      TEXT("Subtract")},
 	{MaterialX::Category::Multiply, TEXT("Multiply")},
@@ -75,21 +76,30 @@ UInterchangeMaterialXTranslator::UInterchangeMaterialXTranslator()
 	{MaterialX::Category::Min,      TEXT("Min")},
 	{MaterialX::Category::Combine2, TEXT("AppendVector")},
 	{MaterialX::Category::Combine3, TEXT("AppendVector")},
-	{MaterialX::Category::Combine4, TEXT("AppendVector")}
-},
+	{MaterialX::Category::Combine4, TEXT("AppendVector")}},
 UEInputs{ TEXT("A"), TEXT("B"), TEXT("Input"), TEXT("Factor"), TEXT("Min"), TEXT("Max") }
 #endif // WITH_EDITOR
 {
 }
 
+EInterchangeTranslatorType UInterchangeMaterialXTranslator::GetTranslatorType() const
+{
+	return EInterchangeTranslatorType::Scenes;
+}
+
+bool UInterchangeMaterialXTranslator::DoesSupportAssetType(EInterchangeTranslatorAssetType AssetType) const
+{
+	return AssetType == EInterchangeTranslatorAssetType::Materials;
+}
+
 TArray<FString> UInterchangeMaterialXTranslator::GetSupportedFormats() const
 {
 	// Call to UInterchangeMaterialXTranslator::GetSupportedFormats is not supported out of game thread
 	// A more global solution must be found for translators which require some initialization
-	if (!IsInGameThread())
-		{
+	if(!IsInGameThread())
+	{
 		return TArray<FString>{};
-			}
+	}
 
 	return UE::Interchange::MaterialX::IsStandardSurfacePackageLoaded() ? TArray<FString>{ TEXT("mtlx;MaterialX File Format") } : TArray<FString>{};
 }
@@ -106,83 +116,116 @@ bool UInterchangeMaterialXTranslator::Translate(UInterchangeBaseNodeContainer& B
 	{
 		return false;
 	}
-
-	mx::FileSearchPath MaterialXLibFolder{ TCHAR_TO_UTF8(*FPaths::Combine(
-		FPaths::EngineDir(),
-		TEXT("Binaries"),
-		TEXT("ThirdParty"),
-		TEXT("MaterialX"),
-		TEXT("libraries"))) };
-
-	mx::DocumentPtr MaterialXLibrary = mx::createDocument();
-
-	mx::StringSet LoadedLibs = mx::loadLibraries({ mx::Library::Std, mx::Library::Pbr, mx::Library::Bxdf }, MaterialXLibFolder, MaterialXLibrary);
-	if(LoadedLibs.empty())
+	try
 	{
-		UInterchangeResultError_Generic* Message = AddMessage<UInterchangeResultError_Generic>();
-		Message->Text = FText::Format(LOCTEXT("MaterialXLibrariesNotFound", "Couldn't load MaterialX libraries from {0}"),
-			FText::FromString(MaterialXLibFolder.asString().c_str()));
-		return false;
-	}
-
-	mx::DocumentPtr Document = mx::createDocument();
-	mx::readFromXmlFile(Document, TCHAR_TO_UTF8(*Filename));
+		mx::FileSearchPath MaterialXLibFolder{ TCHAR_TO_UTF8(*FPaths::Combine(
+			FPaths::EngineDir(),
+			TEXT("Binaries"),
+			TEXT("ThirdParty"),
+			TEXT("MaterialX"),
+			TEXT("libraries"))) };
 
-	Document->importLibrary(MaterialXLibrary);
-	bIsDocumentValid = Document->validate();
+		mx::DocumentPtr MaterialXLibrary = mx::createDocument();
 
-	for(mx::ElementPtr Elem : Document->traverseTree())
-	{
-		//make sure to read only the current file otherwise we'll process the entire library
-		if(Elem->getActiveSourceUri() != Document->getActiveSourceUri())
+		mx::StringSet LoadedLibs = mx::loadLibraries({ mx::Library::Std, mx::Library::Pbr, mx::Library::Bxdf, mx::Library::Lights }, MaterialXLibFolder, MaterialXLibrary);
+		if(LoadedLibs.empty())
 		{
-			continue;
+			UInterchangeResultError_Generic* Message = AddMessage<UInterchangeResultError_Generic>();
+			Message->Text = FText::Format(LOCTEXT("MaterialXLibrariesNotFound", "Couldn't load MaterialX libraries from {0}"),
+										  FText::FromString(MaterialXLibFolder.asString().c_str()));
+			return false;
 		}
 
-		mx::NodePtr Node = Elem->asA<mx::Node>();
+		mx::DocumentPtr Document = mx::createDocument();
+		mx::readFromXmlFile(Document, TCHAR_TO_UTF8(*Filename));
+		Document->importLibrary(MaterialXLibrary);
 
-		if(Node)
+		std::string MaterialXMessage;
+		bIsDocumentValid = Document->validate(&MaterialXMessage);
+		if(!bIsDocumentValid)
 		{
-			//The entry point for materials is only on a surfacematerial node
-			if(Node->getType() == mx::MATERIAL_TYPE_STRING && Node->getCategory() == mx::SURFACE_MATERIAL_NODE_STRING)
+			UInterchangeResultError_Generic* Message = AddMessage<UInterchangeResultError_Generic>();
+			Message->Text = FText::Format(LOCTEXT("MaterialXDocumentInvalid", "{0}"),
+										  FText::FromString(MaterialXMessage.c_str()));
+			return false;
+		}
+
+		for(mx::ElementPtr Elem : Document->traverseTree())
+		{
+			//make sure to read only the current file otherwise we'll process the entire library
+			if(Elem->getActiveSourceUri() != Document->getActiveSourceUri())
 			{
-				if(!Node->getTypeDef())
-				{
-					UInterchangeResultError_Generic* Message = AddMessage<UInterchangeResultError_Generic>();
-					Message->Text = FText::Format(LOCTEXT("TypeDefNotFound", "<{0}> has no matching TypeDef"),
-						FText::FromString(Node->getName().c_str()));
+				continue;
+			}
 
-					bIsReferencesValid = false;
-					break;
-				}
+			mx::NodePtr Node = Elem->asA<mx::Node>();
 
-				bool bHasStandardSurface = false;
+			if(Node)
+			{
+				std::string NameOfNodeDocument = Document->getCategory();
+				std::string NameOfNode = Node->getCategory();
+				printf("%s, %s", NameOfNodeDocument.c_str(), NameOfNode.c_str());
 
-				for(mx::InputPtr Input : Node->getInputs())
-				{
-					//we only support standard_surface for now
-					mx::NodePtr ConnectedNode = Input->getConnectedNode();
+				bool bIsMaterialShader = Node->getType() == mx::MATERIAL_TYPE_STRING;
+				bool bIsLightShader = Node->getType() == mx::LIGHT_SHADER_TYPE_STRING;
 
-					if(ConnectedNode && ConnectedNode->getCategory() == "standard_surface")
+				if(bIsMaterialShader || bIsLightShader)
+				{
+					if(!Node->getTypeDef())
 					{
-						ProcessStandardSurface(BaseNodeContainer, ConnectedNode, Document);
-						bHasStandardSurface = true;
+						UInterchangeResultError_Generic* Message = AddMessage<UInterchangeResultError_Generic>();
+						Message->Text = FText::Format(LOCTEXT("TypeDefNotFound", "<{0}> has no matching TypeDef"),
+													  FText::FromString(Node->getName().c_str()));
+
+						bIsReferencesValid = false;
+						break;
 					}
-				}
 
-				if(!bHasStandardSurface)
-				{
-					UInterchangeResultWarning_Generic* Message = AddMessage<UInterchangeResultWarning_Generic>();
-					Message->Text = FText::Format(LOCTEXT("StandardSurfaceNotFound", "<{0}> has no standard_surface inputs"),
-						FText::FromString(Node->getName().c_str()));
+					//The entry point for materials is only on a surfacematerial node
+					if(bIsMaterialShader && Node->getCategory() == mx::SURFACE_MATERIAL_NODE_STRING)
+					{
+						bool bHasStandardSurface = false;
+
+						for(mx::InputPtr Input : Node->getInputs())
+						{
+							//we only support standard_surface for now
+							mx::NodePtr ConnectedNode = Input->getConnectedNode();
+
+							if(ConnectedNode && ConnectedNode->getCategory() == mx::Category::StandardSurface)
+							{
+								ProcessStandardSurface(BaseNodeContainer, ConnectedNode, Document);
+								bHasStandardSurface = true;
+							}
+						}
+
+						if(!bHasStandardSurface)
+						{
+							UInterchangeResultWarning_Generic* Message = AddMessage<UInterchangeResultWarning_Generic>();
+							Message->Text = FText::Format(LOCTEXT("StandardSurfaceNotFound", "<{0}> has no standard_surface inputs"),
+														  FText::FromString(Node->getName().c_str()));
+						}
+					}
+					else if(bIsLightShader)
+					{
+						bIsReferencesValid = true;
+						ProcessLightShader(BaseNodeContainer, Node, Document);
+					}
 				}
 			}
 		}
 	}
+	catch(std::exception& Exception)
+	{
+		bIsDocumentValid = false;
+		bIsReferencesValid = false;
+		UInterchangeResultError_Generic* Message = AddMessage<UInterchangeResultError_Generic>();
+		Message->Text = FText::Format(LOCTEXT("MaterialXException", "{0}"),
+									  FText::FromString(Exception.what()));
+	}
 
 #endif // WITH_EDITOR
 
-	if (bIsDocumentValid && bIsReferencesValid)
+	if(bIsDocumentValid && bIsReferencesValid)
 	{
 		UInterchangeSourceNode* SourceNode = UInterchangeSourceNode::FindOrCreateUniqueInstance(&BaseNodeContainer);
 		SourceNode->SetCustomImportUnusedMaterial(true);
@@ -218,9 +261,9 @@ void UInterchangeMaterialXTranslator::ProcessStandardSurface(UInterchangeBaseNod
 {
 	using namespace UE::Interchange::Materials;
 
-	FString StandardSurfaceName{ FPaths::GetBaseFilename(StandardSurfaceNode->getActiveSourceUri().c_str()) };
+	const FString StandardSurfaceName{ StandardSurfaceNode->getName().c_str() };
 	UInterchangeShaderGraphNode* ShaderGraphNode = NewObject<UInterchangeShaderGraphNode>(&NodeContainer);
-	FString ShaderGraphNodeUID = TEXT("\\Material\\") + StandardSurfaceName;
+	const FString ShaderGraphNodeUID = TEXT("\\Material\\") + StandardSurfaceName;
 	ShaderGraphNode->InitializeNode(ShaderGraphNodeUID, StandardSurfaceName, EInterchangeNodeContainerType::TranslatedAsset);
 	NodeContainer.AddNode(ShaderGraphNode);
 	ShaderGraphNode->SetCustomShaderType(StandardSurface::Name.ToString());
@@ -490,6 +533,157 @@ void UInterchangeMaterialXTranslator::ProcessStandardSurface(UInterchangeBaseNod
 	}
 }
 
+void UInterchangeMaterialXTranslator::ProcessLightShader(UInterchangeBaseNodeContainer& NodeContainer, MaterialX::NodePtr LightShaderNode, MaterialX::DocumentPtr Document) const
+{
+	const FString FileName{ FPaths::GetBaseFilename(LightShaderNode->getActiveSourceUri().c_str()) };
+	const FString LightNodeLabel = FString(LightShaderNode->getName().c_str());
+	UInterchangeSceneNode* SceneNode = NewObject<UInterchangeSceneNode>(&NodeContainer);
+	const FString SceneNodeUID = TEXT("\\Light\\") + FileName + TEXT("\\") + FString(LightShaderNode->getName().c_str());
+	SceneNode->InitializeNode(SceneNodeUID, LightNodeLabel, EInterchangeNodeContainerType::TranslatedScene);
+	NodeContainer.AddNode(SceneNode);
+
+	UInterchangeBaseLightNode* LightNode = nullptr;
+	if(LightShaderNode->getCategory() == mx::Category::PointLight)
+	{
+		LightNode = CreatePointLightNode(LightShaderNode, SceneNode, NodeContainer, Document);
+	}
+	else if(LightShaderNode->getCategory() == mx::Category::DirectionalLight)
+	{
+		LightNode = CreateDirectionalLightNode(LightShaderNode, SceneNode, NodeContainer, Document);
+	}
+	else if(LightShaderNode->getCategory() == mx::Category::SpotLight)
+	{
+		LightNode = CreateSpotLightNode(LightShaderNode, SceneNode, NodeContainer, Document);
+	}
+	else // MaterialX has no standardized lights, these 3 are the most common ones though and serves as an example in the format
+	{
+		LightNode = CreatePointLightNode(LightShaderNode, SceneNode, NodeContainer, Document);
+	}
+
+	const FString LightNodeUid = TEXT("\\Light\\") + LightNodeLabel;
+	LightNode->InitializeNode(LightNodeUid, LightNodeLabel, EInterchangeNodeContainerType::TranslatedAsset);
+	NodeContainer.AddNode(LightNode);
+	const FString Uid = LightNode->GetUniqueID();
+	SceneNode->SetCustomAssetInstanceUid(Uid);
+
+	// Color
+	{
+		mx::InputPtr LightColor = LightShaderNode->getInput(mx::Lights::Input::Color);
+		const mx::Color3 Color = mx::fromValueString<mx::Color3>(LightColor->getValueString());
+		LightNode->SetCustomLightColor({ Color[0], Color[1], Color[2] });
+	}
+
+	// Intensity
+	{
+		mx::InputPtr LightIntensity = LightShaderNode->getInput(mx::Lights::Input::Intensity);
+		LightNode->SetCustomIntensity(mx::fromValueString<float>(LightIntensity->getValueString()));
+	}
+}
+
+static void GetLightDirection(mx::InputPtr DirectionInput, FTransform & Transform)
+{
+	mx::Vector3 Direction = mx::fromValueString<mx::Vector3>(DirectionInput->getValueString());
+
+	//Get rotation to go from UE's default direction of directional light and the direction of the MX directional light node
+	const FVector LightDirection{ Direction[2], Direction[0], Direction[1] };
+	const FVector TransformDirection = Transform.GetUnitAxis(EAxis::X); // it's the default direction of a UE directional light
+	Transform.SetRotation(FQuat::FindBetween(LightDirection, TransformDirection));
+}
+
+static void GetLightPosition(mx::InputPtr PositionInput, FTransform& Transform)
+{
+	const mx::Vector3 Position = mx::fromValueString<mx::Vector3>(PositionInput->getValueString());
+	Transform.SetLocation(FVector{ Position[0] * 100, Position[1] * 100, Position[2] * 100 });
+}
+
+UInterchangeBaseLightNode* UInterchangeMaterialXTranslator::CreateDirectionalLightNode(MaterialX::NodePtr DirectionalLightShaderNode, UInterchangeSceneNode* SceneNode, UInterchangeBaseNodeContainer& NodeContainer, MaterialX::DocumentPtr Document) const
+{
+	UInterchangeDirectionalLightNode* LightNode = NewObject<UInterchangeDirectionalLightNode>(&NodeContainer);
+
+	// Direction
+	{
+		mx::InputPtr DirectionInput = GetDirectionalLightInput(DirectionalLightShaderNode, mx::Lights::DirectionalLight::Input::Direction, Document);
+
+		FTransform Transform;
+		GetLightDirection(DirectionInput, Transform);
+
+		SceneNode->SetCustomLocalTransform(&NodeContainer, Transform);
+	}
+
+	return LightNode;
+}
+
+UInterchangeBaseLightNode* UInterchangeMaterialXTranslator::CreatePointLightNode(MaterialX::NodePtr PointLightShaderNode, UInterchangeSceneNode* SceneNode, UInterchangeBaseNodeContainer& NodeContainer, MaterialX::DocumentPtr Document) const
+{
+	UInterchangePointLightNode* LightNode = NewObject<UInterchangePointLightNode>(&NodeContainer);
+	LightNode->SetCustomIntensityUnits(ELightUnits::Candelas);
+
+	// Decay rate
+	{
+		mx::InputPtr DecayRateInput = GetPointLightInput(PointLightShaderNode, mx::Lights::PointLight::Input::DecayRate, Document);
+		const float DecayRate = mx::fromValueString<float>(DecayRateInput->getValueString());
+		LightNode->SetCustomUseInverseSquaredFalloff(false);
+		LightNode->SetCustomLightFalloffExponent(DecayRate);
+	}
+
+	// Position
+	{
+		mx::InputPtr PositionInput = GetPointLightInput(PointLightShaderNode, mx::Lights::PointLight::Input::Position, Document);
+		FTransform Transform;
+		GetLightPosition(PositionInput, Transform);
+		SceneNode->SetCustomLocalTransform(&NodeContainer, Transform);
+	}
+
+	return LightNode;
+}
+
+UInterchangeBaseLightNode* UInterchangeMaterialXTranslator::CreateSpotLightNode(MaterialX::NodePtr SpotLightShaderNode, UInterchangeSceneNode* SceneNode, UInterchangeBaseNodeContainer& NodeContainer, MaterialX::DocumentPtr Document) const
+{
+	UInterchangeSpotLightNode* LightNode = NewObject<UInterchangeSpotLightNode>(&NodeContainer);
+
+	LightNode->SetCustomIntensityUnits(ELightUnits::Candelas);
+
+	// Decay rate
+	{
+		mx::InputPtr DecayRateInput = GetSpotLightInput(SpotLightShaderNode, mx::Lights::SpotLight::Input::DecayRate, Document);
+		const float DecayRate = mx::fromValueString<float>(DecayRateInput->getValueString());
+		LightNode->SetCustomUseInverseSquaredFalloff(false);
+		LightNode->SetCustomLightFalloffExponent(DecayRate);
+	}
+
+	{
+		FTransform Transform;
+		// Position
+		{
+			mx::InputPtr PositionInput = GetSpotLightInput(SpotLightShaderNode, mx::Lights::SpotLight::Input::Position, Document);
+			GetLightPosition(PositionInput, Transform);
+		}
+
+		// Direction
+		{
+			mx::InputPtr DirectionInput = GetSpotLightInput(SpotLightShaderNode, mx::Lights::DirectionalLight::Input::Direction, Document);
+			GetLightDirection(DirectionInput, Transform);
+		}
+		SceneNode->SetCustomLocalTransform(&NodeContainer, Transform);
+	}
+
+	// Inner angle
+	{
+		mx::InputPtr InnerAngleInput = GetSpotLightInput(SpotLightShaderNode, mx::Lights::SpotLight::Input::InnerAngle, Document);
+		const float InnerAngle = FMath::RadiansToDegrees(mx::fromValueString<float>(InnerAngleInput->getValueString()));
+		LightNode->SetCustomInnerConeAngle(InnerAngle);
+	}
+
+	// Outer angle
+	{
+		mx::InputPtr OuterAngleInput = GetSpotLightInput(SpotLightShaderNode, mx::Lights::SpotLight::Input::OuterAngle, Document);
+		const float OuterAngle = FMath::RadiansToDegrees(mx::fromValueString<float>(OuterAngleInput->getValueString()));
+		LightNode->SetCustomOuterConeAngle(OuterAngle);
+	}
+
+	return LightNode;
+}
+
 bool UInterchangeMaterialXTranslator::ConnectNodeGraphOutputToInput(MaterialX::InputPtr InputToNodeGraph, UInterchangeShaderNode* ShaderNode, const FString& ParentInputName, TMap<FString, UInterchangeShaderNode*>& NamesToShaderNodes, UInterchangeBaseNodeContainer& NodeContainer) const
 {
 	using namespace UE::Interchange::Materials::Standard::Nodes;
@@ -506,7 +700,7 @@ bool UInterchangeMaterialXTranslator::ConnectNodeGraphOutputToInput(MaterialX::I
 		{
 			UInterchangeResultWarning_Generic* Message = AddMessage<UInterchangeResultWarning_Generic>();
 			Message->Text = FText::Format(LOCTEXT("OutputNotFound", "Couldn't find a connected output to ({0})"),
-				FText::FromString(InputToNodeGraph->getName().c_str()));
+										  FText::FromString(InputToNodeGraph->getName().c_str()));
 
 			return false;
 		}
@@ -544,8 +738,8 @@ bool UInterchangeMaterialXTranslator::ConnectNodeGraphOutputToInput(MaterialX::I
 					{
 						UInterchangeResultWarning_Generic* Message = AddMessage<UInterchangeResultWarning_Generic>();
 						Message->Text = FText::Format(LOCTEXT("InputTypeNotSupported", "<{0}>: \"{1}\" is not supported yet"),
-							FText::FromString(InputConstant->getName().c_str()),
-							FText::FromString(InputConstant->getType().c_str()));
+													  FText::FromString(InputConstant->getName().c_str()),
+													  FText::FromString(InputConstant->getType().c_str()));
 					}
 
 				}
@@ -555,7 +749,7 @@ bool UInterchangeMaterialXTranslator::ConnectNodeGraphOutputToInput(MaterialX::I
 
 					if(mx::InputPtr InputIndex = UpstreamNode->getInput("index"))
 					{
-						int32 Index = mx::fromValueString<int>(InputIndex->getValueString());
+						const int32 Index = mx::fromValueString<int>(InputIndex->getValueString());
 						switch(Index)
 						{
 						case 0: MaskShaderNode->AddBooleanAttribute(Mask::Attributes::R, true); break;
@@ -608,7 +802,7 @@ bool UInterchangeMaterialXTranslator::ConnectNodeGraphOutputToInput(MaterialX::I
 				{
 					UInterchangeResultWarning_Generic* Message = AddMessage<UInterchangeResultWarning_Generic>();
 					Message->Text = FText::Format(LOCTEXT("NodeCategoryNotSupported", "<{0}> is not supported yet"),
-						FText::FromString(UpstreamNode->getCategory().c_str()));
+												  FText::FromString(UpstreamNode->getCategory().c_str()));
 				}
 			}
 		}
@@ -677,11 +871,13 @@ UInterchangeTextureNode* UInterchangeMaterialXTranslator::CreateTextureNode(Mate
 		if(mx::InputPtr InputFile = Node->getInput("file"); InputFile && InputFile->hasValue())
 		{
 			FString Filepath{ InputFile->getValueString().c_str() };
-			FString Filename = FPaths::GetCleanFilename(Filepath);
-			FString TextureNodeUID = TEXT("\\Texture\\") + Filename;
+			const FString FilePrefix = GetFilePrefix(InputFile);
+			Filepath = FPaths::Combine(FilePrefix, Filepath);
+			const FString Filename = FPaths::GetCleanFilename(Filepath);
+			const FString TextureNodeUID = TEXT("\\Texture\\") + Filename;
 
 			//Only add the TextureNode once
-			TextureNode = const_cast<UInterchangeTextureNode *>(Cast<UInterchangeTextureNode>(NodeContainer.GetNode(TextureNodeUID)));
+			TextureNode = const_cast<UInterchangeTextureNode*>(Cast<UInterchangeTextureNode>(NodeContainer.GetNode(TextureNodeUID)));
 			if(TextureNode == nullptr)
 			{
 				TextureNode = NewObject<UInterchangeTexture2DNode>(&NodeContainer);
@@ -694,6 +890,10 @@ UInterchangeTextureNode* UInterchangeMaterialXTranslator::CreateTextureNode(Mate
 				}
 
 				TextureNode->SetPayLoadKey(Filepath);
+
+				const FString ColorSpace = GetColorSpace(InputFile);
+				const bool bIsSRGB = ColorSpace == TEXT("srgb_texture");
+				TextureNode->SetCustomSRGB(bIsSRGB);
 			}
 		}
 	}
@@ -763,6 +963,42 @@ MaterialX::InputPtr UInterchangeMaterialXTranslator::GetStandardSurfaceInput(Mat
 	return Input;
 }
 
+MaterialX::InputPtr UInterchangeMaterialXTranslator::GetPointLightInput(MaterialX::NodePtr PointLight, const char* InputName, MaterialX::DocumentPtr Document) const
+{
+	mx::InputPtr Input = PointLight->getInput(InputName);
+
+	if(!Input)
+	{
+		Input = Document->getNodeDef(mx::NodeDefinition::PointLight)->getInput(InputName);
+	}
+
+	return Input;
+}
+
+MaterialX::InputPtr UInterchangeMaterialXTranslator::GetDirectionalLightInput(MaterialX::NodePtr DirectionalLight, const char* InputName, MaterialX::DocumentPtr Document) const
+{
+	mx::InputPtr Input = DirectionalLight->getInput(InputName);
+
+	if(!Input)
+	{
+		Input = Document->getNodeDef(mx::NodeDefinition::DirectionalLight)->getInput(InputName);
+	}
+
+	return Input;
+}
+
+MaterialX::InputPtr UInterchangeMaterialXTranslator::GetSpotLightInput(MaterialX::NodePtr SpotLight, const char* InputName, MaterialX::DocumentPtr Document) const
+{
+	mx::InputPtr Input = SpotLight->getInput(InputName);
+
+	if(!Input)
+	{
+		Input = Document->getNodeDef(mx::NodeDefinition::SpotLight)->getInput(InputName);
+	}
+
+	return Input;
+}
+
 bool UInterchangeMaterialXTranslator::AddAttribute(MaterialX::InputPtr Input, const FString& InputChannelName, UInterchangeShaderNode* ShaderNode) const
 {
 	if(Input)
@@ -834,6 +1070,44 @@ bool UInterchangeMaterialXTranslator::AddLinearColorAttribute(MaterialX::InputPt
 
 	return false;
 }
+
+FString UInterchangeMaterialXTranslator::GetFilePrefix(MaterialX::ElementPtr Element) const
+{
+	FString FilePrefix;
+
+	if(Element)
+	{
+		if(Element->hasFilePrefix())
+		{
+			return FString(Element->getFilePrefix().c_str());
+		}
+		else
+		{
+			return GetFilePrefix(Element->getParent());
+		}
+	}
+
+	return FilePrefix;
+}
+
+FString UInterchangeMaterialXTranslator::GetColorSpace(MaterialX::ElementPtr Element) const
+{
+	FString ColorSpace;
+
+	if(Element)
+	{
+		if(Element->hasColorSpace())
+		{
+			return FString(Element->getColorSpace().c_str());
+		}
+		else
+		{
+			return GetColorSpace(Element->getParent());
+		}
+	}
+
+	return ColorSpace;
+}
 #endif //WITH_EDITOR
 
 #undef LOCTEXT_NAMESPACE
\ No newline at end of file
diff --git a/Engine/Plugins/Experimental/Interchange/Runtime/Source/Import/Public/MaterialX/InterchangeMaterialXTranslator.h b/Engine/Plugins/Experimental/Interchange/Runtime/Source/Import/Public/MaterialX/InterchangeMaterialXTranslator.h
--- a/Engine/Plugins/Experimental/Interchange/Runtime/Source/Import/Public/MaterialX/InterchangeMaterialXTranslator.h
+++ b/Engine/Plugins/Experimental/Interchange/Runtime/Source/Import/Public/MaterialX/InterchangeMaterialXTranslator.h
@@ -18,6 +18,8 @@ PRAGMA_ENABLE_MISSING_BRACES_WARNINGS
 
 class UInterchangeTextureNode;
 class UInterchangeShaderNode;
+class UInterchangeBaseLightNode;
+class UInterchangeSceneNode;
 
 UCLASS(BlueprintType, Experimental)
 class UInterchangeMaterialXTranslator : public UInterchangeTranslatorBase, public IInterchangeTexturePayloadInterface
@@ -29,6 +31,11 @@ public:
 	UInterchangeMaterialXTranslator();
 
 	/** Begin UInterchangeTranslatorBase API*/
+
+	virtual EInterchangeTranslatorType GetTranslatorType() const override;
+
+	virtual bool DoesSupportAssetType(EInterchangeTranslatorAssetType AssetType) const override;
+
 	virtual TArray<FString> GetSupportedFormats() const override;
 
 	/**
@@ -50,9 +57,48 @@ protected:
 	 * 
 	 * @param NodeContainer - The unreal objects descriptions container where to put the translated source data.
 	 * @param StandardSurfaceNode - The <standard_surface> in the MaterialX file
+	 * @param Document - The MaterialX Document that contains all the definitions of the loaded libraries
 	 */
 	void ProcessStandardSurface(UInterchangeBaseNodeContainer & NodeContainer, MaterialX::NodePtr StandardSurfaceNode, MaterialX::DocumentPtr Document) const;
 
+	/**
+	 * Process light shader, MaterialX doesn't standardized lights, but defines the 3 common ones, directional, point and spot
+	 * 
+	 * @param NodeContainer - The unreal objects descriptions container where to put the translated source data.
+	 * @param LightShaderNode - a MaterialX light shader node
+	 * @param Document - The MaterialX Document that contains all the definitions of the loaded libraries
+	 */
+	void ProcessLightShader(UInterchangeBaseNodeContainer& NodeContainer, MaterialX::NodePtr LightShaderNode, MaterialX::DocumentPtr Document) const;
+
+	/**
+	 * Create a directional light node and set the proper transform in the scene node
+	 *
+	 * @param DirectionalLightShaderNode - The MaterialX node related to a <directional_light>
+	 * @param SceneNode - The interchange scene node, to which the node will be attached to
+	 * @param NodeContainer - The unreal objects descriptions container where to put the translated source data.
+	 * @param Document - The MaterialX Document that contains all the definitions of the loaded libraries
+	 */
+	UInterchangeBaseLightNode* CreateDirectionalLightNode(MaterialX::NodePtr DirectionalLightShaderNode, UInterchangeSceneNode* SceneNode, UInterchangeBaseNodeContainer& NodeContainer, MaterialX::DocumentPtr Document) const;
+
+	/**
+	 * Create a point  light node and set the proper transform in the scene node
+	 *
+	 * @param PointLightShaderNode - The MaterialX node related to a <point_light>
+	 * @param SceneNode - The interchange scene node, to which the node will be attached to
+	 * @param NodeContainer - The unreal objects descriptions container where to put the translated source data.
+	 * @param Document - The MaterialX Document that contains all the definitions of the loaded libraries
+	 */
+	UInterchangeBaseLightNode* CreatePointLightNode(MaterialX::NodePtr PointLightShaderNode, UInterchangeSceneNode* SceneNode, UInterchangeBaseNodeContainer& NodeContainer, MaterialX::DocumentPtr Document) const;
+	
+	/**
+	 * Create a spot light node and set the proper transform in the scene node
+	 *
+	 * @param SpotLightShaderNode - The MaterialX node related to a <spot_light>
+	 * @param SceneNode - The interchange scene node, to which the node will be attached to
+	 * @param NodeContainer - The unreal objects descriptions container where to put the translated source data.
+	 * @param Document - The MaterialX Document that contains all the definitions of the loaded libraries
+	 */
+	UInterchangeBaseLightNode* CreateSpotLightNode(MaterialX::NodePtr SpotLightShaderNode, UInterchangeSceneNode* SceneNode, UInterchangeBaseNodeContainer& NodeContainer, MaterialX::DocumentPtr Document) const;
 	/**
 	 * Connect an ouput in the NodeGraph to the ShaderGraph
 	 * 
@@ -141,6 +187,45 @@ protected:
 	 */
 	MaterialX::InputPtr GetStandardSurfaceInput(MaterialX::NodePtr StandardSurface, const char* InputName, MaterialX::DocumentPtr Document) const;
 
+	/**
+	 * Retrieve the input from a point_light node, or take the default input from the library,
+	 * this function should only be called after testing the MaterialX libraries have been successfully imported, meaning the node definition of the point_light
+	 * should always be valid
+	 *
+	 * @param PointLight - the <point_light> node
+	 * @param InputName - the input name to retrieve
+	 * @param Document - the MaterialX library which has the node definition of the standard_surface
+	 *
+	 * @return the input from the given name
+	 */
+	MaterialX::InputPtr GetPointLightInput(MaterialX::NodePtr PointLight, const char* InputName, MaterialX::DocumentPtr Document) const;
+
+	/**
+	 * Retrieve the input from a directional_light node, or take the default input from the library,
+	 * this function should only be called after testing the MaterialX libraries have been successfully imported, meaning the node definition of the directional_light
+	 * should always be valid
+	 *
+	 * @param DirectionalLight - the <directional_light> node
+	 * @param InputName - the input name to retrieve
+	 * @param Document - the MaterialX library which has the node definition of the standard_surface
+	 *
+	 * @return the input from the given name
+	 */
+	MaterialX::InputPtr GetDirectionalLightInput(MaterialX::NodePtr DirectionalLight, const char* InputName, MaterialX::DocumentPtr Document) const;
+
+	/**
+	 * Retrieve the input from a spot_light node, or take the default input from the library,
+	 * this function should only be called after testing the MaterialX libraries have been successfully imported, meaning the node definition of the spot_light
+	 * should always be valid
+	 *
+	 * @param SpotLight - the <spot_light > node
+	 * @param InputName - the input name to retrieve
+	 * @param Document - the MaterialX library which has the node definition of the standard_surface
+	 *
+	 * @return the input from the given name
+	 */
+	MaterialX::InputPtr GetSpotLightInput(MaterialX::NodePtr SpotLight, const char* InputName, MaterialX::DocumentPtr Document) const;
+
 	/**
 	 * Add an attribute to a shader node, only floats and linear colors are supported for the moment
 	 * 
@@ -176,6 +261,26 @@ protected:
 	 */
 	bool AddLinearColorAttribute(MaterialX::InputPtr Input, const FString& InputChannelName, UInterchangeShaderNode* ShaderNode, const FLinearColor& DefaultValue) const;
 
+	/**
+	 * Return the innermost file prefix of an element in the current scope, if it has none, it will take the one from its parents
+	 * 
+	 * @param Element - the Element to retrieve the file prefix from (can be anything, an input, a node, a nodegraph, etc.)
+	 * 
+	 * @return a file prefix or an empty string
+	 * 
+	 */
+	FString GetFilePrefix(MaterialX::ElementPtr Element) const;
+
+	/**
+	 * Return the innermost color space of an element in the current scope, if it has none, it will take the one from its parents
+	 *
+	 * @param Element - the Element to retrieve the color space from (can be anything, an input, a node, a nodegraph, etc.)
+	 *
+	 * @return a color space or an empty string
+	 *
+	 */
+	FString GetColorSpace(MaterialX::ElementPtr Element) const;
+
 private:
 
 	TMap<FString, FString> InputNamesMaterialX2UE;
diff --git a/Engine/Plugins/Experimental/Interchange/Runtime/Source/Nodes/Private/InterchangeLightNode.cpp b/Engine/Plugins/Experimental/Interchange/Runtime/Source/Nodes/Private/InterchangeLightNode.cpp
new file mode 100644
--- /dev/null
+++ b/Engine/Plugins/Experimental/Interchange/Runtime/Source/Nodes/Private/InterchangeLightNode.cpp
@@ -0,0 +1,180 @@
+// Copyright Epic Games, Inc. All Rights Reserved.
+
+#include "InterchangeLightNode.h"
+
+FStringView UInterchangeBaseLightNode::StaticAssetTypeName()
+{
+	return TEXT("Light");
+}
+
+FString UInterchangeBaseLightNode::GetTypeName() const
+{
+	return TEXT("BaseLightNode");
+}
+
+bool UInterchangeBaseLightNode::GetCustomLightColor(FLinearColor& AttributeValue) const
+{
+	IMPLEMENT_NODE_ATTRIBUTE_GETTER(LightColor, FLinearColor)
+}
+
+bool UInterchangeBaseLightNode::SetCustomLightColor(const FLinearColor& AttributeValue)
+{
+	IMPLEMENT_NODE_ATTRIBUTE_SETTER_NODELEGATE(LightColor, FLinearColor)
+}
+
+bool UInterchangeBaseLightNode::GetCustomIntensity(float& AttributeValue) const
+{
+	IMPLEMENT_NODE_ATTRIBUTE_GETTER(LightIntensity, float)
+}
+
+bool UInterchangeBaseLightNode::SetCustomIntensity(float AttributeValue)
+{
+	IMPLEMENT_NODE_ATTRIBUTE_SETTER_NODELEGATE(LightIntensity, float)
+}
+
+bool UInterchangeBaseLightNode::GetCustomTemperature(float& AttributeValue) const
+{
+	IMPLEMENT_NODE_ATTRIBUTE_GETTER(Temperature, float)
+}
+
+bool UInterchangeBaseLightNode::SetCustomTemperature(float AttributeValue)
+{
+	IMPLEMENT_NODE_ATTRIBUTE_SETTER_NODELEGATE(Temperature, float)
+}
+
+bool UInterchangeBaseLightNode::GetCustomUseTemperature(bool& AttributeValue) const
+{
+	IMPLEMENT_NODE_ATTRIBUTE_GETTER(UseTemperature, bool)
+}
+
+bool UInterchangeBaseLightNode::SetCustomUseTemperature(bool AttributeValue)
+{
+	IMPLEMENT_NODE_ATTRIBUTE_SETTER_NODELEGATE(UseTemperature, bool)
+}
+
+FString UInterchangeLightNode::GetTypeName() const
+{
+	return TEXT("LightNode");
+}
+
+bool UInterchangeLightNode::GetCustomIntensityUnits(ELightUnits& AttributeValue) const
+{
+	IMPLEMENT_NODE_ATTRIBUTE_GETTER(IntensityUnits, ELightUnits)
+}
+
+bool UInterchangeLightNode::SetCustomIntensityUnits(ELightUnits AttributeValue)
+{
+	IMPLEMENT_NODE_ATTRIBUTE_SETTER_NODELEGATE(IntensityUnits, ELightUnits)
+}
+
+bool UInterchangeLightNode::GetCustomAttenuationRadius(float& AttributeValue) const
+{
+	IMPLEMENT_NODE_ATTRIBUTE_GETTER(AttenuationRadius, float)
+}
+
+bool UInterchangeLightNode::SetCustomAttenuationRadius(float AttributeValue)
+{
+	IMPLEMENT_NODE_ATTRIBUTE_SETTER_NODELEGATE(AttenuationRadius, float)
+}
+
+FString UInterchangePointLightNode::GetTypeName() const
+{
+	return TEXT("PointLightNode");
+}
+
+bool UInterchangePointLightNode::GetCustomUseInverseSquaredFalloff(bool& AttributeValue) const
+{
+	IMPLEMENT_NODE_ATTRIBUTE_GETTER(UseInverseSquaredFalloff, bool)
+}
+
+bool UInterchangePointLightNode::SetCustomUseInverseSquaredFalloff(bool AttributeValue)
+{
+	IMPLEMENT_NODE_ATTRIBUTE_SETTER_NODELEGATE(UseInverseSquaredFalloff, bool)
+}
+
+bool UInterchangePointLightNode::GetCustomLightFalloffExponent(float& AttributeValue) const
+{
+	IMPLEMENT_NODE_ATTRIBUTE_GETTER(LightFalloffExponent, float)
+}
+
+bool UInterchangePointLightNode::SetCustomLightFalloffExponent(float AttributeValue)
+{
+	IMPLEMENT_NODE_ATTRIBUTE_SETTER_NODELEGATE(LightFalloffExponent, float)
+}
+
+FString UInterchangeSpotLightNode::GetTypeName() const
+{
+	return TEXT("SpotLightNode");
+}
+
+bool UInterchangeSpotLightNode::GetCustomInnerConeAngle(float& AttributeValue) const
+{
+	IMPLEMENT_NODE_ATTRIBUTE_GETTER(InnerConeAngle, float)
+}
+
+bool UInterchangeSpotLightNode::SetCustomInnerConeAngle(float AttributeValue)
+{
+	IMPLEMENT_NODE_ATTRIBUTE_SETTER_NODELEGATE(InnerConeAngle, float)
+}
+
+bool UInterchangeSpotLightNode::GetCustomOuterConeAngle(float& AttributeValue) const
+{
+	IMPLEMENT_NODE_ATTRIBUTE_GETTER(OuterConeAngle, float)
+}
+
+bool UInterchangeSpotLightNode::SetCustomOuterConeAngle(float AttributeValue)
+{
+	IMPLEMENT_NODE_ATTRIBUTE_SETTER_NODELEGATE(OuterConeAngle, float)
+}
+
+FString UInterchangeRectLightNode::GetTypeName() const
+{
+	return TEXT("RectLightNode");
+}
+
+//For the moment n/a, need more discussion specially about barn door's angle
+/*
+bool UInterchangeRectLightNode::GetCustomSourceWidth(float& AttributeValue) const
+{
+	IMPLEMENT_NODE_ATTRIBUTE_GETTER(SourceWidth, float)
+}
+
+bool UInterchangeRectLightNode::SetCustomSourceWidth(float AttributeValue)
+{
+	IMPLEMENT_NODE_ATTRIBUTE_SETTER_NODELEGATE(SourceWidth, float)
+}
+
+bool UInterchangeRectLightNode::GetCustomSourceHeight(float& AttributeValue) const
+{
+	IMPLEMENT_NODE_ATTRIBUTE_GETTER(SourceHeight, float)
+}
+
+bool UInterchangeRectLightNode::SetCustomSourceHeight(float AttributeValue)
+{
+	IMPLEMENT_NODE_ATTRIBUTE_SETTER_NODELEGATE(SourceHeight, float)
+}
+
+bool UInterchangeRectLightNode::GetCustomBarnDoorAngle(float& AttributeValue) const
+{
+	IMPLEMENT_NODE_ATTRIBUTE_GETTER(BarnDoorAngle, float)
+}
+
+bool UInterchangeRectLightNode::SetCustomBarnDoorAngle(float AttributeValue)
+{
+	IMPLEMENT_NODE_ATTRIBUTE_SETTER_NODELEGATE(BarnDoorAngle, float)
+}
+
+bool UInterchangeRectLightNode::GetCustomBarnDoorLength(float& AttributeValue) const
+{
+	IMPLEMENT_NODE_ATTRIBUTE_GETTER(BarnDoorLength, float)
+}
+
+bool UInterchangeRectLightNode::SetCustomBarnDoorLength(float AttributeValue)
+{
+	IMPLEMENT_NODE_ATTRIBUTE_SETTER_NODELEGATE(BarnDoorLength, float)
+}
+*/
+FString UInterchangeDirectionalLightNode::GetTypeName() const
+{
+	return TEXT("DirectionalLightNode");
+}
\ No newline at end of file
diff --git a/Engine/Plugins/Experimental/Interchange/Runtime/Source/Nodes/Public/InterchangeLightNode.h b/Engine/Plugins/Experimental/Interchange/Runtime/Source/Nodes/Public/InterchangeLightNode.h
--- a/Engine/Plugins/Experimental/Interchange/Runtime/Source/Nodes/Public/InterchangeLightNode.h
+++ b/Engine/Plugins/Experimental/Interchange/Runtime/Source/Nodes/Public/InterchangeLightNode.h
@@ -8,25 +8,75 @@
 #include "InterchangeLightNode.generated.h"
 
 UCLASS(BlueprintType, Experimental)
-class INTERCHANGENODES_API UInterchangeLightNode : public UInterchangeBaseNode
+class INTERCHANGENODES_API UInterchangeBaseLightNode : public UInterchangeBaseNode
 {
 	GENERATED_BODY()
 
 public:
-	static FStringView StaticAssetTypeName()
-	{
-		return TEXT("Light");
-	}
+	static FStringView StaticAssetTypeName();
 
 	/**
 	 * Return the node type name of the class, we use this when reporting errors
 	 */
-	virtual FString GetTypeName() const override
-	{
-		const FString TypeName = TEXT("LightNode");
-		return TypeName;
-	}
+	virtual FString GetTypeName() const override;
 
+	UFUNCTION(BlueprintCallable, Category = "Interchange | Node | Light")
+	bool GetCustomLightColor(FLinearColor& AttributeValue) const;
+
+	UFUNCTION(BlueprintCallable, Category = "Interchange | Node | Light")
+	bool SetCustomLightColor(const FLinearColor& AttributeValue);
+
+	UFUNCTION(BlueprintCallable, Category = "Interchange | Node | Light")
+	bool GetCustomIntensity(float& AttributeValue) const;
+
+	UFUNCTION(BlueprintCallable, Category = "Interchange | Node | Light")
+	bool SetCustomIntensity(float AttributeValue);
+
+	UFUNCTION(BlueprintCallable, Category = "Interchange | Node | Light")
+	bool GetCustomTemperature(float& AttributeValue) const;
+
+	UFUNCTION(BlueprintCallable, Category = "Interchange | Node | Light")
+	bool SetCustomTemperature(float AttributeValue);
+
+	UFUNCTION(BlueprintCallable, Category = "Interchange | Node | Light")
+	bool GetCustomUseTemperature(bool & AttributeValue) const;
+
+	UFUNCTION(BlueprintCallable, Category = "Interchange | Node | Light")
+	bool SetCustomUseTemperature(bool AttributeValue);
+
+private:
+
+	IMPLEMENT_NODE_ATTRIBUTE_KEY(LightColor)
+	IMPLEMENT_NODE_ATTRIBUTE_KEY(LightIntensity)
+	IMPLEMENT_NODE_ATTRIBUTE_KEY(Temperature)
+	IMPLEMENT_NODE_ATTRIBUTE_KEY(UseTemperature)
+};
+
+UCLASS(BlueprintType)
+class INTERCHANGENODES_API UInterchangeLightNode : public UInterchangeBaseLightNode
+{
+	GENERATED_BODY()
+
+public:
+
+	virtual FString GetTypeName() const override;
+
+	UFUNCTION(BlueprintCallable, Category = "Interchange | Node | LocalLight")
+	bool GetCustomIntensityUnits(ELightUnits& AttributeValue) const;
+
+	UFUNCTION(BlueprintCallable, Category = "Interchange | Node | LocalLight")
+	bool SetCustomIntensityUnits(ELightUnits AttributeValue);
+
+	UFUNCTION(BlueprintCallable, Category = "Interchange | Node | LocalLight")
+	bool GetCustomAttenuationRadius(float& AttributeValue) const;
+
+	UFUNCTION(BlueprintCallable, Category = "Interchange | Node | LocalLight")
+	bool SetCustomAttenuationRadius(float AttributeValue);
+
+private:
+
+	IMPLEMENT_NODE_ATTRIBUTE_KEY(IntensityUnits)
+	IMPLEMENT_NODE_ATTRIBUTE_KEY(AttenuationRadius)
 };
 
 UCLASS(BlueprintType)
@@ -36,6 +86,23 @@ class INTERCHANGENODES_API UInterchangePointLightNode : public UInterchangeLight
 
 public:
 
+	virtual FString GetTypeName() const override;
+
+	UFUNCTION(BlueprintCallable, Category = "Interchange | Node | PointLight")
+	bool GetCustomUseInverseSquaredFalloff(bool& AttributeValue) const;
+
+	UFUNCTION(BlueprintCallable, Category = "Interchange | Node | PointLight")
+	bool SetCustomUseInverseSquaredFalloff(bool AttributeValue);
+
+	UFUNCTION(BlueprintCallable, Category = "Interchange | Node | PointLight")
+	bool GetCustomLightFalloffExponent(float& AttributeValue) const;
+
+	UFUNCTION(BlueprintCallable, Category = "Interchange | Node | PointLight")
+	bool SetCustomLightFalloffExponent(float AttributeValue);
+
+private:
+	IMPLEMENT_NODE_ATTRIBUTE_KEY(UseInverseSquaredFalloff)
+	IMPLEMENT_NODE_ATTRIBUTE_KEY(LightFalloffExponent)
 };
 
 UCLASS(BlueprintType)
@@ -45,6 +112,24 @@ class INTERCHANGENODES_API UInterchangeSpotLightNode : public UInterchangePointL
 
 public:
 
+	virtual FString GetTypeName() const override;
+
+	UFUNCTION(BlueprintCallable, Category = "Interchange | Node | SpotLight")
+	bool GetCustomInnerConeAngle(float& AttributeValue) const;
+
+	UFUNCTION(BlueprintCallable, Category = "Interchange | Node | SpotLight")
+	bool SetCustomInnerConeAngle(float AttributeValue);
+
+	UFUNCTION(BlueprintCallable, Category = "Interchange | Node | SpotLight")
+	bool GetCustomOuterConeAngle(float& AttributeValue) const;
+
+	UFUNCTION(BlueprintCallable, Category = "Interchange | Node | SpotLight")
+	bool SetCustomOuterConeAngle(float AttributeValue);
+
+private:
+
+	IMPLEMENT_NODE_ATTRIBUTE_KEY(InnerConeAngle)
+	IMPLEMENT_NODE_ATTRIBUTE_KEY(OuterConeAngle)
 };
 
 UCLASS(BlueprintType)
@@ -54,13 +139,48 @@ class INTERCHANGENODES_API UInterchangeRectLightNode : public UInterchangeLightN
 
 public:
 
+	virtual FString GetTypeName() const override;
+
+	//For the moment n/a, need more discussion specially about barn door's angle
+	/*
+	UFUNCTION(BlueprintCallable, Category = "Interchange | Node | RectLightFactory")
+	bool GetCustomSourceWidth(float& AttributeValue) const;
+
+	UFUNCTION(BlueprintCallable, Category = "Interchange | Node | RectLightFactory")
+	bool SetCustomSourceWidth(float AttributeValue);
+
+	UFUNCTION(BlueprintCallable, Category = "Interchange | Node | RectLightFactory")
+	bool GetCustomSourceHeight(float& AttributeValue) const;
+
+	UFUNCTION(BlueprintCallable, Category = "Interchange | Node | RectLightFactory")
+	bool SetCustomSourceHeight(float AttributeValue);
+
+	UFUNCTION(BlueprintCallable, Category = "Interchange | Node | RectLightFactory")
+	bool GetCustomBarnDoorAngle(float& AttributeValue) const;
+
+	UFUNCTION(BlueprintCallable, Category = "Interchange | Node | RectLightFactory")
+	bool SetCustomBarnDoorAngle(float AttributeValue);
+
+	UFUNCTION(BlueprintCallable, Category = "Interchange | Node | RectLightFactory")
+	bool GetCustomBarnDoorLength(float& AttributeValue) const;
+
+	UFUNCTION(BlueprintCallable, Category = "Interchange | Node | RectLightFactory")
+	bool SetCustomBarnDoorLength(float AttributeValue);
+
+private:
+	IMPLEMENT_NODE_ATTRIBUTE_KEY(SourceWidth)
+	IMPLEMENT_NODE_ATTRIBUTE_KEY(SourceHeight)
+	IMPLEMENT_NODE_ATTRIBUTE_KEY(BarnDoorAngle)
+	IMPLEMENT_NODE_ATTRIBUTE_KEY(BarnDoorLength)
+	*/
 };
 
 UCLASS(BlueprintType)
-class INTERCHANGENODES_API UInterchangeDirectionalLightNode : public UInterchangeLightNode
+class INTERCHANGENODES_API UInterchangeDirectionalLightNode : public UInterchangeBaseLightNode
 {
 	GENERATED_BODY()
 
 public:
 
+	virtual FString GetTypeName() const override;
 };
diff --git a/Engine/Plugins/Experimental/Interchange/Runtime/Source/Pipelines/Private/InterchangeGenericScenesPipeline.cpp b/Engine/Plugins/Experimental/Interchange/Runtime/Source/Pipelines/Private/InterchangeGenericScenesPipeline.cpp
--- a/Engine/Plugins/Experimental/Interchange/Runtime/Source/Pipelines/Private/InterchangeGenericScenesPipeline.cpp
+++ b/Engine/Plugins/Experimental/Interchange/Runtime/Source/Pipelines/Private/InterchangeGenericScenesPipeline.cpp
@@ -7,11 +7,13 @@
 #include "InterchangeCineCameraFactoryNode.h"
 #include "InterchangeCommonPipelineDataFactoryNode.h"
 #include "InterchangeLightNode.h"
+#include "InterchangeLightFactoryNode.h"
 #include "InterchangeMeshActorFactoryNode.h"
 #include "InterchangeMeshNode.h"
 #include "InterchangePipelineLog.h"
 #include "InterchangePipelineMeshesUtilities.h"
 #include "InterchangeSceneNode.h"
+#include "Nodes/InterchangeUserDefinedAttribute.h"
 
 #include "Animation/SkeletalMeshActor.h"
 #include "CineCameraActor.h"
@@ -123,18 +125,35 @@ void UInterchangeGenericLevelPipeline::ExecuteSceneNodePreImport(UInterchangeBas
 
 UInterchangeActorFactoryNode* UInterchangeGenericLevelPipeline::CreateActorFactoryNode(const UInterchangeSceneNode* SceneNode, const UInterchangeBaseNode* TranslatedAssetNode, UInterchangeBaseNodeContainer* FactoryNodeContainer) const
 {
-	if (TranslatedAssetNode && TranslatedAssetNode->IsA<UInterchangeCameraNode>())
+	if(TranslatedAssetNode)
 	{
-		return NewObject<UInterchangeCineCameraFactoryNode>(FactoryNodeContainer, NAME_None);
-	}
-	else if (TranslatedAssetNode && TranslatedAssetNode->IsA<UInterchangeMeshNode>())
-	{
-		return NewObject<UInterchangeMeshActorFactoryNode>(FactoryNodeContainer, NAME_None);
-	}
-	else
-	{
-		return NewObject<UInterchangeActorFactoryNode>(FactoryNodeContainer, NAME_None);
+		if(TranslatedAssetNode->IsA<UInterchangeCameraNode>())
+		{
+			return NewObject<UInterchangeCineCameraFactoryNode>(FactoryNodeContainer, NAME_None);
+		}
+		else if(TranslatedAssetNode->IsA<UInterchangeMeshNode>())
+		{
+			return NewObject<UInterchangeMeshActorFactoryNode>(FactoryNodeContainer, NAME_None);
+		}
+		else if(TranslatedAssetNode->IsA<UInterchangeSpotLightNode>())
+		{
+			return NewObject<UInterchangeSpotLightFactoryNode>(FactoryNodeContainer, NAME_None);
+		}
+		else if(TranslatedAssetNode->IsA<UInterchangePointLightNode>())
+		{
+			return NewObject<UInterchangePointLightFactoryNode>(FactoryNodeContainer, NAME_None);
+		}
+		else if(TranslatedAssetNode->IsA<UInterchangeRectLightNode>())
+		{
+			return NewObject<UInterchangeRectLightFactoryNode>(FactoryNodeContainer, NAME_None);
+		}
+		else if(TranslatedAssetNode->IsA<UInterchangeDirectionalLightNode>())
+		{
+			return NewObject<UInterchangeDirectionalLightFactoryNode>(FactoryNodeContainer, NAME_None);
+		}
 	}
+
+	return NewObject<UInterchangeActorFactoryNode>(FactoryNodeContainer, NAME_None);
 }
 
 void UInterchangeGenericLevelPipeline::SetUpFactoryNode(UInterchangeActorFactoryNode* ActorFactoryNode, const UInterchangeSceneNode* SceneNode, const UInterchangeBaseNode* TranslatedAssetNode, UInterchangeBaseNodeContainer* FactoryNodeContainer) const
@@ -166,22 +185,124 @@ void UInterchangeGenericLevelPipeline::SetUpFactoryNode(UInterchangeActorFactory
 			MeshActorFactoryNode->AddFactoryDependencyUid(UInterchangeFactoryBaseNode::BuildFactoryNodeUid(MeshNode->GetUniqueID()));
 		}
 	}
-	else if (const UInterchangeLightNode* LightNode = Cast<UInterchangeLightNode>(TranslatedAssetNode))
+	else if (const UInterchangeBaseLightNode* BaseLightNode = Cast<UInterchangeBaseLightNode>(TranslatedAssetNode))
 	{
+		if (UInterchangeBaseLightFactoryNode* BaseLightFactoryNode = Cast<UInterchangeBaseLightFactoryNode>(ActorFactoryNode))
+		{
+			if (FLinearColor LightColor; BaseLightNode->GetCustomLightColor(LightColor))
+			{
+				BaseLightFactoryNode->SetCustomLightColor(LightColor.ToFColor(true));
+			}
+
+			if (float Intensity; BaseLightNode->GetCustomIntensity(Intensity))
+			{
+				BaseLightFactoryNode->SetCustomIntensity(Intensity);
+			}
+
+			if(bool bUseTemperature; BaseLightNode->GetCustomUseTemperature(bUseTemperature))
+			{
+				BaseLightFactoryNode->SetCustomUseTemperature(bUseTemperature);
+
+				if(float Temperature; BaseLightNode->GetCustomTemperature(Temperature))
+				{
+					BaseLightFactoryNode->SetCustomTemperature(Temperature);
+				}
+			}
+
+			if (const UInterchangeLightNode* LightNode = Cast<UInterchangeLightNode>(BaseLightNode))
+			{
+				if (UInterchangeLightFactoryNode* LightFactoryNode = Cast<UInterchangeLightFactoryNode>(BaseLightFactoryNode))
+				{
+					if (ELightUnits IntensityUnits; LightNode->GetCustomIntensityUnits(IntensityUnits))
+					{
+						LightFactoryNode->SetCustomIntensityUnits(IntensityUnits);
+					}
+
+					if (float AttenuationRadius; LightNode->GetCustomAttenuationRadius(AttenuationRadius))
+					{
+						LightFactoryNode->SetCustomAttenuationRadius(AttenuationRadius);
+					}
+
+					// RectLight
+					// Needs more discussion, n/a for the moment
+					/*
+					if(const UInterchangeRectLightNode* RectLightNode = Cast<UInterchangeRectLightNode>(LightNode))
+					{
+						if(UInterchangeRectLightFactoryNode* RectLightFactoryNode = Cast<UInterchangeRectLightFactoryNode>(LightFactoryNode))
+						{
+							if(float SourceWidth; RectLightNode->GetCustomSourceWidth(SourceWidth))
+							{
+								RectLightFactoryNode->SetCustomSourceWidth(SourceWidth);
+							}
+
+							if(float SourceHeight; RectLightNode->GetCustomSourceHeight(SourceHeight))
+							{
+								RectLightFactoryNode->SetCustomSourceHeight(SourceHeight);
+							}
+
+							if(float BarnDoorAngle; RectLightNode->GetCustomBarnDoorAngle(BarnDoorAngle))
+							{
+								RectLightFactoryNode->SetCustomBarnDoorAngle(BarnDoorAngle);
+							}
+
+							if(float BarnDoorLength; RectLightNode->GetCustomBarnDoorLength(BarnDoorLength))
+							{
+								RectLightFactoryNode->SetCustomBarnDoorLength(BarnDoorLength);
+							}
+						}
+					}
+					*/
+
+					// Point Light
+					if (const UInterchangePointLightNode* PointLightNode = Cast<UInterchangePointLightNode>(LightNode))
+					{
+						if (UInterchangePointLightFactoryNode* PointLightFactoryNode = Cast<UInterchangePointLightFactoryNode>(LightFactoryNode))
+						{
+							if (bool bUseInverseSquaredFalloff; PointLightNode->GetCustomUseInverseSquaredFalloff(bUseInverseSquaredFalloff))
+							{
+								PointLightFactoryNode->SetCustomUseInverseSquaredFalloff(bUseInverseSquaredFalloff);
+
+								if (float LightFalloffExponent; PointLightNode->GetCustomLightFalloffExponent(LightFalloffExponent))
+								{
+									PointLightFactoryNode->SetCustomLightFalloffExponent(LightFalloffExponent);
+								}
+							}
+
+
+							// Spot Light
+							if (const UInterchangeSpotLightNode* SpotLightNode = Cast<UInterchangeSpotLightNode>(PointLightNode))
+							{
+								UInterchangeSpotLightFactoryNode* SpotLightFactoryNode = Cast<UInterchangeSpotLightFactoryNode>(PointLightFactoryNode);
+								if (float InnerConeAngle; SpotLightNode->GetCustomInnerConeAngle(InnerConeAngle))
+								{
+									SpotLightFactoryNode->SetCustomInnerConeAngle(InnerConeAngle);
+								}
+
+								if (float OuterConeAngle; SpotLightNode->GetCustomOuterConeAngle(OuterConeAngle))
+								{
+									SpotLightFactoryNode->SetCustomOuterConeAngle(OuterConeAngle);
+								}
+							}
+						}
+					}
+				}
+			}
+		}
+
 		//Test for spot before point since a spot light is a point light
-		if (LightNode->IsA<UInterchangeSpotLightNode>())
+		if (BaseLightNode->IsA<UInterchangeSpotLightNode>())
 		{
 			ActorFactoryNode->SetCustomActorClassName(ASpotLight::StaticClass()->GetPathName());
 		}
-		else if (LightNode->IsA<UInterchangePointLightNode>())
+		else if (BaseLightNode->IsA<UInterchangePointLightNode>())
 		{
 			ActorFactoryNode->SetCustomActorClassName(APointLight::StaticClass()->GetPathName());
 		}
-		else if (LightNode->IsA<UInterchangeRectLightNode>())
+		else if (BaseLightNode->IsA<UInterchangeRectLightNode>())
 		{
 			ActorFactoryNode->SetCustomActorClassName(ARectLight::StaticClass()->GetPathName());
 		}
-		else if (LightNode->IsA<UInterchangeDirectionalLightNode>())
+		else if (BaseLightNode->IsA<UInterchangeDirectionalLightNode>())
 		{
 			ActorFactoryNode->SetCustomActorClassName(ADirectionalLight::StaticClass()->GetPathName());
 		}
