- Fixing MaskNode not taking properly values of the different channels because of the || operator only evaluating the first operand if it's true.

diff --git a/Engine/Plugins/Experimental/Interchange/Runtime/Source/Pipelines/Private/InterchangeGenericMaterialPipeline.cpp b/Engine/Plugins/Experimental/Interchange/Runtime/Source/Pipelines/Private/InterchangeGenericMaterialPipeline.cpp
--- a/Engine/Plugins/Experimental/Interchange/Runtime/Source/Pipelines/Private/InterchangeGenericMaterialPipeline.cpp
+++ b/Engine/Plugins/Experimental/Interchange/Runtime/Source/Pipelines/Private/InterchangeGenericMaterialPipeline.cpp
@@ -1186,14 +1186,14 @@ void UInterchangeGenericMaterialPipeline::HandleMaskNode(const UInterchangeShade
 	MaskFactoryNode->SetCustomExpressionClassName(UMaterialExpressionComponentMask::StaticClass()->GetName());
 
 	bool bRChannel = false;
+	ShaderNode->GetBooleanAttribute(Mask::Attributes::R, bRChannel);
 	bool bGChannel = false;
+	ShaderNode->GetBooleanAttribute(Mask::Attributes::G, bGChannel);
 	bool bBChannel = false;
+	ShaderNode->GetBooleanAttribute(Mask::Attributes::B, bBChannel);
 	bool bAChannel = false;
-	bool bIsAnyMaskChannelSet = 
-		ShaderNode->GetBooleanAttribute(Mask::Attributes::R, bRChannel) ||
-		ShaderNode->GetBooleanAttribute(Mask::Attributes::G, bGChannel) ||
-		ShaderNode->GetBooleanAttribute(Mask::Attributes::B, bBChannel) ||
-		ShaderNode->GetBooleanAttribute(Mask::Attributes::A, bAChannel);
+	ShaderNode->GetBooleanAttribute(Mask::Attributes::A, bAChannel);
+	bool bIsAnyMaskChannelSet = bRChannel || bGChannel || bBChannel || bAChannel;
 
 	if(bIsAnyMaskChannelSet)
 	{
