- Fixing compile errors
- Fixing warning: missing copyright boilerplate

diff --git a/Engine/Plugins/Experimental/Interchange/Runtime/Source/FactoryNodes/Private/InterchangeLightFactoryNode.cpp b/Engine/Plugins/Experimental/Interchange/Runtime/Source/FactoryNodes/Private/InterchangeLightFactoryNode.cpp
--- a/Engine/Plugins/Experimental/Interchange/Runtime/Source/FactoryNodes/Private/InterchangeLightFactoryNode.cpp
+++ b/Engine/Plugins/Experimental/Interchange/Runtime/Source/FactoryNodes/Private/InterchangeLightFactoryNode.cpp
@@ -1,3 +1,5 @@
+// Copyright Epic Games, Inc. All Rights Reserved.
+
 #include "InterchangeLightFactoryNode.h"
 
 #include "Components/RectLightComponent.h"
diff --git a/Engine/Plugins/Experimental/Interchange/Runtime/Source/FactoryNodes/Public/InterchangeLightFactoryNode.h b/Engine/Plugins/Experimental/Interchange/Runtime/Source/FactoryNodes/Public/InterchangeLightFactoryNode.h
--- a/Engine/Plugins/Experimental/Interchange/Runtime/Source/FactoryNodes/Public/InterchangeLightFactoryNode.h
+++ b/Engine/Plugins/Experimental/Interchange/Runtime/Source/FactoryNodes/Public/InterchangeLightFactoryNode.h
@@ -3,6 +3,7 @@
 #pragma once
 
 #include "CoreMinimal.h"
+#include "Engine/Scene.h"
 #include "InterchangeActorFactoryNode.h"
 
 #include "InterchangeLightFactoryNode.generated.h"
diff --git a/Engine/Plugins/Experimental/Interchange/Runtime/Source/Nodes/Public/InterchangeLightNode.h b/Engine/Plugins/Experimental/Interchange/Runtime/Source/Nodes/Public/InterchangeLightNode.h
--- a/Engine/Plugins/Experimental/Interchange/Runtime/Source/Nodes/Public/InterchangeLightNode.h
+++ b/Engine/Plugins/Experimental/Interchange/Runtime/Source/Nodes/Public/InterchangeLightNode.h
@@ -3,6 +3,7 @@
 #pragma once
 
 #include "CoreMinimal.h"
+#include "Engine/Scene.h"
 #include "Nodes/InterchangeBaseNode.h"
 
 #include "InterchangeLightNode.generated.h"
