MaterialX:
- Changing Map structure, to contain the node category, in case the inputs of UE are not standardized against MaterialX, such as Power for example
- Adding the support of coat_color
- Adding support of power node
- Changing default values of base and base_color according to Autodesk's Standard surface node definition
- handling Interfacename for inputs, which is an alias for inputs inside a nodegraph
- Adding support of colorspace (only srgb is supported)
- properly handling uv tilling/offset because this input is on the texture in MaterialX unlike UE which has it on the texture coordinate node
- Rename OldName attribute to OriginalName
- Modifying the RenameInputName method to keep uniqueness of the name

diff --git a/Engine/Plugins/Experimental/Interchange/Runtime/Source/Import/Private/MaterialX/InterchangeMaterialXDefinitions.h b/Engine/Plugins/Experimental/Interchange/Runtime/Source/Import/Private/MaterialX/InterchangeMaterialXDefinitions.h
--- a/Engine/Plugins/Experimental/Interchange/Runtime/Source/Import/Private/MaterialX/InterchangeMaterialXDefinitions.h
+++ b/Engine/Plugins/Experimental/Interchange/Runtime/Source/Import/Private/MaterialX/InterchangeMaterialXDefinitions.h
@@ -23,7 +23,7 @@ namespace MaterialX
 			static constexpr const char* TransmissionScatter = "transmission_scatter";
 			static constexpr const char* TransmissionScatterAnisotropy = "transmission_scatter_anisotropy";
 			static constexpr const char* TransmissionDispersion = "transmission_dispersion";
-			static constexpr const char* TransmissionExtraRoughness= "transmission_extra_roughness";
+			static constexpr const char* TransmissionExtraRoughness = "transmission_extra_roughness";
 			static constexpr const char* Subsurface = "subsurface";
 			static constexpr const char* SubsurfaceColor = "subsurface_color";
 			static constexpr const char* SubsurfaceRadius = "subsurface_radius";
@@ -46,7 +46,7 @@ namespace MaterialX
 			static constexpr const char* Emission = "emission";
 			static constexpr const char* EmissionColor = "emission_color";
 			static constexpr const char* Opacity = "opacity";
-			static constexpr const char* ThinWalled= "thin_walled";
+			static constexpr const char* ThinWalled = "thin_walled";
 			static constexpr const char* Normal = "normal";
 			static constexpr const char* Tangent = "tangent";
 		}
@@ -55,7 +55,7 @@ namespace MaterialX
 		{
 			namespace Float
 			{
-				static float Base = 0.8f;
+				static float Base = 1.f;
 				static float DiffuseRoughness = 0.f;
 				static float Metalness = 0.f;
 				static float Specular = 1.f;
@@ -87,7 +87,7 @@ namespace MaterialX
 
 			namespace Color3
 			{
-				constexpr FLinearColor BaseColor{ 1.0, 1.0, 1.0 };
+				constexpr FLinearColor BaseColor{ 0.8, 0.8, 0.8 };
 				constexpr FLinearColor SpecularColor{ 1, 1, 1 };
 				constexpr FLinearColor TransmissionColor{ 1, 1, 1 };
 				constexpr FLinearColor TransmissionScatter{ 0, 0, 0 };
@@ -121,7 +121,7 @@ namespace MaterialX
 		}
 
 		namespace DirectionalLight
-		{			
+		{
 			namespace Input
 			{
 				using namespace Lights::Input;
@@ -144,7 +144,7 @@ namespace MaterialX
 	namespace Attributes
 	{
 		static constexpr const char* IsVisited = "UE:IsVisited";
-		static constexpr const char* OldName = "UE:OldName";
+		static constexpr const char* OriginalName = "UE:OriginalName";
 	}
 
 	namespace Category
@@ -173,6 +173,10 @@ namespace MaterialX
 		static constexpr const char* DirectionalLight = "directional_light";
 		static constexpr const char* SpotLight = "spot_light";
 		static constexpr const char* StandardSurface = "standard_surface";
+		//static constexpr const char* Convert = "convert";
+		static constexpr const char* Power = "power";
+		static constexpr const char* RgbToHsv = "rgbtohsv";
+		static constexpr const char* HsvToRgb = "hsvtorgb";
 	}
 
 	namespace NodeDefinition
@@ -186,8 +190,82 @@ namespace MaterialX
 	namespace Library
 	{
 		static constexpr const char* Std = "stdlib";
-		static constexpr const char* Pbr= "pbrlib";
+		static constexpr const char* Pbr = "pbrlib";
 		static constexpr const char* Bxdf = "bxdf";
 		static constexpr const char* Lights = "lights";
 	}
+
+	namespace Type
+	{
+		static constexpr const char* Boolean = "boolean";
+		static constexpr const char* Integer = "integer";
+		static constexpr const char* Float = "float";
+		static constexpr const char* Color3 = "color3";
+		static constexpr const char* Color4 = "color4";
+		static constexpr const char* Vector2 = "vector2";
+		static constexpr const char* Vector3 = "vector3";
+		static constexpr const char* Vector4 = "vector4";
+		static constexpr const char* Matrix33 = "matrix33";
+		static constexpr const char* Matrix44 = "matrix44";
+		static constexpr const char* String = "string";
+		static constexpr const char* Filename = "filename";
+		static constexpr const char* GeomName = "geomname";
+		static constexpr const char* SurfaceShader = "surfaceshader";
+		static constexpr const char* DisplacementShader = "displacementshader";
+		static constexpr const char* VolumeShader = "volumeshader";
+		static constexpr const char* LightShader = "lightshader";
+		static constexpr const char* Material = "material";
+		static constexpr const char* None = "none";
+		static constexpr const char* IntegerArray = "integerarray";
+		static constexpr const char* FloatArray = "floatarray";
+		static constexpr const char* Color3Array = "color3array";
+		static constexpr const char* Color4Array = "color4array";
+		static constexpr const char* Vector2Array = "vector2array";
+		static constexpr const char* Vector3Array = "vector3array";
+		static constexpr const char* Vector4Array = "vector4array";
+		static constexpr const char* StringArray = "stringarray";
+		static constexpr const char* GeomNameArray = "geomnamearray";
+	}
+
+	namespace NodeGroup
+	{
+		namespace Texture2D
+		{
+			namespace Inputs
+			{
+				static constexpr const char* File = "file";
+				static constexpr const char* Default = "default";
+				static constexpr const char* TexCoord = "texcoord";
+				static constexpr const char* FilterType = "filtertype";
+				static constexpr const char* FrameRange = "framerange";
+				static constexpr const char* FrameOffset = "frameoffset";
+				static constexpr const char* FrameEndAction = "frameendaction";
+			}
+		}
+	}
+
+	namespace Image
+	{
+		namespace Inputs
+		{
+			using namespace NodeGroup::Texture2D::Inputs;
+
+			static constexpr const char* Layer = "layer";
+			static constexpr const char* UAddressMode = "uaddressmode";
+			static constexpr const char* VAddressMode = "vaddressmode";
+		}
+	}
+
+	namespace TiledImage
+	{
+		namespace Inputs
+		{
+			using namespace NodeGroup::Texture2D::Inputs;
+
+			static constexpr const char* UVTiling = "uvtiling";
+			static constexpr const char* UVOffset = "uvoffset";
+			static constexpr const char* RealWorldImageSize = "realworldimagesize";
+			static constexpr const char* RealWorldTileSize = "realworldtilesize";
+		}
+	}
 }
\ No newline at end of file
diff --git a/Engine/Plugins/Experimental/Interchange/Runtime/Source/Import/Private/MaterialX/InterchangeMaterialXTranslator.cpp b/Engine/Plugins/Experimental/Interchange/Runtime/Source/Import/Private/MaterialX/InterchangeMaterialXTranslator.cpp
--- a/Engine/Plugins/Experimental/Interchange/Runtime/Source/Import/Private/MaterialX/InterchangeMaterialXTranslator.cpp
+++ b/Engine/Plugins/Experimental/Interchange/Runtime/Source/Import/Private/MaterialX/InterchangeMaterialXTranslator.cpp
@@ -4,13 +4,11 @@
 #if WITH_EDITOR
 #include "MaterialXFormat/Util.h"
 #endif
-#include "Engine/Scene.h"
 #include "InterchangeImportLog.h"
 #include "InterchangeLightNode.h"
 #include "InterchangeManager.h"
 #include "InterchangeMaterialDefinitions.h"
 #include "InterchangeSceneNode.h"
-#include "InterchangeShaderGraphNode.h"
 #include "InterchangeMaterialXDefinitions.h"
 #include "InterchangeTexture2DNode.h"
 #include "Nodes/InterchangeSourceNode.h"
@@ -55,15 +53,17 @@ UInterchangeMaterialXTranslator::UInterchangeMaterialXTranslator()
 #if WITH_EDITOR
 	:
 InputNamesMaterialX2UE{
-	{TEXT("in"),        TEXT("Input")},
-	{TEXT("in1"),       TEXT("A")},
-	{TEXT("in2"),       TEXT("B")},
-	{TEXT("fg"),        TEXT("A")},
-	{TEXT("bg"),        TEXT("B")},
-	{TEXT("mix"),       TEXT("Factor")},
-	{TEXT("low"),       TEXT("Min")},
-	{TEXT("high"),      TEXT("Max")},
-	{TEXT("texcoord"),  UE::Interchange::Materials::Standard::Nodes::TextureSample::Inputs::Coordinates.ToString()} },
+	{{TEXT(""),						 TEXT("in")},		TEXT("Input")},
+	{{TEXT(""),						 TEXT("in1")},		TEXT("A")},
+	{{TEXT(""),						 TEXT("in2")},		TEXT("B")},
+	{{TEXT(""),						 TEXT("fg")},		TEXT("A")},
+	{{TEXT(""),						 TEXT("bg")},		TEXT("B")},
+	{{TEXT(""),						 TEXT("mix")},		TEXT("Factor")},
+	{{TEXT(""),						 TEXT("low")},		TEXT("Min")},
+	{{TEXT(""),						 TEXT("high")},		TEXT("Max")},
+	{{TEXT(""),						 TEXT("texcoord")}, TEXT("Coordinates")},
+	{{MaterialX::Category::Power,    TEXT("in1")},		TEXT("Base")},
+	{{MaterialX::Category::Power,	 TEXT("in2")},		TEXT("Exponent")}},
 NodeNamesMaterialX2UE{
 	{MaterialX::Category::Add,      TEXT("Add")},
 	{MaterialX::Category::Sub,      TEXT("Subtract")},
@@ -76,8 +76,9 @@ NodeNamesMaterialX2UE{
 	{MaterialX::Category::Min,      TEXT("Min")},
 	{MaterialX::Category::Combine2, TEXT("AppendVector")},
 	{MaterialX::Category::Combine3, TEXT("AppendVector")},
-	{MaterialX::Category::Combine4, TEXT("AppendVector")}},
-UEInputs{ TEXT("A"), TEXT("B"), TEXT("Input"), TEXT("Factor"), TEXT("Min"), TEXT("Max") }
+	{MaterialX::Category::Combine4, TEXT("AppendVector")},
+	{MaterialX::Category::Power,    TEXT("Power")}},
+UEInputs{ TEXT("A"), TEXT("B"), TEXT("Input"), TEXT("Factor"), TEXT("Min"), TEXT("Max"), TEXT("Base"), TEXT("Exponent") }
 #endif // WITH_EDITOR
 {
 }
@@ -162,21 +163,16 @@ bool UInterchangeMaterialXTranslator::Translate(UInterchangeBaseNodeContainer& B
 
 			if(Node)
 			{
-				std::string NameOfNodeDocument = Document->getCategory();
-				std::string NameOfNode = Node->getCategory();
-				printf("%s, %s", NameOfNodeDocument.c_str(), NameOfNode.c_str());
-
-				bool bIsMaterialShader = Node->getType() == mx::MATERIAL_TYPE_STRING;
-				bool bIsLightShader = Node->getType() == mx::LIGHT_SHADER_TYPE_STRING;
+				bool bIsMaterialShader = Node->getType() == mx::Type::Material;
+				bool bIsLightShader = Node->getType() == mx::Type::LightShader;
 
 				if(bIsMaterialShader || bIsLightShader)
 				{
 					if(!Node->getTypeDef())
 					{
 						UInterchangeResultError_Generic* Message = AddMessage<UInterchangeResultError_Generic>();
-						Message->Text = FText::Format(LOCTEXT("TypeDefNotFound", "<{0}> has no matching TypeDef"),
+						Message->Text = FText::Format(LOCTEXT("TypeDefNotFound", "<{0}> has no matching TypeDef, aborting import..."),
 													  FText::FromString(Node->getName().c_str()));
-
 						bIsReferencesValid = false;
 						break;
 					}
@@ -261,15 +257,9 @@ void UInterchangeMaterialXTranslator::ProcessStandardSurface(UInterchangeBaseNod
 {
 	using namespace UE::Interchange::Materials;
 
-	const FString StandardSurfaceName{ StandardSurfaceNode->getName().c_str() };
-	UInterchangeShaderGraphNode* ShaderGraphNode = NewObject<UInterchangeShaderGraphNode>(&NodeContainer);
-	const FString ShaderGraphNodeUID = TEXT("\\Material\\") + StandardSurfaceName;
-	ShaderGraphNode->InitializeNode(ShaderGraphNodeUID, StandardSurfaceName, EInterchangeNodeContainerType::TranslatedAsset);
-	NodeContainer.AddNode(ShaderGraphNode);
-	ShaderGraphNode->SetCustomShaderType(StandardSurface::Name.ToString());
-
 	TMap<FString, UInterchangeShaderNode*> NamesToShaderNodes;
-	NamesToShaderNodes.Add(StandardSurfaceName, ShaderGraphNode);
+
+	UInterchangeShaderGraphNode* ShaderGraphNode = CreateShaderNode<UInterchangeShaderGraphNode>(StandardSurfaceNode->getName().c_str(), StandardSurface::Name.ToString(), TEXT(""), NamesToShaderNodes, NodeContainer);
 
 	//Base
 	{
@@ -455,6 +445,16 @@ void UInterchangeMaterialXTranslator::ProcessStandardSurface(UInterchangeBaseNod
 			}
 		}
 
+		//Color
+		{
+			mx::InputPtr InputCoatColor = GetStandardSurfaceInput(StandardSurfaceNode, mx::StandardSurface::Input::CoatColor, Document);
+
+			if(!ConnectNodeGraphOutputToInput(InputCoatColor, ShaderGraphNode, StandardSurface::Parameters::Coat.ToString(), NamesToShaderNodes, NodeContainer))
+			{
+				AddLinearColorAttribute(InputCoatColor, StandardSurface::Parameters::CoatColor.ToString(), ShaderGraphNode, mx::StandardSurface::DefaultValue::Color3::CoatColor);
+			}
+		}
+
 		//Roughness
 		{
 			mx::InputPtr InputCoatRoughness = GetStandardSurfaceInput(StandardSurfaceNode, mx::StandardSurface::Input::CoatRoughness, Document);
@@ -700,8 +700,7 @@ bool UInterchangeMaterialXTranslator::ConnectNodeGraphOutputToInput(MaterialX::I
 		{
 			UInterchangeResultWarning_Generic* Message = AddMessage<UInterchangeResultWarning_Generic>();
 			Message->Text = FText::Format(LOCTEXT("OutputNotFound", "Couldn't find a connected output to ({0})"),
-										  FText::FromString(InputToNodeGraph->getName().c_str()));
-
+										  FText::FromString(GetInputName(InputToNodeGraph)));
 			return false;
 		}
 
@@ -713,7 +712,7 @@ bool UInterchangeMaterialXTranslator::ConnectNodeGraphOutputToInput(MaterialX::I
 				FString InputChannelName = ParentInputName;
 
 				//Replace the input's name by the one used in UE
-				RenameInputsNames(UpstreamNode);
+				RenameNodeInputs(UpstreamNode);
 
 				if(mx::NodePtr DownstreamNode = Edge.getDownstreamElement()->asA<mx::Node>())
 				{
@@ -724,7 +723,7 @@ bool UInterchangeMaterialXTranslator::ConnectNodeGraphOutputToInput(MaterialX::I
 
 					if(mx::InputPtr ConnectedInput = Edge.getConnectingElement()->asA<mx::Input>())
 					{
-						InputChannelName = ConnectedInput->getName().c_str();
+						InputChannelName = GetInputName(ConnectedInput);
 					}
 				}
 
@@ -738,14 +737,14 @@ bool UInterchangeMaterialXTranslator::ConnectNodeGraphOutputToInput(MaterialX::I
 					{
 						UInterchangeResultWarning_Generic* Message = AddMessage<UInterchangeResultWarning_Generic>();
 						Message->Text = FText::Format(LOCTEXT("InputTypeNotSupported", "<{0}>: \"{1}\" is not supported yet"),
-													  FText::FromString(InputConstant->getName().c_str()),
+													  FText::FromString(GetInputName(InputConstant)),
 													  FText::FromString(InputConstant->getType().c_str()));
 					}
 
 				}
 				else if(UpstreamNode->getCategory() == mx::Category::Extract)
 				{
-					UInterchangeShaderNode* MaskShaderNode = CreateShaderNode(UpstreamNode->getName().c_str(), Mask::Name.ToString(), ParentShaderNode, NamesToShaderNodes, NodeContainer);
+					UInterchangeShaderNode* MaskShaderNode = CreateShaderNode<UInterchangeShaderNode>(UpstreamNode->getName().c_str(), Mask::Name.ToString(), ParentShaderNode->GetUniqueID(), NamesToShaderNodes, NodeContainer);
 
 					if(mx::InputPtr InputIndex = UpstreamNode->getInput("index"))
 					{
@@ -768,9 +767,9 @@ bool UInterchangeMaterialXTranslator::ConnectNodeGraphOutputToInput(MaterialX::I
 				//dot means identity, input == output
 				else if(UpstreamNode->getCategory() == mx::Category::Dot || UpstreamNode->getCategory() == mx::Category::NormalMap)
 				{
-					if(mx::InputPtr Input = GetInputFromOldName(UpstreamNode, "in"))
+					if(mx::InputPtr Input = GetInputFromOriginalName(UpstreamNode, "in"))
 					{
-						Input->setName(TCHAR_TO_UTF8(*InputChannelName)); //let's take the parent node's input name
+						RenameInput(Input, TCHAR_TO_UTF8(*InputChannelName)); //let's take the parent node's input name
 						NamesToShaderNodes.FindOrAdd(UpstreamNode->getName().c_str(), ParentShaderNode);
 					}
 				}
@@ -781,21 +780,76 @@ bool UInterchangeMaterialXTranslator::ConnectNodeGraphOutputToInput(MaterialX::I
 						//By default set the output of a texture to RGB, if the type is float then it's either up to an extract node or the Material Input to handle it
 						FString OutputChannel{ TEXT("RGB") };
 
-						if(UpstreamNode->getType() == "vector4" || UpstreamNode->getType() == "color4")
+						if(UpstreamNode->getType() == mx::Type::Vector4 || UpstreamNode->getType() == mx::Type::Color4)
 							OutputChannel = TEXT("RGBA");
 
-						UInterchangeShaderNode* TextureShaderNode = CreateShaderNode(UpstreamNode->getName().c_str(), TextureSample::Name.ToString(), ParentShaderNode, NamesToShaderNodes, NodeContainer);
+						UInterchangeShaderNode* TextureShaderNode = CreateShaderNode<UInterchangeShaderNode>(UpstreamNode->getName().c_str(), TextureSample::Name.ToString(), ParentShaderNode->GetUniqueID(), NamesToShaderNodes, NodeContainer);
 						TextureShaderNode->AddStringAttribute(UInterchangeShaderPortsAPI::MakeInputValueKey(TextureSample::Inputs::Texture.ToString()), TextureNode->GetUniqueID());
 						UInterchangeShaderPortsAPI::ConnectOuputToInput(ParentShaderNode, InputChannelName, TextureShaderNode->GetUniqueID(), OutputChannel);
+
+						UInterchangeShaderNode* ImageNode = TextureShaderNode;
+						FString ImageNodeInputName{ TextureSample::Inputs::Coordinates.ToString() };
+
+						auto ConnectUVTransformToOutput = [this, UpstreamNode, &NodeContainer, &NamesToShaderNodes](UInterchangeShaderNode*& ImageNode, FString& ImageNodeInputName, const FString& ShaderType, const char* InputName)
+						{
+							if(mx::InputPtr Input = UpstreamNode->getInput(InputName))
+							{
+								const FString ShaderNodeName{ UpstreamNode->getName().c_str() + FString(TEXT("_")) + InputName };
+								UInterchangeShaderNode* ShaderNode = CreateShaderNode<UInterchangeShaderNode>(ShaderNodeName, ShaderType, ImageNode->GetUniqueID(), NamesToShaderNodes, NodeContainer);
+								UInterchangeShaderPortsAPI::ConnectDefaultOuputToInput(ImageNode, ImageNodeInputName, ShaderNode->GetUniqueID());
+
+								//Vec2
+								if(Input->hasValueString())
+								{
+									const FString MaskNodeName{ UpstreamNode->getName().c_str() + FString(TEXT("_")) + ShaderType+TEXT("MaskNode") };
+									UInterchangeShaderNode* MaskShaderNode = CreateShaderNode<UInterchangeShaderNode>(MaskNodeName, Mask::Name.ToString(), ShaderNode->GetUniqueID(), NamesToShaderNodes, NodeContainer);
+									MaskShaderNode->AddBooleanAttribute(Mask::Attributes::R, true);
+									MaskShaderNode->AddBooleanAttribute(Mask::Attributes::G, true);
+
+									mx::Vector2 Vec2 = mx::fromValueString<mx::Vector2>(Input->getValueString());
+									MaskShaderNode->AddLinearColorAttribute(UInterchangeShaderPortsAPI::MakeInputValueKey(Mask::Inputs::Input.ToString()), FLinearColor(Vec2[0], Vec2[1], 0));
+									UInterchangeShaderPortsAPI::ConnectDefaultOuputToInput(ShaderNode, TEXT("B"), MaskShaderNode->GetUniqueID());
+								}
+								else
+								{
+									RenameInput(Input, "B");
+								}
+
+								//The node will replace the texture node for the input
+								NamesToShaderNodes.Add(UpstreamNode->getName().c_str(), ShaderNode);
+
+								//We also need to replace the name of texcoord input by one of the inputs of this node
+								if(mx::InputPtr InputTexCoord = GetInputFromOriginalName(UpstreamNode, mx::TiledImage::Inputs::TexCoord))
+								{
+									RenameInput(InputTexCoord, "A");
+								}
+								else
+								{
+									//Let's reuse the same texture coordinate node for further use (no need to set a parent)
+									const FString TextureCoordinateName{ FPaths::GetBaseFilename(UpstreamNode->getActiveSourceUri().c_str()) + FString(TEXT("_texcoord")) };
+									UInterchangeShaderNode* TextureCoordinateNode = CreateShaderNode<UInterchangeShaderNode>(TextureCoordinateName, TextureCoordinate::Name.ToString(), TEXT(""), NamesToShaderNodes, NodeContainer);
+									UInterchangeShaderPortsAPI::ConnectDefaultOuputToInput(ShaderNode, TEXT("A"), TextureCoordinateNode->GetUniqueID());
+								}
+
+								ImageNode = ShaderNode;
+								ImageNodeInputName = TEXT("A");
+							}
+						};
+
+						// UV offset (MaterialX defines it as a substraction)
+						// MaterialX spec: the offset for the given image along the U and V axes. Mathematically
+						// equivalent to subtracting the given vector value from the incoming texture coordinates.
+						ConnectUVTransformToOutput(ImageNode, ImageNodeInputName, TEXT("Subtract"), mx::TiledImage::Inputs::UVOffset);
+						ConnectUVTransformToOutput(ImageNode, ImageNodeInputName, TEXT("Multiply"), mx::TiledImage::Inputs::UVTiling);
 					}
 					else
 					{
-						AddAttribute(UpstreamNode->getInput("default"), InputChannelName, ParentShaderNode);
+						AddAttribute(UpstreamNode->getInput(mx::NodeGroup::Texture2D::Inputs::Default), InputChannelName, ParentShaderNode);
 					}
 				}
 				else if(UpstreamNode->getCategory() == mx::Category::TexCoord)
 				{
-					UInterchangeShaderNode* TextureCoordinateNode = CreateShaderNode(UpstreamNode->getName().c_str(), TextureCoordinate::Name.ToString(), ParentShaderNode, NamesToShaderNodes, NodeContainer);
+					UInterchangeShaderNode* TextureCoordinateNode = CreateShaderNode<UInterchangeShaderNode>(UpstreamNode->getName().c_str(), TextureCoordinate::Name.ToString(), ParentShaderNode->GetUniqueID(), NamesToShaderNodes, NodeContainer);
 					UInterchangeShaderPortsAPI::ConnectDefaultOuputToInput(ParentShaderNode, InputChannelName, TextureCoordinateNode->GetUniqueID());
 				}
 				else
@@ -819,17 +873,29 @@ bool UInterchangeMaterialXTranslator::ConnectNodeOutputToInput(MaterialX::NodePt
 	{
 		UInterchangeShaderNode* OperatorNode = nullptr;
 
-		OperatorNode = CreateShaderNode(Node->getName().c_str(), *ShaderType, ParentShaderNode, NamesToShaderNodes, NodeContainer);
+		OperatorNode = CreateShaderNode<UInterchangeShaderNode>(Node->getName().c_str(), *ShaderType, ParentShaderNode->GetUniqueID(), NamesToShaderNodes, NodeContainer);
 
 		for(mx::InputPtr Input : Node->getInputs())
 		{
 			if(Input->hasValue())
 			{
-				if(const FString* InputNameFound = UEInputs.Find(Input->getName().c_str()))
+				if(const FString* InputNameFound = UEInputs.Find(GetInputName(Input)))
 				{
 					AddAttribute(Input, *InputNameFound, OperatorNode);
 				}
 			}
+			else if(Input->hasInterfaceName())
+			{
+				mx::InputPtr InputInterface = Input->getInterfaceInput();
+				if(InputInterface->hasValue())
+				{
+					// We need to take the input name from the original Input not the Interface
+					if(const FString* InputNameFound = UEInputs.Find(GetInputName( Input)))
+					{
+						AddAttribute(InputInterface, *InputNameFound, OperatorNode);
+					}
+				}
+			}
 		}
 
 		bIsConnected = UInterchangeShaderPortsAPI::ConnectDefaultOuputToInput(ParentShaderNode, InputChannelName, OperatorNode->GetUniqueID());
@@ -838,29 +904,6 @@ bool UInterchangeMaterialXTranslator::ConnectNodeOutputToInput(MaterialX::NodePt
 	return bIsConnected;
 }
 
-UInterchangeShaderNode* UInterchangeMaterialXTranslator::CreateShaderNode(const FString& NodeName, const FString& ShaderType, UInterchangeShaderNode* ParentNode, TMap<FString, UInterchangeShaderNode*>& NamesToShaderNodes, UInterchangeBaseNodeContainer& NodeContainer) const
-{
-	UInterchangeShaderNode* Node;
-
-	if(UInterchangeShaderNode** FoundNode = NamesToShaderNodes.Find(NodeName); !FoundNode)
-	{
-		const FString NodeUID = ParentNode->GetUniqueID() + TEXT("_") + NodeName;
-		Node = NewObject<UInterchangeShaderNode>(&NodeContainer);
-		Node->InitializeNode(NodeUID, NodeName, EInterchangeNodeContainerType::TranslatedAsset);
-		NodeContainer.AddNode(Node);
-		NodeContainer.SetNodeParentUid(NodeUID, ParentNode->GetUniqueID());
-		Node->SetCustomShaderType(ShaderType);
-
-		NamesToShaderNodes.Add(NodeName, Node);
-	}
-	else
-	{
-		Node = *FoundNode;
-	}
-
-	return Node;
-}
-
 UInterchangeTextureNode* UInterchangeMaterialXTranslator::CreateTextureNode(MaterialX::NodePtr Node, UInterchangeBaseNodeContainer& NodeContainer) const
 {
 	UInterchangeTextureNode* TextureNode = nullptr;
@@ -901,22 +944,29 @@ UInterchangeTextureNode* UInterchangeMaterialXTranslator::CreateTextureNode(Mate
 	return TextureNode;
 }
 
-const FString& UInterchangeMaterialXTranslator::GetMatchedInputName(MaterialX::InputPtr Input) const
+const FString& UInterchangeMaterialXTranslator::GetMatchedInputName(MaterialX::NodePtr Node, MaterialX::InputPtr Input) const
 {
-	static FString Empty{ "" };
+	static FString EmptyString;
 
 	if(Input)
 	{
-		if(const FString* Result = InputNamesMaterialX2UE.Find(Input->getName().c_str()))
+		const FString NodeCategory{ Node->getCategory().c_str() };
+		const FString InputName{ GetInputName(Input) };
+		TPair<FString, FString> Pair{ NodeCategory, InputName };
+		if(const FString* Result = InputNamesMaterialX2UE.Find(Pair))
+		{
+			return *Result;
+		}
+		else if((Result = InputNamesMaterialX2UE.Find({ EmptyString, InputName })))
 		{
 			return *Result;
 		}
 	}
 
-	return Empty;
+	return EmptyString;
 }
 
-void UInterchangeMaterialXTranslator::RenameInputsNames(MaterialX::NodePtr Node) const
+void UInterchangeMaterialXTranslator::RenameNodeInputs(MaterialX::NodePtr Node) const
 {
 	if(Node)
 	{
@@ -926,31 +976,64 @@ void UInterchangeMaterialXTranslator::RenameInputsNames(MaterialX::NodePtr Node)
 
 			for(mx::InputPtr Input : Node->getInputs())
 			{
-				if(const FString& Name = GetMatchedInputName(Input); !Name.IsEmpty())
+				if(const FString& Name = GetMatchedInputName(Node, Input); !Name.IsEmpty())
 				{
-					Input->setAttribute(mx::Attributes::OldName, Input->getName()); // keep the old name for further processing
-					Input->setName(TCHAR_TO_UTF8(*Name));
+					RenameInput(Input, TCHAR_TO_UTF8(*Name));
 				}
 			}
 		}
 	}
 }
 
-MaterialX::InputPtr UInterchangeMaterialXTranslator::GetInputFromOldName(MaterialX::NodePtr Node, const char* OldNameAttribute) const
+void UInterchangeMaterialXTranslator::RenameInput(MaterialX::InputPtr Input, const char* NewName) const
+{
+	std::string OriginalName;
+
+	if(!Input->hasAttribute(mx::Attributes::OriginalName))
+	{
+		OriginalName = Input->getName();
+		Input->setAttribute(mx::Attributes::OriginalName, OriginalName); // keep the original name for further processing
+	}
+	else
+	{
+		OriginalName = Input->getAttribute(mx::Attributes::OriginalName);
+	}
+
+	Input->setName(OriginalName + "_" + NewName);
+}
+
+MaterialX::InputPtr UInterchangeMaterialXTranslator::GetInputFromOriginalName(MaterialX::NodePtr Node, const char* OriginalNameAttribute) const
 {
 	mx::InputPtr InputRes{ nullptr };
 
 	for(mx::InputPtr Input : Node->getInputs())
 	{
-		if(Input->getAttribute(mx::Attributes::OldName) == OldNameAttribute)
+		if(Input->getAttribute(mx::Attributes::OriginalName) == OriginalNameAttribute)
 		{
 			InputRes = Input;
+			break;
 		}
 	}
 
 	return InputRes;
 }
 
+FString UInterchangeMaterialXTranslator::GetInputName(MaterialX::InputPtr Input) const
+{
+	std::string Name = Input->getName();
+
+	if(Input->hasAttribute(mx::Attributes::OriginalName))
+	{
+		std::string OriginalName;
+		OriginalName = Input->getAttribute(mx::Attributes::OriginalName);
+		OriginalName += "_";
+		Name.replace(0, OriginalName.size(), "");
+	}
+
+
+	return Name.c_str();
+}
+
 MaterialX::InputPtr UInterchangeMaterialXTranslator::GetStandardSurfaceInput(MaterialX::NodePtr StandardSurface, const char* InputName, MaterialX::DocumentPtr Document) const
 {
 	mx::InputPtr Input = StandardSurface->getInput(InputName);
@@ -1003,22 +1086,26 @@ bool UInterchangeMaterialXTranslator::AddAttribute(MaterialX::InputPtr Input, co
 {
 	if(Input)
 	{
-		if(Input->getType() == "float")
+		if(Input->getType() == mx::Type::Float)
 		{
 			return ShaderNode->AddFloatAttribute(UInterchangeShaderPortsAPI::MakeInputValueKey(InputChannelName), mx::fromValueString<float>(Input->getValueString()));
 		}
+		else if(Input->getType() == mx::Type::Integer)
+		{
+			printf("");
+		}
 
 		FLinearColor LinearColor;
 
 		bool bIsColor = false;
 
-		if(bool bIsColor3 = (Input->getType() == "color3" || Input->getType() == "vector3"))
+		if(bool bIsColor3 = (Input->getType() == mx::Type::Color3 || Input->getType() == mx::Type::Vector3))
 		{
 			mx::Color3 Color = mx::fromValueString<mx::Color3>(Input->getValueString());
 			LinearColor = FLinearColor{ Color[0], Color[1], Color[2] };
 			bIsColor = true;
 		}
-		else if(bool bIsColor4 = (Input->getType() == "color4" || Input->getType() == "vector4"))
+		else if(bool bIsColor4 = (Input->getType() == mx::Type::Color4|| Input->getType() == mx::Type::Vector4))
 		{
 			mx::Color4 Color = mx::fromValueString<mx::Color4>(Input->getValueString());
 			LinearColor = FLinearColor{ Color[0], Color[1], Color[2], Color[3] };
diff --git a/Engine/Plugins/Experimental/Interchange/Runtime/Source/Import/Public/MaterialX/InterchangeMaterialXTranslator.h b/Engine/Plugins/Experimental/Interchange/Runtime/Source/Import/Public/MaterialX/InterchangeMaterialXTranslator.h
--- a/Engine/Plugins/Experimental/Interchange/Runtime/Source/Import/Public/MaterialX/InterchangeMaterialXTranslator.h
+++ b/Engine/Plugins/Experimental/Interchange/Runtime/Source/Import/Public/MaterialX/InterchangeMaterialXTranslator.h
@@ -1,6 +1,7 @@
 // Copyright Epic Games, Inc. All Rights Reserved.
 #pragma once
 
+#include "InterchangeShaderGraphNode.h"
 #include "InterchangeTranslatorBase.h"
 #include "Texture/InterchangeTexturePayloadInterface.h"
 
@@ -17,7 +18,6 @@ PRAGMA_ENABLE_MISSING_BRACES_WARNINGS
 #include "InterchangeMaterialXTranslator.generated.h"
 
 class UInterchangeTextureNode;
-class UInterchangeShaderNode;
 class UInterchangeBaseLightNode;
 class UInterchangeSceneNode;
 
@@ -41,7 +41,7 @@ public:
 	/**
 	 * Translate the associated source data into a node hold by the specified nodes container.
 	 *
-	 * @param BaseNodeContainer - The unreal objects descriptions container where to put the translated source data.
+	 * @param BaseNodeContainer - The container where to add the translated Interchange nodes.
 	 * @return true if the translator can translate the source data, false otherwise.
 	 */
 	virtual bool Translate( UInterchangeBaseNodeContainer& BaseNodeContainer ) const override;
@@ -55,7 +55,7 @@ protected:
 	/**
 	 * Process Autodesk's standard surface shader
 	 * 
-	 * @param NodeContainer - The unreal objects descriptions container where to put the translated source data.
+	 * @param NodeContainer - The container where to add the translated Interchange nodes.
 	 * @param StandardSurfaceNode - The <standard_surface> in the MaterialX file
 	 * @param Document - The MaterialX Document that contains all the definitions of the loaded libraries
 	 */
@@ -64,7 +64,7 @@ protected:
 	/**
 	 * Process light shader, MaterialX doesn't standardized lights, but defines the 3 common ones, directional, point and spot
 	 * 
-	 * @param NodeContainer - The unreal objects descriptions container where to put the translated source data.
+	 * @param The container where to add the translated Interchange nodes.
 	 * @param LightShaderNode - a MaterialX light shader node
 	 * @param Document - The MaterialX Document that contains all the definitions of the loaded libraries
 	 */
@@ -75,7 +75,7 @@ protected:
 	 *
 	 * @param DirectionalLightShaderNode - The MaterialX node related to a <directional_light>
 	 * @param SceneNode - The interchange scene node, to which the node will be attached to
-	 * @param NodeContainer - The unreal objects descriptions container where to put the translated source data.
+	 * @param The container where to add the translated Interchange nodes.
 	 * @param Document - The MaterialX Document that contains all the definitions of the loaded libraries
 	 */
 	UInterchangeBaseLightNode* CreateDirectionalLightNode(MaterialX::NodePtr DirectionalLightShaderNode, UInterchangeSceneNode* SceneNode, UInterchangeBaseNodeContainer& NodeContainer, MaterialX::DocumentPtr Document) const;
@@ -85,7 +85,7 @@ protected:
 	 *
 	 * @param PointLightShaderNode - The MaterialX node related to a <point_light>
 	 * @param SceneNode - The interchange scene node, to which the node will be attached to
-	 * @param NodeContainer - The unreal objects descriptions container where to put the translated source data.
+	 * @param The container where to add the translated Interchange nodes.
 	 * @param Document - The MaterialX Document that contains all the definitions of the loaded libraries
 	 */
 	UInterchangeBaseLightNode* CreatePointLightNode(MaterialX::NodePtr PointLightShaderNode, UInterchangeSceneNode* SceneNode, UInterchangeBaseNodeContainer& NodeContainer, MaterialX::DocumentPtr Document) const;
@@ -95,7 +95,7 @@ protected:
 	 *
 	 * @param SpotLightShaderNode - The MaterialX node related to a <spot_light>
 	 * @param SceneNode - The interchange scene node, to which the node will be attached to
-	 * @param NodeContainer - The unreal objects descriptions container where to put the translated source data.
+	 * @param The container where to add the translated Interchange nodes.
 	 * @param Document - The MaterialX Document that contains all the definitions of the loaded libraries
 	 */
 	UInterchangeBaseLightNode* CreateSpotLightNode(MaterialX::NodePtr SpotLightShaderNode, UInterchangeSceneNode* SceneNode, UInterchangeBaseNodeContainer& NodeContainer, MaterialX::DocumentPtr Document) const;
@@ -106,7 +106,7 @@ protected:
 	 * @param ShaderNode - The Interchange shader node to connect the MaterialX's node graph to
 	 * @param ParentInputName - The name of the input of the shader node to which we want the node graph to be connected to
 	 * @param NamesToShaderNodes - Map of the shader nodes already created
-	 * @param NodeContainer - The unreal objects descriptions container where to put the translated source data.
+	 * @param The container where to add the translated Interchange nodes.
 	 * 
 	 * @return true if the given input is attached to one of the outputs of a node graph
 	 */
@@ -119,60 +119,104 @@ protected:
 	 * @param ParentShaderNode - The shader node to connect to
 	 * @param InputChannelName - The input of the ParentShaderNode to connect to
 	 * @param NamesToShaderNodes - Map of the shader nodes already created
-	 * @param NodeContainer - The unreal objects descriptions container where to put the translated source data.
+	 * @param The container where to add the translated Interchange nodes.
 	 * 
 	 * @return true if a shader node has been successfully created and is connected to the given input
 	 */
 	bool ConnectNodeOutputToInput(MaterialX::NodePtr Node, UInterchangeShaderNode* ParentShaderNode, const FString& InputChannelName, TMap<FString, UInterchangeShaderNode*>& NamesToShaderNodes, UInterchangeBaseNodeContainer& NodeContainer) const;
 
 	/**
-	 * Helper function to create an InterchangeShaderNode
+	 * Helper template function to create an InterchangeShaderNode
 	 * 
 	 * @param NodeName - The name of the shader node
 	 * @param ShaderType - The shader node's type we want to create
 	 * @param ParentNode - The parent node of the created node
 	 * @param NamesToShaderNodes - Map of the shader nodes already created
-	 * @param NodeContainer - The unreal objects descriptions container where to put the translated source data.
+	 * @param The container where to add the translated Interchange nodes.
 	 * 
 	 * @return The shader node that was created
 	 */
-	UInterchangeShaderNode* CreateShaderNode(const FString & NodeName, const FString & ShaderType, UInterchangeShaderNode * ParentNode, TMap<FString, UInterchangeShaderNode*>& NamesToShaderNodes, UInterchangeBaseNodeContainer & NodeContainer) const;
+	template<typename ShaderNodeType>
+	ShaderNodeType* CreateShaderNode(const FString& NodeName, const FString& ShaderType, const FString & ParentNodeUID, TMap<FString, UInterchangeShaderNode*>& NamesToShaderNodes, UInterchangeBaseNodeContainer& NodeContainer) const
+	{
+		static_assert(std::is_convertible_v<ShaderNodeType*, UInterchangeShaderNode*>, "CreateShaderNode only accepts type that derived from UInterchangeShaderNode");
+
+		ShaderNodeType* Node;
+
+		const FString NodeUID = UInterchangeShaderNode::MakeNodeUid(NodeName, ParentNodeUID);
+		
+		//Test directly in the NodeContainer, because the NamesToShaderNodes can be altered during the node graph either by the parent (dot/normalmap),
+		//or by putting an intermediary node between the child and the parent (tiledimage)
+		if(Node = const_cast<ShaderNodeType*>(Cast<ShaderNodeType>(NodeContainer.GetNode(NodeUID))); !Node)
+		{
+			Node = NewObject<ShaderNodeType>(&NodeContainer);
+			Node->InitializeNode(NodeUID, NodeName, EInterchangeNodeContainerType::TranslatedAsset);
+			NodeContainer.AddNode(Node);
+			if constexpr(std::is_same_v<ShaderNodeType, UInterchangeShaderGraphNode>)
+			{
+				NodeContainer.SetNodeParentUid(NodeUID, ParentNodeUID);
+			}
+			Node->SetCustomShaderType(ShaderType);
+
+			NamesToShaderNodes.Add(NodeName, Node);
+		}
+
+		return Node;
+	}
 
 	/**
 	 * Helper function to create an InterchangeTextureNode
 	 *
 	 * @param Node - The MaterialX node, it should be of the category <image> no test is done on it
-	 * @param NodeContainer - The unreal objects descriptions container where to put the translated source data.
+	 * @param The container where to add the translated Interchange nodes.
 	 * 
 	 * @return The texture node that was created
 	 */
 	UInterchangeTextureNode* CreateTextureNode(MaterialX::NodePtr Node, UInterchangeBaseNodeContainer& NodeContainer) const;
 
 	/**
-	 * Get the UE corresponding name of MaterialX input of a material
+	 * Get the UE corresponding name of a MaterialX Node category and input for a material
 	 * 
 	 * @param Input - MaterialX input
 	 * 
-	 * @return The matched name of the input else empty string
+	 * @return The matched name of the Node/Input else empty string
 	 */
-	const FString& GetMatchedInputName(MaterialX::InputPtr Input) const;
-
+	const FString& GetMatchedInputName(MaterialX::NodePtr Node, MaterialX::InputPtr Input) const;
+	
 	/**
 	 * Rename the inputs names of a node to correspond to the one used by UE, it will keep the old inputs names under the attribute "oldname"
 	 * 
 	 * @param Node - Look up to all inputs of Node and rename them to match UE names
 	 */
-	void RenameInputsNames(MaterialX::NodePtr Node) const;
+	void RenameNodeInputs(MaterialX::NodePtr Node) const;
+
+	/**
+	 * Rename the input name, it will keep the original input name under the attribute MaterialX::Attributes::OriginalName.
+	 * It will keep the uniqueness of the name inside the MaterialX Document
+	 *
+	 * @param Input - The input to rename
+	 * @param NewName - the new name of the input
+	 */
+	void RenameInput(MaterialX::InputPtr Input, const char * NewName) const;
 
 	/**
-	 * Helper function to retrieve an Input in a Node from its old name (after a renaming)
+	 * Helper function to retrieve an Input in a Node from its original name (after a renaming)
 	 * 
 	 * @param Node - The node with the Inputs to look up
-	 * @param OldNameAttribute - The previous name of an Input
+	 * @param OriginalNameAttribute - The previous name of an Input
 	 * 
 	 * @return The found Input, nullptr otherwise
 	 */
-	MaterialX::InputPtr GetInputFromOldName(MaterialX::NodePtr Node, const char * OldNameAttribute) const;
+	MaterialX::InputPtr GetInputFromOriginalName(MaterialX::NodePtr Node, const char * OriginalNameAttribute) const;
+
+	/**
+	 * Get the input name, use this function instead of getName, because a renaming may have occured and we ensure to have the proper name that will be used by UE inputs
+	 * 
+	 * @param Input - The input to retrieve the name from
+	 * 
+	 * @return The input name
+	 */
+	FString GetInputName(MaterialX::InputPtr Input) const;
 
 	/**
 	 * Retrieve the input from a standard_surface node, or take the default input from the library, 
@@ -283,7 +327,7 @@ protected:
 
 private:
 
-	TMap<FString, FString> InputNamesMaterialX2UE;
+	TMap<TPair<FString, FString>, FString> InputNamesMaterialX2UE; //given a MaterialX node (category - input), return the UE/Interchange input name.
 	TMap<FString, FString> NodeNamesMaterialX2UE; //given a MaterialX node category, return the UE category
 	TSet<FString> UEInputs;
 #endif // WITH_EDITOR
