Adding Coat color support for Standard Surface shading model

diff --git a/Engine/Plugins/Experimental/Interchange/Runtime/Source/Pipelines/Private/InterchangeGenericMaterialPipeline.cpp b/Engine/Plugins/Experimental/Interchange/Runtime/Source/Pipelines/Private/InterchangeGenericMaterialPipeline.cpp
--- a/Engine/Plugins/Experimental/Interchange/Runtime/Source/Pipelines/Private/InterchangeGenericMaterialPipeline.cpp
+++ b/Engine/Plugins/Experimental/Interchange/Runtime/Source/Pipelines/Private/InterchangeGenericMaterialPipeline.cpp
@@ -504,6 +504,7 @@ bool UInterchangeGenericMaterialPipeline::HandleStandardSurfaceModel(const UInte
 			ArgumentsType{StandardSurface::Parameters::SheenColor.ToString(), true, false},			// SheenColor
 			ArgumentsType{StandardSurface::Parameters::SheenRoughness.ToString(), true, false},		// SheenRoughness
 			ArgumentsType{StandardSurface::Parameters::Coat.ToString(), true, false},				// Coat
+			ArgumentsType{StandardSurface::Parameters::CoatColor.ToString(), false, false},			// Coat Color
 			ArgumentsType{StandardSurface::Parameters::CoatRoughness.ToString(), true, false},		// CoatRoughness
 			ArgumentsType{StandardSurface::Parameters::ThinFilmThickness.ToString(), true, false},	// ThinFilmThickness
 			ArgumentsType{StandardSurface::Parameters::Emission.ToString(), true, false},			// Emission
