On behalf of Emil Kirichev

Fix for missing textures coming from .glb files (embedded textures)
swarm https://p4-swarm.epicgames.net/reviews/20554094

diff --git a/Engine/Plugins/Experimental/Interchange/Runtime/Source/Import/Private/Gltf/InterchangeGltfTranslator.cpp b/Engine/Plugins/Experimental/Interchange/Runtime/Source/Import/Private/Gltf/InterchangeGltfTranslator.cpp
--- a/Engine/Plugins/Experimental/Interchange/Runtime/Source/Import/Private/Gltf/InterchangeGltfTranslator.cpp
+++ b/Engine/Plugins/Experimental/Interchange/Runtime/Source/Import/Private/Gltf/InterchangeGltfTranslator.cpp
@@ -15,6 +15,8 @@
 #include "InterchangeShaderGraphNode.h"
 #include "InterchangeTexture2DNode.h"
 
+#include "Texture/InterchangeImageWrapperTranslator.h"
+
 #include "Algo/Find.h"
 #include "StaticMeshAttributes.h"
 #include "UObject/GCObjectScopeGuard.h"
@@ -41,7 +43,9 @@ namespace UE::Interchange::Gltf::Private
 		}
 		else
 		{
-			return GenerateNameForGltfNode( FPaths::GetBaseFilename( Texture.Source.URI ), NodeIndex );
+			// If texture has no URI, we consider it an embedded texture (coming from a .glb file)
+			const FString NodeName = Texture.Source.URI.IsEmpty() ? "EmbeddedTexture" : FPaths::GetBaseFilename(Texture.Source.URI);
+			return GenerateNameForGltfNode(NodeName, NodeIndex);
 		}
 	}
 
@@ -888,23 +892,33 @@ TOptional< UE::Interchange::FImportImage > UInterchangeGltfTranslator::GetTextur
 	}
 
 	const GLTF::FTexture& GltfTexture = GltfAsset.Textures[ TextureIndex ];
-	const FString TextureFilePath = FPaths::ConvertRelativePathToFull( GltfTexture.Source.FilePath );
 
-	UInterchangeSourceData* PayloadSourceData = UInterchangeManager::GetInterchangeManager().CreateSourceData( TextureFilePath );
-	FGCObjectScopeGuard ScopedSourceData( PayloadSourceData );
-	
-	if ( !PayloadSourceData )
+	if (GltfTexture.Source.FilePath.IsEmpty())
 	{
-		return TOptional<UE::Interchange::FImportImage>();
+		// Embedded texture -- try using ImageWrapper to decode it
+		TArray64<uint8> ImageData(GltfTexture.Source.Data, GltfTexture.Source.DataByteLength);
+		return UInterchangeImageWrapperTranslator::GetTexturePayloadDataFromBuffer(ImageData);
 	}
-
-	UInterchangeTranslatorBase* SourceTranslator = UInterchangeManager::GetInterchangeManager().GetTranslatorForSourceData( PayloadSourceData );
-	FGCObjectScopeGuard ScopedSourceTranslator( SourceTranslator );
-	const IInterchangeTexturePayloadInterface* TextureTranslator = Cast< IInterchangeTexturePayloadInterface >( SourceTranslator );
-	if ( !ensure( TextureTranslator ) )
+	else
 	{
-		return TOptional<UE::Interchange::FImportImage>();
-	}
+		const FString TextureFilePath = FPaths::ConvertRelativePathToFull(GltfTexture.Source.FilePath);
+
+		UInterchangeSourceData* PayloadSourceData = UInterchangeManager::GetInterchangeManager().CreateSourceData(TextureFilePath);
+		FGCObjectScopeGuard ScopedSourceData(PayloadSourceData);
+
+		if (!PayloadSourceData)
+		{
+			return TOptional<UE::Interchange::FImportImage>();
+		}
+
+		UInterchangeTranslatorBase* SourceTranslator = UInterchangeManager::GetInterchangeManager().GetTranslatorForSourceData(PayloadSourceData);
+		FGCObjectScopeGuard ScopedSourceTranslator(SourceTranslator);
+		const IInterchangeTexturePayloadInterface* TextureTranslator = Cast< IInterchangeTexturePayloadInterface >(SourceTranslator);
+		if (!ensure(TextureTranslator))
+		{
+			return TOptional<UE::Interchange::FImportImage>();
+		}
 
-	return TextureTranslator->GetTexturePayloadData( PayloadSourceData, TextureFilePath );
+		return TextureTranslator->GetTexturePayloadData(PayloadSourceData, TextureFilePath);
+	}
 }
\ No newline at end of file
diff --git a/Engine/Plugins/Experimental/Interchange/Runtime/Source/Import/Private/Texture/InterchangeImageWrapperTranslator.cpp b/Engine/Plugins/Experimental/Interchange/Runtime/Source/Import/Private/Texture/InterchangeImageWrapperTranslator.cpp
--- a/Engine/Plugins/Experimental/Interchange/Runtime/Source/Import/Private/Texture/InterchangeImageWrapperTranslator.cpp
+++ b/Engine/Plugins/Experimental/Interchange/Runtime/Source/Import/Private/Texture/InterchangeImageWrapperTranslator.cpp
@@ -71,6 +71,11 @@ TOptional<UE::Interchange::FImportImage> UInterchangeImageWrapperTranslator::Get
 		return TOptional<UE::Interchange::FImportImage>();
 	}
 
+	return GetTexturePayloadDataFromBuffer(SourceDataBuffer);
+}
+
+TOptional<UE::Interchange::FImportImage> UInterchangeImageWrapperTranslator::GetTexturePayloadDataFromBuffer(const TArray64<uint8>& SourceDataBuffer)
+{
 	const uint8* Buffer = SourceDataBuffer.GetData();
 	const uint8* BufferEnd = Buffer + SourceDataBuffer.Num();
 
@@ -151,4 +156,3 @@ TOptional<UE::Interchange::FImportImage> UInterchangeImageWrapperTranslator::Get
 
 	return PayloadData;
 }
-
diff --git a/Engine/Plugins/Experimental/Interchange/Runtime/Source/Import/Public/Texture/InterchangeImageWrapperTranslator.h b/Engine/Plugins/Experimental/Interchange/Runtime/Source/Import/Public/Texture/InterchangeImageWrapperTranslator.h
--- a/Engine/Plugins/Experimental/Interchange/Runtime/Source/Import/Public/Texture/InterchangeImageWrapperTranslator.h
+++ b/Engine/Plugins/Experimental/Interchange/Runtime/Source/Import/Public/Texture/InterchangeImageWrapperTranslator.h
@@ -46,4 +46,6 @@ public:
 	virtual TOptional<UE::Interchange::FImportImage> GetTexturePayloadData(const UInterchangeSourceData* PayloadSourceData, const FString& PayLoadKey) const override;
 
 	/* IInterchangeTexturePayloadInterface End */
+
+	static TOptional<UE::Interchange::FImportImage> GetTexturePayloadDataFromBuffer(const TArray64<uint8>& SourceDataBuffer);
 };
\ No newline at end of file
