On behalf of Emil Kirichev
Added support for transform animations to Interchange GLTF translator

diff --git a/Engine/Plugins/Experimental/Interchange/Runtime/Source/Import/Private/Gltf/InterchangeGltfTranslator.cpp b/Engine/Plugins/Experimental/Interchange/Runtime/Source/Import/Private/Gltf/InterchangeGltfTranslator.cpp
--- a/Engine/Plugins/Experimental/Interchange/Runtime/Source/Import/Private/Gltf/InterchangeGltfTranslator.cpp
+++ b/Engine/Plugins/Experimental/Interchange/Runtime/Source/Import/Private/Gltf/InterchangeGltfTranslator.cpp
@@ -6,7 +6,9 @@
 #include "GLTFMeshFactory.h"
 #include "GLTFReader.h"
 
+#include "InterchangeAnimationTrackSetNode.h"
 #include "InterchangeCameraNode.h"
+#include "InterchangeImportLog.h"
 #include "InterchangeLightNode.h"
 #include "InterchangeManager.h"
 #include "InterchangeMaterialDefinitions.h"
@@ -15,40 +17,16 @@
 #include "InterchangeShaderGraphNode.h"
 #include "InterchangeTexture2DNode.h"
 
+#include "Sections/MovieScene3DTransformSection.h"
 #include "Texture/InterchangeImageWrapperTranslator.h"
 
 #include "Algo/Find.h"
+#include "Async/Async.h"
 #include "StaticMeshAttributes.h"
 #include "UObject/GCObjectScopeGuard.h"
 
 namespace UE::Interchange::Gltf::Private
 {
-	FString GenerateNameForGltfNode( const FString& NodeName, int32 NodeIndex )
-	{
-		if ( NodeIndex == 0 )
-		{
-			return NodeName;
-		}
-		else
-		{
-			return NodeName + TEXT("_") + LexToString( NodeIndex );
-		}
-	}
-
-	FString GenerateTextureName( const GLTF::FTexture& Texture, int32 NodeIndex )
-	{
-		if ( !Texture.Name.IsEmpty() )
-		{
-			return GenerateNameForGltfNode( Texture.Name, NodeIndex );
-		}
-		else
-		{
-			// If texture has no URI, we consider it an embedded texture (coming from a .glb file)
-			const FString NodeName = Texture.Source.URI.IsEmpty() ? "EmbeddedTexture" : FPaths::GetBaseFilename(Texture.Source.URI);
-			return GenerateNameForGltfNode(NodeName, NodeIndex);
-		}
-	}
-
 	EInterchangeTextureWrapMode ConvertWrap(GLTF::FSampler::EWrap Wrap)
 	{
 		switch (Wrap)
@@ -64,21 +42,98 @@ namespace UE::Interchange::Gltf::Private
 			return EInterchangeTextureWrapMode::Wrap;
 		}
 	}
+
+	void GetAnimationTransformPayloadData(const GLTF::FAnimation& Animation, int32 ChannelIndex, UE::Interchange::FAnimationCurvePayloadData& OutPayloadData)
+	{
+		const GLTF::FAnimation::FChannel& Channel = Animation.Channels[ChannelIndex];
+		const GLTF::FAnimation::FSampler& Sampler = Animation.Samplers[Channel.Sampler];
+
+		TArray<float> FrameTimeBuffer;
+		TArray<float> FrameDataBuffer;
+
+		Sampler.Input.GetFloatArray(FrameTimeBuffer);
+
+		ERichCurveInterpMode InterpolationMode = static_cast<ERichCurveInterpMode>(Sampler.Interpolation);
+
+		TFunction<void(TFunction<FVector3f(const float*)>, int32, int32, bool)> ConvertCurve = [InterpolationMode, &FrameTimeBuffer, &FrameDataBuffer, &OutPayloadData](TFunction<FVector3f(const float*)> ConvertSample, int32 CurvesIndexOffset, int32 SampleSize, bool bUnwindRotation)
+		{
+			const int32 SampleCount = FrameDataBuffer.Num() / SampleSize;
+
+			FRichCurve& Curve0 = OutPayloadData.Curves[CurvesIndexOffset + 0];
+			FRichCurve& Curve1 = OutPayloadData.Curves[CurvesIndexOffset + 1];
+			FRichCurve& Curve2 = OutPayloadData.Curves[CurvesIndexOffset + 2];
+
+			for (int32 SampleIndex = 0; SampleIndex < SampleCount; ++SampleIndex)
+			{
+				const float* RawSample = FrameDataBuffer.GetData() + (SampleIndex * SampleSize);
+
+				FVector3f Sample = ConvertSample(RawSample);
+
+				FKeyHandle Key0 = Curve0.AddKey(FrameTimeBuffer[SampleIndex], Sample[0], bUnwindRotation);
+				FKeyHandle Key1 = Curve1.AddKey(FrameTimeBuffer[SampleIndex], Sample[1], bUnwindRotation);
+				FKeyHandle Key2 = Curve2.AddKey(FrameTimeBuffer[SampleIndex], Sample[2], bUnwindRotation);
+
+				Curve0.SetKeyInterpMode(Key0, InterpolationMode);
+				Curve1.SetKeyInterpMode(Key1, InterpolationMode);
+				Curve2.SetKeyInterpMode(Key2, InterpolationMode);
+			}
+		};
+
+		switch (Channel.Target.Path)
+		{
+			case GLTF::FAnimation::EPath::Translation:
+			{
+				FrameDataBuffer.SetNumUninitialized(Sampler.Output.Count * 3);
+				Sampler.Output.GetCoordArray(reinterpret_cast<FVector3f*>(FrameDataBuffer.GetData()));
+				const float Scale = 100.f; // Convert m to cm
+				ConvertCurve([&Scale](const float* VecData)
+					{
+						return FVector3f(VecData[0] * Scale, VecData[1] * Scale, VecData[2] * Scale);
+					}, 0, 3, false);
+				break;
+			}
+			case GLTF::FAnimation::EPath::Rotation:
+			{
+				FrameDataBuffer.SetNumUninitialized(Sampler.Output.Count * 4);
+				Sampler.Output.GetQuatArray(reinterpret_cast<FVector4f*>(FrameDataBuffer.GetData()));
+				ConvertCurve([](const float* QuatData)
+					{
+						const FQuat4f Quat(QuatData[0], QuatData[1], QuatData[2], QuatData[3]);
+						return Quat.Euler();
+					}, 3, 4, true);
+				break;
+			}
+			case GLTF::FAnimation::EPath::Scale:
+			{
+				FrameDataBuffer.SetNumUninitialized(Sampler.Output.Count * 3);
+				Sampler.Output.GetCoordArray(reinterpret_cast<FVector3f*>(FrameDataBuffer.GetData()));
+				ConvertCurve([](const float* VecData)
+					{
+						return FVector3f(VecData[0], VecData[1], VecData[2]);
+					}, 6, 3, false);
+				break;
+			}
+			default:
+				UE_LOG(LogInterchangeImport, Warning, TEXT("Animation type not supported"));
+				break;
+		}
+	}
 }
 
-void UInterchangeGltfTranslator::HandleGltfNode( UInterchangeBaseNodeContainer& NodeContainer, const GLTF::FNode& GltfNode, const FString& ParentNodeUid, const int32 NodeIndex ) const
+void UInterchangeGltfTranslator::HandleGltfNode( UInterchangeBaseNodeContainer& NodeContainer, const GLTF::FNode& GltfNode, const FString& ParentNodeUid, const int32 NodeIndex, FNodeUidMap& NodeUidMap ) const
 {
 	using namespace UE::Interchange::Gltf::Private;
 
-	const FString NodeName = GenerateNameForGltfNode( GltfNode.Name, NodeIndex );
-	const FString NodeUid = ParentNodeUid + TEXT("\\") + NodeName;
+	const FString NodeUid = ParentNodeUid + TEXT("\\") + GltfNode.Name;
 
 	const UInterchangeSceneNode* ParentSceneNode = Cast< UInterchangeSceneNode >( NodeContainer.GetNode( ParentNodeUid ) );
 
 	UInterchangeSceneNode* InterchangeSceneNode = NewObject< UInterchangeSceneNode >( &NodeContainer );
-	InterchangeSceneNode->InitializeNode( NodeUid, NodeName, EInterchangeNodeContainerType::TranslatedScene );
+	InterchangeSceneNode->InitializeNode( NodeUid, GltfNode.Name, EInterchangeNodeContainerType::TranslatedScene );
 	NodeContainer.AddNode( InterchangeSceneNode );
 
+	NodeUidMap.Add( &GltfNode, NodeUid );
+
 	FTransform Transform = GltfNode.Transform;
 
 	constexpr float MetersToCentimeters = 100.f;
@@ -90,7 +145,7 @@ void UInterchangeGltfTranslator::HandleGltfNode( UInterchangeBaseNodeContainer&
 		{
 			if ( GltfAsset.Meshes.IsValidIndex( GltfNode.MeshIndex ) )
 			{
-				const FString MeshNodeUid = TEXT("\\Mesh\\") + GenerateNameForGltfNode( GltfAsset.Meshes[ GltfNode.MeshIndex ].Name, GltfNode.MeshIndex );
+				const FString MeshNodeUid = TEXT("\\Mesh\\") + GltfAsset.Meshes[ GltfNode.MeshIndex ].Name;
 				InterchangeSceneNode->SetCustomAssetInstanceUid( MeshNodeUid );
 			}
 			break;
@@ -102,7 +157,7 @@ void UInterchangeGltfTranslator::HandleGltfNode( UInterchangeBaseNodeContainer&
 
 			if ( GltfAsset.Cameras.IsValidIndex( GltfNode.CameraIndex ) )
 			{
-				const FString CameraNodeUid = TEXT("\\Camera\\") + GenerateNameForGltfNode( GltfAsset.Cameras[ GltfNode.CameraIndex ].Name, GltfNode.CameraIndex );
+				const FString CameraNodeUid = TEXT("\\Camera\\") + GltfAsset.Cameras[ GltfNode.CameraIndex ].Name;
 				InterchangeSceneNode->SetCustomAssetInstanceUid( CameraNodeUid );
 			}
 			break;
@@ -114,7 +169,7 @@ void UInterchangeGltfTranslator::HandleGltfNode( UInterchangeBaseNodeContainer&
 
 			if ( GltfAsset.Lights.IsValidIndex( GltfNode.LightIndex ) )
 			{
-				const FString LightNodeUid = TEXT("\\Light\\") + GenerateNameForGltfNode( GltfAsset.Lights[ GltfNode.LightIndex ].Name, GltfNode.LightIndex );
+				const FString LightNodeUid = TEXT("\\Light\\") + GltfAsset.Lights[ GltfNode.LightIndex ].Name;
 				InterchangeSceneNode->SetCustomAssetInstanceUid( LightNodeUid );
 			}
 		}
@@ -139,7 +194,7 @@ void UInterchangeGltfTranslator::HandleGltfNode( UInterchangeBaseNodeContainer&
 	{
 		if ( GltfAsset.Nodes.IsValidIndex( ChildIndex ) )
 		{
-			HandleGltfNode( NodeContainer, GltfAsset.Nodes[ ChildIndex ], NodeUid, ChildIndex );
+			HandleGltfNode( NodeContainer, GltfAsset.Nodes[ ChildIndex ], NodeUid, ChildIndex, NodeUidMap );
 		}
 	}
 }
@@ -183,7 +238,7 @@ void UInterchangeGltfTranslator::HandleGltfMaterialParameter( UInterchangeBaseNo
 		UInterchangeShaderNode* ColorNode = UInterchangeShaderNode::Create( &NodeContainer, ColorNodeName, ShaderNode.GetUniqueID() );
 		ColorNode->SetCustomShaderType( Standard::Nodes::TextureSample::Name.ToString() );
 
-		const FString TextureName = GenerateTextureName( GltfAsset.Textures[ TextureMap.TextureIndex ], TextureMap.TextureIndex );
+		const FString TextureName = GltfAsset.Textures[ TextureMap.TextureIndex ].Name;
 		const FString TextureUid = UInterchangeTextureNode::MakeNodeUid( TextureName );
 
 		ColorNode->AddStringAttribute( UInterchangeShaderPortsAPI::MakeInputValueKey( Standard::Nodes::TextureSample::Inputs::Texture.ToString() ), TextureUid );
@@ -691,7 +746,7 @@ EInterchangeTranslatorType UInterchangeGltfTranslator::GetTranslatorType() const
 bool UInterchangeGltfTranslator::DoesSupportAssetType(EInterchangeTranslatorAssetType AssetType) const
 {
 	//gltf translator support Meshes and Materials
-	if (AssetType == EInterchangeTranslatorAssetType::Materials || AssetType == EInterchangeTranslatorAssetType::Meshes)
+	if (AssetType == EInterchangeTranslatorAssetType::Materials || AssetType == EInterchangeTranslatorAssetType::Meshes || AssetType == EInterchangeTranslatorAssetType::Animations)
 	{
 		return true;
 	}
@@ -713,8 +768,8 @@ bool UInterchangeGltfTranslator::Translate( UInterchangeBaseNodeContainer& NodeC
 {
 	using namespace UE::Interchange::Gltf::Private;
 
-	FString Filename = GetSourceData()->GetFilename();
-	if ( !FPaths::FileExists( Filename ) )
+	FString FilePath = GetSourceData()->GetFilename();
+	if ( !FPaths::FileExists( FilePath ) )
 	{
 		return false;
 	}
@@ -723,14 +778,17 @@ bool UInterchangeGltfTranslator::Translate( UInterchangeBaseNodeContainer& NodeC
 
 	const bool bLoadImageData = false;
 	const bool bLoadMetaData = false;
-	GltfFileReader.ReadFile( Filename, bLoadImageData, bLoadMetaData, const_cast< UInterchangeGltfTranslator* >( this )->GltfAsset );
+	GltfFileReader.ReadFile( FilePath, bLoadImageData, bLoadMetaData, const_cast< UInterchangeGltfTranslator* >( this )->GltfAsset );
+
+	const FString FileName = FPaths::GetBaseFilename(FilePath);
+	const_cast< UInterchangeGltfTranslator* >( this )->GltfAsset.GenerateNames(FileName);
 
 	// Textures
 	{
 		int32 TextureIndex = 0;
 		for ( const GLTF::FTexture& GltfTexture : GltfAsset.Textures )
 		{
-			const FString TextureName = GenerateTextureName( GltfTexture, TextureIndex );
+			const FString TextureName = GltfTexture.Name;
 
 			UInterchangeTexture2DNode* TextureNode = UInterchangeTexture2DNode::Create( &NodeContainer, TextureName );
 			TextureNode->SetPayLoadKey( LexToString( TextureIndex++ ) );
@@ -745,7 +803,7 @@ bool UInterchangeGltfTranslator::Translate( UInterchangeBaseNodeContainer& NodeC
 		int32 MaterialIndex = 0;
 		for ( const GLTF::FMaterial& GltfMaterial : GltfAsset.Materials )
 		{
-			UInterchangeShaderGraphNode* ShaderGraphNode = UInterchangeShaderGraphNode::Create( &NodeContainer, TEXT("Material_") + GenerateNameForGltfNode( GltfMaterial.Name, MaterialIndex ) );
+			UInterchangeShaderGraphNode* ShaderGraphNode = UInterchangeShaderGraphNode::Create( &NodeContainer, TEXT("Material_") + GltfMaterial.Name );
 
 			HandleGltfMaterial( NodeContainer, GltfMaterial, *ShaderGraphNode );
 			++MaterialIndex;
@@ -758,9 +816,9 @@ bool UInterchangeGltfTranslator::Translate( UInterchangeBaseNodeContainer& NodeC
 		for ( const GLTF::FMesh& GltfMesh : GltfAsset.Meshes )
 		{
 			UInterchangeMeshNode* MeshNode = NewObject< UInterchangeMeshNode >( &NodeContainer );
-			FString MeshNodeUid = TEXT("\\Mesh\\") + GenerateNameForGltfNode( GltfMesh.Name, MeshIndex );
+			FString MeshNodeUid = TEXT("\\Mesh\\") + GltfMesh.Name;
 
-			MeshNode->InitializeNode( MeshNodeUid, GenerateNameForGltfNode( GltfMesh.Name, MeshIndex ), EInterchangeNodeContainerType::TranslatedAsset );
+			MeshNode->InitializeNode( MeshNodeUid, GltfMesh.Name, EInterchangeNodeContainerType::TranslatedAsset );
 			MeshNode->SetPayLoadKey( LexToString( MeshIndex++ ) );
 			NodeContainer.AddNode( MeshNode );
 
@@ -769,7 +827,7 @@ bool UInterchangeGltfTranslator::Translate( UInterchangeBaseNodeContainer& NodeC
 			{
 				if ( GltfAsset.Materials.IsValidIndex( Primitive.MaterialIndex ) )
 				{
-					const FString MaterialName = GenerateNameForGltfNode(GltfAsset.Materials[Primitive.MaterialIndex].Name, Primitive.MaterialIndex);
+					const FString MaterialName = GltfAsset.Materials[Primitive.MaterialIndex].Name;
 					const FString ShaderGraphNodeUid = UInterchangeShaderGraphNode::MakeNodeUid( TEXT("Material_") + MaterialName);
 					MeshNode->SetSlotMaterialDependencyUid(MaterialName, ShaderGraphNodeUid);
 				}
@@ -783,7 +841,7 @@ bool UInterchangeGltfTranslator::Translate( UInterchangeBaseNodeContainer& NodeC
 		for ( const GLTF::FCamera& GltfCamera : GltfAsset.Cameras )
 		{
 			UInterchangeCameraNode* CameraNode = NewObject< UInterchangeCameraNode >( &NodeContainer );
-			FString CameraNodeUid = TEXT("\\Camera\\") + GenerateNameForGltfNode( GltfCamera.Name, CameraIndex );
+			FString CameraNodeUid = TEXT("\\Camera\\") + GltfCamera.Name;
 			CameraNode->InitializeNode( CameraNodeUid, GltfCamera.Name, EInterchangeNodeContainerType::TranslatedAsset );
 			NodeContainer.AddNode( CameraNode );
 			++CameraIndex;
@@ -796,13 +854,16 @@ bool UInterchangeGltfTranslator::Translate( UInterchangeBaseNodeContainer& NodeC
 		for ( const GLTF::FLight& GltfLight : GltfAsset.Lights )
 		{
 			UInterchangeLightNode* LightNode = NewObject< UInterchangeLightNode >( &NodeContainer );
-			FString LightNodeUid = TEXT("\\Light\\") + GenerateNameForGltfNode( GltfLight.Name, LightIndex );
+			FString LightNodeUid = TEXT("\\Light\\") + GltfLight.Name;
 			LightNode->InitializeNode( LightNodeUid, GltfLight.Name, EInterchangeNodeContainerType::TranslatedAsset );
 			NodeContainer.AddNode( LightNode );
 			++LightIndex;
 		}
 	}
 
+	// Cache created scene nodes UIDs to use later for animation binding
+	FNodeUidMap NodeUidMap;
+
 	// Scenes
 	{
 		int32 SceneIndex = 0;
@@ -813,11 +874,9 @@ bool UInterchangeGltfTranslator::Translate( UInterchangeBaseNodeContainer& NodeC
 			FString SceneName = GltfScene.Name;
 			if (SceneName.IsEmpty())
 			{
-				SceneName = TEXT("Scene");
+				SceneName = LexToString(SceneIndex++) + TEXT("_scene_") + FileName;
 			}
 
-			SceneName = GenerateNameForGltfNode( SceneName, SceneIndex );
-
 			FString SceneNodeUid = TEXT("\\Scene\\") + SceneName;
 			SceneNode->InitializeNode( SceneNodeUid, SceneName, EInterchangeNodeContainerType::TranslatedScene );
 			NodeContainer.AddNode( SceneNode );
@@ -826,11 +885,17 @@ bool UInterchangeGltfTranslator::Translate( UInterchangeBaseNodeContainer& NodeC
 			{
 				if ( GltfAsset.Nodes.IsValidIndex( NodeIndex ) )
 				{
-					HandleGltfNode( NodeContainer, GltfAsset.Nodes[ NodeIndex ], SceneNodeUid, NodeIndex );
+					HandleGltfNode( NodeContainer, GltfAsset.Nodes[ NodeIndex ], SceneNodeUid, NodeIndex, NodeUidMap );
 				}
 			}
+		}
+	}
 
-			++SceneIndex;
+	// Animations
+	{
+		for (int32 AnimationIndex = 0; AnimationIndex < GltfAsset.Animations.Num(); ++AnimationIndex)
+		{
+			HandleGltfAnimation(NodeContainer, NodeUidMap, AnimationIndex);
 		}
 	}
 
@@ -870,7 +935,7 @@ TFuture< TOptional< UE::Interchange::FStaticMeshPayloadData > > UInterchangeGltf
 
 			if ( GltfAsset.Materials.IsValidIndex( MaterialIndex ) )
 			{
-				const FString MaterialName = GenerateNameForGltfNode(GltfAsset.Materials[MaterialIndex].Name, MaterialIndex);
+				const FString MaterialName = GltfAsset.Materials[MaterialIndex].Name;
 				StaticMeshAttributes.GetPolygonGroupMaterialSlotNames()[ MaterialSlotIndex ] = *MaterialName;
 			}
 		}
@@ -921,4 +986,181 @@ TOptional< UE::Interchange::FImportImage > UInterchangeGltfTranslator::GetTextur
 
 		return TextureTranslator->GetTexturePayloadData(PayloadSourceData, TextureFilePath);
 	}
-}
\ No newline at end of file
+}
+TFuture<TOptional<UE::Interchange::FAnimationCurvePayloadData>> UInterchangeGltfTranslator::GetAnimationCurvePayloadData(const FString& PayLoadKey) const
+{
+	using namespace UE::Interchange;
+	using namespace UE::Interchange::Gltf::Private;
+
+	TSharedPtr<TPromise<TOptional<UE::Interchange::FAnimationCurvePayloadData>>> Promise = MakeShared<TPromise<TOptional<UE::Interchange::FAnimationCurvePayloadData>>>();
+	Promise->SetValue(TOptional<UE::Interchange::FAnimationCurvePayloadData>());
+
+	TArray<FString> PayloadTokens;
+
+	// We need at minimum animation index + one channel index
+	if (2 > PayLoadKey.ParseIntoArray(PayloadTokens, TEXT(":")))
+	{
+		// Invalid payload
+		return Promise->GetFuture();
+	}
+
+	TArray<int32> Tokens;
+	Tokens.Reserve(PayloadTokens.Num());
+
+	for (const FString& Token : PayloadTokens)
+	{
+		int32 IntToken = INDEX_NONE;
+		LexFromString(IntToken, *Token);
+
+		if (IntToken != INDEX_NONE)
+		{
+			Tokens.Add(IntToken);
+		}
+	}
+
+	const int32 AnimationIndex = Tokens[0];
+
+	if (!ensure(GltfAsset.Animations.IsValidIndex(AnimationIndex)))
+	{
+		// Invalid payload
+		return Promise->GetFuture();
+	}
+
+	return Async(EAsyncExecution::TaskGraph, [this, AnimationIndex, Tokens = MoveTemp(Tokens)]
+		{
+			TOptional<FAnimationCurvePayloadData> Result;
+
+			const GLTF::FAnimation& GltfAnimation = GltfAsset.Animations[AnimationIndex];
+
+			UE::Interchange::FAnimationCurvePayloadData TransformPayloadData;
+
+			TransformPayloadData.Curves.SetNum(9);
+
+			for (int32 TokenIndex = 1; TokenIndex < Tokens.Num(); ++TokenIndex)
+			{
+				const int32 ChannelIndex = Tokens[TokenIndex];
+				if (!ensure(GltfAnimation.Channels.IsValidIndex(ChannelIndex)))
+				{
+					continue;
+				}
+
+				GetAnimationTransformPayloadData(GltfAnimation, ChannelIndex, TransformPayloadData);
+			}
+
+			Result.Emplace(MoveTemp(TransformPayloadData));
+
+			return Result;
+		}
+	);
+}
+
+TFuture<TOptional<UE::Interchange::FAnimationStepCurvePayloadData>> UInterchangeGltfTranslator::GetAnimationStepCurvePayloadData(const FString& PayLoadKey) const
+{
+	using namespace UE::Interchange;
+
+	TSharedPtr<TPromise<TOptional<FAnimationStepCurvePayloadData>>> Promise = MakeShared<TPromise<TOptional<FAnimationStepCurvePayloadData>>>();
+	Promise->SetValue(TOptional<FAnimationStepCurvePayloadData>());
+	return Promise->GetFuture();
+}
+
+TFuture<TOptional<UE::Interchange::FAnimationBakeTransformPayloadData>> UInterchangeGltfTranslator::GetAnimationBakeTransformPayloadData(const FString& PayLoadKey, const double BakeFrequency, const double RangeStartSecond, const double RangeStopSecond) const
+{
+	TSharedPtr<TPromise<TOptional<UE::Interchange::FAnimationBakeTransformPayloadData>>> Promise = MakeShared<TPromise<TOptional<UE::Interchange::FAnimationBakeTransformPayloadData>>>();
+	Promise->SetValue(TOptional<UE::Interchange::FAnimationBakeTransformPayloadData>());
+	return Promise->GetFuture();
+}
+
+void UInterchangeGltfTranslator::HandleGltfAnimation(UInterchangeBaseNodeContainer& NodeContainer, const FNodeUidMap& NodeUidMap, int32 AnimationIndex) const
+{
+	const GLTF::FAnimation& GltfAnimation = GltfAsset.Animations[AnimationIndex];
+
+	TMap< const GLTF::FNode*, TArray<int32> > NodeChannelsMap;
+
+	for (int32 ChannelIndex = 0; ChannelIndex < GltfAnimation.Channels.Num(); ++ChannelIndex)
+	{
+		const GLTF::FNode* AnimatedNode = &GltfAnimation.Channels[ChannelIndex].Target.Node;
+
+		TArray<int32>& NodeChannels = NodeChannelsMap.FindOrAdd(AnimatedNode);
+		NodeChannels.Add(ChannelIndex);
+	}
+
+	if (NodeChannelsMap.IsEmpty())
+	{
+		return;
+	}
+
+	UInterchangeAnimationTrackSetNode* TrackSetNode = NewObject< UInterchangeAnimationTrackSetNode >(&NodeContainer);
+
+	const FString AnimTrackSetName = GltfAnimation.Name;
+	const FString AnimTrackSetNodeUid = TEXT("\\Animation\\") + AnimTrackSetName;
+	TrackSetNode->InitializeNode(AnimTrackSetNodeUid, AnimTrackSetName, EInterchangeNodeContainerType::TranslatedAsset);
+
+	for (const TTuple<const GLTF::FNode*, TArray<int32>>& NodeChannelsEntry : NodeChannelsMap)
+	{
+		const TArray<int32>& ChannelIndices = NodeChannelsEntry.Value;
+		if (ChannelIndices.Num() == 0)
+		{
+			continue;
+		}
+
+		const GLTF::FNode* GltfNode = NodeChannelsEntry.Key;
+		const FString* NodeUid = NodeUidMap.Find(GltfNode);
+		if (!ensure(NodeUid))
+		{
+			continue;
+		}
+
+		UInterchangeTransformAnimationTrackNode* TransformAnimTrackNode = NewObject< UInterchangeTransformAnimationTrackNode >(&NodeContainer);
+
+		const FString TransformAnimTrackNodeName = FString::Printf(TEXT("%s_%s"), *GltfNode->Name, *GltfAnimation.Name);
+		const FString TransformAnimTrackNodeUid = TEXT("\\AnimationTrack\\") + TransformAnimTrackNodeName;
+
+		TransformAnimTrackNode->InitializeNode(TransformAnimTrackNodeUid, TransformAnimTrackNodeName, EInterchangeNodeContainerType::TranslatedAsset);
+
+		TransformAnimTrackNode->SetCustomActorDependencyUid(*NodeUid);
+
+		FString PayloadKey = FString::FromInt(AnimationIndex);
+
+		constexpr int32 TranslationChannel = 0x0001 | 0x0002 | 0x0004;
+		constexpr int32 RotationChannel = 0x0008 | 0x0010 | 0x0020;
+		constexpr int32 ScaleChannel = 0x0040 | 0x0080 | 0x0100;
+
+		int32 UsedChannels = 0;
+
+		for (int32 ChannelIndex : ChannelIndices)
+		{
+			PayloadKey += TEXT(":") + FString::FromInt(ChannelIndex);
+
+			const GLTF::FAnimation::FChannel& Channel = GltfAnimation.Channels[ChannelIndex];
+
+			switch (Channel.Target.Path)
+			{
+				case GLTF::FAnimation::EPath::Translation:
+				{
+					UsedChannels |= TranslationChannel;
+				} break;
+
+				case GLTF::FAnimation::EPath::Rotation:
+				{
+					UsedChannels |= RotationChannel;
+				} break;
+
+				case GLTF::FAnimation::EPath::Scale:
+				{
+					UsedChannels |= ScaleChannel;
+				} break;
+
+				default: break;
+			}
+		}
+
+		TransformAnimTrackNode->SetCustomAnimationPayloadKey(PayloadKey);
+		TransformAnimTrackNode->SetCustomUsedChannels(UsedChannels);
+
+		NodeContainer.AddNode(TransformAnimTrackNode);
+
+		TrackSetNode->AddCustomAnimationTrackUid(TransformAnimTrackNodeUid);
+	}
+
+	NodeContainer.AddNode(TrackSetNode);
+}
diff --git a/Engine/Plugins/Experimental/Interchange/Runtime/Source/Import/Public/Gltf/InterchangeGltfTranslator.h b/Engine/Plugins/Experimental/Interchange/Runtime/Source/Import/Public/Gltf/InterchangeGltfTranslator.h
--- a/Engine/Plugins/Experimental/Interchange/Runtime/Source/Import/Public/Gltf/InterchangeGltfTranslator.h
+++ b/Engine/Plugins/Experimental/Interchange/Runtime/Source/Import/Public/Gltf/InterchangeGltfTranslator.h
@@ -3,6 +3,7 @@
 
 #include "GLTFAsset.h"
 #include "InterchangeTranslatorBase.h"
+#include "Animation/InterchangeAnimationPayloadInterface.h"
 #include "Mesh/InterchangeStaticMeshPayload.h"
 #include "Mesh/InterchangeStaticMeshPayloadInterface.h"
 #include "Nodes/InterchangeBaseNodeContainer.h"
@@ -19,7 +20,9 @@ class UInterchangeShaderNode;
 
 UCLASS(BlueprintType, Experimental)
 class UInterchangeGltfTranslator : public UInterchangeTranslatorBase,
-	public IInterchangeStaticMeshPayloadInterface, public IInterchangeTexturePayloadInterface
+	public IInterchangeStaticMeshPayloadInterface, 
+	public IInterchangeTexturePayloadInterface,
+	public IInterchangeAnimationPayloadInterface
 {
 	GENERATED_BODY()
 
@@ -43,11 +46,20 @@ public:
 
 	/* IInterchangeTexturePayloadInterface End */
 
+	/* IInterchangeAnimationPayloadInterface Begin */
+	virtual TFuture<TOptional<UE::Interchange::FAnimationCurvePayloadData>> GetAnimationCurvePayloadData(const FString& PayLoadKey) const override;
+	virtual TFuture<TOptional<UE::Interchange::FAnimationStepCurvePayloadData>> GetAnimationStepCurvePayloadData(const FString& PayLoadKey) const override;
+	virtual TFuture<TOptional<UE::Interchange::FAnimationBakeTransformPayloadData>> GetAnimationBakeTransformPayloadData(const FString& PayLoadKey, const double BakeFrequency, const double RangeStartSecond, const double RangeStopSecond) const override;
+	/* IInterchangeAnimationPayloadInterface End */
+
 protected:
-	void HandleGltfNode( UInterchangeBaseNodeContainer& NodeContainer, const GLTF::FNode& GltfNode, const FString& ParentNodeUid, const int32 NodeIndex ) const;
+	using FNodeUidMap = TMap<const GLTF::FNode*, FString>;
+
+	void HandleGltfNode( UInterchangeBaseNodeContainer& NodeContainer, const GLTF::FNode& GltfNode, const FString& ParentNodeUid, const int32 NodeIndex, FNodeUidMap& NodeUidMap ) const;
 	void HandleGltfMaterial( UInterchangeBaseNodeContainer& NodeContainer, const GLTF::FMaterial& GltfMaterial, UInterchangeShaderGraphNode& ShaderGraphNode ) const;
 	void HandleGltfMaterialParameter( UInterchangeBaseNodeContainer& NodeContainer, const GLTF::FTextureMap& TextureMap, UInterchangeShaderNode& ShaderNode,
 		const FString& MapName, const TVariant< FLinearColor, float >& MapFactor, const FString& OutputChannel, const bool bInverse = false, const bool bIsNormal = false ) const;
+	void HandleGltfAnimation(UInterchangeBaseNodeContainer& NodeContainer, const FNodeUidMap& NodeUidMap, int32 AnimationIndex) const;
 
 	/** Support for KHR_materials_clearcoat */
 	void HandleGltfClearCoat( UInterchangeBaseNodeContainer& NodeContainer, const GLTF::FMaterial& GltfMaterial, UInterchangeShaderGraphNode& ShaderGraphNode ) const;
