MaterialX
- Fixing coat_color connected to coat input in the translator
- Renaming fuzz color and cloth outhputs in the MX_StandardSurface to reflect the outputs in the generic material pipeline

diff --git a/Engine/Plugins/Experimental/Interchange/Runtime/Source/Import/Private/MaterialX/InterchangeMaterialXTranslator.cpp b/Engine/Plugins/Experimental/Interchange/Runtime/Source/Import/Private/MaterialX/InterchangeMaterialXTranslator.cpp
--- a/Engine/Plugins/Experimental/Interchange/Runtime/Source/Import/Private/MaterialX/InterchangeMaterialXTranslator.cpp
+++ b/Engine/Plugins/Experimental/Interchange/Runtime/Source/Import/Private/MaterialX/InterchangeMaterialXTranslator.cpp
@@ -449,7 +449,7 @@ void UInterchangeMaterialXTranslator::ProcessStandardSurface(UInterchangeBaseNod
 		{
 			mx::InputPtr InputCoatColor = GetStandardSurfaceInput(StandardSurfaceNode, mx::StandardSurface::Input::CoatColor, Document);
 
-			if(!ConnectNodeGraphOutputToInput(InputCoatColor, ShaderGraphNode, StandardSurface::Parameters::Coat.ToString(), NamesToShaderNodes, NodeContainer))
+			if(!ConnectNodeGraphOutputToInput(InputCoatColor, ShaderGraphNode, StandardSurface::Parameters::CoatColor.ToString(), NamesToShaderNodes, NodeContainer))
 			{
 				AddLinearColorAttribute(InputCoatColor, StandardSurface::Parameters::CoatColor.ToString(), ShaderGraphNode, mx::StandardSurface::DefaultValue::Color3::CoatColor);
 			}
