[Interchange] Support for Variants
  - Added UInterchangeVariantSetNode and UInterchangeSceneVariantSetNode class to translate Variant data
  - Added UInterchangeSceneVariantSetFactoryNode class to convert translated variant data into UE
  - Added UInterchangeSceneVariantSetFactory class to create VariantSets asset
  - Added IInterchangeVariantSetPayloadInterface to retrieve payload attached to a UInterchangeVariantSetNode
  - Added FVariantSetPayloadData struct and its related struct to store variant's data in the payload
  - Added processing of UInterchangeSceneVariantSetNode nodes in the UInterchangeGenericLevelPipeline
  - Moved the assignment of the static mesh to a static mesh actor from UInterchangeStaticMeshActorFactory::PostImportPreCompletedCallback to UInterchangeStaticMeshActorFactory::SetupStaticMeshActor

Unrelated: Removed conditional dependency to the Engine module in the InterchangeNodes module. This is done in order to quickly identify potential build break in the InterchangeWorker program

diff --git a/Engine/Plugins/Experimental/Interchange/Runtime/Interchange.uplugin b/Engine/Plugins/Experimental/Interchange/Runtime/Interchange.uplugin
--- a/Engine/Plugins/Experimental/Interchange/Runtime/Interchange.uplugin
+++ b/Engine/Plugins/Experimental/Interchange/Runtime/Interchange.uplugin
@@ -75,6 +75,14 @@
 		{
 			"Name": "GLTFImporter",
 			"Enabled": true
+		},
+		{
+			"Name": "VariantManagerContent",
+			"Enabled": true
+		},
+		{
+			"Name": "VariantManager",
+			"Enabled": true
 		}
 	],
 	"IsExperimentalVersion": false,
diff --git a/Engine/Plugins/Experimental/Interchange/Runtime/Source/FactoryNodes/InterchangeFactoryNodes.Build.cs b/Engine/Plugins/Experimental/Interchange/Runtime/Source/FactoryNodes/InterchangeFactoryNodes.Build.cs
--- a/Engine/Plugins/Experimental/Interchange/Runtime/Source/FactoryNodes/InterchangeFactoryNodes.Build.cs
+++ b/Engine/Plugins/Experimental/Interchange/Runtime/Source/FactoryNodes/InterchangeFactoryNodes.Build.cs
@@ -16,6 +16,7 @@ namespace UnrealBuildTool.Rules
 					"InterchangeCore",
 					"InterchangeNodes",
 					"TextureUtilitiesCommon",
+					"VariantManagerContent",
 				}
 			);
 
diff --git a/Engine/Plugins/Experimental/Interchange/Runtime/Source/FactoryNodes/Private/InterchangeSceneVariantSetsFactoryNode.cpp b/Engine/Plugins/Experimental/Interchange/Runtime/Source/FactoryNodes/Private/InterchangeSceneVariantSetsFactoryNode.cpp
new file mode 100644
--- /dev/null
+++ b/Engine/Plugins/Experimental/Interchange/Runtime/Source/FactoryNodes/Private/InterchangeSceneVariantSetsFactoryNode.cpp
@@ -0,0 +1,40 @@
+// Copyright Epic Games, Inc. All Rights Reserved.
+
+#include "InterchangeSceneVariantSetsFactoryNode.h"
+
+#include "LevelVariantSets.h"
+
+UInterchangeSceneVariantSetsFactoryNode::UInterchangeSceneVariantSetsFactoryNode()
+{
+	CustomVariantSetUids.Initialize(Attributes, TEXT("VariantSetUids"));
+}
+
+UClass* UInterchangeSceneVariantSetsFactoryNode::GetObjectClass() const
+{
+	return ULevelVariantSets::StaticClass();
+}
+
+int32 UInterchangeSceneVariantSetsFactoryNode::GetCustomVariantSetUidCount() const
+{
+	return CustomVariantSetUids.GetCount();
+}
+
+void UInterchangeSceneVariantSetsFactoryNode::GetCustomVariantSetUids(TArray<FString>& OutVariantSetUids) const
+{
+	CustomVariantSetUids.GetItems(OutVariantSetUids);
+}
+
+void UInterchangeSceneVariantSetsFactoryNode::GetCustomVariantSetUid(const int32 Index, FString& OutVariantSetUid) const
+{
+	CustomVariantSetUids.GetItem(Index, OutVariantSetUid);
+}
+
+bool UInterchangeSceneVariantSetsFactoryNode::AddCustomVariantSetUid(const FString& VariantSetUid)
+{
+	return CustomVariantSetUids.AddItem(VariantSetUid);
+}
+
+bool UInterchangeSceneVariantSetsFactoryNode::RemoveCustomVariantSetUid(const FString& VariantSetUid)
+{
+	return CustomVariantSetUids.RemoveItem(VariantSetUid);
+}
diff --git a/Engine/Plugins/Experimental/Interchange/Runtime/Source/FactoryNodes/Public/InterchangeSceneVariantSetsFactoryNode.h b/Engine/Plugins/Experimental/Interchange/Runtime/Source/FactoryNodes/Public/InterchangeSceneVariantSetsFactoryNode.h
new file mode 100644
--- /dev/null
+++ b/Engine/Plugins/Experimental/Interchange/Runtime/Source/FactoryNodes/Public/InterchangeSceneVariantSetsFactoryNode.h
@@ -0,0 +1,68 @@
+// Copyright Epic Games, Inc. All Rights Reserved.
+
+#pragma once
+
+#include "CoreMinimal.h"
+#include "Nodes/InterchangeFactoryBaseNode.h"
+
+#include "InterchangeSceneVariantSetsFactoryNode.generated.h"
+
+UCLASS(BlueprintType, Experimental)
+class INTERCHANGEFACTORYNODES_API UInterchangeSceneVariantSetsFactoryNode : public UInterchangeFactoryBaseNode
+{
+	GENERATED_BODY()
+
+	UInterchangeSceneVariantSetsFactoryNode();
+
+public:
+	static FStringView StaticAssetTypeName()
+	{
+		return TEXT("SceneVariantSetFactory");
+	}
+
+	/**
+	 * Return the node type name of the class, we use this when reporting errors
+	 */
+	virtual FString GetTypeName() const override
+	{
+		const FString TypeName = TEXT("SceneVariantSetFactoryNode");
+		return TypeName;
+	}
+
+	/** Get the class this node want to create */
+	UFUNCTION(BlueprintCallable, Category = "Interchange | Node | SceneVariantSetFactory")
+	virtual class UClass* GetObjectClass() const override;
+
+	/**
+	 * This function allow to retrieve the number of translated VariantSets' unique ids for this object.
+	 */
+	UFUNCTION(BlueprintCallable, Category = "Interchange | Node | SceneVariantSetFactory")
+	int32 GetCustomVariantSetUidCount() const;
+
+	/**
+	 * This function allow to retrieve all the translated VariantSets' unique ids for this object.
+	 */
+	UFUNCTION(BlueprintCallable, Category = "Interchange | Node | SceneVariantSetFactory")
+	void GetCustomVariantSetUids(TArray<FString>& OutVariantUids) const;
+
+	/**
+	 * This function allow to retrieve one translated VariantSet's unique id for this object.
+	 */
+	UFUNCTION(BlueprintCallable, Category = "Interchange | Node | SceneVariantSetFactory")
+	void GetCustomVariantSetUid(const int32 Index, FString& OutVariantUid) const;
+
+	/**
+	 * Add one translated VariantSet's unique id to this object.
+	 */
+	UFUNCTION(BlueprintCallable, Category = "Interchange | Node | SceneVariantSetFactory")
+	bool AddCustomVariantSetUid(const FString& VariantUid);
+
+	/**
+	 * Remove one translated VariantSet's unique id from this object.
+	 */
+	UFUNCTION(BlueprintCallable, Category = "Interchange | Node | SceneVariantSetFactory")
+	bool RemoveCustomVariantSetUid(const FString& VariantUid);
+
+private:
+	UE::Interchange::TArrayAttributeHelper<FString> CustomVariantSetUids;
+};
diff --git a/Engine/Plugins/Experimental/Interchange/Runtime/Source/Import/InterchangeImport.Build.cs b/Engine/Plugins/Experimental/Interchange/Runtime/Source/Import/InterchangeImport.Build.cs
--- a/Engine/Plugins/Experimental/Interchange/Runtime/Source/Import/InterchangeImport.Build.cs
+++ b/Engine/Plugins/Experimental/Interchange/Runtime/Source/Import/InterchangeImport.Build.cs
@@ -41,6 +41,7 @@ namespace UnrealBuildTool.Rules
 					"RenderCore",
 					"RHI",
 					"TextureUtilitiesCommon",
+					"VariantManagerContent",
 				}
 			);
 
@@ -52,6 +53,7 @@ namespace UnrealBuildTool.Rules
 						"BSPUtils",
 						"MaterialEditor",
 						"UnrealEd",
+						"VariantManager",
 					}
 				);
 
diff --git a/Engine/Plugins/Experimental/Interchange/Runtime/Source/Import/Private/InterchangeImportModule.cpp b/Engine/Plugins/Experimental/Interchange/Runtime/Source/Import/Private/InterchangeImportModule.cpp
--- a/Engine/Plugins/Experimental/Interchange/Runtime/Source/Import/Private/InterchangeImportModule.cpp
+++ b/Engine/Plugins/Experimental/Interchange/Runtime/Source/Import/Private/InterchangeImportModule.cpp
@@ -20,6 +20,7 @@
 #include "Modules/ModuleManager.h"
 #include "Scene/InterchangeActorFactory.h"
 #include "Scene/InterchangeCineCameraActorFactory.h"
+#include "Scene/InterchangeSceneVariantSetsFactory.h"
 #include "Scene/InterchangeStaticMeshActorFactory.h"
 #include "Scene/InterchangeSkeletalMeshActorFactory.h"
 #include "Texture/InterchangeImageWrapperTranslator.h"
@@ -86,6 +87,7 @@ void FInterchangeImportModule::StartupModule()
 		InterchangeManager.RegisterFactory(UInterchangeCineCameraActorFactory::StaticClass());
 		InterchangeManager.RegisterFactory(UInterchangeStaticMeshActorFactory::StaticClass());
 		InterchangeManager.RegisterFactory(UInterchangeSkeletalMeshActorFactory::StaticClass());
+		InterchangeManager.RegisterFactory(UInterchangeSceneVariantSetsFactory::StaticClass());
 	};
 	
 	if (GEngine)
diff --git a/Engine/Plugins/Experimental/Interchange/Runtime/Source/Import/Private/Scene/InterchangeSceneVariantSetsFactory.cpp b/Engine/Plugins/Experimental/Interchange/Runtime/Source/Import/Private/Scene/InterchangeSceneVariantSetsFactory.cpp
new file mode 100644
--- /dev/null
+++ b/Engine/Plugins/Experimental/Interchange/Runtime/Source/Import/Private/Scene/InterchangeSceneVariantSetsFactory.cpp
@@ -0,0 +1,426 @@
+// Copyright Epic Games, Inc. All Rights Reserved.
+
+#include "Scene/InterchangeSceneVariantSetsFactory.h"
+
+#include "InterchangeImportLog.h"
+#include "InterchangeMaterialFactoryNode.h"
+#include "InterchangeResult.h"
+#include "InterchangeSourceData.h"
+#include "InterchangeTranslatorBase.h"
+#include "InterchangeSceneVariantSetsFactoryNode.h"
+#include "InterchangeVariantSetNode.h"
+#include "Nodes/InterchangeBaseNode.h"
+#include "Nodes/InterchangeBaseNodeContainer.h"
+
+#include "GameFramework/Actor.h"
+#include "LevelVariantSets.h"
+#include "Materials/MaterialInterface.h"
+#include "PropertyValue.h"
+#include "Variant.h"
+#include "VariantObjectBinding.h"
+#include "VariantSet.h"
+#if WITH_EDITOR
+#include "VariantManager.h"
+#include "VariantManagerModule.h"
+#endif
+
+#define LOCTEXT_NAMESPACE "InterchangeSceneVariantSetsFactory"
+
+#if WITH_EDITOR
+namespace UE::Interchange::Private
+{
+	class FLevelVariantSetsHelper
+	{
+	public:
+		FLevelVariantSetsHelper(ULevelVariantSets& InLevelSequence, UInterchangeSceneVariantSetsFactoryNode& InFactoryNode, const UInterchangeBaseNodeContainer& InNodeContainer, const IInterchangeVariantSetPayloadInterface& InPayloadInterface)
+			: LevelVariantSets(InLevelSequence)
+			, FactoryNode(InFactoryNode)
+			, NodeContainer(InNodeContainer)
+			, PayloadInterface(InPayloadInterface)
+		{
+		}
+
+		void PopulateLevelVariantSets();
+		TSharedPtr<FVariantManager> GetVariantManager() { return VariantManager; }
+
+		TArray<FText> Warnings;
+
+	private:
+		void AddVariantSet(const UInterchangeVariantSetNode& VariantSetNode);
+		void AddVariant(const FVariant& Variant, UVariantSet& VariantSet);
+		void AddBinding(const FVariantBinding& Binding, UVariant& Variant);
+		void AddPropertyCapture(const FVariantPropertyCaptureData& PropertyCapture, UVariantObjectBinding& ObjectBinding);
+
+		AActor* GetActor(const FString& ActorNodeUid);
+
+	private:
+		ULevelVariantSets& LevelVariantSets;
+		UInterchangeSceneVariantSetsFactoryNode& FactoryNode;
+		const UInterchangeBaseNodeContainer& NodeContainer;
+		const IInterchangeVariantSetPayloadInterface& PayloadInterface;
+		TSharedPtr<FVariantManager> VariantManager;
+	};
+
+	void FLevelVariantSetsHelper::PopulateLevelVariantSets()
+	{
+		if (FactoryNode.GetCustomVariantSetUidCount() == 0)
+		{
+			return;
+		}
+
+		IVariantManagerModule& VariantManagerModule = FModuleManager::LoadModuleChecked<IVariantManagerModule>("VariantManager");
+		VariantManager = VariantManagerModule.CreateVariantManager(&LevelVariantSets);
+		if (!ensure(VariantManager))
+		{
+			return;
+		}
+
+		TArray<FString> VariantSetUids;
+		FactoryNode.GetCustomVariantSetUids(VariantSetUids);
+
+		for (const FString& VariantSetUid : VariantSetUids)
+		{
+			if (const UInterchangeVariantSetNode* VariantSetNode = Cast<UInterchangeVariantSetNode>(NodeContainer.GetNode(VariantSetUid)))
+			{
+				AddVariantSet(*VariantSetNode);
+			}
+		}
+	}
+
+	void FLevelVariantSetsHelper::AddVariantSet(const UInterchangeVariantSetNode& VariantSetNode)
+	{
+		FString PayloadKey;
+		if (!VariantSetNode.GetCustomVariantsPayloadKey(PayloadKey))
+		{
+			return;
+		}
+
+		FString DisplayText;
+		VariantSetNode.GetCustomDisplayText(DisplayText);
+
+		const TOptional<UE::Interchange::FVariantSetPayloadData>& PayloadData = PayloadInterface.GetVariantSetPayloadData(PayloadKey).Get();
+		if (!PayloadData.IsSet())
+		{
+			UE_LOG(LogInterchangeImport, Warning, TEXT("No payload for variant set %s"), *DisplayText);
+			return;
+		}
+
+		UVariantSet* VariantSet = VariantManager->CreateVariantSet(&LevelVariantSets);
+
+		VariantSet->SetDisplayText(FText::FromString(DisplayText));
+
+		for (const FVariant& Variant : PayloadData->Variants)
+		{
+			AddVariant(Variant, *VariantSet);
+		}
+	}
+
+	void FLevelVariantSetsHelper::AddVariant(const FVariant& VariantData, UVariantSet& VariantSet)
+	{
+
+		UVariant* Variant = VariantManager->CreateVariant(&VariantSet);
+		Variant->SetDisplayText(FText::FromString(VariantData.DisplayText));
+
+		for (const FVariantBinding& Binding : VariantData.Bindings)
+		{
+			AddBinding(Binding, *Variant);
+		}
+	}
+
+	void FLevelVariantSetsHelper::AddBinding(const FVariantBinding& Binding, UVariant& Variant)
+	{
+		// Get targeted actor exists
+		AActor* Actor = GetActor(Binding.TargetUid);
+		if (!Actor)
+		{
+			UE_LOG(LogInterchangeImport, Warning, TEXT("Cannot find actor for variant %s"), *Variant.GetDisplayText().ToString());
+			return;
+		}
+
+		TArray<UVariantObjectBinding*> Bindings = VariantManager->CreateObjectBindings({ Actor }, { &Variant });
+		if (Bindings.Num() != 1 || Bindings[0] == nullptr)
+		{
+			return;
+		}
+
+		UVariantObjectBinding& ObjectBinding = *(Bindings[0]);
+		for (const FVariantPropertyCaptureData& PropertyCapture : Binding.Captures)
+		{
+			AddPropertyCapture(PropertyCapture, ObjectBinding);
+		}
+	}
+
+	void FLevelVariantSetsHelper::AddPropertyCapture(const FVariantPropertyCaptureData& PropertyCapture, UVariantObjectBinding& ObjectBinding)
+	{
+		using namespace UE::Interchange;
+
+		static_assert((uint8)EVariantPropertyCaptureCategory::Undefined == (uint8)EPropertyValueCategory::Undefined, "INVALID_ENUM_VALUE");
+		static_assert((uint8)EVariantPropertyCaptureCategory::Generic == (uint8)EPropertyValueCategory::Generic, "INVALID_ENUM_VALUE");
+		static_assert((uint8)EVariantPropertyCaptureCategory::RelativeLocation == (uint8)EPropertyValueCategory::RelativeLocation, "INVALID_ENUM_VALUE");
+		static_assert((uint8)EVariantPropertyCaptureCategory::RelativeRotation == (uint8)EPropertyValueCategory::RelativeRotation, "INVALID_ENUM_VALUE");
+		static_assert((uint8)EVariantPropertyCaptureCategory::RelativeScale3D == (uint8)EPropertyValueCategory::RelativeScale3D, "INVALID_ENUM_VALUE");
+		static_assert((uint8)EVariantPropertyCaptureCategory::Visibility == (uint8)EPropertyValueCategory::Visibility, "INVALID_ENUM_VALUE");
+		static_assert((uint8)EVariantPropertyCaptureCategory::Material == (uint8)EPropertyValueCategory::Material, "INVALID_ENUM_VALUE");
+		static_assert((uint8)EVariantPropertyCaptureCategory::Color == (uint8)EPropertyValueCategory::Color, "INVALID_ENUM_VALUE");
+		static_assert((uint8)EVariantPropertyCaptureCategory::Option == (uint8)EPropertyValueCategory::Option, "INVALID_ENUM_VALUE");
+
+		if (PropertyCapture.Buffer.IsSet())
+		{
+			const TArray<uint8>& RecordedData = *PropertyCapture.Buffer;
+
+			UPropertyValue* PropertyValue = nullptr;
+
+			if (EnumHasAnyFlags(PropertyCapture.Category, EVariantPropertyCaptureCategory::Visibility))
+			{
+				TArray<UPropertyValue*> PropertyValues = VariantManager->CreateVisibilityPropertyCaptures({ &ObjectBinding });
+				if (PropertyValues.Num() == 1)
+				{
+					PropertyValue = PropertyValues[0];
+				}
+			}
+			else if (EnumHasAnyFlags(PropertyCapture.Category, EVariantPropertyCaptureCategory::RelativeLocation))
+			{
+				TArray<UPropertyValue*> PropertyValues = VariantManager->CreateLocationPropertyCaptures({ &ObjectBinding });
+				if (PropertyValues.Num() == 1)
+				{
+					PropertyValue = PropertyValues[0];
+				}
+			}
+			else if (EnumHasAnyFlags(PropertyCapture.Category, EVariantPropertyCaptureCategory::RelativeRotation))
+			{
+				TArray<UPropertyValue*> PropertyValues = VariantManager->CreateRotationPropertyCaptures({ &ObjectBinding });
+				if (PropertyValues.Num() == 1)
+				{
+					PropertyValue = PropertyValues[0];
+				}
+			}
+			else if (EnumHasAnyFlags(PropertyCapture.Category, EVariantPropertyCaptureCategory::RelativeScale3D))
+			{
+				TArray<UPropertyValue*> PropertyValues = VariantManager->CreateScale3DPropertyCaptures({ &ObjectBinding });
+				if (PropertyValues.Num() == 1)
+				{
+					PropertyValue = PropertyValues[0];
+				}
+			}
+
+			if (PropertyValue)
+			{
+				PropertyValue->SetRecordedData(RecordedData.GetData(), RecordedData.Num());
+			}
+		}
+		else if (PropertyCapture.ObjectUid.IsSet())
+		{
+			const FString TargetFactoryNodeUid = UInterchangeFactoryBaseNode::BuildFactoryNodeUid(*PropertyCapture.ObjectUid);
+
+			if (EnumHasAnyFlags(PropertyCapture.Category, EVariantPropertyCaptureCategory::Material))
+			{
+				if (const UInterchangeBaseMaterialFactoryNode* TargetFactoryNode = Cast<UInterchangeBaseMaterialFactoryNode>(NodeContainer.GetNode(TargetFactoryNodeUid)))
+				{
+					if (const UMaterialInterface* TargetMaterial = Cast<UMaterialInterface>(TargetFactoryNode->ReferenceObject.TryLoad()))
+					{
+						TArray<UPropertyValue*> PropertyValues = VariantManager->CreateMaterialPropertyCaptures({ &ObjectBinding });
+						if (PropertyValues.Num() == 1)
+						{
+							PropertyValues[0]->SetRecordedData((uint8*)&TargetMaterial, sizeof(UMaterialInterface*));
+						}
+						else
+						{
+							for (UPropertyValue* PropertyValue : PropertyValues)
+							{
+								if (PropertyValue && PropertyValue->GetFullDisplayString() == PropertyCapture.PropertyPath)
+								{
+									PropertyValue->SetRecordedData((uint8*)&TargetMaterial, sizeof(UMaterialInterface*));
+									break;
+								}
+							}
+						}
+					}
+					else
+					{
+						UE_LOG(LogInterchangeImport, Warning, TEXT("Did not find material '%s' when creating variant asset."), *TargetFactoryNode->GetDisplayLabel());
+					}
+				}
+				else
+				{
+					UE_LOG(LogInterchangeImport, Warning, TEXT("Did not find factory when creating variant assets."));
+				}
+			}
+		}
+	}
+
+	AActor* FLevelVariantSetsHelper::GetActor(const FString& ActorNodeUid)
+	{
+		const FString ActorFactoryNodeUid = UInterchangeFactoryBaseNode::BuildFactoryNodeUid(ActorNodeUid);
+		const UInterchangeFactoryBaseNode* ActorFactoryNode = Cast<UInterchangeFactoryBaseNode>(NodeContainer.GetNode(ActorFactoryNodeUid));
+
+		if (ActorFactoryNode && ActorFactoryNode->ReferenceObject.IsValid())
+		{
+			return Cast<AActor>(ActorFactoryNode->ReferenceObject.TryLoad());
+		}
+
+		return nullptr;
+	}
+} //namespace UE::Interchange::Private
+#endif //WITH_EDITOR
+
+UClass* UInterchangeSceneVariantSetsFactory::GetFactoryClass() const
+{
+	return ULevelVariantSets::StaticClass();
+}
+
+UObject* UInterchangeSceneVariantSetsFactory::CreateEmptyAsset(const FCreateAssetParams& Arguments)
+{
+#if !WITH_EDITOR || !WITH_EDITORONLY_DATA
+
+	UE_LOG(LogInterchangeImport, Error, TEXT("Cannot import levelsequence asset in runtime, this is an editor only feature."));
+	return nullptr;
+#else
+	if (Arguments.ReimportObject)
+	{
+		UInterchangeResultError_Generic* Message = AddMessage<UInterchangeResultError_Generic>();
+		Message->SourceAssetName = Arguments.SourceData->GetFilename();
+		Message->DestinationAssetName = Arguments.AssetName;
+		Message->AssetType = ULevelVariantSets::StaticClass();
+		Message->Text = LOCTEXT("CreateEmptyAssetUnsupportedReimport", "Re-import of ULevelVariantSets not supported yet.");
+
+		return nullptr;
+	}
+
+	if (!Cast<IInterchangeVariantSetPayloadInterface>(Arguments.Translator))
+	{
+		UE_LOG(LogInterchangeImport, Error, TEXT("Cannot import LevelVariantSets, the translator do not implement the IInterchangeVariantSetPayloadInterface interface."));
+		return nullptr;
+	}
+
+	ULevelVariantSets* LevelVariantSets = nullptr;
+	if (!Arguments.AssetNode || !Arguments.AssetNode->GetObjectClass()->IsChildOf(GetFactoryClass()))
+	{
+		return nullptr;
+	}
+
+	const UInterchangeSceneVariantSetsFactoryNode* FactoryNode = Cast<UInterchangeSceneVariantSetsFactoryNode>(Arguments.AssetNode);
+	if (FactoryNode == nullptr)
+	{
+		return nullptr;
+	}
+
+	// create an asset if it doesn't exist
+	UObject* ExistingAsset = StaticFindObject(nullptr, Arguments.Parent, *Arguments.AssetName);
+
+	// create a new material or overwrite existing asset, if possible
+	if (!ExistingAsset)
+	{
+		LevelVariantSets = NewObject<ULevelVariantSets>(Arguments.Parent, *Arguments.AssetName, RF_Public | RF_Standalone);
+	}
+	else if (ExistingAsset->GetClass()->IsChildOf(ULevelVariantSets::StaticClass()))
+	{
+		// This is a reimport, we are just re-updating the source data
+		LevelVariantSets = Cast<ULevelVariantSets>(ExistingAsset);
+	}
+
+	if (!LevelVariantSets)
+	{
+		UE_LOG(LogInterchangeImport, Warning, TEXT("Could not create LevelVariantSets asset %s"), *Arguments.AssetName);
+		return nullptr;
+	}
+
+	return LevelVariantSets;
+#endif //else !WITH_EDITOR || !WITH_EDITORONLY_DATA
+}
+
+UObject* UInterchangeSceneVariantSetsFactory::CreateAsset(const FCreateAssetParams& Arguments)
+{
+#if !WITH_EDITOR || !WITH_EDITORONLY_DATA
+	// TODO: Can we import ULevelVariantSets at runtime
+	UE_LOG(LogInterchangeImport, Error, TEXT("Cannot import LevelVariantSets asset in runtime, this is an editor only feature."));
+	return nullptr;
+
+#else
+	using namespace UE::Interchange;
+
+	// Re-import is not supported yet
+	// Need to add an AssetImportData property to ULevelVariantSets
+	if (Arguments.ReimportObject)
+	{
+		UInterchangeResultError_Generic* Message = AddMessage<UInterchangeResultError_Generic>();
+		Message->SourceAssetName = Arguments.SourceData->GetFilename();
+		Message->DestinationAssetName = Arguments.AssetName;
+		Message->AssetType = ULevelVariantSets::StaticClass();
+		Message->Text = LOCTEXT("CreateAssetUnsupportedReimport", "Re-import of ULevelVariantSets not supported yet.");
+
+		return nullptr;
+	}
+
+	if (!Arguments.NodeContainer || !Arguments.AssetNode || !Arguments.AssetNode->GetObjectClass()->IsChildOf(GetFactoryClass()))
+	{
+		return nullptr;
+	}
+
+	UInterchangeSceneVariantSetsFactoryNode* FactoryNode = Cast<UInterchangeSceneVariantSetsFactoryNode>(Arguments.AssetNode);
+	if (!FactoryNode)
+	{
+		return nullptr;
+	}
+
+	Translator = Arguments.Translator;
+	const IInterchangeVariantSetPayloadInterface* LevelVariantSetsPayloadInterface = Cast<IInterchangeVariantSetPayloadInterface>(Arguments.Translator);
+	if (!LevelVariantSetsPayloadInterface)
+	{
+		UE_LOG(LogInterchangeImport, Error, TEXT("Cannot import LevelVariantSets, the translator do not implement the IInterchangeVariantSetPayloadInterface interface."));
+		return nullptr;
+	}
+
+	const UClass* LevelVariantSetsClass = FactoryNode->GetObjectClass();
+	check(LevelVariantSetsClass && LevelVariantSetsClass->IsChildOf(GetFactoryClass()));
+
+	// create an asset if it doesn't exist
+	UObject* ExistingAsset = StaticFindObject(nullptr, Arguments.Parent, *Arguments.AssetName);
+
+	ULevelVariantSets* LevelVariantSets = nullptr;
+	// create a new level sequence or overwrite existing asset, if possible
+	if (!ExistingAsset)
+	{
+		//NewObject is not thread safe, the asset registry directory watcher tick on the main thread can trig before we finish initializing the UObject and will crash
+		//The UObject should have been create by calling CreateEmptyAsset on the main thread.
+		check(IsInGameThread());
+		LevelVariantSets = NewObject<ULevelVariantSets>(Arguments.Parent, LevelVariantSetsClass, *Arguments.AssetName, RF_Public | RF_Standalone);
+	}
+	else if (ExistingAsset->GetClass()->IsChildOf(LevelVariantSetsClass))
+	{
+		//This is a reimport, we are just re-updating the source data
+		LevelVariantSets = Cast<ULevelVariantSets>(ExistingAsset);
+	}
+
+	if (!ensure(LevelVariantSets))
+	{
+		UInterchangeResultError_Generic* Message = AddMessage<UInterchangeResultError_Generic>();
+		Message->SourceAssetName = Arguments.SourceData->GetFilename();
+		Message->DestinationAssetName = Arguments.AssetName;
+		Message->AssetType = ULevelVariantSets::StaticClass();
+		Message->Text = FText::Format(LOCTEXT("CreateAssetFailed", "Could not create nor find LevelVariantSets asset {0}."), FText::FromString(Arguments.AssetName));
+		return nullptr;
+	}
+
+	Private::FLevelVariantSetsHelper Helper(*LevelVariantSets, *FactoryNode, *Arguments.NodeContainer, *LevelVariantSetsPayloadInterface);
+	Helper.PopulateLevelVariantSets();
+
+	/** Apply all FactoryNode custom attributes to the level sequence asset */
+	FactoryNode->ApplyAllCustomAttributeToObject(LevelVariantSets);
+
+	//Getting the file Hash will cache it into the source data
+	Arguments.SourceData->GetFileContentHash();
+
+	//The interchange completion task (call in the GameThread after the factories pass), will call PostEditChange which will trig another asynchronous system that will build all material in parallel
+	return LevelVariantSets;
+#endif //else !WITH_EDITOR || !WITH_EDITORONLY_DATA
+}
+
+void UInterchangeSceneVariantSetsFactory::PreImportPreCompletedCallback(const FImportPreCompletedCallbackParams& Arguments)
+{
+	check(IsInGameThread());
+	Super::PreImportPreCompletedCallback(Arguments);
+
+	// TODO: Need to add an AssetImportData property to ULevelVariantSets
+}
+
+#undef LOCTEXT_NAMESPACE
\ No newline at end of file
diff --git a/Engine/Plugins/Experimental/Interchange/Runtime/Source/Import/Private/Scene/InterchangeStaticMeshActorFactory.cpp b/Engine/Plugins/Experimental/Interchange/Runtime/Source/Import/Private/Scene/InterchangeStaticMeshActorFactory.cpp
--- a/Engine/Plugins/Experimental/Interchange/Runtime/Source/Import/Private/Scene/InterchangeStaticMeshActorFactory.cpp
+++ b/Engine/Plugins/Experimental/Interchange/Runtime/Source/Import/Private/Scene/InterchangeStaticMeshActorFactory.cpp
@@ -45,30 +45,18 @@ void UInterchangeStaticMeshActorFactory::SetupStaticMeshActor(const UInterchange
 	if (UStaticMeshComponent* StaticMeshComponent = StaticMeshActor->GetStaticMeshComponent())
 	{
 		StaticMeshComponent->UnregisterComponent();
-	}
-}
-
-void UInterchangeStaticMeshActorFactory::PostImportPreCompletedCallback(const FImportPreCompletedCallbackParams& Arguments)
-{
-	// Set the static mesh on the component in the post import callback, once the static mesh has been fully imported.
-	// The component doesn't like being set a static mesh with uninitialized render data.
 
-	if (AStaticMeshActor* StaticMeshActor = Cast<AStaticMeshActor>(Arguments.ImportedObject))
-	{
-		if (UStaticMeshComponent* StaticMeshComponent = StaticMeshActor->GetStaticMeshComponent())
+		if (const UInterchangeFactoryBaseNode* MeshNode = UE::Interchange::ActorHelper::FindAssetInstanceFactoryNode(NodeContainer, ActorFactoryNode))
 		{
-			if (const UInterchangeFactoryBaseNode* MeshNode = UE::Interchange::ActorHelper::FindAssetInstanceFactoryNode(Arguments.NodeContainer, Arguments.FactoryNode))
+			if (UStaticMesh* StaticMesh = Cast<UStaticMesh>(MeshNode->ReferenceObject.TryLoad()))
 			{
-				if (UStaticMesh* StaticMesh = Cast<UStaticMesh>(MeshNode->ReferenceObject.TryLoad()))
-				{
-					StaticMeshComponent->SetStaticMesh(StaticMesh);
+				StaticMeshComponent->SetStaticMesh(StaticMesh);
 
-					if (const UInterchangeMeshActorFactoryNode* MeshActorFactoryNode = Cast<UInterchangeMeshActorFactoryNode>(Arguments.FactoryNode))
-					{
-						UE::Interchange::ActorHelper::ApplySlotMaterialDependencies(*Arguments.NodeContainer, *MeshActorFactoryNode, *StaticMeshComponent);
-					}
+				if (const UInterchangeMeshActorFactoryNode* MeshActorFactoryNode = Cast<UInterchangeMeshActorFactoryNode>(ActorFactoryNode))
+				{
+					UE::Interchange::ActorHelper::ApplySlotMaterialDependencies(*NodeContainer, *MeshActorFactoryNode, *StaticMeshComponent);
 				}
 			}
 		}
 	}
-}
\ No newline at end of file
+}
diff --git a/Engine/Plugins/Experimental/Interchange/Runtime/Source/Import/Public/Scene/InterchangeSceneVariantSetsFactory.h b/Engine/Plugins/Experimental/Interchange/Runtime/Source/Import/Public/Scene/InterchangeSceneVariantSetsFactory.h
new file mode 100644
--- /dev/null
+++ b/Engine/Plugins/Experimental/Interchange/Runtime/Source/Import/Public/Scene/InterchangeSceneVariantSetsFactory.h
@@ -0,0 +1,42 @@
+// Copyright Epic Games, Inc. All Rights Reserved.
+
+#pragma once
+
+#include "CoreMinimal.h"
+#include "InterchangeFactoryBase.h"
+#include "UObject/Object.h"
+#include "UObject/ObjectMacros.h"
+
+#include "Scene/InterchangeVariantSetPayloadInterface.h"
+
+#include "InterchangeSceneVariantSetsFactory.generated.h"
+
+class UVariantObjectBinding;
+
+UCLASS(BlueprintType, Experimental)
+class INTERCHANGEIMPORT_API UInterchangeSceneVariantSetsFactory : public UInterchangeFactoryBase
+{
+	GENERATED_BODY()
+public:
+
+	//////////////////////////////////////////////////////////////////////////
+	// Interchange factory base interface begin
+
+	virtual UClass* GetFactoryClass() const override;
+	virtual UObject* CreateEmptyAsset(const FCreateAssetParams& Arguments) override;
+	virtual UObject* CreateAsset(const FCreateAssetParams& Arguments) override;
+	virtual void PreImportPreCompletedCallback(const FImportPreCompletedCallbackParams& Arguments) override;
+	virtual bool CanExecuteOnAnyThread() const override
+	{
+		//Currently we cannot use the anim sequence controller outside of the game thread
+		return false;
+	}
+
+	// Interchange factory base interface end
+	//////////////////////////////////////////////////////////////////////////
+
+private:
+	const UInterchangeTranslatorBase* Translator = nullptr;
+};
+
+
diff --git a/Engine/Plugins/Experimental/Interchange/Runtime/Source/Import/Public/Scene/InterchangeStaticMeshActorFactory.h b/Engine/Plugins/Experimental/Interchange/Runtime/Source/Import/Public/Scene/InterchangeStaticMeshActorFactory.h
--- a/Engine/Plugins/Experimental/Interchange/Runtime/Source/Import/Public/Scene/InterchangeStaticMeshActorFactory.h
+++ b/Engine/Plugins/Experimental/Interchange/Runtime/Source/Import/Public/Scene/InterchangeStaticMeshActorFactory.h
@@ -31,8 +31,6 @@ public:
 		return false; 
 	}
 
-	virtual void PostImportPreCompletedCallback(const FImportPreCompletedCallbackParams& Arguments) override;
-
 	// Interchange factory base interface end
 	//////////////////////////////////////////////////////////////////////////
 
diff --git a/Engine/Plugins/Experimental/Interchange/Runtime/Source/Import/Public/Scene/InterchangeVariantSetPayloadInterface.h b/Engine/Plugins/Experimental/Interchange/Runtime/Source/Import/Public/Scene/InterchangeVariantSetPayloadInterface.h
new file mode 100644
--- /dev/null
+++ b/Engine/Plugins/Experimental/Interchange/Runtime/Source/Import/Public/Scene/InterchangeVariantSetPayloadInterface.h
@@ -0,0 +1,80 @@
+// Copyright Epic Games, Inc. All Rights Reserved.
+
+#pragma once
+
+#include "CoreMinimal.h"
+#include "Async/Future.h"
+#include "InterchangeSourceData.h"
+#include "UObject/Object.h"
+#include "UObject/ObjectMacros.h"
+#include "UObject/Interface.h"
+
+#include "InterchangeVariantSetPayloadInterface.generated.h"
+
+namespace UE::Interchange
+{
+	enum class EVariantPropertyCaptureCategory : uint16
+	{
+		Undefined         = 0x0000,
+		Generic           = 0x0001,
+		RelativeLocation  = 0x0002,
+		RelativeRotation  = 0x0004,
+		RelativeScale3D   = 0x0008,
+		Visibility        = 0x0010,
+		Material          = 0x0020,
+		Color             = 0x0040,
+		Option            = 0x0080
+	};
+	ENUM_CLASS_FLAGS(EVariantPropertyCaptureCategory)
+
+	struct FVariantPropertyCaptureData
+	{
+		FString PropertyPath;
+		EVariantPropertyCaptureCategory Category = EVariantPropertyCaptureCategory::Undefined;
+		TOptional<TArray<uint8>> Buffer;
+		TOptional<FString> ObjectUid;
+	};
+
+	struct FVariantBinding
+	{
+		FString TargetUid;
+		TArray<FVariantPropertyCaptureData> Captures;
+	};
+
+	struct FVariant
+	{
+		FString DisplayText;
+		TArray<FVariantBinding> Bindings;
+	};
+
+	struct FVariantSetPayloadData
+	{
+		TArray<FVariant> Variants;
+	};
+}
+
+UINTERFACE()
+class INTERCHANGEIMPORT_API UInterchangeVariantSetPayloadInterface : public UInterface
+{
+	GENERATED_BODY()
+};
+
+/**
+ * VariantSet payload interface. Derive from this interface if your payload can import Variant's data
+ */
+class INTERCHANGEIMPORT_API IInterchangeVariantSetPayloadInterface
+{
+	GENERATED_BODY()
+public:
+
+	/**
+	 * Once the translation is done, the import process need a way to retrieve payload data.
+	 * This payload will be use by the factories to create the asset.
+	 *
+	 * @param PayloadKey - The key to retrieve the a particular payload contain into the specified source data.
+	 * @return a PayloadData containing the the data ask with the key.
+	 */
+	virtual TFuture<TOptional<UE::Interchange::FVariantSetPayloadData>> GetVariantSetPayloadData(const FString& PayloadKey) const = 0;
+};
+
+
diff --git a/Engine/Plugins/Experimental/Interchange/Runtime/Source/Nodes/InterchangeNodes.Build.cs b/Engine/Plugins/Experimental/Interchange/Runtime/Source/Nodes/InterchangeNodes.Build.cs
--- a/Engine/Plugins/Experimental/Interchange/Runtime/Source/Nodes/InterchangeNodes.Build.cs
+++ b/Engine/Plugins/Experimental/Interchange/Runtime/Source/Nodes/InterchangeNodes.Build.cs
@@ -6,6 +6,9 @@ namespace UnrealBuildTool.Rules
 	{
 		public InterchangeNodes(ReadOnlyTargetRules Target) : base(Target)
 		{
+			// This module does not depend on the Engine module. This is intentional.
+			// The InterchangeWorker program does not compile against the Engine module
+			// but depends on this module.
 			PublicDependencyModuleNames.AddRange(
 				new string[]
 				{
@@ -14,17 +17,6 @@ namespace UnrealBuildTool.Rules
 					"InterchangeCore",
 				}
 			);
-
-			if(Target.bCompileAgainstEngine)
-			{
-				PublicDependencyModuleNames.AddRange(
-					new string[]
-					{
-						"CinematicCamera",
-						"Engine"
-					}
-				);
-			}
 		}
 	}
 }
diff --git a/Engine/Plugins/Experimental/Interchange/Runtime/Source/Nodes/Private/InterchangeVariantSetNode.cpp b/Engine/Plugins/Experimental/Interchange/Runtime/Source/Nodes/Private/InterchangeVariantSetNode.cpp
new file mode 100644
--- /dev/null
+++ b/Engine/Plugins/Experimental/Interchange/Runtime/Source/Nodes/Private/InterchangeVariantSetNode.cpp
@@ -0,0 +1,87 @@
+// Copyright Epic Games, Inc. All Rights Reserved.
+
+#include "InterchangeVariantSetNode.h"
+
+// UInterchangeSceneVariantSetsNode
+
+UInterchangeSceneVariantSetsNode::UInterchangeSceneVariantSetsNode()
+{
+	CustomVariantSetUids.Initialize(Attributes, TEXT("VariantSetUids"));
+}
+
+int32 UInterchangeSceneVariantSetsNode::GetCustomVariantSetUidCount() const
+{
+	return CustomVariantSetUids.GetCount();
+}
+
+void UInterchangeSceneVariantSetsNode::GetCustomVariantSetUids(TArray<FString>& OutVariantSetUids) const
+{
+	CustomVariantSetUids.GetItems(OutVariantSetUids);
+}
+
+void UInterchangeSceneVariantSetsNode::GetCustomVariantSetUid(const int32 Index, FString& OutVariantSetUid) const
+{
+	CustomVariantSetUids.GetItem(Index, OutVariantSetUid);
+}
+
+bool UInterchangeSceneVariantSetsNode::AddCustomVariantSetUid(const FString& VariantSetUid)
+{
+	return CustomVariantSetUids.AddItem(VariantSetUid);
+}
+
+bool UInterchangeSceneVariantSetsNode::RemoveCustomVariantSetUid(const FString& VariantSetUid)
+{
+	return CustomVariantSetUids.RemoveItem(VariantSetUid);
+}
+
+// UInterchangeVariantSetNode
+
+UInterchangeVariantSetNode::UInterchangeVariantSetNode()
+{
+	CustomDependencyUids.Initialize(Attributes, TEXT("DependencyUids"));
+}
+
+bool UInterchangeVariantSetNode::GetCustomDisplayText(FString& AttributeValue) const
+{
+	IMPLEMENT_NODE_ATTRIBUTE_GETTER(DisplayText, FString);
+}
+
+bool UInterchangeVariantSetNode::SetCustomDisplayText(const FString& AttributeValue)
+{
+	IMPLEMENT_NODE_ATTRIBUTE_SETTER_NODELEGATE(DisplayText, FString);
+}
+
+bool UInterchangeVariantSetNode::GetCustomVariantsPayloadKey(FString& AttributeValue) const
+{
+	IMPLEMENT_NODE_ATTRIBUTE_GETTER(VariantsPayload, FString);
+}
+
+bool UInterchangeVariantSetNode::SetCustomVariantsPayloadKey(const FString& AttributeValue)
+{
+	IMPLEMENT_NODE_ATTRIBUTE_SETTER_NODELEGATE(VariantsPayload, FString);
+}
+
+int32 UInterchangeVariantSetNode::GetCustomDependencyUidCount() const
+{
+	return CustomDependencyUids.GetCount();
+}
+
+void UInterchangeVariantSetNode::GetCustomDependencyUids(TArray<FString>& OutDependencyUids) const
+{
+	CustomDependencyUids.GetItems(OutDependencyUids);
+}
+
+void UInterchangeVariantSetNode::GetCustomDependencyUid(const int32 Index, FString& OutDependencyUid) const
+{
+	CustomDependencyUids.GetItem(Index, OutDependencyUid);
+}
+
+bool UInterchangeVariantSetNode::AddCustomDependencyUid(const FString& DependencyUid)
+{
+	return CustomDependencyUids.AddItem(DependencyUid);
+}
+
+bool UInterchangeVariantSetNode::RemoveCustomDependencyUid(const FString& DependencyUid)
+{
+	return CustomDependencyUids.RemoveItem(DependencyUid);
+}
diff --git a/Engine/Plugins/Experimental/Interchange/Runtime/Source/Nodes/Public/InterchangeVariantSetNode.h b/Engine/Plugins/Experimental/Interchange/Runtime/Source/Nodes/Public/InterchangeVariantSetNode.h
new file mode 100644
--- /dev/null
+++ b/Engine/Plugins/Experimental/Interchange/Runtime/Source/Nodes/Public/InterchangeVariantSetNode.h
@@ -0,0 +1,153 @@
+// Copyright Epic Games, Inc. All Rights Reserved.
+
+#pragma once
+
+#include "CoreMinimal.h"
+#include "Nodes/InterchangeBaseNode.h"
+
+#include "InterchangeVariantSetNode.generated.h"
+
+/**
+ * Class to represent a set of variants
+ */
+UCLASS(BlueprintType, Experimental)
+class INTERCHANGENODES_API UInterchangeVariantSetNode : public UInterchangeBaseNode
+{
+	GENERATED_BODY()
+
+	UInterchangeVariantSetNode();
+
+public:
+	static FStringView StaticAssetTypeName()
+	{
+		return TEXT("VariantSet");
+	}
+
+	/**
+	 * Return the node type name of the class, we use this when reporting errors
+	 */
+	virtual FString GetTypeName() const override
+	{
+		const FString TypeName = TEXT("VariantSetNode");
+		return TypeName;
+	}
+
+	/**
+	 * This function allow to retrieve the text which is displayed in the UI for this VariantSet.
+	 */
+	UFUNCTION(BlueprintCallable, Category = "Interchange | Node | VariantSet")
+	bool GetCustomDisplayText(FString& AttributeValue) const;
+
+	/**
+	 * This function allow to set the text to be displayed in the UI for this VariantSet.
+	 */
+	UFUNCTION(BlueprintCallable, Category = "Interchange | Node | VariantSet")
+	bool SetCustomDisplayText(const FString& AttributeValue);
+
+	/**
+	 * Get the payload key needed to retrieve the variants for this VariantSet.
+	 */
+	UFUNCTION(BlueprintCallable, Category = "Interchange | Node | VariantSet")
+	bool GetCustomVariantsPayloadKey(FString& PayloadKey) const;
+
+	/**
+	 * Set the payload key needed to retrieve the variants for this VariantSet.
+	 */
+	UFUNCTION(BlueprintCallable, Category = "Interchange | Node | VariantSet")
+	bool SetCustomVariantsPayloadKey(const FString& PayloadKey);
+
+	/**
+	 * This function allow to retrieve the number of translated node's unique ids for this VariantSet.
+	 */
+	UFUNCTION(BlueprintCallable, Category = "Interchange | Node | SceneVariantSet")
+	int32 GetCustomDependencyUidCount() const;
+
+	/**
+	 * This function allow to retrieve all the translated node's unique ids for this VariantSet.
+	 */
+	UFUNCTION(BlueprintCallable, Category = "Interchange | Node | SceneVariantSet")
+	void GetCustomDependencyUids(TArray<FString>& OutDependencyUids) const;
+
+	/**
+	 * This function allow to retrieve a specific translated node's unique id for this VariantSet.
+	 */
+	UFUNCTION(BlueprintCallable, Category = "Interchange | Node | SceneVariantSet")
+	void GetCustomDependencyUid(const int32 Index, FString& OutDependencyUid) const;
+
+	/**
+	 * Add one translated node's unique id to this VariantSet.
+	 */
+	UFUNCTION(BlueprintCallable, Category = "Interchange | Node | SceneVariantSet")
+	bool AddCustomDependencyUid(const FString& DependencyUid);
+
+	/**
+	 * Remove one translated node's unique id from this VariantSet.
+	 */
+	UFUNCTION(BlueprintCallable, Category = "Interchange | Node | SceneVariantSet")
+	bool RemoveCustomDependencyUid(const FString& DependencyUid);
+
+private:
+	const UE::Interchange::FAttributeKey Macro_CustomDisplayTextKey = UE::Interchange::FAttributeKey(TEXT("DisplayText"));
+	const UE::Interchange::FAttributeKey Macro_CustomVariantsPayloadKey = UE::Interchange::FAttributeKey(TEXT("VariantsPayload"));
+	UE::Interchange::TArrayAttributeHelper<FString> CustomDependencyUids;
+};
+
+/**
+ * Class to represent a set of VariantSet nodes
+ */
+UCLASS(BlueprintType, Experimental)
+class INTERCHANGENODES_API UInterchangeSceneVariantSetsNode : public UInterchangeBaseNode
+{
+	GENERATED_BODY()
+
+	UInterchangeSceneVariantSetsNode();
+
+public:
+	static FStringView StaticAssetTypeName()
+	{
+		return TEXT("SceneVariantSet");
+	}
+
+	/**
+	* Return the node type name of the class, we use this when reporting errors
+	*/
+	virtual FString GetTypeName() const override
+	{
+		const FString TypeName = TEXT("SceneVariantSetNode");
+		return TypeName;
+	}
+
+	/**
+	 * This function allow to retrieve the number of VariantSets for this object.
+	 */
+	UFUNCTION(BlueprintCallable, Category = "Interchange | Node | SceneVariantSet")
+	int32 GetCustomVariantSetUidCount() const;
+
+	/**
+	 * This function allow to retrieve all the VariantSets' unique ids for this object.
+	 */
+	UFUNCTION(BlueprintCallable, Category = "Interchange | Node | SceneVariantSet")
+	void GetCustomVariantSetUids(TArray<FString>& OutVariantUids) const;
+
+	/**
+	 * This function allow to retrieve one VariantSet's unique id for this object.
+	 */
+	UFUNCTION(BlueprintCallable, Category = "Interchange | Node | SceneVariantSet")
+	void GetCustomVariantSetUid(const int32 Index, FString& OutVariantUid) const;
+
+	/**
+	 * Add one VariantSet's unique id to this object.
+	 */
+	UFUNCTION(BlueprintCallable, Category = "Interchange | Node | SceneVariantSet")
+	bool AddCustomVariantSetUid(const FString& VariantUid);
+
+	/**
+	 * Remove one VariantSet's unique id from this object.
+	 */
+	UFUNCTION(BlueprintCallable, Category = "Interchange | Node | SceneVariantSet")
+	bool RemoveCustomVariantSetUid(const FString& VariantUid);
+
+private:
+	UE::Interchange::TArrayAttributeHelper<FString> CustomVariantSetUids;
+};
+
diff --git a/Engine/Plugins/Experimental/Interchange/Runtime/Source/Pipelines/Private/InterchangeGenericScenesPipeline.cpp b/Engine/Plugins/Experimental/Interchange/Runtime/Source/Pipelines/Private/InterchangeGenericScenesPipeline.cpp
--- a/Engine/Plugins/Experimental/Interchange/Runtime/Source/Pipelines/Private/InterchangeGenericScenesPipeline.cpp
+++ b/Engine/Plugins/Experimental/Interchange/Runtime/Source/Pipelines/Private/InterchangeGenericScenesPipeline.cpp
@@ -13,6 +13,8 @@
 #include "InterchangePipelineLog.h"
 #include "InterchangePipelineMeshesUtilities.h"
 #include "InterchangeSceneNode.h"
+#include "InterchangeSceneVariantSetsFactoryNode.h"
+#include "InterchangeVariantSetNode.h"
 #include "Nodes/InterchangeUserDefinedAttribute.h"
 
 #include "Animation/SkeletalMeshActor.h"
@@ -31,8 +33,10 @@ void UInterchangeGenericLevelPipeline::ExecutePreImportPipeline(UInterchangeBase
 		return;
 	}
 
+	BaseNodeContainer = InBaseNodeContainer;
+
 	FTransform GlobalOffsetTransform = FTransform::Identity;
-	if (UInterchangeCommonPipelineDataFactoryNode* CommonPipelineDataFactoryNode = UInterchangeCommonPipelineDataFactoryNode::GetUniqueInstance(InBaseNodeContainer))
+	if (UInterchangeCommonPipelineDataFactoryNode* CommonPipelineDataFactoryNode = UInterchangeCommonPipelineDataFactoryNode::GetUniqueInstance(BaseNodeContainer))
 	{
 		CommonPipelineDataFactoryNode->GetCustomGlobalOffsetTransform(GlobalOffsetTransform);
 	}
@@ -40,7 +44,7 @@ void UInterchangeGenericLevelPipeline::ExecutePreImportPipeline(UInterchangeBase
 	TArray<UInterchangeSceneNode*> SceneNodes;
 
 	//Find all translated node we need for this pipeline
-	InBaseNodeContainer->IterateNodes([&SceneNodes](const FString& NodeUid, UInterchangeBaseNode* Node)
+	BaseNodeContainer->IterateNodes([&SceneNodes](const FString& NodeUid, UInterchangeBaseNode* Node)
 	{
 		switch(Node->GetNodeContainerType())
 		{
@@ -69,14 +73,30 @@ void UInterchangeGenericLevelPipeline::ExecutePreImportPipeline(UInterchangeBase
 					continue;
 				}
 			}
-			ExecuteSceneNodePreImport(InBaseNodeContainer, GlobalOffsetTransform, SceneNode, InBaseNodeContainer);
+			ExecuteSceneNodePreImport(GlobalOffsetTransform, SceneNode);
+		}
+	}
+
+	//Find all translated scene variant sets
+	TArray<UInterchangeSceneVariantSetsNode*> SceneVariantSetNodes;
+
+	InBaseNodeContainer->IterateNodesOfType<UInterchangeSceneVariantSetsNode>([&SceneVariantSetNodes](const FString& NodeUid, UInterchangeSceneVariantSetsNode* Node)
+		{
+			SceneVariantSetNodes.Add(Node);
+		});
+
+	for (const UInterchangeSceneVariantSetsNode* SceneVariantSetNode : SceneVariantSetNodes)
+	{
+		if (SceneVariantSetNode)
+		{
+			ExecuteSceneVariantSetNodePreImport(*SceneVariantSetNode);
 		}
 	}
 }
 
-void UInterchangeGenericLevelPipeline::ExecuteSceneNodePreImport(UInterchangeBaseNodeContainer* InBaseNodeContainer, const FTransform& GlobalOffsetTransform, const UInterchangeSceneNode* SceneNode, UInterchangeBaseNodeContainer* FactoryNodeContainer)
+void UInterchangeGenericLevelPipeline::ExecuteSceneNodePreImport(const FTransform& GlobalOffsetTransform, const UInterchangeSceneNode* SceneNode)
 {
-	if (!SceneNode)
+	if (!BaseNodeContainer || !SceneNode)
 	{
 		return;
 	}
@@ -86,10 +106,10 @@ void UInterchangeGenericLevelPipeline::ExecuteSceneNodePreImport(UInterchangeBas
 	FString AssetInstanceUid;
 	if (SceneNode->GetCustomAssetInstanceUid(AssetInstanceUid))
 	{
-		TranslatedAssetNode = FactoryNodeContainer->GetNode(AssetInstanceUid);
+		TranslatedAssetNode = BaseNodeContainer->GetNode(AssetInstanceUid);
 	}
 
-	UInterchangeActorFactoryNode* ActorFactoryNode = CreateActorFactoryNode(SceneNode, TranslatedAssetNode, FactoryNodeContainer);
+	UInterchangeActorFactoryNode* ActorFactoryNode = CreateActorFactoryNode(SceneNode, TranslatedAssetNode);
 
 	if (!ensure(ActorFactoryNode))
 	{
@@ -97,11 +117,11 @@ void UInterchangeGenericLevelPipeline::ExecuteSceneNodePreImport(UInterchangeBas
 	}
 
 	ActorFactoryNode->InitializeNode(UInterchangeFactoryBaseNode::BuildFactoryNodeUid(SceneNode->GetUniqueID()), SceneNode->GetDisplayLabel(), EInterchangeNodeContainerType::FactoryData);
-	const FString ActorFactoryNodeUid = FactoryNodeContainer->AddNode(ActorFactoryNode);
+	const FString ActorFactoryNodeUid = BaseNodeContainer->AddNode(ActorFactoryNode);
 	if (!SceneNode->GetParentUid().IsEmpty())
 	{
 		const FString ParentFactoryNodeUid = TEXT("Factory_") + SceneNode->GetParentUid();
-		FactoryNodeContainer->SetNodeParentUid(ActorFactoryNodeUid, ParentFactoryNodeUid);
+		BaseNodeContainer->SetNodeParentUid(ActorFactoryNodeUid, ParentFactoryNodeUid);
 		ActorFactoryNode->AddFactoryDependencyUid(ParentFactoryNodeUid);
 	}
 
@@ -110,7 +130,7 @@ void UInterchangeGenericLevelPipeline::ExecuteSceneNodePreImport(UInterchangeBas
 
 	//TODO move this code to the factory, a stack over pipeline can change the global offset transform which will affect this value.
 	FTransform GlobalTransform;
-	if (SceneNode->GetCustomGlobalTransform(InBaseNodeContainer, GlobalOffsetTransform, GlobalTransform))
+	if (SceneNode->GetCustomGlobalTransform(BaseNodeContainer, GlobalOffsetTransform, GlobalTransform))
 	{
 		ActorFactoryNode->SetCustomGlobalTransform(GlobalTransform);
 	}
@@ -119,46 +139,51 @@ void UInterchangeGenericLevelPipeline::ExecuteSceneNodePreImport(UInterchangeBas
 
 	if (TranslatedAssetNode)
 	{
-		SetUpFactoryNode(ActorFactoryNode, SceneNode, TranslatedAssetNode, FactoryNodeContainer);
+		SetUpFactoryNode(ActorFactoryNode, SceneNode, TranslatedAssetNode);
 	}
 }
 
-UInterchangeActorFactoryNode* UInterchangeGenericLevelPipeline::CreateActorFactoryNode(const UInterchangeSceneNode* SceneNode, const UInterchangeBaseNode* TranslatedAssetNode, UInterchangeBaseNodeContainer* FactoryNodeContainer) const
+UInterchangeActorFactoryNode* UInterchangeGenericLevelPipeline::CreateActorFactoryNode(const UInterchangeSceneNode* SceneNode, const UInterchangeBaseNode* TranslatedAssetNode) const
 {
+	if (!ensure(BaseNodeContainer))
+	{
+		return nullptr;
+	}
+
 	if(TranslatedAssetNode)
 	{
 		if(TranslatedAssetNode->IsA<UInterchangeCameraNode>())
 		{
-			return NewObject<UInterchangeCineCameraFactoryNode>(FactoryNodeContainer, NAME_None);
+			return NewObject<UInterchangeCineCameraFactoryNode>(BaseNodeContainer, NAME_None);
 		}
 		else if(TranslatedAssetNode->IsA<UInterchangeMeshNode>())
 		{
-			return NewObject<UInterchangeMeshActorFactoryNode>(FactoryNodeContainer, NAME_None);
+			return NewObject<UInterchangeMeshActorFactoryNode>(BaseNodeContainer, NAME_None);
 		}
 		else if(TranslatedAssetNode->IsA<UInterchangeSpotLightNode>())
 		{
-			return NewObject<UInterchangeSpotLightFactoryNode>(FactoryNodeContainer, NAME_None);
+			return NewObject<UInterchangeSpotLightFactoryNode>(BaseNodeContainer, NAME_None);
 		}
 		else if(TranslatedAssetNode->IsA<UInterchangePointLightNode>())
 		{
-			return NewObject<UInterchangePointLightFactoryNode>(FactoryNodeContainer, NAME_None);
+			return NewObject<UInterchangePointLightFactoryNode>(BaseNodeContainer, NAME_None);
 		}
 		else if(TranslatedAssetNode->IsA<UInterchangeRectLightNode>())
 		{
-			return NewObject<UInterchangeRectLightFactoryNode>(FactoryNodeContainer, NAME_None);
+			return NewObject<UInterchangeRectLightFactoryNode>(BaseNodeContainer, NAME_None);
 		}
 		else if(TranslatedAssetNode->IsA<UInterchangeDirectionalLightNode>())
 		{
-			return NewObject<UInterchangeDirectionalLightFactoryNode>(FactoryNodeContainer, NAME_None);
+			return NewObject<UInterchangeDirectionalLightFactoryNode>(BaseNodeContainer, NAME_None);
 		}
 	}
 
-	return NewObject<UInterchangeActorFactoryNode>(FactoryNodeContainer, NAME_None);
+	return NewObject<UInterchangeActorFactoryNode>(BaseNodeContainer, NAME_None);
 }
 
-void UInterchangeGenericLevelPipeline::SetUpFactoryNode(UInterchangeActorFactoryNode* ActorFactoryNode, const UInterchangeSceneNode* SceneNode, const UInterchangeBaseNode* TranslatedAssetNode, UInterchangeBaseNodeContainer* FactoryNodeContainer) const
+void UInterchangeGenericLevelPipeline::SetUpFactoryNode(UInterchangeActorFactoryNode* ActorFactoryNode, const UInterchangeSceneNode* SceneNode, const UInterchangeBaseNode* TranslatedAssetNode) const
 {
-	if (!ensure(ActorFactoryNode && SceneNode && TranslatedAssetNode && FactoryNodeContainer))
+	if (!ensure(BaseNodeContainer && ActorFactoryNode && SceneNode && TranslatedAssetNode))
 	{
 		return;
 	}
@@ -180,7 +205,7 @@ void UInterchangeGenericLevelPipeline::SetUpFactoryNode(UInterchangeActorFactory
 			TMap<FString, FString> SlotMaterialDependencies;
 			SceneNode->GetSlotMaterialDependencies(SlotMaterialDependencies);
 
-			UE::Interchange::MeshesUtilities::ApplySlotMaterialDependencies(*MeshActorFactoryNode, SlotMaterialDependencies, *FactoryNodeContainer);
+			UE::Interchange::MeshesUtilities::ApplySlotMaterialDependencies(*MeshActorFactoryNode, SlotMaterialDependencies, *BaseNodeContainer);
 
 			MeshActorFactoryNode->AddFactoryDependencyUid(UInterchangeFactoryBaseNode::BuildFactoryNodeUid(MeshNode->GetUniqueID()));
 		}
@@ -333,3 +358,47 @@ void UInterchangeGenericLevelPipeline::SetUpFactoryNode(UInterchangeActorFactory
 		}
 	}
 }
+
+void UInterchangeGenericLevelPipeline::ExecuteSceneVariantSetNodePreImport(const UInterchangeSceneVariantSetsNode& SceneVariantSetNode)
+{
+	if (!ensure(BaseNodeContainer))
+	{
+		return;
+	}
+
+	const FString FactoryNodeUid = UInterchangeFactoryBaseNode::BuildFactoryNodeUid(SceneVariantSetNode.GetUniqueID());
+
+	UInterchangeSceneVariantSetsFactoryNode* FactoryNode = NewObject<UInterchangeSceneVariantSetsFactoryNode>(BaseNodeContainer, NAME_None);
+
+	FactoryNode->InitializeNode(FactoryNodeUid, SceneVariantSetNode.GetDisplayLabel(), EInterchangeNodeContainerType::FactoryData);
+	FactoryNode->SetEnabled(true);
+
+	TArray<FString> VariantSetUids;
+	SceneVariantSetNode.GetCustomVariantSetUids(VariantSetUids);
+
+	for (const FString& VariantSetUid : VariantSetUids)
+	{
+		FactoryNode->AddCustomVariantSetUid(VariantSetUid);
+
+		// Update factory's dependencies
+		if (const UInterchangeVariantSetNode* TrackNode = Cast<UInterchangeVariantSetNode>(BaseNodeContainer->GetNode(VariantSetUid)))
+		{
+			TArray<FString> ActorNodeUids;
+			TrackNode->GetCustomDependencyUids(ActorNodeUids);
+
+			for (const FString& ActorNodeUid : ActorNodeUids)
+			{
+				
+				const FString ActorFactoryNodeUid = UInterchangeFactoryBaseNode::BuildFactoryNodeUid(ActorNodeUid);
+				FactoryNode->AddFactoryDependencyUid(ActorFactoryNodeUid);
+			}
+		}
+	}
+
+	UInterchangeUserDefinedAttributesAPI::DuplicateAllUserDefinedAttribute(&SceneVariantSetNode, FactoryNode, false);
+
+	FactoryNode->AddTargetNodeUid(SceneVariantSetNode.GetUniqueID());
+	SceneVariantSetNode.AddTargetNodeUid(FactoryNode->GetUniqueID());
+
+	BaseNodeContainer->AddNode(FactoryNode);
+}
diff --git a/Engine/Plugins/Experimental/Interchange/Runtime/Source/Pipelines/Public/InterchangeGenericScenesPipeline.h b/Engine/Plugins/Experimental/Interchange/Runtime/Source/Pipelines/Public/InterchangeGenericScenesPipeline.h
--- a/Engine/Plugins/Experimental/Interchange/Runtime/Source/Pipelines/Public/InterchangeGenericScenesPipeline.h
+++ b/Engine/Plugins/Experimental/Interchange/Runtime/Source/Pipelines/Public/InterchangeGenericScenesPipeline.h
@@ -13,6 +13,8 @@
 
 class UInterchangeActorFactoryNode;
 class UInterchangeSceneNode;
+class UInterchangeSceneVariantSetsNode;
+
 
 
 UCLASS(BlueprintType, editinlinenew, Experimental)
@@ -32,18 +34,25 @@ protected:
 	/**
 	 * PreImport step called for each translated SceneNode.
 	 */
-	virtual void ExecuteSceneNodePreImport(UInterchangeBaseNodeContainer* InBaseNodeContainer, const FTransform& GlobalOffsetTransform, const UInterchangeSceneNode* SceneNode, UInterchangeBaseNodeContainer* FactoryNodeContainer);
+	virtual void ExecuteSceneNodePreImport(const FTransform& GlobalOffsetTransform, const UInterchangeSceneNode* SceneNode);
+
+	/**
+	 * PreImport step called for each translated SceneVariantSetNode.
+	 */
+	virtual void ExecuteSceneVariantSetNodePreImport(const UInterchangeSceneVariantSetsNode& SceneVariantSetNode);
 
 	/**
 	 * Return a new Actor Factory Node to be used for the given SceneNode.
 	 */
-	virtual UInterchangeActorFactoryNode* CreateActorFactoryNode(const UInterchangeSceneNode* SceneNode, const UInterchangeBaseNode* TranslatedAssetNode, UInterchangeBaseNodeContainer* FactoryNodeContainer) const;
+	virtual UInterchangeActorFactoryNode* CreateActorFactoryNode(const UInterchangeSceneNode* SceneNode, const UInterchangeBaseNode* TranslatedAssetNode) const;
 
 	/**
 	 * Use to set up the given factory node's attributes after its initialization.
 	 */
-	virtual void SetUpFactoryNode(UInterchangeActorFactoryNode* ActorFactoryNode, const UInterchangeSceneNode* SceneNode, const UInterchangeBaseNode* TranslatedAssetNode, UInterchangeBaseNodeContainer* FactoryNodeContainer) const;
+	virtual void SetUpFactoryNode(UInterchangeActorFactoryNode* ActorFactoryNode, const UInterchangeSceneNode* SceneNode, const UInterchangeBaseNode* TranslatedAssetNode) const;
 	
+protected:
+	UInterchangeBaseNodeContainer* BaseNodeContainer = nullptr;
 };
 
 
