Adding the support of color space conversion for colors (only gamma22 for the moment until new tests come up)

diff --git a/Engine/Plugins/Experimental/Interchange/Runtime/Source/Import/Private/MaterialX/InterchangeMaterialXTranslator.cpp b/Engine/Plugins/Experimental/Interchange/Runtime/Source/Import/Private/MaterialX/InterchangeMaterialXTranslator.cpp
--- a/Engine/Plugins/Experimental/Interchange/Runtime/Source/Import/Private/MaterialX/InterchangeMaterialXTranslator.cpp
+++ b/Engine/Plugins/Experimental/Interchange/Runtime/Source/Import/Private/MaterialX/InterchangeMaterialXTranslator.cpp
@@ -569,8 +569,8 @@ void UInterchangeMaterialXTranslator::ProcessLightShader(UInterchangeBaseNodeCon
 	// Color
 	{
 		mx::InputPtr LightColor = LightShaderNode->getInput(mx::Lights::Input::Color);
-		const mx::Color3 Color = mx::fromValueString<mx::Color3>(LightColor->getValueString());
-		LightNode->SetCustomLightColor({ Color[0], Color[1], Color[2] });
+		const FLinearColor Color = MakeLinearColorFromColor3(LightColor);
+		LightNode->SetCustomLightColor(Color);
 	}
 
 	// Intensity
@@ -1101,14 +1101,12 @@ bool UInterchangeMaterialXTranslator::AddAttribute(MaterialX::InputPtr Input, co
 
 		if(bool bIsColor3 = (Input->getType() == mx::Type::Color3 || Input->getType() == mx::Type::Vector3))
 		{
-			mx::Color3 Color = mx::fromValueString<mx::Color3>(Input->getValueString());
-			LinearColor = FLinearColor{ Color[0], Color[1], Color[2] };
+			LinearColor = MakeLinearColorFromColor3(Input);
 			bIsColor = true;
 		}
 		else if(bool bIsColor4 = (Input->getType() == mx::Type::Color4|| Input->getType() == mx::Type::Vector4))
 		{
-			mx::Color4 Color = mx::fromValueString<mx::Color4>(Input->getValueString());
-			LinearColor = FLinearColor{ Color[0], Color[1], Color[2], Color[3] };
+			LinearColor = MakeLinearColorFromColor4(Input);
 			bIsColor = true;
 		}
 
@@ -1145,8 +1143,7 @@ bool UInterchangeMaterialXTranslator::AddLinearColorAttribute(MaterialX::InputPt
 	{
 		if(Input->hasValueString())
 		{
-			mx::Color3 Color = mx::fromValueString<mx::Color3>(Input->getValueString());
-			FLinearColor Value{ Color[0], Color[1], Color[2] };
+			const FLinearColor Value = MakeLinearColorFromColor3(Input);
 
 			if(!Value.Equals(DefaultValue))
 			{
@@ -1195,6 +1192,62 @@ FString UInterchangeMaterialXTranslator::GetColorSpace(MaterialX::ElementPtr Ele
 
 	return ColorSpace;
 }
+
+FLinearColor UInterchangeMaterialXTranslator::MakeLinearColorFromColor3(MaterialX::InputPtr Input) const
+{
+	mx::Color3 Color = mx::fromValueString<mx::Color3>(Input->getValueString());
+
+	//we assume that the default color space is linear
+	FLinearColor LinearColor(Color[0], Color[1], Color[2]);
+	const FString ColorSpace = GetColorSpace(Input);
+	
+	if(ColorSpace.IsEmpty() || ColorSpace == TEXT("lin_rec709") || ColorSpace == TEXT("none"))
+	{
+		;//noop
+	}
+	else if(ColorSpace == TEXT("gamma22"))
+	{
+		LinearColor = FLinearColor::FromPow22Color(FColor(Color[0]/255.f, Color[1]/255.f, Color[2]/255.f));
+	}
+	else
+	{
+		UInterchangeResultWarning_Generic* Message = AddMessage<UInterchangeResultWarning_Generic>();
+		Message->Text = FText::Format(LOCTEXT("ColorSpaceNotSupported", "<{0}>-<{1}>: Colorspace {2} is not supported yet, falling back to linear"),
+									  FText::FromString(Input->getParent()->getName().c_str()),
+									  FText::FromString(Input->getName().c_str()),
+									  FText::FromString(ColorSpace));
+	}
+
+	return LinearColor;
+}
+
+FLinearColor UInterchangeMaterialXTranslator::MakeLinearColorFromColor4(MaterialX::InputPtr Input) const
+{
+	mx::Color4 Color = mx::fromValueString<mx::Color4>(Input->getValueString());
+
+	FLinearColor LinearColor(Color[0], Color[1], Color[2], Color[3]);
+	const FString ColorSpace = GetColorSpace(Input);
+
+	if(ColorSpace == TEXT("") || ColorSpace == TEXT("lin_rec709") || ColorSpace == TEXT("none"))
+	{
+		;//no op
+	}
+	if(ColorSpace == TEXT("gamma22"))
+	{
+		LinearColor = FLinearColor::FromPow22Color(FColor(Color[0] / 255.f, Color[1] / 255.f, Color[2] / 255.f, Color[3] / 255.f));
+	}
+	else
+	{
+		UInterchangeResultWarning_Generic* Message = AddMessage<UInterchangeResultWarning_Generic>();
+		Message->Text = FText::Format(LOCTEXT("ColorSpaceNotSupported", "<{0}>-<{1}>: Colorspace {2} is not supported yet, falling back to linear"),
+									  FText::FromString(Input->getParent()->getName().c_str()),
+									  FText::FromString(Input->getName().c_str()),
+									  FText::FromString(ColorSpace));
+	}
+
+	return LinearColor;
+}
+
 #endif //WITH_EDITOR
 
 #undef LOCTEXT_NAMESPACE
\ No newline at end of file
diff --git a/Engine/Plugins/Experimental/Interchange/Runtime/Source/Import/Public/MaterialX/InterchangeMaterialXTranslator.h b/Engine/Plugins/Experimental/Interchange/Runtime/Source/Import/Public/MaterialX/InterchangeMaterialXTranslator.h
--- a/Engine/Plugins/Experimental/Interchange/Runtime/Source/Import/Public/MaterialX/InterchangeMaterialXTranslator.h
+++ b/Engine/Plugins/Experimental/Interchange/Runtime/Source/Import/Public/MaterialX/InterchangeMaterialXTranslator.h
@@ -325,6 +325,24 @@ protected:
 	 */
 	FString GetColorSpace(MaterialX::ElementPtr Element) const;
 
+	/**
+	 * Helper function that returns a color after a color space conversion, the function makes no assumption on the input, and it should have a value of Color3 type
+	 * 
+	 * @param Input - The input that has a Color3 value in it
+	 * 
+	 * @return The linear color after color space conversion
+	 */
+	FLinearColor MakeLinearColorFromColor3(MaterialX::InputPtr Input) const;
+
+	/**
+	 * Helper function that returns a color after a color space conversion, the function makes no assumption on the input, and it should have a value of Color3 type
+	 *
+	 * @param Input - The input that has a Color3 value in it
+	 *
+	 * @return The linear color after color space conversion
+	 */
+	FLinearColor MakeLinearColorFromColor4(MaterialX::InputPtr Input) const;
+
 private:
 
 	TMap<TPair<FString, FString>, FString> InputNamesMaterialX2UE; //given a MaterialX node (category - input), return the UE/Interchange input name.
