Fix crash in the fbx translator by passing the correct base class when creating light node.

#rnx

diff --git a/Engine/Plugins/Experimental/Interchange/Runtime/Source/Parsers/Fbx/Private/FbxLight.cpp b/Engine/Plugins/Experimental/Interchange/Runtime/Source/Parsers/Fbx/Private/FbxLight.cpp
--- a/Engine/Plugins/Experimental/Interchange/Runtime/Source/Parsers/Fbx/Private/FbxLight.cpp
+++ b/Engine/Plugins/Experimental/Interchange/Runtime/Source/Parsers/Fbx/Private/FbxLight.cpp
@@ -18,7 +18,7 @@ namespace UE
 	{
 		namespace Private
 		{
-			UInterchangeLightNode* FFbxLight::CreateLightNode(UInterchangeBaseNodeContainer& NodeContainer, const FString& NodeUid, const FString& NodeName, const FbxLight& LightAttribute)
+			UInterchangeBaseLightNode* FFbxLight::CreateLightNode(UInterchangeBaseNodeContainer& NodeContainer, const FString& NodeUid, const FString& NodeName, const FbxLight& LightAttribute)
 			{
 				UClass* LightClass;
 
@@ -42,18 +42,18 @@ namespace UE
 					break;
 				}
 
-				UInterchangeLightNode* LightNode = NewObject<UInterchangeLightNode>(&NodeContainer, LightClass, NAME_None);
-				if (!ensure(LightNode))
+				UInterchangeBaseLightNode* BaseLightNode = NewObject<UInterchangeBaseLightNode>(&NodeContainer, LightClass, NAME_None);
+				if (!ensure(BaseLightNode))
 				{
 					UInterchangeResultError_Generic* Message = Parser.AddMessage<UInterchangeResultError_Generic>();
 					Message->Text = LOCTEXT("CannotAllocateNode", "Cannot allocate a node when importing FBX.");
 					return nullptr;
 				}
 
-				LightNode->InitializeNode(NodeUid, NodeName, EInterchangeNodeContainerType::TranslatedAsset);
-				NodeContainer.AddNode(LightNode);
+				BaseLightNode->InitializeNode(NodeUid, NodeName, EInterchangeNodeContainerType::TranslatedAsset);
+				NodeContainer.AddNode(BaseLightNode);
 
-				return LightNode;
+				return BaseLightNode;
 			}
 
 			void FFbxLight::AddLightsRecursively(FbxNode* Node, UInterchangeBaseNodeContainer& NodeContainer)
@@ -65,10 +65,10 @@ namespace UE
 
 					if (NodeAttribute && NodeAttribute->GetAttributeType() == FbxNodeAttribute::eLight)
 					{
-						FString NodeName = FFbxHelper::GetNodeAttributeName(NodeAttribute, UInterchangeLightNode::StaticAssetTypeName());
-						FString NodeUid = FFbxHelper::GetNodeAttributeUniqueID(NodeAttribute, UInterchangeLightNode::StaticAssetTypeName());
+						FString NodeName = FFbxHelper::GetNodeAttributeName(NodeAttribute, UInterchangeBaseLightNode::StaticAssetTypeName());
+						FString NodeUid = FFbxHelper::GetNodeAttributeUniqueID(NodeAttribute, UInterchangeBaseLightNode::StaticAssetTypeName());
 
-						const UInterchangeLightNode* LightNode = Cast<const UInterchangeLightNode>(NodeContainer.GetNode(NodeUid));
+						const UInterchangeBaseLightNode* LightNode = Cast<const UInterchangeBaseLightNode>(NodeContainer.GetNode(NodeUid));
 
 						if (!LightNode)
 						{
diff --git a/Engine/Plugins/Experimental/Interchange/Runtime/Source/Parsers/Fbx/Private/FbxLight.h b/Engine/Plugins/Experimental/Interchange/Runtime/Source/Parsers/Fbx/Private/FbxLight.h
--- a/Engine/Plugins/Experimental/Interchange/Runtime/Source/Parsers/Fbx/Private/FbxLight.h
+++ b/Engine/Plugins/Experimental/Interchange/Runtime/Source/Parsers/Fbx/Private/FbxLight.h
@@ -5,7 +5,7 @@
 #include "FbxInclude.h"
 
 class UInterchangeBaseNodeContainer;
-class UInterchangeLightNode;
+class UInterchangeBaseLightNode;
 
 namespace UE
 {
@@ -23,14 +23,14 @@ namespace UE
 				{}
 
 				/**
-				 * Create a UInterchangeLightNode and add it to the NodeContainer for all FbxNodeAttributes of type eLight the fbx file contains.
+				 * Create a UInterchangeBaseLightNode and add it to the NodeContainer for all FbxNodeAttributes of type eLight the fbx file contains.
 				 * 
 				 * @note - Any node that already exist in the NodeContainer will not be created or modified.
 				 */
 				void AddAllLights(FbxScene* SDKScene, UInterchangeBaseNodeContainer& NodeContainer);
 
 			protected:
-				UInterchangeLightNode* CreateLightNode(UInterchangeBaseNodeContainer& NodeContainer, const FString& NodeUID, const FString& NodeName, const FbxLight& LightAttribute);
+				UInterchangeBaseLightNode* CreateLightNode(UInterchangeBaseNodeContainer& NodeContainer, const FString& NodeUID, const FString& NodeName, const FbxLight& LightAttribute);
 				void AddLightsRecursively(FbxNode* Node, UInterchangeBaseNodeContainer& NodeContainer);
 
 			private:
