Inclusive Terminology Push
Replaced "slave" with "follower" at InterchangeSkeletonHelper.cpp

diff --git a/Engine/Plugins/Experimental/Interchange/Runtime/Source/Import/Private/Mesh/InterchangeSkeletonHelper.cpp b/Engine/Plugins/Experimental/Interchange/Runtime/Source/Import/Private/Mesh/InterchangeSkeletonHelper.cpp
--- a/Engine/Plugins/Experimental/Interchange/Runtime/Source/Import/Private/Mesh/InterchangeSkeletonHelper.cpp
+++ b/Engine/Plugins/Experimental/Interchange/Runtime/Source/Import/Private/Mesh/InterchangeSkeletonHelper.cpp
@@ -151,7 +151,7 @@ namespace UE::Interchange::Private
 		}
 
 		// originally we made sure at least matches more than 50% 
-		// but then slave components can't play since they're only partial
+		// but then follower components can't play since they're only partial
 		// if the hierarchy matches, and if it's more then 1 bone, we allow
 		return (NumOfBoneMatches > 0);
 	}
