Bug fix for identifying normalMaps in Interchange Texture pipeline at PostImport step

diff --git a/Engine/Plugins/Experimental/Interchange/Runtime/Source/Pipelines/Private/InterchangeGenericTexturePipeline.cpp b/Engine/Plugins/Experimental/Interchange/Runtime/Source/Pipelines/Private/InterchangeGenericTexturePipeline.cpp
--- a/Engine/Plugins/Experimental/Interchange/Runtime/Source/Pipelines/Private/InterchangeGenericTexturePipeline.cpp
+++ b/Engine/Plugins/Experimental/Interchange/Runtime/Source/Pipelines/Private/InterchangeGenericTexturePipeline.cpp
@@ -292,7 +292,7 @@ void UInterchangeGenericTexturePipeline::PostImportTextureAssetImport(UObject* C
 		// Verify if the texture is a normal map
 		if (UTexture* Texture = Cast<UTexture>(CreatedAsset))
 		{
-			if (!Texture->IsNormalMap())
+			if (Texture->IsNormalMap())
 			{
 				// This can create 2 build of the texture (we should revisit this at some point)
 				if (FTextureCompilingManager::Get().IsCompilingTexture(Texture))
