InterChange Gltf Import - NormalMap - Flip Green Channel via Material functions
Introduced Material Node structure for Interchange-GLTF, the ability to flip the green channel of the normal map

diff --git a/Engine/Plugins/Experimental/Interchange/Runtime/Source/Import/Private/Gltf/InterchangeGltfTranslator.cpp b/Engine/Plugins/Experimental/Interchange/Runtime/Source/Import/Private/Gltf/InterchangeGltfTranslator.cpp
--- a/Engine/Plugins/Experimental/Interchange/Runtime/Source/Import/Private/Gltf/InterchangeGltfTranslator.cpp
+++ b/Engine/Plugins/Experimental/Interchange/Runtime/Source/Import/Private/Gltf/InterchangeGltfTranslator.cpp
@@ -259,6 +259,24 @@ void UInterchangeGltfTranslator::HandleGltfMaterialParameter( UInterchangeBaseNo
 			bNeedsFactorNode = !MapFactor.Get< FLinearColor >().Equals( FLinearColor::White );
 		}
 
+		if (bIsNormal)
+		{
+			//MakeFloat3:
+			UInterchangeShaderNode* MakeFloat3Node = UInterchangeShaderNode::Create(&NodeContainer, ColorNodeName + TEXT("_MakeFloat3"), ShaderNode.GetUniqueID());
+			MakeFloat3Node->SetCustomShaderType(Standard::Nodes::MakeFloat3::Name.ToString());
+
+			UInterchangeShaderNode* MultiplyNode = UInterchangeShaderNode::Create(&NodeContainer, ColorNodeName + TEXT("_Flip"), ShaderNode.GetUniqueID());
+			MultiplyNode->SetCustomShaderType(Standard::Nodes::Multiply::Name.ToString());
+			MultiplyNode->AddFloatAttribute(UInterchangeShaderPortsAPI::MakeInputValueKey(Standard::Nodes::Multiply::Inputs::B.ToString()), -1.0);
+			UInterchangeShaderPortsAPI::ConnectOuputToInput(MultiplyNode, Standard::Nodes::Multiply::Inputs::A.ToString(), ColorNode->GetUniqueID(), "G");
+
+			UInterchangeShaderPortsAPI::ConnectOuputToInput(MakeFloat3Node, Standard::Nodes::MakeFloat3::Inputs::X.ToString(), ColorNode->GetUniqueID(), "R");
+			UInterchangeShaderPortsAPI::ConnectDefaultOuputToInput(MakeFloat3Node, Standard::Nodes::MakeFloat3::Inputs::Y.ToString(), MultiplyNode->GetUniqueID());
+			UInterchangeShaderPortsAPI::ConnectOuputToInput(MakeFloat3Node, Standard::Nodes::MakeFloat3::Inputs::Z.ToString(), ColorNode->GetUniqueID(), "B");
+
+			ColorNode = MakeFloat3Node;
+		}
+
 		if ( bNeedsFactorNode )
 		{
 			UInterchangeShaderNode* FactorNode = UInterchangeShaderNode::Create( &NodeContainer, ColorNodeName + TEXT("_Factor"), ShaderNode.GetUniqueID() );
@@ -778,7 +796,7 @@ bool UInterchangeGltfTranslator::Translate( UInterchangeBaseNodeContainer& NodeC
 
 	const bool bLoadImageData = false;
 	const bool bLoadMetaData = false;
-	GltfFileReader.ReadFile( FilePath, bLoadImageData, bLoadMetaData, const_cast< UInterchangeGltfTranslator* >( this )->GltfAsset );
+	GltfFileReader.ReadFile( FilePath, bLoadImageData, bLoadMetaData, const_cast< UInterchangeGltfTranslator* >( this )->GltfAsset );	
 
 	const FString FileName = FPaths::GetBaseFilename(FilePath);
 	const_cast< UInterchangeGltfTranslator* >( this )->GltfAsset.GenerateNames(FileName);
diff --git a/Engine/Plugins/Experimental/Interchange/Runtime/Source/Nodes/Public/InterchangeMaterialDefinitions.h b/Engine/Plugins/Experimental/Interchange/Runtime/Source/Nodes/Public/InterchangeMaterialDefinitions.h
--- a/Engine/Plugins/Experimental/Interchange/Runtime/Source/Nodes/Public/InterchangeMaterialDefinitions.h
+++ b/Engine/Plugins/Experimental/Interchange/Runtime/Source/Nodes/Public/InterchangeMaterialDefinitions.h
@@ -75,6 +75,17 @@ namespace Interchange
 					}
 				}
 
+				namespace MakeFloat3
+				{
+					const FName Name = TEXT("MakeFloat3");
+					namespace Inputs
+					{
+						const FName X = TEXT("X");
+						const FName Y = TEXT("Y");
+						const FName Z = TEXT("Z");
+					}
+				}
+
 				namespace OneMinus
 				{
 					const FName Name = TEXT("OneMinus");
diff --git a/Engine/Plugins/Experimental/Interchange/Runtime/Source/Pipelines/Private/InterchangeGenericMaterialPipeline.cpp b/Engine/Plugins/Experimental/Interchange/Runtime/Source/Pipelines/Private/InterchangeGenericMaterialPipeline.cpp
--- a/Engine/Plugins/Experimental/Interchange/Runtime/Source/Pipelines/Private/InterchangeGenericMaterialPipeline.cpp
+++ b/Engine/Plugins/Experimental/Interchange/Runtime/Source/Pipelines/Private/InterchangeGenericMaterialPipeline.cpp
@@ -945,6 +945,42 @@ void UInterchangeGenericMaterialPipeline::HandleFlattenNormalNode(const UInterch
 	}
 }
 
+void UInterchangeGenericMaterialPipeline::HandleMakeFloat3Node(const UInterchangeShaderNode* ShaderNode, UInterchangeBaseMaterialFactoryNode* MaterialFactoryNode,
+	UInterchangeMaterialExpressionFactoryNode* MakeFloat3FactoryNode)
+{
+	using namespace UE::Interchange::Materials::Standard::Nodes::MakeFloat3;
+
+	MakeFloat3FactoryNode->SetCustomExpressionClassName(UMaterialExpressionMaterialFunctionCall::StaticClass()->GetName());
+
+	const FName MaterialFunctionMemberName = GET_MEMBER_NAME_CHECKED(UMaterialExpressionMaterialFunctionCall, MaterialFunction);
+	MakeFloat3FactoryNode->AddStringAttribute(MaterialFunctionMemberName, TEXT("/Engine/Functions/Engine_MaterialFunctions02/Utility/MakeFloat3.MakeFloat3"));
+	MakeFloat3FactoryNode->AddApplyAndFillDelegates<FString>(MaterialFunctionMemberName, UMaterialExpressionMaterialFunctionCall::StaticClass(), MaterialFunctionMemberName);
+
+	TTuple<UInterchangeMaterialExpressionFactoryNode*, FString> RedChannelExpression =
+		CreateMaterialExpressionForInput(MaterialFactoryNode, ShaderNode, Inputs::X.ToString(), MakeFloat3FactoryNode->GetUniqueID());
+	if (RedChannelExpression.Get<0>())
+	{
+		UInterchangeShaderPortsAPI::ConnectOuputToInput(MakeFloat3FactoryNode, TEXT("X"),
+			RedChannelExpression.Get<0>()->GetUniqueID(), RedChannelExpression.Get<1>());
+	}
+
+	TTuple<UInterchangeMaterialExpressionFactoryNode*, FString> GreenChannelExpression =
+		CreateMaterialExpressionForInput(MaterialFactoryNode, ShaderNode, Inputs::Y.ToString(), MakeFloat3FactoryNode->GetUniqueID());
+	if (GreenChannelExpression.Get<0>())
+	{
+		UInterchangeShaderPortsAPI::ConnectOuputToInput(MakeFloat3FactoryNode, TEXT("Y"),
+			GreenChannelExpression.Get<0>()->GetUniqueID(), GreenChannelExpression.Get<1>());
+	}
+
+	TTuple<UInterchangeMaterialExpressionFactoryNode*, FString> BlueChannelExpression =
+		CreateMaterialExpressionForInput(MaterialFactoryNode, ShaderNode, Inputs::Z.ToString(), MakeFloat3FactoryNode->GetUniqueID());
+	if (BlueChannelExpression.Get<0>())
+	{
+		UInterchangeShaderPortsAPI::ConnectOuputToInput(MakeFloat3FactoryNode, TEXT("Z"),
+			BlueChannelExpression.Get<0>()->GetUniqueID(), BlueChannelExpression.Get<1>());
+	}
+}
+
 void UInterchangeGenericMaterialPipeline::HandleTextureSampleNode(const UInterchangeShaderNode* ShaderNode, UInterchangeBaseMaterialFactoryNode* MaterialFactoryNode, UInterchangeMaterialExpressionFactoryNode* TextureSampleFactoryNode)
 {
 	using namespace UE::Interchange::Materials::Standard::Nodes::TextureSample;
@@ -1288,6 +1324,10 @@ UInterchangeMaterialExpressionFactoryNode* UInterchangeGenericMaterialPipeline::
 	{
 		HandleFlattenNormalNode(ShaderNode, MaterialFactoryNode, MaterialExpression);
 	}
+	else if (*ShaderType == Nodes::MakeFloat3::Name)
+	{
+		HandleMakeFloat3Node(ShaderNode, MaterialFactoryNode, MaterialExpression);
+	}
 	else if (*ShaderType == Nodes::Lerp::Name)
 	{
 		HandleLerpNode(ShaderNode, MaterialFactoryNode, MaterialExpression);
diff --git a/Engine/Plugins/Experimental/Interchange/Runtime/Source/Pipelines/Public/InterchangeGenericMaterialPipeline.h b/Engine/Plugins/Experimental/Interchange/Runtime/Source/Pipelines/Public/InterchangeGenericMaterialPipeline.h
--- a/Engine/Plugins/Experimental/Interchange/Runtime/Source/Pipelines/Public/InterchangeGenericMaterialPipeline.h
+++ b/Engine/Plugins/Experimental/Interchange/Runtime/Source/Pipelines/Public/InterchangeGenericMaterialPipeline.h
@@ -100,6 +100,7 @@ private:
 	bool HandleBxDFInput(const UInterchangeShaderGraphNode* ShaderGraphNode, UInterchangeMaterialFactoryNode* MaterialFactoryNode);
 
 	void HandleFlattenNormalNode(const UInterchangeShaderNode* ShaderNode, UInterchangeBaseMaterialFactoryNode* MaterialFactoryNode, UInterchangeMaterialExpressionFactoryNode* FlattenNormalFactoryNode);
+	void HandleMakeFloat3Node(const UInterchangeShaderNode* ShaderNode, UInterchangeBaseMaterialFactoryNode* MaterialFactoryNode, UInterchangeMaterialExpressionFactoryNode* MakeFloat3FactoryNode);
 	void HandleTextureSampleNode(const UInterchangeShaderNode* ShaderNode, UInterchangeBaseMaterialFactoryNode* MaterialFactoryNode, UInterchangeMaterialExpressionFactoryNode* TextureSampleFactoryNode);
 	void HandleTextureCoordinateNode(const UInterchangeShaderNode* ShaderNode, UInterchangeBaseMaterialFactoryNode* MaterialFactoryNode, UInterchangeMaterialExpressionFactoryNode*& TextureSampleFactoryNode);
 	void HandleLerpNode(const UInterchangeShaderNode* ShaderNode, UInterchangeBaseMaterialFactoryNode* MaterialFactoryNode, UInterchangeMaterialExpressionFactoryNode* LerpFactoryNode);
