Adding a new Light Factory to Interchange.
- Adding support of IES Texture

diff --git a/Engine/Plugins/Experimental/Interchange/Runtime/Source/FactoryNodes/Private/InterchangeLightFactoryNode.cpp b/Engine/Plugins/Experimental/Interchange/Runtime/Source/FactoryNodes/Private/InterchangeLightFactoryNode.cpp
--- a/Engine/Plugins/Experimental/Interchange/Runtime/Source/FactoryNodes/Private/InterchangeLightFactoryNode.cpp
+++ b/Engine/Plugins/Experimental/Interchange/Runtime/Source/FactoryNodes/Private/InterchangeLightFactoryNode.cpp
@@ -65,6 +65,16 @@ bool UInterchangeLightFactoryNode::SetCustomAttenuationRadius(float AttributeVal
 	IMPLEMENT_NODE_ATTRIBUTE_SETTER(UInterchangeLightFactoryNode, AttenuationRadius, float, ULocalLightComponent)
 }
 
+bool UInterchangeLightFactoryNode::GetCustomIESTexture(FString& AttributeValue) const
+{
+	IMPLEMENT_NODE_ATTRIBUTE_GETTER(IESTexture, FString);
+}
+
+bool UInterchangeLightFactoryNode::SetCustomIESTexture(const FString& AttributeValue)
+{
+	IMPLEMENT_NODE_ATTRIBUTE_SETTER_NODELEGATE(IESTexture, FString)
+}
+
 bool UInterchangePointLightFactoryNode::GetCustomUseInverseSquaredFalloff(bool& AttributeValue) const
 {
 	IMPLEMENT_NODE_ATTRIBUTE_GETTER(bUseInverseSquaredFalloff, bool)
diff --git a/Engine/Plugins/Experimental/Interchange/Runtime/Source/FactoryNodes/Public/InterchangeLightFactoryNode.h b/Engine/Plugins/Experimental/Interchange/Runtime/Source/FactoryNodes/Public/InterchangeLightFactoryNode.h
--- a/Engine/Plugins/Experimental/Interchange/Runtime/Source/FactoryNodes/Public/InterchangeLightFactoryNode.h
+++ b/Engine/Plugins/Experimental/Interchange/Runtime/Source/FactoryNodes/Public/InterchangeLightFactoryNode.h
@@ -70,6 +70,12 @@ public:
 
 	UFUNCTION(BlueprintCallable, Category = "Interchange | Node | LightFactory")
 	bool SetCustomAttenuationRadius(float AttributeValue, bool bAddApplyDelegate = true);
+	
+	UFUNCTION(BlueprintCallable, Category = "Interchange | Node | LightFactory")
+	bool GetCustomIESTexture(FString& AttributeValue) const;
+
+	UFUNCTION(BlueprintCallable, Category = "Interchange | Node | LightFactory")
+	bool SetCustomIESTexture(const FString& AttributeValue);
 
 private:
 
diff --git a/Engine/Plugins/Experimental/Interchange/Runtime/Source/Import/Private/InterchangeImportModule.cpp b/Engine/Plugins/Experimental/Interchange/Runtime/Source/Import/Private/InterchangeImportModule.cpp
--- a/Engine/Plugins/Experimental/Interchange/Runtime/Source/Import/Private/InterchangeImportModule.cpp
+++ b/Engine/Plugins/Experimental/Interchange/Runtime/Source/Import/Private/InterchangeImportModule.cpp
@@ -20,6 +20,7 @@
 #include "Modules/ModuleManager.h"
 #include "Scene/InterchangeActorFactory.h"
 #include "Scene/InterchangeCineCameraActorFactory.h"
+#include "Scene/InterchangeLightActorFactory.h"
 #include "Scene/InterchangeSceneVariantSetsFactory.h"
 #include "Scene/InterchangeStaticMeshActorFactory.h"
 #include "Scene/InterchangeSkeletalMeshActorFactory.h"
@@ -88,6 +89,7 @@ void FInterchangeImportModule::StartupModule()
 		InterchangeManager.RegisterFactory(UInterchangeStaticMeshActorFactory::StaticClass());
 		InterchangeManager.RegisterFactory(UInterchangeSkeletalMeshActorFactory::StaticClass());
 		InterchangeManager.RegisterFactory(UInterchangeSceneVariantSetsFactory::StaticClass());
+		InterchangeManager.RegisterFactory(UInterchangeLightActorFactory::StaticClass());
 	};
 	
 	if (GEngine)
diff --git a/Engine/Plugins/Experimental/Interchange/Runtime/Source/Import/Private/Scene/InterchangeLightActorFactory.cpp b/Engine/Plugins/Experimental/Interchange/Runtime/Source/Import/Private/Scene/InterchangeLightActorFactory.cpp
new file mode 100644
--- /dev/null
+++ b/Engine/Plugins/Experimental/Interchange/Runtime/Source/Import/Private/Scene/InterchangeLightActorFactory.cpp
@@ -0,0 +1,51 @@
+// Copyright Epic Games, Inc. All Rights Reserved.
+
+#include "Scene/InterchangeLightActorFactory.h"
+#include "Scene/InterchangeActorHelper.h"
+#include "Components/LightComponent.h"
+#include "Nodes/InterchangeFactoryBaseNode.h"
+#include "Nodes/InterchangeBaseNodeContainer.h"
+#include "InterchangeTextureLightProfileFactoryNode.h"
+
+#include "CoreMinimal.h"
+#include "UObject/Object.h"
+#include "UObject/ObjectMacros.h"
+#include "InterchangeLightFactoryNode.h"
+#include "Engine/Light.h"
+
+UClass* UInterchangeLightActorFactory::GetFactoryClass() const
+{
+	return ALight::StaticClass();
+}
+
+UObject* UInterchangeLightActorFactory::CreateSceneObject(const UInterchangeFactoryBase::FCreateSceneObjectsParams& CreateSceneObjectsParams)
+{
+	ALight * SpawnedActor = Cast<ALight>(UE::Interchange::ActorHelper::SpawnFactoryActor(CreateSceneObjectsParams));
+
+	if(!SpawnedActor)
+	{
+		return nullptr;
+	}
+
+	if(ULightComponent* LightComponent = SpawnedActor->GetLightComponent())
+	{
+		LightComponent->UnregisterComponent();
+		if(const UInterchangeLightFactoryNode* LightFactoryNode = Cast<UInterchangeLightFactoryNode>(CreateSceneObjectsParams.FactoryNode))
+		{
+			if(FString IESTexture; LightFactoryNode->GetCustomIESTexture(IESTexture))
+			{
+				IESTexture = UInterchangeFactoryBaseNode::BuildFactoryNodeUid(IESTexture);
+				if(const UInterchangeTextureLightProfileFactoryNode* TextureFactoryNode = Cast<UInterchangeTextureLightProfileFactoryNode>(CreateSceneObjectsParams.NodeContainer->GetNode(IESTexture)))
+				{
+					if(UTextureLightProfile* Texture = Cast<UTextureLightProfile>(TextureFactoryNode->ReferenceObject.TryLoad()))
+					{
+						LightComponent->SetIESTexture(Texture);
+					}
+				}
+			}
+		}
+		CreateSceneObjectsParams.FactoryNode->ApplyAllCustomAttributeToObject(LightComponent);
+	}
+
+	return SpawnedActor;
+}
diff --git a/Engine/Plugins/Experimental/Interchange/Runtime/Source/Import/Public/Scene/InterchangeLightActorFactory.h b/Engine/Plugins/Experimental/Interchange/Runtime/Source/Import/Public/Scene/InterchangeLightActorFactory.h
new file mode 100644
--- /dev/null
+++ b/Engine/Plugins/Experimental/Interchange/Runtime/Source/Import/Public/Scene/InterchangeLightActorFactory.h
@@ -0,0 +1,29 @@
+// Copyright Epic Games, Inc. All Rights Reserved.
+
+#pragma once
+
+#include "InterchangeFactoryBase.h"
+
+#include "InterchangeLightActorFactory.generated.h"
+
+UCLASS(BlueprintType, Experimental)
+class INTERCHANGEIMPORT_API UInterchangeLightActorFactory : public UInterchangeFactoryBase
+{
+	GENERATED_BODY()
+
+public:
+	//////////////////////////////////////////////////////////////////////////
+	// Interchange factory base interface begin
+
+	virtual UClass* GetFactoryClass() const override;
+
+	virtual UObject* CreateSceneObject(const UInterchangeFactoryBase::FCreateSceneObjectsParams& CreateSceneObjectsParams) override;
+
+	virtual bool CanExecuteOnAnyThread() const override
+	{
+		return false;
+	}
+
+	// Interchange factory base interface end
+	//////////////////////////////////////////////////////////////////////////
+};
\ No newline at end of file
diff --git a/Engine/Plugins/Experimental/Interchange/Runtime/Source/Nodes/Private/InterchangeLightNode.cpp b/Engine/Plugins/Experimental/Interchange/Runtime/Source/Nodes/Private/InterchangeLightNode.cpp
--- a/Engine/Plugins/Experimental/Interchange/Runtime/Source/Nodes/Private/InterchangeLightNode.cpp
+++ b/Engine/Plugins/Experimental/Interchange/Runtime/Source/Nodes/Private/InterchangeLightNode.cpp
@@ -77,6 +77,16 @@ bool UInterchangeLightNode::SetCustomAttenuationRadius(float AttributeValue)
 	IMPLEMENT_NODE_ATTRIBUTE_SETTER_NODELEGATE(AttenuationRadius, float)
 }
 
+bool UInterchangeLightNode::GetCustomIESTexture(FString& AttributeValue) const
+{
+	IMPLEMENT_NODE_ATTRIBUTE_GETTER(IESTexture, FString)
+}
+
+bool UInterchangeLightNode::SetCustomIESTexture(const FString& AttributeValue)
+{
+	IMPLEMENT_NODE_ATTRIBUTE_SETTER_NODELEGATE(IESTexture, FString)
+}
+
 FString UInterchangePointLightNode::GetTypeName() const
 {
 	return TEXT("PointLightNode");
diff --git a/Engine/Plugins/Experimental/Interchange/Runtime/Source/Nodes/Public/InterchangeLightNode.h b/Engine/Plugins/Experimental/Interchange/Runtime/Source/Nodes/Public/InterchangeLightNode.h
--- a/Engine/Plugins/Experimental/Interchange/Runtime/Source/Nodes/Public/InterchangeLightNode.h
+++ b/Engine/Plugins/Experimental/Interchange/Runtime/Source/Nodes/Public/InterchangeLightNode.h
@@ -82,10 +82,17 @@ public:
 	UFUNCTION(BlueprintCallable, Category = "Interchange | Node | LocalLight")
 	bool SetCustomAttenuationRadius(float AttributeValue);
 
+	UFUNCTION(BlueprintCallable, Category = "Interchange | Node | LocalLight")
+	bool GetCustomIESTexture(FString& AttributeValue) const;
+
+	UFUNCTION(BlueprintCallable, Category = "Interchange | Node | LocalLight")
+	bool SetCustomIESTexture(const FString& AttributeValue);
+
 private:
 
 	IMPLEMENT_NODE_ATTRIBUTE_KEY(IntensityUnits)
 	IMPLEMENT_NODE_ATTRIBUTE_KEY(AttenuationRadius)
+	IMPLEMENT_NODE_ATTRIBUTE_KEY(IESTexture)
 };
 
 UCLASS(BlueprintType)
diff --git a/Engine/Plugins/Experimental/Interchange/Runtime/Source/Pipelines/Private/InterchangeGenericScenesPipeline.cpp b/Engine/Plugins/Experimental/Interchange/Runtime/Source/Pipelines/Private/InterchangeGenericScenesPipeline.cpp
--- a/Engine/Plugins/Experimental/Interchange/Runtime/Source/Pipelines/Private/InterchangeGenericScenesPipeline.cpp
+++ b/Engine/Plugins/Experimental/Interchange/Runtime/Source/Pipelines/Private/InterchangeGenericScenesPipeline.cpp
@@ -256,6 +256,11 @@ void UInterchangeGenericLevelPipeline::SetUpFactoryNode(UInterchangeActorFactory
 						LightFactoryNode->SetCustomAttenuationRadius(AttenuationRadius);
 					}
 
+					if(FString IESTexture; LightNode->GetCustomIESTexture(IESTexture))
+					{
+						LightFactoryNode->SetCustomIESTexture(IESTexture);
+					}
+
 					// RectLight
 					if(const UInterchangeRectLightNode* RectLightNode = Cast<UInterchangeRectLightNode>(LightNode))
 					{
