Fixing colors on gltf import (Interchange).
(base+diffuse)Textures are now correctly imported with sRGB settings.

diff --git a/Engine/Plugins/Experimental/Interchange/Runtime/Source/Import/Private/Gltf/InterchangeGltfTranslator.cpp b/Engine/Plugins/Experimental/Interchange/Runtime/Source/Import/Private/Gltf/InterchangeGltfTranslator.cpp
--- a/Engine/Plugins/Experimental/Interchange/Runtime/Source/Import/Private/Gltf/InterchangeGltfTranslator.cpp
+++ b/Engine/Plugins/Experimental/Interchange/Runtime/Source/Import/Private/Gltf/InterchangeGltfTranslator.cpp
@@ -338,6 +338,10 @@ void UInterchangeGltfTranslator::HandleGltfMaterial( UInterchangeBaseNodeContain
 
 	ShaderGraphNode.SetCustomTwoSided( GltfMaterial.bIsDoubleSided );
 
+	//Based on the gltf specification the basecolor and emissive textures have SRGB colors:
+	SetTextureSRGB(NodeContainer, GltfMaterial.BaseColor);
+	SetTextureSRGB(NodeContainer, GltfMaterial.Emissive);
+
 	if ( GltfMaterial.ShadingModel == GLTF::FMaterial::EShadingModel::MetallicRoughness )
 	{
 		// Base Color
@@ -1182,3 +1186,14 @@ void UInterchangeGltfTranslator::HandleGltfAnimation(UInterchangeBaseNodeContain
 
 	NodeContainer.AddNode(TrackSetNode);
 }
+
+void UInterchangeGltfTranslator::SetTextureSRGB(UInterchangeBaseNodeContainer& NodeContainer, const GLTF::FTextureMap& TextureMap) const
+{
+	if (GltfAsset.Textures.IsValidIndex(TextureMap.TextureIndex))
+	{
+		const FString TextureName = GltfAsset.Textures[TextureMap.TextureIndex].Name;
+		const FString TextureUid = UInterchangeTextureNode::MakeNodeUid(TextureName);
+		UInterchangeTextureNode* TextureNode = const_cast<UInterchangeTextureNode*>(Cast<UInterchangeTextureNode>(NodeContainer.GetNode(TextureUid)));
+		TextureNode->SetCustomSRGB(true);
+	}
+}
\ No newline at end of file
diff --git a/Engine/Plugins/Experimental/Interchange/Runtime/Source/Import/Public/Gltf/InterchangeGltfTranslator.h b/Engine/Plugins/Experimental/Interchange/Runtime/Source/Import/Public/Gltf/InterchangeGltfTranslator.h
--- a/Engine/Plugins/Experimental/Interchange/Runtime/Source/Import/Public/Gltf/InterchangeGltfTranslator.h
+++ b/Engine/Plugins/Experimental/Interchange/Runtime/Source/Import/Public/Gltf/InterchangeGltfTranslator.h
@@ -72,6 +72,8 @@ protected:
 	void HandleGltfTextureTransform( UInterchangeBaseNodeContainer& NodeContainer, const GLTF::FTextureTransform& TextureTransform, const int32 TexCoordIndex, UInterchangeShaderNode& ShaderNode ) const;
 
 private:
+	void SetTextureSRGB(UInterchangeBaseNodeContainer& NodeContainer, const GLTF::FTextureMap& TextureMap) const;
+
 	GLTF::FAsset GltfAsset;
 };
 
