Refactoring GreenChannelFlipping for GLTF Interchange imports (for normalmaps)
(refactoring from in Shader green channel flipping to in Texture one)

diff --git a/Engine/Plugins/Experimental/Interchange/Runtime/Source/Import/Private/Gltf/InterchangeGltfTranslator.cpp b/Engine/Plugins/Experimental/Interchange/Runtime/Source/Import/Private/Gltf/InterchangeGltfTranslator.cpp
--- a/Engine/Plugins/Experimental/Interchange/Runtime/Source/Import/Private/Gltf/InterchangeGltfTranslator.cpp
+++ b/Engine/Plugins/Experimental/Interchange/Runtime/Source/Import/Private/Gltf/InterchangeGltfTranslator.cpp
@@ -259,24 +259,6 @@ void UInterchangeGltfTranslator::HandleGltfMaterialParameter( UInterchangeBaseNo
 			bNeedsFactorNode = !MapFactor.Get< FLinearColor >().Equals( FLinearColor::White );
 		}
 
-		if (bIsNormal)
-		{
-			//MakeFloat3:
-			UInterchangeShaderNode* MakeFloat3Node = UInterchangeShaderNode::Create(&NodeContainer, ColorNodeName + TEXT("_MakeFloat3"), ShaderNode.GetUniqueID());
-			MakeFloat3Node->SetCustomShaderType(Standard::Nodes::MakeFloat3::Name.ToString());
-
-			UInterchangeShaderNode* MultiplyNode = UInterchangeShaderNode::Create(&NodeContainer, ColorNodeName + TEXT("_Flip"), ShaderNode.GetUniqueID());
-			MultiplyNode->SetCustomShaderType(Standard::Nodes::Multiply::Name.ToString());
-			MultiplyNode->AddFloatAttribute(UInterchangeShaderPortsAPI::MakeInputValueKey(Standard::Nodes::Multiply::Inputs::B.ToString()), -1.0);
-			UInterchangeShaderPortsAPI::ConnectOuputToInput(MultiplyNode, Standard::Nodes::Multiply::Inputs::A.ToString(), ColorNode->GetUniqueID(), "G");
-
-			UInterchangeShaderPortsAPI::ConnectOuputToInput(MakeFloat3Node, Standard::Nodes::MakeFloat3::Inputs::X.ToString(), ColorNode->GetUniqueID(), "R");
-			UInterchangeShaderPortsAPI::ConnectDefaultOuputToInput(MakeFloat3Node, Standard::Nodes::MakeFloat3::Inputs::Y.ToString(), MultiplyNode->GetUniqueID());
-			UInterchangeShaderPortsAPI::ConnectOuputToInput(MakeFloat3Node, Standard::Nodes::MakeFloat3::Inputs::Z.ToString(), ColorNode->GetUniqueID(), "B");
-
-			ColorNode = MakeFloat3Node;
-		}
-
 		if ( bNeedsFactorNode )
 		{
 			UInterchangeShaderNode* FactorNode = UInterchangeShaderNode::Create( &NodeContainer, ColorNodeName + TEXT("_Factor"), ShaderNode.GetUniqueID() );
@@ -341,6 +323,9 @@ void UInterchangeGltfTranslator::HandleGltfMaterial( UInterchangeBaseNodeContain
 	//Based on the gltf specification the basecolor and emissive textures have SRGB colors:
 	SetTextureSRGB(NodeContainer, GltfMaterial.BaseColor);
 	SetTextureSRGB(NodeContainer, GltfMaterial.Emissive);
+	//According to GLTF documentation the normal maps should have their green channel flipped:
+	SetTextureFlipGreenChannel(NodeContainer, GltfMaterial.Normal);
+	SetTextureFlipGreenChannel(NodeContainer, GltfMaterial.ClearCoat.NormalMap);
 
 	if ( GltfMaterial.ShadingModel == GLTF::FMaterial::EShadingModel::MetallicRoughness )
 	{
@@ -1193,7 +1178,21 @@ void UInterchangeGltfTranslator::SetTextureSRGB(UInterchangeBaseNodeContainer& N
 	{
 		const FString TextureName = GltfAsset.Textures[TextureMap.TextureIndex].Name;
 		const FString TextureUid = UInterchangeTextureNode::MakeNodeUid(TextureName);
-		UInterchangeTextureNode* TextureNode = const_cast<UInterchangeTextureNode*>(Cast<UInterchangeTextureNode>(NodeContainer.GetNode(TextureUid)));
-		TextureNode->SetCustomSRGB(true);
+		if (UInterchangeTextureNode* TextureNode = const_cast<UInterchangeTextureNode*>(Cast<UInterchangeTextureNode>(NodeContainer.GetNode(TextureUid))))
+		{
+			TextureNode->SetCustomSRGB(true);
+		}
+	}
+}
+void UInterchangeGltfTranslator::SetTextureFlipGreenChannel(UInterchangeBaseNodeContainer& NodeContainer, const GLTF::FTextureMap& TextureMap) const
+{
+	if (GltfAsset.Textures.IsValidIndex(TextureMap.TextureIndex))
+	{
+		const FString TextureName = GltfAsset.Textures[TextureMap.TextureIndex].Name;
+		const FString TextureUid = UInterchangeTextureNode::MakeNodeUid(TextureName);
+		if (UInterchangeTextureNode* TextureNode = const_cast<UInterchangeTextureNode*>(Cast<UInterchangeTextureNode>(NodeContainer.GetNode(TextureUid))))
+		{
+			TextureNode->SetCustombFlipGreenChannel(true);
+		}
 	}
 }
\ No newline at end of file
diff --git a/Engine/Plugins/Experimental/Interchange/Runtime/Source/Import/Public/Gltf/InterchangeGltfTranslator.h b/Engine/Plugins/Experimental/Interchange/Runtime/Source/Import/Public/Gltf/InterchangeGltfTranslator.h
--- a/Engine/Plugins/Experimental/Interchange/Runtime/Source/Import/Public/Gltf/InterchangeGltfTranslator.h
+++ b/Engine/Plugins/Experimental/Interchange/Runtime/Source/Import/Public/Gltf/InterchangeGltfTranslator.h
@@ -73,6 +73,7 @@ protected:
 
 private:
 	void SetTextureSRGB(UInterchangeBaseNodeContainer& NodeContainer, const GLTF::FTextureMap& TextureMap) const;
+	void SetTextureFlipGreenChannel(UInterchangeBaseNodeContainer& NodeContainer, const GLTF::FTextureMap& TextureMap) const;
 
 	GLTF::FAsset GltfAsset;
 };
diff --git a/Engine/Plugins/Experimental/Interchange/Runtime/Source/Nodes/Public/InterchangeTextureNode.h b/Engine/Plugins/Experimental/Interchange/Runtime/Source/Nodes/Public/InterchangeTextureNode.h
--- a/Engine/Plugins/Experimental/Interchange/Runtime/Source/Nodes/Public/InterchangeTextureNode.h
+++ b/Engine/Plugins/Experimental/Interchange/Runtime/Source/Nodes/Public/InterchangeTextureNode.h
@@ -113,7 +113,20 @@ public:
 		IMPLEMENT_NODE_ATTRIBUTE_SETTER_NODELEGATE(SRGB, bool)
 	}
 
+	UFUNCTION(BlueprintCallable, Category = "Interchange | Node | Texture")
+	bool GetCustombFlipGreenChannel(bool& AttributeValue) const
+	{
+		IMPLEMENT_NODE_ATTRIBUTE_GETTER(bFlipGreenChannel, bool);
+	}
+
+	UFUNCTION(BlueprintCallable, Category = "Interchange | Node | Texture")
+	bool SetCustombFlipGreenChannel(const bool& AttributeValue)
+	{
+		IMPLEMENT_NODE_ATTRIBUTE_SETTER_NODELEGATE(bFlipGreenChannel, bool)
+	}
+
 private:
 
 	IMPLEMENT_NODE_ATTRIBUTE_KEY(SRGB)
+	IMPLEMENT_NODE_ATTRIBUTE_KEY(bFlipGreenChannel)
 };
diff --git a/Engine/Plugins/Experimental/Interchange/Runtime/Source/Pipelines/Private/InterchangeGenericMaterialPipeline.cpp b/Engine/Plugins/Experimental/Interchange/Runtime/Source/Pipelines/Private/InterchangeGenericMaterialPipeline.cpp
--- a/Engine/Plugins/Experimental/Interchange/Runtime/Source/Pipelines/Private/InterchangeGenericMaterialPipeline.cpp
+++ b/Engine/Plugins/Experimental/Interchange/Runtime/Source/Pipelines/Private/InterchangeGenericMaterialPipeline.cpp
@@ -1043,6 +1043,13 @@ void UInterchangeGenericMaterialPipeline::HandleTextureSampleNode(const UInterch
 				TextureFactoryNode->SetCustomSRGB(bSRGB);
 			}
 		}
+		if (bool bFlipGreenChannel; TextureNode->GetCustombFlipGreenChannel(bFlipGreenChannel))
+		{
+			if (UInterchangeTextureFactoryNode* TextureFactoryNode = Cast<UInterchangeTextureFactoryNode>(BaseNodeContainer->GetFactoryNode(TextureFactoryUid)))
+			{
+				TextureFactoryNode->SetCustombFlipGreenChannel(bFlipGreenChannel);
+			}
+		}
 	}
 	else
 	{
