[Interchange GLTF] Added support for Variant
Modified logic in the variant sets factory for material variations when there is more than one

diff --git a/Engine/Plugins/Experimental/Interchange/Runtime/Source/Import/Private/Gltf/InterchangeGltfTranslator.cpp b/Engine/Plugins/Experimental/Interchange/Runtime/Source/Import/Private/Gltf/InterchangeGltfTranslator.cpp
--- a/Engine/Plugins/Experimental/Interchange/Runtime/Source/Import/Private/Gltf/InterchangeGltfTranslator.cpp
+++ b/Engine/Plugins/Experimental/Interchange/Runtime/Source/Import/Private/Gltf/InterchangeGltfTranslator.cpp
@@ -16,6 +16,7 @@
 #include "InterchangeSceneNode.h"
 #include "InterchangeShaderGraphNode.h"
 #include "InterchangeTexture2DNode.h"
+#include "InterchangeVariantSetNode.h"
 
 #include "Sections/MovieScene3DTransformSection.h"
 #include "Texture/InterchangeImageWrapperTranslator.h"
@@ -118,9 +119,31 @@ namespace UE::Interchange::Gltf::Private
 				break;
 		}
 	}
+
+	bool CheckForVariants(const GLTF::FMesh& Mesh, int32 VariantCount, int32 MaterialCount)
+	{
+		for (const GLTF::FPrimitive& Primitive : Mesh.Primitives)
+		{
+			for (const GLTF::FVariantMapping& VariantMapping : Primitive.VariantMappings)
+			{
+				if (FMath::IsWithin(VariantMapping.MaterialIndex, 0, MaterialCount))
+				{
+					for (int32 VariantIndex : VariantMapping.VariantIndices)
+					{
+						if (FMath::IsWithin(VariantIndex, 0, VariantCount))
+						{
+							return true;
+						}
+					}
+				}
+			}
+		}
+
+		return false;
+	}
 }
 
-void UInterchangeGltfTranslator::HandleGltfNode( UInterchangeBaseNodeContainer& NodeContainer, const GLTF::FNode& GltfNode, const FString& ParentNodeUid, const int32 NodeIndex, FNodeUidMap& NodeUidMap ) const
+void UInterchangeGltfTranslator::HandleGltfNode( UInterchangeBaseNodeContainer& NodeContainer, const GLTF::FNode& GltfNode, const FString& ParentNodeUid, const int32 NodeIndex, bool &bHasVariants ) const
 {
 	using namespace UE::Interchange::Gltf::Private;
 
@@ -147,6 +170,11 @@ void UInterchangeGltfTranslator::HandleGltfNode( UInterchangeBaseNodeContainer&
 			{
 				const FString MeshNodeUid = TEXT("\\Mesh\\") + GltfAsset.Meshes[ GltfNode.MeshIndex ].Name;
 				InterchangeSceneNode->SetCustomAssetInstanceUid( MeshNodeUid );
+
+				if (!bHasVariants && GltfAsset.Variants.Num() > 0)
+				{
+					bHasVariants |= CheckForVariants(GltfAsset.Meshes[ GltfNode.MeshIndex ], GltfAsset.Variants.Num(), GltfAsset.Materials.Num());
+				}
 			}
 			break;
 		}
@@ -194,7 +222,7 @@ void UInterchangeGltfTranslator::HandleGltfNode( UInterchangeBaseNodeContainer&
 	{
 		if ( GltfAsset.Nodes.IsValidIndex( ChildIndex ) )
 		{
-			HandleGltfNode( NodeContainer, GltfAsset.Nodes[ ChildIndex ], NodeUid, ChildIndex, NodeUidMap );
+			HandleGltfNode( NodeContainer, GltfAsset.Nodes[ ChildIndex ], NodeUid, ChildIndex, bHasVariants );
 		}
 	}
 }
@@ -869,7 +897,7 @@ bool UInterchangeGltfTranslator::Translate( UInterchangeBaseNodeContainer& NodeC
 	}
 
 	// Cache created scene nodes UIDs to use later for animation binding
-	FNodeUidMap NodeUidMap;
+	bool bHasVariants = false;
 
 	// Scenes
 	{
@@ -892,7 +920,7 @@ bool UInterchangeGltfTranslator::Translate( UInterchangeBaseNodeContainer& NodeC
 			{
 				if ( GltfAsset.Nodes.IsValidIndex( NodeIndex ) )
 				{
-					HandleGltfNode( NodeContainer, GltfAsset.Nodes[ NodeIndex ], SceneNodeUid, NodeIndex, NodeUidMap );
+					HandleGltfNode( NodeContainer, GltfAsset.Nodes[ NodeIndex ], SceneNodeUid, NodeIndex, bHasVariants );
 				}
 			}
 		}
@@ -902,10 +930,16 @@ bool UInterchangeGltfTranslator::Translate( UInterchangeBaseNodeContainer& NodeC
 	{
 		for (int32 AnimationIndex = 0; AnimationIndex < GltfAsset.Animations.Num(); ++AnimationIndex)
 		{
-			HandleGltfAnimation(NodeContainer, NodeUidMap, AnimationIndex);
+			HandleGltfAnimation(NodeContainer, AnimationIndex);
 		}
 	}
 
+	// Create SceneVariantSetsNode if model contains variants.
+	if (bHasVariants)
+	{
+		HandleGltfVariants(NodeContainer, FileName);
+	}
+
 	return true;
 }
 
@@ -1077,7 +1111,7 @@ TFuture<TOptional<UE::Interchange::FAnimationBakeTransformPayloadData>> UInterch
 	return Promise->GetFuture();
 }
 
-void UInterchangeGltfTranslator::HandleGltfAnimation(UInterchangeBaseNodeContainer& NodeContainer, const FNodeUidMap& NodeUidMap, int32 AnimationIndex) const
+void UInterchangeGltfTranslator::HandleGltfAnimation(UInterchangeBaseNodeContainer& NodeContainer, int32 AnimationIndex) const
 {
 	const GLTF::FAnimation& GltfAnimation = GltfAsset.Animations[AnimationIndex];
 
@@ -1195,4 +1229,219 @@ void UInterchangeGltfTranslator::SetTextureFlipGreenChannel(UInterchangeBaseNode
 			TextureNode->SetCustombFlipGreenChannel(true);
 		}
 	}
-}
\ No newline at end of file
+}
+
+TFuture<TOptional<UE::Interchange::FVariantSetPayloadData>> UInterchangeGltfTranslator::GetVariantSetPayloadData(const FString& PayloadKey) const
+{
+	using namespace UE::Interchange;
+
+	TPromise<TOptional<FVariantSetPayloadData>> EmptyPromise;
+	EmptyPromise.SetValue(TOptional<FVariantSetPayloadData>());
+
+	TArray<FString> PayloadTokens;
+
+	// We need two indices to build the payload: index of LevelVariantSet and index of VariantSetIndex
+	if (GltfAsset.Variants.Num() + 1 != PayloadKey.ParseIntoArray(PayloadTokens, TEXT(";")))
+	{
+		// Invalid payload
+		return EmptyPromise.GetFuture();
+	}
+
+	//FString PayloadKey = FileName;
+	for (int32 Index = 0; Index < GltfAsset.Variants.Num(); ++Index)
+	{
+		if (PayloadTokens[Index + 1] != GltfAsset.Variants[Index])
+		{
+			// Invalid payload
+			return EmptyPromise.GetFuture();
+		}
+	}
+
+	return Async(EAsyncExecution::TaskGraph, [this]
+			{
+				FVariantSetPayloadData PayloadData;
+				TOptional<FVariantSetPayloadData> Result;
+
+				if (this->GetVariantSetPayloadData(PayloadData))
+				{
+					Result.Emplace(MoveTemp(PayloadData));
+				}
+
+				return Result;
+			}
+		);
+}
+
+void UInterchangeGltfTranslator::HandleGltfVariants(UInterchangeBaseNodeContainer& NodeContainer, const FString& FileName) const
+{
+	UInterchangeVariantSetNode* VariantSetNode = nullptr;
+	VariantSetNode = NewObject<UInterchangeVariantSetNode>(&NodeContainer);
+
+	FString VariantSetNodeUid = TEXT("\\VariantSet\\") + FileName;
+	VariantSetNode->InitializeNode(VariantSetNodeUid, FileName, EInterchangeNodeContainerType::TranslatedScene);
+	NodeContainer.AddNode(VariantSetNode);
+
+	VariantSetNode->SetCustomDisplayText(FileName);
+
+	FString PayloadKey = FileName;
+	for (int32 Index = 0; Index < GltfAsset.Variants.Num(); ++Index)
+	{
+		PayloadKey += TEXT(";") + GltfAsset.Variants[Index];
+	}
+	VariantSetNode->SetCustomVariantsPayloadKey(PayloadKey);
+
+	TFunction<void(const TArray<int32>& Nodes)> CollectDependencies;
+	CollectDependencies = [this, VariantSetNode, &CollectDependencies](const TArray<int32>& Nodes) -> void
+	{
+		const TArray<GLTF::FMaterial>& Materials = this->GltfAsset.Materials;
+
+		for (const int32 NodeIndex : Nodes)
+		{
+			if (this->GltfAsset.Nodes.IsValidIndex(NodeIndex))
+			{
+				const GLTF::FNode& GltfNode = this->GltfAsset.Nodes[NodeIndex];
+
+				if (GltfNode.Type == GLTF::FNode::EType::Mesh && this->GltfAsset.Meshes.IsValidIndex(GltfNode.MeshIndex))
+				{
+					const GLTF::FMesh& Mesh = this->GltfAsset.Meshes[GltfNode.MeshIndex];
+					const FString* NodeUidPtr = this->NodeUidMap.Find(&GltfNode);
+					ensure(NodeUidPtr);
+
+					VariantSetNode->AddCustomDependencyUid(*NodeUidPtr);
+
+					for (const GLTF::FPrimitive& Primitive : Mesh.Primitives)
+					{
+						if (Primitive.VariantMappings.Num() > 0)
+						{
+							for (const GLTF::FVariantMapping& VariantMapping : Primitive.VariantMappings)
+							{
+								if (!ensure(Materials.IsValidIndex(VariantMapping.MaterialIndex)))
+								{
+									continue;
+								}
+
+								const GLTF::FMaterial& GltfMaterial = Materials[VariantMapping.MaterialIndex];
+								const FString MaterialUid = UInterchangeShaderGraphNode::MakeNodeUid(TEXT("Material_") + GltfMaterial.Name);
+
+								VariantSetNode->AddCustomDependencyUid(MaterialUid);
+
+							}
+						}
+					}
+				}
+
+				if (GltfNode.Children.Num() > 0)
+				{
+					CollectDependencies(GltfNode.Children);
+				}
+			}
+		}
+	};
+
+	for (const GLTF::FScene& GltfScene : GltfAsset.Scenes)
+	{
+		CollectDependencies(GltfScene.Nodes);
+	}
+
+	UInterchangeSceneVariantSetsNode* SceneVariantSetsNode = NewObject<UInterchangeSceneVariantSetsNode>(&NodeContainer);
+
+	FString SceneVariantSetsNodeUid = TEXT("\\SceneVariantSets\\") + FileName;
+	SceneVariantSetsNode->InitializeNode(SceneVariantSetsNodeUid, FileName, EInterchangeNodeContainerType::TranslatedScene);
+	NodeContainer.AddNode(SceneVariantSetsNode);
+
+	SceneVariantSetsNode->AddCustomVariantSetUid(VariantSetNodeUid);
+}
+
+bool UInterchangeGltfTranslator::GetVariantSetPayloadData(UE::Interchange::FVariantSetPayloadData& PayloadData) const
+{
+	using namespace UE;
+
+	PayloadData.Variants.SetNum(GltfAsset.Variants.Num());
+
+	TMap<const TCHAR*, Interchange::FVariant*> VariantMap;
+	VariantMap.Reserve(GltfAsset.Variants.Num());
+
+	for (int32 VariantIndex = 0; VariantIndex < GltfAsset.Variants.Num(); ++VariantIndex)
+	{
+		const FString& VariantName = GltfAsset.Variants[VariantIndex];
+		PayloadData.Variants[VariantIndex].DisplayText = VariantName;
+		VariantMap.Add(*VariantName, &PayloadData.Variants[VariantIndex]);
+	}
+
+	TFunction<void(const TArray<int32>& Nodes)> BuildPayloadData;
+	BuildPayloadData = [this, VariantMap = MoveTemp(VariantMap), &BuildPayloadData](const TArray<int32>& Nodes) -> void
+	{
+		const TArray<FString>& VariantNames = this->GltfAsset.Variants;
+		const TArray<GLTF::FMaterial>& Materials = this->GltfAsset.Materials;
+
+		for (const int32 NodeIndex : Nodes)
+		{
+			if (!ensure(this->GltfAsset.Nodes.IsValidIndex(NodeIndex)))
+			{
+				continue;
+			}
+
+			const GLTF::FNode& GltfNode = this->GltfAsset.Nodes[NodeIndex];
+
+			if (GltfNode.Type == GLTF::FNode::EType::Mesh && this->GltfAsset.Meshes.IsValidIndex(GltfNode.MeshIndex))
+			{
+				const GLTF::FMesh& Mesh = this->GltfAsset.Meshes[GltfNode.MeshIndex];
+				const FString* NodeUidPtr = this->NodeUidMap.Find(&GltfNode);
+				if (!ensure(NodeUidPtr))
+				{
+					continue;
+				}
+
+				for (const GLTF::FPrimitive& Primitive : Mesh.Primitives)
+				{
+					for (const GLTF::FVariantMapping& VariantMapping : Primitive.VariantMappings)
+					{
+						if (!ensure(Materials.IsValidIndex(VariantMapping.MaterialIndex)))
+						{
+							continue;
+						}
+
+						const GLTF::FMaterial& GltfMaterial = Materials[VariantMapping.MaterialIndex];
+						const FString MaterialUid = UInterchangeShaderGraphNode::MakeNodeUid(TEXT("Material_") + GltfMaterial.Name);
+
+						for (int32 VariantIndex : VariantMapping.VariantIndices)
+						{
+							if (!ensure(VariantMap.Contains(*VariantNames[VariantIndex])))
+							{
+								continue;
+							}
+
+							// This is on par with the Datasmith GLTF translator but might be wrong.
+							// Each primitive should be a section of the static mesh and 
+							// TODO: Revisit creation of static mesh and handling of variants: UE-159945.
+							Interchange::FVariantPropertyCaptureData PropertyCaptureData;
+
+							PropertyCaptureData.Category = Interchange::EVariantPropertyCaptureCategory::Material;
+							PropertyCaptureData.ObjectUid = MaterialUid;
+
+							Interchange::FVariant& VariantData = *(VariantMap[*VariantNames[VariantIndex]]);
+
+							Interchange::FVariantBinding& Binding = VariantData.Bindings.AddDefaulted_GetRef();
+
+							Binding.TargetUid = *NodeUidPtr;
+							Binding.Captures.Add(PropertyCaptureData);
+						}
+
+					}
+				}
+			}
+
+			if (GltfNode.Children.Num() > 0)
+			{
+				BuildPayloadData(GltfNode.Children);
+			}
+		}
+	};
+
+	for (const GLTF::FScene& GltfScene : GltfAsset.Scenes)
+	{
+		BuildPayloadData(GltfScene.Nodes);
+	}
+
+	return true;
+}
diff --git a/Engine/Plugins/Experimental/Interchange/Runtime/Source/Import/Private/Scene/InterchangeSceneVariantSetsFactory.cpp b/Engine/Plugins/Experimental/Interchange/Runtime/Source/Import/Private/Scene/InterchangeSceneVariantSetsFactory.cpp
--- a/Engine/Plugins/Experimental/Interchange/Runtime/Source/Import/Private/Scene/InterchangeSceneVariantSetsFactory.cpp
+++ b/Engine/Plugins/Experimental/Interchange/Runtime/Source/Import/Private/Scene/InterchangeSceneVariantSetsFactory.cpp
@@ -218,7 +218,29 @@ namespace UE::Interchange::Private
 				{
 					if (const UMaterialInterface* TargetMaterial = Cast<UMaterialInterface>(TargetFactoryNode->ReferenceObject.TryLoad()))
 					{
-						TArray<UPropertyValue*> PropertyValues = VariantManager->CreateMaterialPropertyCaptures({ &ObjectBinding });
+						// FVariantManager::CreateMaterialPropertyCaptures returns an array of the newly found properties
+						// Therefore, if a mesh actor has more than 1 material and a binding has more than 1 material variant
+						// The second call to FVariantManager::CreateMaterialPropertyCaptures will return an empty array.
+						// Consequently, we have to check if material properties has not been extracted yet
+						TArray<UPropertyValue*> PropertyValues;
+						{
+							const TArray<UPropertyValue*>& ExistingPropertyValues = ObjectBinding.GetCapturedProperties();
+
+							for (const UPropertyValue* PropertyValue : ExistingPropertyValues)
+							{
+								if (EnumHasAnyFlags(PropertyValue->GetPropCategory(), EPropertyValueCategory::Material))
+								{
+									PropertyValues.Add(const_cast<UPropertyValue*>(PropertyValue));
+								}
+							}
+
+							// Material properties have not been asked for yet.
+							if (PropertyValues.Num() == 0)
+							{
+								PropertyValues = VariantManager->CreateMaterialPropertyCaptures({ &ObjectBinding });
+							}
+						}
+
 						if (PropertyValues.Num() == 1)
 						{
 							PropertyValues[0]->SetRecordedData((uint8*)&TargetMaterial, sizeof(UMaterialInterface*));
diff --git a/Engine/Plugins/Experimental/Interchange/Runtime/Source/Import/Public/Gltf/InterchangeGltfTranslator.h b/Engine/Plugins/Experimental/Interchange/Runtime/Source/Import/Public/Gltf/InterchangeGltfTranslator.h
--- a/Engine/Plugins/Experimental/Interchange/Runtime/Source/Import/Public/Gltf/InterchangeGltfTranslator.h
+++ b/Engine/Plugins/Experimental/Interchange/Runtime/Source/Import/Public/Gltf/InterchangeGltfTranslator.h
@@ -7,6 +7,7 @@
 #include "Mesh/InterchangeStaticMeshPayload.h"
 #include "Mesh/InterchangeStaticMeshPayloadInterface.h"
 #include "Nodes/InterchangeBaseNodeContainer.h"
+#include "Scene/InterchangeVariantSetPayloadInterface.h"
 #include "Texture/InterchangeTexturePayloadInterface.h"
 #include "UObject/Object.h"
 #include "UObject/ObjectMacros.h"
@@ -15,14 +16,16 @@
 
 class UInterchangeShaderGraphNode;
 class UInterchangeShaderNode;
+class UInterchangeVariantSetNode;
 
 /* Gltf translator class support import of texture, material, static mesh, skeletal mesh, */
 
 UCLASS(BlueprintType, Experimental)
-class UInterchangeGltfTranslator : public UInterchangeTranslatorBase,
-	public IInterchangeStaticMeshPayloadInterface, 
-	public IInterchangeTexturePayloadInterface,
-	public IInterchangeAnimationPayloadInterface
+class UInterchangeGltfTranslator : public UInterchangeTranslatorBase
+	, public IInterchangeStaticMeshPayloadInterface
+	, public IInterchangeTexturePayloadInterface
+	, public IInterchangeAnimationPayloadInterface
+	, public IInterchangeVariantSetPayloadInterface
 {
 	GENERATED_BODY()
 
@@ -52,14 +55,19 @@ public:
 	virtual TFuture<TOptional<UE::Interchange::FAnimationBakeTransformPayloadData>> GetAnimationBakeTransformPayloadData(const FString& PayLoadKey, const double BakeFrequency, const double RangeStartSecond, const double RangeStopSecond) const override;
 	/* IInterchangeAnimationPayloadInterface End */
 
+	/* IInterchangeVariantSetPayloadInterface Begin */
+	virtual TFuture<TOptional<UE::Interchange::FVariantSetPayloadData>> GetVariantSetPayloadData(const FString& PayloadKey) const override;
+	/* IInterchangeVariantSetPayloadInterface End */
+
 protected:
 	using FNodeUidMap = TMap<const GLTF::FNode*, FString>;
 
-	void HandleGltfNode( UInterchangeBaseNodeContainer& NodeContainer, const GLTF::FNode& GltfNode, const FString& ParentNodeUid, const int32 NodeIndex, FNodeUidMap& NodeUidMap ) const;
+	void HandleGltfNode( UInterchangeBaseNodeContainer& NodeContainer, const GLTF::FNode& GltfNode, const FString& ParentNodeUid, const int32 NodeIndex, bool &bHasVariants ) const;
 	void HandleGltfMaterial( UInterchangeBaseNodeContainer& NodeContainer, const GLTF::FMaterial& GltfMaterial, UInterchangeShaderGraphNode& ShaderGraphNode ) const;
 	void HandleGltfMaterialParameter( UInterchangeBaseNodeContainer& NodeContainer, const GLTF::FTextureMap& TextureMap, UInterchangeShaderNode& ShaderNode,
 		const FString& MapName, const TVariant< FLinearColor, float >& MapFactor, const FString& OutputChannel, const bool bInverse = false, const bool bIsNormal = false ) const;
-	void HandleGltfAnimation(UInterchangeBaseNodeContainer& NodeContainer, const FNodeUidMap& NodeUidMap, int32 AnimationIndex) const;
+	void HandleGltfAnimation(UInterchangeBaseNodeContainer& NodeContainer, int32 AnimationIndex) const;
+	void HandleGltfVariants(UInterchangeBaseNodeContainer& NodeContainer, const FString& FileName) const;
 
 	/** Support for KHR_materials_clearcoat */
 	void HandleGltfClearCoat( UInterchangeBaseNodeContainer& NodeContainer, const GLTF::FMaterial& GltfMaterial, UInterchangeShaderGraphNode& ShaderGraphNode ) const;
@@ -71,11 +79,13 @@ protected:
 	/** Support for KHR_texture_transform */
 	void HandleGltfTextureTransform( UInterchangeBaseNodeContainer& NodeContainer, const GLTF::FTextureTransform& TextureTransform, const int32 TexCoordIndex, UInterchangeShaderNode& ShaderNode ) const;
 
+	bool GetVariantSetPayloadData(UE::Interchange::FVariantSetPayloadData& PayloadData) const;
 private:
 	void SetTextureSRGB(UInterchangeBaseNodeContainer& NodeContainer, const GLTF::FTextureMap& TextureMap) const;
 	void SetTextureFlipGreenChannel(UInterchangeBaseNodeContainer& NodeContainer, const GLTF::FTextureMap& TextureMap) const;
 
 	GLTF::FAsset GltfAsset;
+	mutable FNodeUidMap NodeUidMap;
 };
 
 
diff --git a/Engine/Plugins/Experimental/Interchange/Runtime/Source/Pipelines/Private/InterchangeGenericScenesPipeline.cpp b/Engine/Plugins/Experimental/Interchange/Runtime/Source/Pipelines/Private/InterchangeGenericScenesPipeline.cpp
--- a/Engine/Plugins/Experimental/Interchange/Runtime/Source/Pipelines/Private/InterchangeGenericScenesPipeline.cpp
+++ b/Engine/Plugins/Experimental/Interchange/Runtime/Source/Pipelines/Private/InterchangeGenericScenesPipeline.cpp
@@ -371,12 +371,15 @@ void UInterchangeGenericLevelPipeline::ExecuteSceneVariantSetNodePreImport(const
 		return;
 	}
 
+	// We may eventually want to optionally import variants
+	static bool bEnableSceneVariantSet = true;
+
 	const FString FactoryNodeUid = UInterchangeFactoryBaseNode::BuildFactoryNodeUid(SceneVariantSetNode.GetUniqueID());
 
 	UInterchangeSceneVariantSetsFactoryNode* FactoryNode = NewObject<UInterchangeSceneVariantSetsFactoryNode>(BaseNodeContainer, NAME_None);
 
 	FactoryNode->InitializeNode(FactoryNodeUid, SceneVariantSetNode.GetDisplayLabel(), EInterchangeNodeContainerType::FactoryData);
-	FactoryNode->SetEnabled(true);
+	FactoryNode->SetEnabled(bEnableSceneVariantSet);
 
 	TArray<FString> VariantSetUids;
 	SceneVariantSetNode.GetCustomVariantSetUids(VariantSetUids);
@@ -388,14 +391,22 @@ void UInterchangeGenericLevelPipeline::ExecuteSceneVariantSetNodePreImport(const
 		// Update factory's dependencies
 		if (const UInterchangeVariantSetNode* TrackNode = Cast<UInterchangeVariantSetNode>(BaseNodeContainer->GetNode(VariantSetUid)))
 		{
-			TArray<FString> ActorNodeUids;
-			TrackNode->GetCustomDependencyUids(ActorNodeUids);
+			TArray<FString> DependencyNodeUids;
+			TrackNode->GetCustomDependencyUids(DependencyNodeUids);
 
-			for (const FString& ActorNodeUid : ActorNodeUids)
+			for (const FString& DependencyNodeUid : DependencyNodeUids)
 			{
 				
-				const FString ActorFactoryNodeUid = UInterchangeFactoryBaseNode::BuildFactoryNodeUid(ActorNodeUid);
-				FactoryNode->AddFactoryDependencyUid(ActorFactoryNodeUid);
+				const FString DependencyFactoryNodeUid = UInterchangeFactoryBaseNode::BuildFactoryNodeUid(DependencyNodeUid);
+				FactoryNode->AddFactoryDependencyUid(DependencyFactoryNodeUid);
+
+				if (UInterchangeFactoryBaseNode* DependencyFactoryNode = BaseNodeContainer->GetFactoryNode(DependencyFactoryNodeUid))
+				{
+					if (bEnableSceneVariantSet && !DependencyFactoryNode->IsEnabled())
+					{
+						DependencyFactoryNode->SetEnabled(true);
+					}
+				}
 			}
 		}
 	}
