Interchange - Cleanup for SRGB and GreenChannelFlipping setting for TextureFactoryNode and expanding comment regarding normal Y channel flip in GLTF translator

diff --git a/Engine/Plugins/Experimental/Interchange/Runtime/Source/Import/Private/Gltf/InterchangeGltfTranslator.cpp b/Engine/Plugins/Experimental/Interchange/Runtime/Source/Import/Private/Gltf/InterchangeGltfTranslator.cpp
--- a/Engine/Plugins/Experimental/Interchange/Runtime/Source/Import/Private/Gltf/InterchangeGltfTranslator.cpp
+++ b/Engine/Plugins/Experimental/Interchange/Runtime/Source/Import/Private/Gltf/InterchangeGltfTranslator.cpp
@@ -351,7 +351,9 @@ void UInterchangeGltfTranslator::HandleGltfMaterial( UInterchangeBaseNodeContain
 	//Based on the gltf specification the basecolor and emissive textures have SRGB colors:
 	SetTextureSRGB(NodeContainer, GltfMaterial.BaseColor);
 	SetTextureSRGB(NodeContainer, GltfMaterial.Emissive);
-	//According to GLTF documentation the normal maps should have their green channel flipped:
+	//According to GLTF documentation the normal maps are right handed (following OpenGL convention),
+	//however UE expects left handed normal maps, this can be resolved by flipping the green channel of the normal textures:
+	//(based on [REDACTED_PROFILE_URL]/glTF-Sample-Models/tree/master/2.0/NormalTangentTest#problem-flipped-y-axis-or-flipped-green-channel)
 	SetTextureFlipGreenChannel(NodeContainer, GltfMaterial.Normal);
 	SetTextureFlipGreenChannel(NodeContainer, GltfMaterial.ClearCoat.NormalMap);
 
diff --git a/Engine/Plugins/Experimental/Interchange/Runtime/Source/Pipelines/Private/InterchangeGenericMaterialPipeline.cpp b/Engine/Plugins/Experimental/Interchange/Runtime/Source/Pipelines/Private/InterchangeGenericMaterialPipeline.cpp
--- a/Engine/Plugins/Experimental/Interchange/Runtime/Source/Pipelines/Private/InterchangeGenericMaterialPipeline.cpp
+++ b/Engine/Plugins/Experimental/Interchange/Runtime/Source/Pipelines/Private/InterchangeGenericMaterialPipeline.cpp
@@ -1032,22 +1032,12 @@ void UInterchangeGenericMaterialPipeline::HandleTextureSampleNode(const UInterch
 		{
 			if (UInterchangeTextureFactoryNode* TextureFactoryNode = Cast<UInterchangeTextureFactoryNode>(BaseNodeContainer->GetFactoryNode(TextureFactoryUid)))
 			{
-				TextureFactoryNode->SetCustomSRGB(false);
-			}
-		}
-
-		if(bool bSRGB; TextureNode->GetCustomSRGB(bSRGB))
-		{
-			if(UInterchangeTextureFactoryNode* TextureFactoryNode = Cast<UInterchangeTextureFactoryNode>(BaseNodeContainer->GetFactoryNode(TextureFactoryUid)))
-			{
-				TextureFactoryNode->SetCustomSRGB(bSRGB);
-			}
-		}
-		if (bool bFlipGreenChannel; TextureNode->GetCustombFlipGreenChannel(bFlipGreenChannel))
-		{
-			if (UInterchangeTextureFactoryNode* TextureFactoryNode = Cast<UInterchangeTextureFactoryNode>(BaseNodeContainer->GetFactoryNode(TextureFactoryUid)))
-			{
-				TextureFactoryNode->SetCustombFlipGreenChannel(bFlipGreenChannel);
+				bool bSRGB;
+				if (!TextureNode->GetCustomSRGB(bSRGB))
+				{
+					//Only set CustomSRGB if it wasn't set by the InterchangeGenericTexturePipeline before
+					TextureFactoryNode->SetCustomSRGB(false);
+				}
 			}
 		}
 	}
diff --git a/Engine/Plugins/Experimental/Interchange/Runtime/Source/Pipelines/Private/InterchangeGenericTexturePipeline.cpp b/Engine/Plugins/Experimental/Interchange/Runtime/Source/Pipelines/Private/InterchangeGenericTexturePipeline.cpp
--- a/Engine/Plugins/Experimental/Interchange/Runtime/Source/Pipelines/Private/InterchangeGenericTexturePipeline.cpp
+++ b/Engine/Plugins/Experimental/Interchange/Runtime/Source/Pipelines/Private/InterchangeGenericTexturePipeline.cpp
@@ -271,6 +271,15 @@ UInterchangeTextureFactoryNode* UInterchangeGenericTexturePipeline::CreateTextur
 		}
 	}
 
+	if (bool bSRGB; TextureNode->GetCustomSRGB(bSRGB))
+	{
+		TextureFactoryNode->SetCustomSRGB(bSRGB);
+	}
+	if (bool bFlipGreenChannel; TextureNode->GetCustombFlipGreenChannel(bFlipGreenChannel))
+	{
+		TextureFactoryNode->SetCustombFlipGreenChannel(bFlipGreenChannel);
+	}
+
 #if WITH_EDITORONLY_DATA
 	if (TextureFactoryNode)
 	{
