Swapping ClearCoatNormal and Normal maps for Interchange GLTF importer.
So that ClearCoatNormal becomes "Top" normal map and default Normal becomes the "bottom" normal map.
(+also updated the default input value for Normals to be 0,0,1)

diff --git a/Engine/Plugins/Experimental/Interchange/Runtime/Source/Import/Private/Gltf/InterchangeGltfTranslator.cpp b/Engine/Plugins/Experimental/Interchange/Runtime/Source/Import/Private/Gltf/InterchangeGltfTranslator.cpp
--- a/Engine/Plugins/Experimental/Interchange/Runtime/Source/Import/Private/Gltf/InterchangeGltfTranslator.cpp
+++ b/Engine/Plugins/Experimental/Interchange/Runtime/Source/Import/Private/Gltf/InterchangeGltfTranslator.cpp
@@ -331,13 +331,21 @@ void UInterchangeGltfTranslator::HandleGltfMaterialParameter( UInterchangeBaseNo
 	}
 	else
 	{
-		if ( MapFactor.IsType< FLinearColor >() )
+		if ( bIsNormal && !bTextureHasImportance )
 		{
-			NodeToConnectTo->AddLinearColorAttribute( UInterchangeShaderPortsAPI::MakeInputValueKey( InputToConnectTo ), MapFactor.Get< FLinearColor >() );
+			//default normal value is 0,0,1 (blue)
+			NodeToConnectTo->AddLinearColorAttribute( UInterchangeShaderPortsAPI::MakeInputValueKey(InputToConnectTo), FLinearColor::Blue );
 		}
-		else if ( MapFactor.IsType< float >() )
+		else
 		{
-			NodeToConnectTo->AddFloatAttribute( UInterchangeShaderPortsAPI::MakeInputValueKey( InputToConnectTo ), MapFactor.Get< float >() );
+			if ( MapFactor.IsType< FLinearColor >() )
+			{
+				NodeToConnectTo->AddLinearColorAttribute( UInterchangeShaderPortsAPI::MakeInputValueKey(InputToConnectTo), MapFactor.Get< FLinearColor >() );
+			}
+			else if ( MapFactor.IsType< float >() )
+			{
+				NodeToConnectTo->AddFloatAttribute( UInterchangeShaderPortsAPI::MakeInputValueKey(InputToConnectTo), MapFactor.Get< float >() );
+			}
 		}
 	}
 }
@@ -356,6 +364,11 @@ void UInterchangeGltfTranslator::HandleGltfMaterial( UInterchangeBaseNodeContain
 	//(based on [REDACTED_PROFILE_URL]/glTF-Sample-Models/tree/master/2.0/NormalTangentTest#problem-flipped-y-axis-or-flipped-green-channel)
 	SetTextureFlipGreenChannel(NodeContainer, GltfMaterial.Normal);
 	SetTextureFlipGreenChannel(NodeContainer, GltfMaterial.ClearCoat.NormalMap);
+	
+	//if there is a clearcoatnormal map then we want to swap it with the normals map
+	//as the Interchange pipeline will connect the clearcoatnormal map to ClearCoatBottomNormalMap
+	//however as per specification the gltf.clearcoatnormal.map should be the top clearcoat and the gltf.normal.map should be the bottom one.
+	bool bSwapNormalAndClearCoatNormal = GltfMaterial.bHasClearCoat;
 
 	if ( GltfMaterial.ShadingModel == GLTF::FMaterial::EShadingModel::MetallicRoughness )
 	{
@@ -438,7 +451,7 @@ void UInterchangeGltfTranslator::HandleGltfMaterial( UInterchangeBaseNodeContain
 			const bool bInverse = false;
 			const bool bIsNormal = true;
 
-			HandleGltfMaterialParameter( NodeContainer, GltfMaterial.Normal, ShaderGraphNode, Common::Parameters::Normal.ToString(),
+			HandleGltfMaterialParameter( NodeContainer, GltfMaterial.Normal, ShaderGraphNode, bSwapNormalAndClearCoatNormal ? ClearCoat::Parameters::ClearCoatNormal.ToString() : Common::Parameters::Normal.ToString(),
 				NormalFactor, Standard::Nodes::TextureSample::Outputs::RGB.ToString(), bInverse, bIsNormal );
 		}
 
@@ -487,7 +500,7 @@ void UInterchangeGltfTranslator::HandleGltfMaterial( UInterchangeBaseNodeContain
 
 	if ( GltfMaterial.bHasClearCoat )
 	{
-		HandleGltfClearCoat( NodeContainer, GltfMaterial, ShaderGraphNode );
+		HandleGltfClearCoat( NodeContainer, GltfMaterial, ShaderGraphNode, bSwapNormalAndClearCoatNormal );
 	}
 
 	if ( GltfMaterial.bHasSheen )
@@ -501,7 +514,7 @@ void UInterchangeGltfTranslator::HandleGltfMaterial( UInterchangeBaseNodeContain
 	}
 }
 
-void UInterchangeGltfTranslator::HandleGltfClearCoat( UInterchangeBaseNodeContainer& NodeContainer, const GLTF::FMaterial& GltfMaterial, UInterchangeShaderGraphNode& ShaderGraphNode ) const
+void UInterchangeGltfTranslator::HandleGltfClearCoat( UInterchangeBaseNodeContainer& NodeContainer, const GLTF::FMaterial& GltfMaterial, UInterchangeShaderGraphNode& ShaderGraphNode, const bool bSwapNormalAndClearCoatNormal) const
 {
 	using namespace UE::Interchange::Materials;
 
@@ -536,7 +549,7 @@ void UInterchangeGltfTranslator::HandleGltfClearCoat( UInterchangeBaseNodeContai
 		const bool bInverse = false;
 		const bool bIsNormal = true;
 
-		HandleGltfMaterialParameter( NodeContainer, GltfMaterial.ClearCoat.NormalMap, ShaderGraphNode, ClearCoat::Parameters::ClearCoatNormal.ToString(),
+		HandleGltfMaterialParameter( NodeContainer, GltfMaterial.ClearCoat.NormalMap, ShaderGraphNode, bSwapNormalAndClearCoatNormal ? Common::Parameters::Normal.ToString() : ClearCoat::Parameters::ClearCoatNormal.ToString(),
 			ClearCoatNormalFactor, Standard::Nodes::TextureSample::Outputs::RGB.ToString(), bInverse, bIsNormal );
 	}
 }
@@ -965,7 +978,7 @@ TFuture< TOptional< UE::Interchange::FStaticMeshPayloadData > > UInterchangeGltf
 	const GLTF::FMesh& GltfMesh = GltfAsset.Meshes[ MeshIndex ];
 	GLTF::FMeshFactory MeshFactory;
 	MeshFactory.SetUniformScale( 100.f ); // GLTF is in meters while UE is in centimeters
-	MeshFactory.FillMeshDescription( GltfMesh, &StaticMeshPayloadData.MeshDescription, true );
+	MeshFactory.FillMeshDescription( GltfMesh, &StaticMeshPayloadData.MeshDescription );
 
 	// Patch polygon groups material slot names to match Interchange expectations (rename material slots from indices to material names)
 	{
diff --git a/Engine/Plugins/Experimental/Interchange/Runtime/Source/Import/Public/Gltf/InterchangeGltfTranslator.h b/Engine/Plugins/Experimental/Interchange/Runtime/Source/Import/Public/Gltf/InterchangeGltfTranslator.h
--- a/Engine/Plugins/Experimental/Interchange/Runtime/Source/Import/Public/Gltf/InterchangeGltfTranslator.h
+++ b/Engine/Plugins/Experimental/Interchange/Runtime/Source/Import/Public/Gltf/InterchangeGltfTranslator.h
@@ -70,7 +70,7 @@ protected:
 	void HandleGltfVariants(UInterchangeBaseNodeContainer& NodeContainer, const FString& FileName) const;
 
 	/** Support for KHR_materials_clearcoat */
-	void HandleGltfClearCoat( UInterchangeBaseNodeContainer& NodeContainer, const GLTF::FMaterial& GltfMaterial, UInterchangeShaderGraphNode& ShaderGraphNode ) const;
+	void HandleGltfClearCoat( UInterchangeBaseNodeContainer& NodeContainer, const GLTF::FMaterial& GltfMaterial, UInterchangeShaderGraphNode& ShaderGraphNode, const bool bSwapNormalAndClearCoatNormal ) const;
 	/** Support for KHR_materials_sheen */
 	void HandleGltfSheen( UInterchangeBaseNodeContainer& NodeContainer, const GLTF::FMaterial& GltfMaterial, UInterchangeShaderGraphNode& ShaderGraphNode ) const;
 	/** Support for KHR_materials_transmission */
