Fixed static analysis issue reported by Horde

diff --git a/Engine/Plugins/Experimental/Interchange/Runtime/Source/Import/Private/Gltf/InterchangeGltfTranslator.cpp b/Engine/Plugins/Experimental/Interchange/Runtime/Source/Import/Private/Gltf/InterchangeGltfTranslator.cpp
--- a/Engine/Plugins/Experimental/Interchange/Runtime/Source/Import/Private/Gltf/InterchangeGltfTranslator.cpp
+++ b/Engine/Plugins/Experimental/Interchange/Runtime/Source/Import/Private/Gltf/InterchangeGltfTranslator.cpp
@@ -1320,7 +1320,10 @@ void UInterchangeGltfTranslator::HandleGltfVariants(UInterchangeBaseNodeContaine
 				{
 					const GLTF::FMesh& Mesh = this->GltfAsset.Meshes[GltfNode.MeshIndex];
 					const FString* NodeUidPtr = this->NodeUidMap.Find(&GltfNode);
-					ensure(NodeUidPtr);
+					if (!ensure(NodeUidPtr))
+					{
+						continue;
+					}
 
 					VariantSetNode->AddCustomDependencyUid(*NodeUidPtr);
 
