Fix scaling issue when re-importing.
Add the legacy skeleton rematch in case the skeleton is incompatible
#rnx

diff --git a/Engine/Plugins/Experimental/Interchange/Runtime/Source/Import/Private/Mesh/InterchangeSkeletalMeshFactory.cpp b/Engine/Plugins/Experimental/Interchange/Runtime/Source/Import/Private/Mesh/InterchangeSkeletalMeshFactory.cpp
--- a/Engine/Plugins/Experimental/Interchange/Runtime/Source/Import/Private/Mesh/InterchangeSkeletalMeshFactory.cpp
+++ b/Engine/Plugins/Experimental/Interchange/Runtime/Source/Import/Private/Mesh/InterchangeSkeletalMeshFactory.cpp
@@ -1,9 +1,13 @@
 // Copyright Epic Games, Inc. All Rights Reserved. 
 #include "Mesh/InterchangeSkeletalMeshFactory.h"
 
+#include "AssetRegistry/AssetRegistryModule.h"
+#include "Async/Async.h"
 #include "Async/ParallelFor.h"
 #include "Components.h"
+#include "CoreGlobals.h"
 #include "Engine/SkeletalMesh.h"
+#include "GenericPlatform/GenericPlatformMisc.h"
 #include "GPUSkinPublicDefs.h"
 #include "InterchangeAssetImportData.h"
 #include "InterchangeCommonPipelineDataFactoryNode.h"
@@ -22,6 +26,7 @@
 #include "Math/GenericOctree.h"
 #include "Mesh/InterchangeSkeletalMeshPayload.h"
 #include "Mesh/InterchangeSkeletalMeshPayloadInterface.h"
+#include "Misc/MessageDialog.h"
 #include "Nodes/InterchangeBaseNode.h"
 #include "Nodes/InterchangeBaseNodeContainer.h"
 #include "PhysicsEngine/PhysicsAsset.h"
@@ -991,8 +996,6 @@ UObject* UInterchangeSkeletalMeshFactory::CreateAsset(const FCreateAssetParams&
 	}
 	else
 	{
-		SkeletalMesh->GetImportedBounds().BoxExtent.Set(0.0, 0.0, 0.0);
-
 		//When we re-import, we force the current skeletalmesh skeleton, to be specified and to be the reference
 		FSoftObjectPath SpecifiedSkeleton = SkeletalMesh->GetSkeleton();
 		SkeletalMeshFactoryNode->SetCustomSkeletonSoftObjectPath(SpecifiedSkeleton);
@@ -1354,17 +1357,17 @@ UObject* UInterchangeSkeletalMeshFactory::CreateAsset(const FCreateAssetParams&
 			}
 		}
 
-		//Add the bound to the skeletal mesh
-		if (SkeletalMesh->GetImportedBounds().BoxExtent.IsNearlyZero())
+		//Update the bounding box if we are importing the LOD 0
+		if(CurrentLodIndex == 0)
 		{
 			FBox3f BoundingBox(SkeletalMeshImportData.Points.GetData(), SkeletalMeshImportData.Points.Num());
 			const FVector3f BoundingBoxSize = BoundingBox.GetSize();
-
-			if (SkeletalMeshImportData.Points.Num() > 2 && BoundingBoxSize.X < THRESH_POINTS_ARE_SAME && BoundingBoxSize.Y < THRESH_POINTS_ARE_SAME && BoundingBoxSize.Z < THRESH_POINTS_ARE_SAME)
+			if (SkeletalMeshImportData.Points.Num() > 2 && BoundingBoxSize.X < UE_THRESH_POINTS_ARE_SAME && BoundingBoxSize.Y < UE_THRESH_POINTS_ARE_SAME && BoundingBoxSize.Z < UE_THRESH_POINTS_ARE_SAME)
 			{
 				//TODO log a user error
 				//AddTokenizedErrorMessage(FTokenizedMessage::Create(EMessageSeverity::Error, FText::Format(LOCTEXT("FbxSkeletaLMeshimport_ErrorMeshTooSmall", "Cannot import this mesh, the bounding box of this mesh is smaller than the supported threshold[{0}]."), FText::FromString(FString::Printf(TEXT("%f"), THRESH_POINTS_ARE_SAME)))), FFbxErrors::SkeletalMesh_FillImportDataFailed);
 			}
+			FBoxSphereBounds BoxSphereBound((FBox)BoundingBox);
 			SkeletalMesh->SetImportedBounds(FBoxSphereBounds((FBox)BoundingBox));
 		}
 
@@ -1373,7 +1376,93 @@ UObject* UInterchangeSkeletalMeshFactory::CreateAsset(const FCreateAssetParams&
 
 	if(SkeletonReference)
 	{
-		SkeletonReference->MergeAllBonesToBoneTree(SkeletalMesh);
+		if ((!bApplySkinningOnly || !bIsReImport) && !SkeletonReference->MergeAllBonesToBoneTree(SkeletalMesh))
+		{
+			TUniqueFunction<bool()> RecreateSkeleton = [this, WeakSkeletalMesh = TWeakObjectPtr<USkeletalMesh>(SkeletalMesh), WeakSkeleton = TWeakObjectPtr<USkeleton>(SkeletonReference)]()
+			{
+
+				check(IsInGameThread());
+
+				USkeleton* SkeletonPtr = WeakSkeleton.Get();
+				USkeletalMesh* SkeletalMeshPtr = WeakSkeletalMesh.Get();
+
+				if (!SkeletonPtr || !SkeletalMeshPtr)
+				{
+					return false;
+				}
+
+				if (GIsRunningUnattendedScript)
+				{
+					UInterchangeResultError_Generic* Message = AddMessage<UInterchangeResultError_Generic>();
+					Message->Text = NSLOCTEXT("InterchangeSkeletalMeshFactory", "ImportWithScriptIncompatibleSkeleton", "Interchange Import UInterchangeSkeletalMeshFactory::CreateAsset, cannot merge bone tree with the existing skeleton.");
+					return false;
+				}
+
+				EAppReturnType::Type MergeBonesChoice = FMessageDialog::Open(EAppMsgType::YesNo
+					, EAppReturnType::No
+					, NSLOCTEXT("InterchangeSkeletalMeshFactory", "SkeletonFailed_BoneMerge", "FAILED TO MERGE BONES:\n\n This could happen if significant hierarchical changes have been made\ne.g. inserting a bone between nodes.\nWould you like to regenerate the Skeleton from this mesh?\n\n***WARNING: THIS MAY INVALIDATE OR REQUIRE RECOMPRESSION OF ANIMATION DATA.***\n"));
+				if (MergeBonesChoice == EAppReturnType::Yes)
+				{
+					if (SkeletonPtr->RecreateBoneTree(SkeletalMeshPtr))
+					{
+						TArray<const USkeletalMesh*> OtherSkeletalMeshUsingSkeleton;
+						FString SkeletalMeshList;
+						FAssetRegistryModule& AssetRegistryModule = FModuleManager::LoadModuleChecked<FAssetRegistryModule>(TEXT("AssetRegistry"));
+						TArray<FAssetData> SkeletalMeshAssetData;
+
+						FARFilter ARFilter;
+						ARFilter.ClassPaths.Add(USkeletalMesh::StaticClass()->GetClassPathName());
+						ARFilter.TagsAndValues.Add(TEXT("Skeleton"), FAssetData(SkeletonPtr).GetExportTextName());
+
+						IAssetRegistry& AssetRegistry = AssetRegistryModule.Get();
+						if (AssetRegistry.GetAssets(ARFilter, SkeletalMeshAssetData))
+						{
+							// look through all skeletalmeshes that uses this skeleton
+							for (int32 AssetId = 0; AssetId < SkeletalMeshAssetData.Num(); ++AssetId)
+							{
+								FAssetData& CurAssetData = SkeletalMeshAssetData[AssetId];
+								const USkeletalMesh* ExtraSkeletalMesh = Cast<USkeletalMesh>(CurAssetData.GetAsset());
+								if (SkeletalMeshPtr != ExtraSkeletalMesh && IsValid(ExtraSkeletalMesh))
+								{
+									OtherSkeletalMeshUsingSkeleton.Add(ExtraSkeletalMesh);
+									SkeletalMeshList += TEXT("\n") + ExtraSkeletalMesh->GetPathName();
+								}
+							}
+						}
+						if (OtherSkeletalMeshUsingSkeleton.Num() > 0)
+						{
+							FText MessageText = FText::Format(
+								NSLOCTEXT("InterchangeSkeletalMeshFactory", "Skeleton_ReAddAllMeshes", "Would you like to merge all SkeletalMeshes using this skeleton to ensure all bones are merged? This will require to load those SkeletalMeshes.{0}")
+								, FText::FromString(SkeletalMeshList));
+							if (FMessageDialog::Open(EAppMsgType::YesNo, MessageText) == EAppReturnType::Yes)
+							{
+								// look through all skeletalmeshes that uses this skeleton
+								for (const USkeletalMesh* ExtraSkeletalMesh : OtherSkeletalMeshUsingSkeleton)
+								{
+									// merge still can fail
+									if (!SkeletonPtr->MergeAllBonesToBoneTree(ExtraSkeletalMesh))
+									{
+										FMessageDialog::Open(EAppMsgType::Ok,
+											FText::Format(NSLOCTEXT("InterchangeSkeletalMeshFactory", "SkeletonRegenError_RemergingBones", "Failed to merge SkeletalMesh '{0}'."), FText::FromString(ExtraSkeletalMesh->GetName())));
+									}
+								}
+							}
+						}
+					}
+				}
+				return true;
+			};
+
+			if (IsInGameThread())
+			{
+				RecreateSkeleton();
+			}
+			else
+			{
+				//Wait until the skeleton is recreate on the game thread
+				Async(EAsyncExecution::TaskGraphMainThread, MoveTemp(RecreateSkeleton)).Wait();
+			}
+		}
 		if (SkeletalMesh->GetSkeleton() != SkeletonReference)
 		{
 			SkeletalMesh->SetSkeleton(SkeletonReference);
@@ -1384,6 +1473,12 @@ UObject* UInterchangeSkeletalMeshFactory::CreateAsset(const FCreateAssetParams&
 		UE_LOG(LogInterchangeImport, Error, TEXT("Interchange Import UInterchangeSkeletalMeshFactory::CreateAsset, USkeleton* SkeletonReference is nullptr."));
 	}
 
+	if (bIsReImport)
+	{
+		//We must reset the matrixs since CalculateInvRefMatrices only do the calculation if matrix count differ from the bone count.
+		SkeletalMesh->GetRefBasesInvMatrix().Reset();
+	}
+
 	SkeletalMesh->CalculateInvRefMatrices();
 
 	if (!bIsReImport)
@@ -1478,8 +1573,8 @@ void UInterchangeSkeletalMeshFactory::PreImportPreCompletedCallback(const FImpor
 					const int32 NewSourceIndex = GetSourceIndexFromContentType(ImportContentType);
 					//NewSourceIndex should be 0, 1 or 2 (All, Geo, Skinning)
 					check(NewSourceIndex >= 0 && NewSourceIndex < 3);
-					const FString DefaultFilename = AssetImportData->ScriptGetFirstFilename();
 					const TArray<FString> OldFilenames = AssetImportData->ScriptExtractFilenames();
+					const FString DefaultFilename = OldFilenames.Num() > 0 ? OldFilenames[0] : NewSourceFilename;
 					for (int32 SourceIndex = 0; SourceIndex < 3; ++SourceIndex)
 					{
 						FString SourceLabel = GetSourceLabelFromSourceIndex(SourceIndex);
