Interchange
- Added cvars to turn on/off support for each file format. Everything is on by default for now.
- Syntax is: Interchange.FeatureFlags.Import.{file extension}
- Removed the "enable Interchange for Import into level" and "enable Interchange for textures only". The file formats cvars should be used instead.

diff --git a/Engine/Plugins/Experimental/Interchange/Runtime/Source/Import/Private/Fbx/InterchangeFbxTranslator.cpp b/Engine/Plugins/Experimental/Interchange/Runtime/Source/Import/Private/Fbx/InterchangeFbxTranslator.cpp
--- a/Engine/Plugins/Experimental/Interchange/Runtime/Source/Import/Private/Fbx/InterchangeFbxTranslator.cpp
+++ b/Engine/Plugins/Experimental/Interchange/Runtime/Source/Import/Private/Fbx/InterchangeFbxTranslator.cpp
@@ -23,6 +23,12 @@
 #include "Texture/InterchangeTexturePayloadData.h"
 #include "UObject/GCObjectScopeGuard.h"
 
+static bool GInterchangeEnableFBXImport = true;
+static FAutoConsoleVariableRef CCvarInterchangeEnableFBXImport(
+	TEXT("Interchange.FeatureFlags.Import.FBX"),
+	GInterchangeEnableFBXImport,
+	TEXT("Whether FBX support is enabled."),
+	ECVF_Default);
 
 UInterchangeFbxTranslator::UInterchangeFbxTranslator()
 {
@@ -47,8 +53,15 @@ bool UInterchangeFbxTranslator::DoesSupportAssetType(EInterchangeTranslatorAsset
 
 TArray<FString> UInterchangeFbxTranslator::GetSupportedFormats() const
 {
-	TArray<FString> Formats {TEXT("fbx;Filmbox")};
-	return Formats;
+	if (GInterchangeEnableFBXImport)
+	{
+		TArray<FString> Formats{ TEXT("fbx;Filmbox") };
+		return Formats;
+	}
+	else
+	{
+		return TArray<FString>{};
+	}
 }
 
 bool UInterchangeFbxTranslator::Translate(UInterchangeBaseNodeContainer& BaseNodeContainer) const
diff --git a/Engine/Plugins/Experimental/Interchange/Runtime/Source/Import/Private/Gltf/InterchangeGltfTranslator.cpp b/Engine/Plugins/Experimental/Interchange/Runtime/Source/Import/Private/Gltf/InterchangeGltfTranslator.cpp
--- a/Engine/Plugins/Experimental/Interchange/Runtime/Source/Import/Private/Gltf/InterchangeGltfTranslator.cpp
+++ b/Engine/Plugins/Experimental/Interchange/Runtime/Source/Import/Private/Gltf/InterchangeGltfTranslator.cpp
@@ -26,6 +26,13 @@
 #include "StaticMeshAttributes.h"
 #include "UObject/GCObjectScopeGuard.h"
 
+static bool GInterchangeEnableGLTFImport = true;
+static FAutoConsoleVariableRef CCvarInterchangeEnableGLTFImport(
+	TEXT("Interchange.FeatureFlags.Import.GLTF"),
+	GInterchangeEnableGLTFImport,
+	TEXT("Whether glTF support is enabled."),
+	ECVF_Default);
+
 namespace UE::Interchange::Gltf::Private
 {
 	EInterchangeTextureWrapMode ConvertWrap(GLTF::FSampler::EWrap Wrap)
@@ -807,9 +814,13 @@ bool UInterchangeGltfTranslator::DoesSupportAssetType(EInterchangeTranslatorAsse
 TArray<FString> UInterchangeGltfTranslator::GetSupportedFormats() const
 {
 	TArray<FString> GltfExtensions;
-	GltfExtensions.Reserve(2);
-	GltfExtensions.Add(TEXT("gltf;GL Transmission Format"));
-	GltfExtensions.Add(TEXT("glb;GL Transmission Format (Binary)"));
+
+	if ( GInterchangeEnableGLTFImport )
+	{
+		GltfExtensions.Reserve(2);
+		GltfExtensions.Add(TEXT("gltf;GL Transmission Format"));
+		GltfExtensions.Add(TEXT("glb;GL Transmission Format (Binary)"));
+	}
 
 	return GltfExtensions;
 }
diff --git a/Engine/Plugins/Experimental/Interchange/Runtime/Source/Import/Private/InterchangeImportModule.cpp b/Engine/Plugins/Experimental/Interchange/Runtime/Source/Import/Private/InterchangeImportModule.cpp
--- a/Engine/Plugins/Experimental/Interchange/Runtime/Source/Import/Private/InterchangeImportModule.cpp
+++ b/Engine/Plugins/Experimental/Interchange/Runtime/Source/Import/Private/InterchangeImportModule.cpp
@@ -62,17 +62,11 @@ void FInterchangeImportModule::StartupModule()
 
 		//Textures
 		InterchangeManager.RegisterTranslator(UInterchangeImageWrapperTranslator::StaticClass());
-		InterchangeManager.RegisterTextureOnlyTranslatorClass(UInterchangeImageWrapperTranslator::StaticClass());
 		InterchangeManager.RegisterTranslator(UInterchangeDDSTranslator::StaticClass());
-		InterchangeManager.RegisterTextureOnlyTranslatorClass(UInterchangeDDSTranslator::StaticClass());
 		InterchangeManager.RegisterTranslator(UInterchangeJPGTranslator::StaticClass());
-		InterchangeManager.RegisterTextureOnlyTranslatorClass(UInterchangeJPGTranslator::StaticClass());
 		InterchangeManager.RegisterTranslator(UInterchangePCXTranslator::StaticClass());
-		InterchangeManager.RegisterTextureOnlyTranslatorClass(UInterchangePCXTranslator::StaticClass());
 		InterchangeManager.RegisterTranslator(UInterchangePSDTranslator::StaticClass());
-		InterchangeManager.RegisterTextureOnlyTranslatorClass(UInterchangePSDTranslator::StaticClass());
 		InterchangeManager.RegisterTranslator(UInterchangeIESTranslator::StaticClass());
-		InterchangeManager.RegisterTextureOnlyTranslatorClass(UInterchangeIESTranslator::StaticClass());
 
 		//Register the factories
 		InterchangeManager.RegisterFactory(UInterchangeTextureFactory::StaticClass());
diff --git a/Engine/Plugins/Experimental/Interchange/Runtime/Source/Import/Private/MaterialX/InterchangeMaterialXTranslator.cpp b/Engine/Plugins/Experimental/Interchange/Runtime/Source/Import/Private/MaterialX/InterchangeMaterialXTranslator.cpp
--- a/Engine/Plugins/Experimental/Interchange/Runtime/Source/Import/Private/MaterialX/InterchangeMaterialXTranslator.cpp
+++ b/Engine/Plugins/Experimental/Interchange/Runtime/Source/Import/Private/MaterialX/InterchangeMaterialXTranslator.cpp
@@ -16,6 +16,13 @@
 
 #define LOCTEXT_NAMESPACE "InterchangeMaterialXTranslator"
 
+static bool GInterchangeEnableMaterialXImport = true;
+static FAutoConsoleVariableRef CCvarInterchangeEnableMaterialXImport(
+	TEXT("Interchange.FeatureFlags.Import.MTLX"),
+	GInterchangeEnableMaterialXImport,
+	TEXT("Whether MaterialX support is enabled."),
+	ECVF_Default);
+
 namespace mx = MaterialX;
 
 namespace UE::Interchange::MaterialX
@@ -97,7 +104,7 @@ TArray<FString> UInterchangeMaterialXTranslator::GetSupportedFormats() const
 {
 	// Call to UInterchangeMaterialXTranslator::GetSupportedFormats is not supported out of game thread
 	// A more global solution must be found for translators which require some initialization
-	if(!IsInGameThread())
+	if(!IsInGameThread() || !GInterchangeEnableMaterialXImport)
 	{
 		return TArray<FString>{};
 	}
diff --git a/Engine/Plugins/Experimental/Interchange/Runtime/Source/Import/Private/Mesh/InterchangeOBJTranslator.cpp b/Engine/Plugins/Experimental/Interchange/Runtime/Source/Import/Private/Mesh/InterchangeOBJTranslator.cpp
--- a/Engine/Plugins/Experimental/Interchange/Runtime/Source/Import/Private/Mesh/InterchangeOBJTranslator.cpp
+++ b/Engine/Plugins/Experimental/Interchange/Runtime/Source/Import/Private/Mesh/InterchangeOBJTranslator.cpp
@@ -24,6 +24,12 @@
 #include "UObject/GCObjectScopeGuard.h"
 #include "UVMapSettings.h"
 
+static bool GInterchangeEnableOBJImport = true;
+static FAutoConsoleVariableRef CCvarInterchangeEnableOBJImport(
+	TEXT("Interchange.FeatureFlags.Import.OBJ"),
+	GInterchangeEnableOBJImport,
+	TEXT("Whether OBJ support is enabled."),
+	ECVF_Default);
 
 /**
  * This is a binary representation of the .obj file 
@@ -854,8 +860,15 @@ UInterchangeOBJTranslator::~UInterchangeOBJTranslator()
 
 TArray<FString> UInterchangeOBJTranslator::GetSupportedFormats() const
 {
-	TArray<FString> Formats {TEXT("obj;OBJ File Format")};
-	return Formats;
+	if (GInterchangeEnableOBJImport)
+	{
+		TArray<FString> Formats{ TEXT("obj;OBJ File Format") };
+		return Formats;
+	}
+	else
+	{
+		return TArray<FString>{};
+	}
 }
 
 bool UInterchangeOBJTranslator::DoesSupportAssetType(EInterchangeTranslatorAssetType AssetType) const
diff --git a/Engine/Plugins/Experimental/Interchange/Runtime/Source/Import/Private/Texture/InterchangeDDSTranslator.cpp b/Engine/Plugins/Experimental/Interchange/Runtime/Source/Import/Private/Texture/InterchangeDDSTranslator.cpp
--- a/Engine/Plugins/Experimental/Interchange/Runtime/Source/Import/Private/Texture/InterchangeDDSTranslator.cpp
+++ b/Engine/Plugins/Experimental/Interchange/Runtime/Source/Import/Private/Texture/InterchangeDDSTranslator.cpp
@@ -19,6 +19,13 @@
 #include "Serialization/Archive.h"
 #include "Texture/TextureTranslatorUtilities.h"
 
+static bool GInterchangeEnableDDSImport = true;
+static FAutoConsoleVariableRef CCvarInterchangeEnableDDSImport(
+	TEXT("Interchange.FeatureFlags.Import.DDS"),
+	GInterchangeEnableDDSImport,
+	TEXT("Whether DDS support is enabled."),
+	ECVF_Default);
+
 namespace UE::Interchange::Private::InterchangeDDSTranslator
 { 
 	bool LoadDDSHeaderFromFile(TArray64<uint8>& OutHeader, const FString& Filename)
@@ -51,8 +58,15 @@ namespace UE::Interchange::Private::InterchangeDDSTranslator
 
 TArray<FString> UInterchangeDDSTranslator::GetSupportedFormats() const
 {
-	TArray<FString> Formats {TEXT("dds;DirectDraw Surface")};
-	return Formats;
+	if (GInterchangeEnableDDSImport)
+	{
+		TArray<FString> Formats{ TEXT("dds;DirectDraw Surface") };
+		return Formats;
+	}
+	else
+	{
+		return TArray<FString>{};
+	}
 }
 
 bool UInterchangeDDSTranslator::CanImportSourceData(const UInterchangeSourceData* InSourceData) const
diff --git a/Engine/Plugins/Experimental/Interchange/Runtime/Source/Import/Private/Texture/InterchangeIESTranslator.cpp b/Engine/Plugins/Experimental/Interchange/Runtime/Source/Import/Private/Texture/InterchangeIESTranslator.cpp
--- a/Engine/Plugins/Experimental/Interchange/Runtime/Source/Import/Private/Texture/InterchangeIESTranslator.cpp
+++ b/Engine/Plugins/Experimental/Interchange/Runtime/Source/Import/Private/Texture/InterchangeIESTranslator.cpp
@@ -11,10 +11,24 @@
 #include "Misc/Paths.h"
 #include "Texture/TextureTranslatorUtilities.h"
 
+static bool GInterchangeEnableIESImport = true;
+static FAutoConsoleVariableRef CCvarInterchangeEnableIESImport(
+	TEXT("Interchange.FeatureFlags.Import.IES"),
+	GInterchangeEnableIESImport,
+	TEXT("Whether IES support is enabled."),
+	ECVF_Default);
+
 TArray<FString> UInterchangeIESTranslator::GetSupportedFormats() const
 {
-	TArray<FString> Formats {TEXT("ies;IES light profile")};
-	return Formats;
+	if (GInterchangeEnableIESImport)
+	{
+		TArray<FString> Formats{ TEXT("ies;IES light profile") };
+		return Formats;
+	}
+	else
+	{
+		return TArray<FString>{};
+	}
 }
 
 bool UInterchangeIESTranslator::Translate(UInterchangeBaseNodeContainer& BaseNodeContainer) const
diff --git a/Engine/Plugins/Experimental/Interchange/Runtime/Source/Import/Private/Texture/InterchangeImageWrapperTranslator.cpp b/Engine/Plugins/Experimental/Interchange/Runtime/Source/Import/Private/Texture/InterchangeImageWrapperTranslator.cpp
--- a/Engine/Plugins/Experimental/Interchange/Runtime/Source/Import/Private/Texture/InterchangeImageWrapperTranslator.cpp
+++ b/Engine/Plugins/Experimental/Interchange/Runtime/Source/Import/Private/Texture/InterchangeImageWrapperTranslator.cpp
@@ -19,19 +19,89 @@
 #include "TextureImportUtils.h"
 #include "TgaImageSupport.h"
 
+static bool GInterchangeEnablePNGImport = true;
+static FAutoConsoleVariableRef CCvarInterchangeEnablePNGImport(
+	TEXT("Interchange.FeatureFlags.Import.PNG"),
+	GInterchangeEnablePNGImport,
+	TEXT("Whether PNG support is enabled."),
+	ECVF_Default);
+
+static bool GInterchangeEnableBMPImport = true;
+static FAutoConsoleVariableRef CCvarInterchangeEnableBMPImport(
+	TEXT("Interchange.FeatureFlags.Import.BMP"),
+	GInterchangeEnableBMPImport,
+	TEXT("Whether BMP support is enabled."),
+	ECVF_Default);
+
+static bool GInterchangeEnableEXRImport = true;
+static FAutoConsoleVariableRef CCvarInterchangeEnableEXRImport(
+	TEXT("Interchange.FeatureFlags.Import.EXR"),
+	GInterchangeEnableEXRImport,
+	TEXT("Whether OpenEXR support is enabled."),
+	ECVF_Default);
+
+static bool GInterchangeEnableHDRImport = true;
+static FAutoConsoleVariableRef CCvarInterchangeEnableHDRImport(
+	TEXT("Interchange.FeatureFlags.Import.HDR"),
+	GInterchangeEnableHDRImport,
+	TEXT("Whether HDR support is enabled."),
+	ECVF_Default);
+
+#if WITH_LIBTIFF
+static bool GInterchangeEnableTIFFImport = true;
+static FAutoConsoleVariableRef CCvarInterchangeEnableTIFFImport(
+	TEXT("Interchange.FeatureFlags.Import.TIFF"),
+	GInterchangeEnableTIFFImport,
+	TEXT("Whether TIFF support is enabled."),
+	ECVF_Default);
+#endif
+
+static bool GInterchangeEnableTGAImport = true;
+static FAutoConsoleVariableRef CCvarInterchangeEnableTGAImport(
+	TEXT("Interchange.FeatureFlags.Import.TGA"),
+	GInterchangeEnableTGAImport,
+	TEXT("Whether TGA support is enabled."),
+	ECVF_Default);
+
 TArray<FString> UInterchangeImageWrapperTranslator::GetSupportedFormats() const
 {
-	return {
-			TEXT("png;Portable Network Graphic"),
-			TEXT("bmp;Bitmap image"),
-			TEXT("exr;OpenEXR image"),
-			TEXT("hdr;High Dynamic Range image"),
+	TArray<FString> Formats;
+	Formats.Reserve(7);
+
+	if (GInterchangeEnablePNGImport)
+	{
+		Formats.Emplace(TEXT("png;Portable Network Graphic"));
+	}
+
+	if (GInterchangeEnableBMPImport)
+	{
+		Formats.Emplace(TEXT("bmp;Bitmap image"));
+	}
+
+	if (GInterchangeEnableEXRImport)
+	{
+		Formats.Emplace(TEXT("exr;OpenEXR image"));
+	}
+
+	if (GInterchangeEnableHDRImport)
+	{
+		Formats.Emplace(TEXT("hdr;High Dynamic Range image"));
+	}
+
 #if WITH_LIBTIFF
-			TEXT("tif;Tag Image File Format"),
-			TEXT("tiff;Tag Image File Format"),
-#endif //WITH_LIBTIFF
-			TEXT("tga;Targa image")
-		};
+	if (GInterchangeEnableTIFFImport)
+	{
+		Formats.Emplace(TEXT("tif;Tag Image File Format"));
+		Formats.Emplace(TEXT("tiff;Tag Image File Format"));
+	}
+#endif
+
+	if (GInterchangeEnableTGAImport)
+	{
+		Formats.Emplace(TEXT("tga;Targa image"));
+	}
+
+	return Formats;
 }
 
 bool UInterchangeImageWrapperTranslator::Translate(UInterchangeBaseNodeContainer& BaseNodeContainer) const
diff --git a/Engine/Plugins/Experimental/Interchange/Runtime/Source/Import/Private/Texture/InterchangeJPGTranslator.cpp b/Engine/Plugins/Experimental/Interchange/Runtime/Source/Import/Private/Texture/InterchangeJPGTranslator.cpp
--- a/Engine/Plugins/Experimental/Interchange/Runtime/Source/Import/Private/Texture/InterchangeJPGTranslator.cpp
+++ b/Engine/Plugins/Experimental/Interchange/Runtime/Source/Import/Private/Texture/InterchangeJPGTranslator.cpp
@@ -17,13 +17,24 @@
 #include "Nodes/InterchangeBaseNodeContainer.h"
 #include "Texture/TextureTranslatorUtilities.h"
 
+static bool GInterchangeEnableJPGImport = true;
+static FAutoConsoleVariableRef CCvarInterchangeEnableJPGImport(
+	TEXT("Interchange.FeatureFlags.Import.JPG"),
+	GInterchangeEnableJPGImport,
+	TEXT("Whether JPG support is enabled."),
+	ECVF_Default);
+
 TArray<FString> UInterchangeJPGTranslator::GetSupportedFormats() const
 {
 	TArray<FString> Formats;
-	Formats.Reserve(2);
 
-	Formats.Add(TEXT("jpg;JPEG image"));
-	Formats.Add(TEXT("jpeg;JPEG image"));
+	if (GInterchangeEnableJPGImport)
+	{
+		Formats.Reserve(2);
+
+		Formats.Add(TEXT("jpg;JPEG image"));
+		Formats.Add(TEXT("jpeg;JPEG image"));
+	}
 
 	return Formats;
 }
diff --git a/Engine/Plugins/Experimental/Interchange/Runtime/Source/Import/Private/Texture/InterchangePCXTranslator.cpp b/Engine/Plugins/Experimental/Interchange/Runtime/Source/Import/Private/Texture/InterchangePCXTranslator.cpp
--- a/Engine/Plugins/Experimental/Interchange/Runtime/Source/Import/Private/Texture/InterchangePCXTranslator.cpp
+++ b/Engine/Plugins/Experimental/Interchange/Runtime/Source/Import/Private/Texture/InterchangePCXTranslator.cpp
@@ -15,6 +15,13 @@
 #include "Nodes/InterchangeBaseNodeContainer.h"
 #include "Texture/TextureTranslatorUtilities.h"
 
+static bool GInterchangeEnablePCXImport = true;
+static FAutoConsoleVariableRef CCvarInterchangeEnablePCXImport(
+	TEXT("Interchange.FeatureFlags.Import.PCX"),
+	GInterchangeEnablePCXImport,
+	TEXT("Whether PCX support is enabled."),
+	ECVF_Default);
+
 //////////////////////////////////////////////////////////////////////////
 // PCX helper local function
 namespace UE
@@ -68,9 +75,15 @@ namespace UE
 
 TArray<FString> UInterchangePCXTranslator::GetSupportedFormats() const
 {
-	TArray<FString> Formats {TEXT("pcx;Picture Exchange")};
-
-	return Formats;
+	if (GInterchangeEnablePCXImport)
+	{
+		TArray<FString> Formats{ TEXT("pcx;Picture Exchange") };
+		return Formats;
+	}
+	else
+	{
+		return TArray<FString>{};
+	}
 }
 
 bool UInterchangePCXTranslator::Translate(UInterchangeBaseNodeContainer& BaseNodeContainer) const
diff --git a/Engine/Plugins/Experimental/Interchange/Runtime/Source/Import/Private/Texture/InterchangePSDTranslator.cpp b/Engine/Plugins/Experimental/Interchange/Runtime/Source/Import/Private/Texture/InterchangePSDTranslator.cpp
--- a/Engine/Plugins/Experimental/Interchange/Runtime/Source/Import/Private/Texture/InterchangePSDTranslator.cpp
+++ b/Engine/Plugins/Experimental/Interchange/Runtime/Source/Import/Private/Texture/InterchangePSDTranslator.cpp
@@ -15,6 +15,13 @@
 #include "Nodes/InterchangeBaseNodeContainer.h"
 #include "Texture/TextureTranslatorUtilities.h"
 
+static bool GInterchangeEnablePSDImport = true;
+static FAutoConsoleVariableRef CCvarInterchangeEnablePSDImport(
+	TEXT("Interchange.FeatureFlags.Import.PSD"),
+	GInterchangeEnablePSDImport,
+	TEXT("Whether PSD support is enabled."),
+	ECVF_Default);
+
 //////////////////////////////////////////////////////////////////////////
 // PSD helper local function
 namespace UE
@@ -369,9 +376,15 @@ namespace UE
 
 TArray<FString> UInterchangePSDTranslator::GetSupportedFormats() const
 {
-	TArray<FString> Formats {TEXT("psd;Photoshop Document")};
-
-	return Formats;
+	if (GInterchangeEnablePSDImport)
+	{
+		TArray<FString> Formats{ TEXT("psd;Photoshop Document") };
+		return Formats;
+	}
+	else
+	{
+		return TArray<FString>{};
+	}
 }
 
 bool UInterchangePSDTranslator::Translate(UInterchangeBaseNodeContainer& BaseNodeContainer) const
diff --git a/Engine/Source/Developer/AssetTools/Private/AssetTools.cpp b/Engine/Source/Developer/AssetTools/Private/AssetTools.cpp
--- a/Engine/Source/Developer/AssetTools/Private/AssetTools.cpp
+++ b/Engine/Source/Developer/AssetTools/Private/AssetTools.cpp
@@ -1749,22 +1749,10 @@ TArray<UObject*> UAssetToolsImpl::ImportAssetsInternal(const TArray<FString>& Fi
 	FScopedSlowTask SlowTask(ValidFiles.Num(), LOCTEXT("ImportSlowTask", "Importing"));
 
 	bool bUseInterchangeFramework = false;
-	bool bUseInterchangeFrameworkForTextureOnly = false;
 	UInterchangeManager& InterchangeManager = UInterchangeManager::GetInterchangeManager();
 #if WITH_EDITOR
 	const UEditorExperimentalSettings* EditorExperimentalSettings = GetDefault<UEditorExperimentalSettings>();
 	bUseInterchangeFramework = EditorExperimentalSettings->bEnableInterchangeFramework;
-
-	if (bUseInterchangeFramework && Params.SpecifiedFactory)
-	{
-		if (Params.SpecifiedFactory->GetClass()->IsChildOf(USceneImportFactory::StaticClass()))
-		{
-			bUseInterchangeFramework = GetDefault<UInterchangeProjectSettings>()->bUseInterchangeWhenImportingIntoLevel;
-		}
-	}
-
-	bUseInterchangeFrameworkForTextureOnly = (!bUseInterchangeFramework) && EditorExperimentalSettings->bEnableInterchangeFrameworkForTextureOnly;
-	bUseInterchangeFramework |= bUseInterchangeFrameworkForTextureOnly;
 #endif
 
 	if (!bUseInterchangeFramework && ValidFiles.Num() > 1)
@@ -1918,16 +1906,7 @@ TArray<UObject*> UAssetToolsImpl::ImportAssetsInternal(const TArray<FString>& Fi
 			{
 				UE::Interchange::FScopedSourceData ScopedSourceData(Filename);
 
-				if (bUseInterchangeFrameworkForTextureOnly)
-				{
-					UInterchangeTranslatorBase* Translator = InterchangeManager.GetTranslatorForSourceData(ScopedSourceData.GetSourceData());
-					if (!Translator || !InterchangeManager.IsTranslatorClassForTextureOnly(Translator->GetClass()))
-					{
-						bOnlyInterchangeImport = false;
-						break;
-					}
-				}
-				else if (!InterchangeManager.CanTranslateSourceData(ScopedSourceData.GetSourceData()))
+				if (!InterchangeManager.CanTranslateSourceData(ScopedSourceData.GetSourceData()))
 				{
 					bOnlyInterchangeImport = false;
 					break;
@@ -1993,17 +1972,7 @@ TArray<UObject*> UAssetToolsImpl::ImportAssetsInternal(const TArray<FString>& Fi
 		{
 			UE::Interchange::FScopedSourceData ScopedSourceData(Filename);
 
-			bool bUseATextureTranslator = false;
-			if (bUseInterchangeFrameworkForTextureOnly)
-			{
-				UInterchangeTranslatorBase* Translator = InterchangeManager.GetTranslatorForSourceData(ScopedSourceData.GetSourceData());
-				if (Translator && InterchangeManager.IsTranslatorClassForTextureOnly(Translator->GetClass()))
-				{
-					bUseATextureTranslator = true;
-				}
-			}
-
-			if (bUseATextureTranslator || (!bUseInterchangeFrameworkForTextureOnly && InterchangeManager.CanTranslateSourceData(ScopedSourceData.GetSourceData())))
+			if (InterchangeManager.CanTranslateSourceData(ScopedSourceData.GetSourceData()))
 			{
 				FImportAssetParameters ImportAssetParameters;
 				ImportAssetParameters.bIsAutomated = bAutomatedImport;
diff --git a/Engine/Source/Editor/UnrealEd/Classes/Settings/EditorExperimentalSettings.h b/Engine/Source/Editor/UnrealEd/Classes/Settings/EditorExperimentalSettings.h
--- a/Engine/Source/Editor/UnrealEd/Classes/Settings/EditorExperimentalSettings.h
+++ b/Engine/Source/Editor/UnrealEd/Classes/Settings/EditorExperimentalSettings.h
@@ -45,10 +45,6 @@ public:
 	UPROPERTY(EditAnywhere, config, Category = Interchange, meta = (DisplayName = "Enable interchange framework import"))
 	bool bEnableInterchangeFramework;
 
-	/** Enable interchange framework, the interchange framework is a new import system which can import asynchronously and in parallel but only add support the texture assets. See the interchange project settings to configure the import pipeline*/
-	UPROPERTY(EditAnywhere, config, Category = Interchange, meta = (DisplayName = "Enable interchange framework import for textures only"))
-	bool bEnableInterchangeFrameworkForTextureOnly;
-
 	/** Allows the editor to run on HDR monitors on Windows 10 */
 	UPROPERTY(EditAnywhere, config, Category = HDR, meta = (ConfigRestartRequired = true, DisplayName = "Enable Editor Support for HDR Monitors"))
 	bool bHDREditor;
diff --git a/Engine/Source/Editor/UnrealEd/Private/Editor.cpp b/Engine/Source/Editor/UnrealEd/Private/Editor.cpp
--- a/Engine/Source/Editor/UnrealEd/Private/Editor.cpp
+++ b/Engine/Source/Editor/UnrealEd/Private/Editor.cpp
@@ -270,8 +270,6 @@ bool FReimportManager::Reimport( UObject* Obj, bool bAskForNewFileIfMissing, boo
 	const UEditorExperimentalSettings* EditorExperimentalSettings = GetDefault<UEditorExperimentalSettings>();
 
 	bUseInterchangeFramework = EditorExperimentalSettings->bEnableInterchangeFramework;
-	const bool bUseInterchangeFrameworkForTextureOnly = (!bUseInterchangeFramework) && EditorExperimentalSettings->bEnableInterchangeFrameworkForTextureOnly;
-	bUseInterchangeFramework |= bUseInterchangeFrameworkForTextureOnly;
 
 	bool bSuccess = false;
 	if ( Obj )
@@ -398,17 +396,8 @@ bool FReimportManager::Reimport( UObject* Obj, bool bAskForNewFileIfMissing, boo
 					UE::Interchange::FScopedSourceData ScopedSourceData(SourceFilenames[RealValidSourceFileIndex]);
 
 					CanReimportHandler->SetReimportSourceIndex(Obj, SourceFileIndex);
-					bool bUseATextureTranslator = false;
-					if (bUseInterchangeFrameworkForTextureOnly)
-					{
-						UInterchangeTranslatorBase* Translator = InterchangeManager.GetTranslatorForSourceData(ScopedSourceData.GetSourceData());
-						if (Translator && InterchangeManager.IsTranslatorClassForTextureOnly(Translator->GetClass()))
-						{
-							bUseATextureTranslator = true;
-						}
-					}
 
-					if (bUseATextureTranslator || (!bUseInterchangeFrameworkForTextureOnly && InterchangeManager.CanTranslateSourceData(ScopedSourceData.GetSourceData())))
+					if (InterchangeManager.CanTranslateSourceData(ScopedSourceData.GetSourceData()))
 					{
 						auto PostImportedLambda = [](UObject* ImportedObject)
 						{
diff --git a/Engine/Source/Editor/UnrealEd/Private/Factories/InterchangeReimportHandler.cpp b/Engine/Source/Editor/UnrealEd/Private/Factories/InterchangeReimportHandler.cpp
--- a/Engine/Source/Editor/UnrealEd/Private/Factories/InterchangeReimportHandler.cpp
+++ b/Engine/Source/Editor/UnrealEd/Private/Factories/InterchangeReimportHandler.cpp
@@ -27,8 +27,7 @@ UInterchangeReimportHandler::UInterchangeReimportHandler(const FObjectInitialize
 bool UInterchangeReimportHandler::CanReimport(UObject* Obj, TArray<FString>& OutFilenames)
 {
 	const UEditorExperimentalSettings* EditorExperimentalSettings = GetDefault<UEditorExperimentalSettings>();
-	const bool bUseInterchange = EditorExperimentalSettings->bEnableInterchangeFramework 
-		|| (EditorExperimentalSettings->bEnableInterchangeFrameworkForTextureOnly && Cast<UTexture>(Obj) != nullptr);
+	const bool bUseInterchange = EditorExperimentalSettings->bEnableInterchangeFramework;
 
 	if (bUseInterchange)
 	{
diff --git a/Engine/Source/Editor/UnrealEd/Private/FileHelpers.cpp b/Engine/Source/Editor/UnrealEd/Private/FileHelpers.cpp
--- a/Engine/Source/Editor/UnrealEd/Private/FileHelpers.cpp
+++ b/Engine/Source/Editor/UnrealEd/Private/FileHelpers.cpp
@@ -1445,7 +1445,7 @@ void FEditorFileUtils::Import(const FString& InFilename)
 	const UEditorExperimentalSettings* EditorExperimentalSettings = GetDefault<UEditorExperimentalSettings>();
 
 	UE::Interchange::FScopedSourceData ScopedSourceData(InFilename);
-	const bool bImportThroughInterchange = (EditorExperimentalSettings->bEnableInterchangeFramework && GetDefault<UInterchangeProjectSettings>()->bUseInterchangeWhenImportingIntoLevel) && UInterchangeManager::GetInterchangeManager().CanTranslateSourceData(ScopedSourceData.GetSourceData());
+	const bool bImportThroughInterchange = EditorExperimentalSettings->bEnableInterchangeFramework && UInterchangeManager::GetInterchangeManager().CanTranslateSourceData(ScopedSourceData.GetSourceData());
 
 	USceneImportFactory* SceneFactory = nullptr;
 
diff --git a/Engine/Source/Editor/UnrealEd/Private/Settings/SettingsClasses.cpp b/Engine/Source/Editor/UnrealEd/Private/Settings/SettingsClasses.cpp
--- a/Engine/Source/Editor/UnrealEd/Private/Settings/SettingsClasses.cpp
+++ b/Engine/Source/Editor/UnrealEd/Private/Settings/SettingsClasses.cpp
@@ -238,20 +238,6 @@ void UEditorExperimentalSettings::PostEditChangeProperty( struct FPropertyChange
 	{
 		CVarEditorHDRNITLevel->Set(HDREditorNITLevel, ECVF_SetByProjectSetting);
 	}
-	else if (Name == GET_MEMBER_NAME_CHECKED(UEditorExperimentalSettings, bEnableInterchangeFramework))
-	{
-		if (bEnableInterchangeFramework)
-		{
-			bEnableInterchangeFrameworkForTextureOnly = false;
-		}
-	}
-	else if (Name == GET_MEMBER_NAME_CHECKED(UEditorExperimentalSettings, bEnableInterchangeFrameworkForTextureOnly))
-	{
-		if (bEnableInterchangeFrameworkForTextureOnly)
-		{
-			bEnableInterchangeFramework = false;
-		}
-	}
 
 	if (!FUnrealEdMisc::Get().IsDeletePreferences())
 	{
diff --git a/Engine/Source/Runtime/Experimental/Interchange/Engine/Private/InterchangeManager.cpp b/Engine/Source/Runtime/Experimental/Interchange/Engine/Private/InterchangeManager.cpp
--- a/Engine/Source/Runtime/Experimental/Interchange/Engine/Private/InterchangeManager.cpp
+++ b/Engine/Source/Runtime/Experimental/Interchange/Engine/Private/InterchangeManager.cpp
@@ -1167,19 +1167,6 @@ UInterchangeTranslatorBase* UInterchangeManager::GetTranslatorSupportingPayloadI
 	return nullptr;
 }
 
-void UInterchangeManager::RegisterTextureOnlyTranslatorClass(const UClass* TranslatorClass)
-{
-	if (TranslatorClass)
-	{
-		TextureOnlyTranslatorClass.Add(TranslatorClass);
-	}
-}
-
-bool UInterchangeManager::IsTranslatorClassForTextureOnly(const UClass* TranslatorClass) const
-{
-	return TextureOnlyTranslatorClass.Contains(TranslatorClass);
-}
-
 bool UInterchangeManager::IsAttended()
 {
 	if (FApp::IsGame())
diff --git a/Engine/Source/Runtime/Experimental/Interchange/Engine/Public/InterchangeManager.h b/Engine/Source/Runtime/Experimental/Interchange/Engine/Public/InterchangeManager.h
--- a/Engine/Source/Runtime/Experimental/Interchange/Engine/Public/InterchangeManager.h
+++ b/Engine/Source/Runtime/Experimental/Interchange/Engine/Public/InterchangeManager.h
@@ -479,19 +479,6 @@ public:
 	 */
 	UInterchangeTranslatorBase* GetTranslatorSupportingPayloadInterfaceForSourceData(const UInterchangeSourceData* SourceData, const UClass* PayloadInterfaceClass) const;
 
-	/**
-	 * Register a translator as a texture translator only
-	 * This used to implement the experimental setting to use interchange for textures only.
-	 * 
-	 * @Param Translator - The UClass of the translator you want to register
-	 */
-	void RegisterTextureOnlyTranslatorClass(const UClass* TranslatorClass);
-
-	/**
-	 * Check if the translator class is registered as a texture only translator
-	 */
-	bool IsTranslatorClassForTextureOnly(const UClass* TranslatorClass) const;
-
 protected:
 
 	/** Return true if we can show some UI */
@@ -579,8 +566,5 @@ private:
 	//This boolean is reset to false when the ImportTasks array is empty.
 	bool bImportAllWithDefault = false;
 
-	UPROPERTY()
-	TSet<TObjectPtr<const UClass>> TextureOnlyTranslatorClass;
-
 	friend class UE::Interchange::FScopedTranslator;
 };
diff --git a/Engine/Source/Runtime/Experimental/Interchange/Engine/Public/InterchangeProjectSettings.h b/Engine/Source/Runtime/Experimental/Interchange/Engine/Public/InterchangeProjectSettings.h
--- a/Engine/Source/Runtime/Experimental/Interchange/Engine/Public/InterchangeProjectSettings.h
+++ b/Engine/Source/Runtime/Experimental/Interchange/Engine/Public/InterchangeProjectSettings.h
@@ -55,10 +55,6 @@ public:
 	UPROPERTY(EditAnywhere, config, Category = "ImportContent")
 	FInterchangeImportSettings ContentImportSettings;
 
-	/** If checked, will use Interchange when importing into level.*/
-	UPROPERTY(EditAnywhere, config, Category = "ImportIntoLevel|Experimental")
-	bool bUseInterchangeWhenImportingIntoLevel = false;
-
 	/**
 	 * Settings used when importing into a level.
 	¸*/
