Fix static analysis warning V595: The 'TextureFactoryNode' pointer was utilized before it was verified against nullptr.
#rnx

diff --git a/Engine/Plugins/Experimental/Interchange/Runtime/Source/Pipelines/Private/InterchangeGenericTexturePipeline.cpp b/Engine/Plugins/Experimental/Interchange/Runtime/Source/Pipelines/Private/InterchangeGenericTexturePipeline.cpp
--- a/Engine/Plugins/Experimental/Interchange/Runtime/Source/Pipelines/Private/InterchangeGenericTexturePipeline.cpp
+++ b/Engine/Plugins/Experimental/Interchange/Runtime/Source/Pipelines/Private/InterchangeGenericTexturePipeline.cpp
@@ -251,6 +251,7 @@ UInterchangeTextureFactoryNode* UInterchangeGenericTexturePipeline::CreateTextur
 		if (!ensure(TextureFactoryNode))
 		{
 			//Log an error
+			return nullptr;
 		}
 	}
 	else
@@ -292,12 +293,9 @@ UInterchangeTextureFactoryNode* UInterchangeGenericTexturePipeline::CreateTextur
 	}
 
 #if WITH_EDITORONLY_DATA
-	if (TextureFactoryNode)
+	if (bPreferCompressedSourceData)
 	{
-		if (bPreferCompressedSourceData)
-		{
-			TextureFactoryNode->SetCustomPreferCompressedSourceData(true);
-		}
+		TextureFactoryNode->SetCustomPreferCompressedSourceData(true);
 	}
 #endif // WITH_EDITORONLY_DATA
 
