Adding image adjustment material expression:
- hsvtorgb
- rgbtohsv

diff --git a/Engine/Plugins/Experimental/Interchange/Runtime/Source/Import/Private/MaterialX/MaterialExpressionHsvToRgb.cpp b/Engine/Plugins/Experimental/Interchange/Runtime/Source/Import/Private/MaterialX/MaterialExpressionHsvToRgb.cpp
new file mode 100644
--- /dev/null
+++ b/Engine/Plugins/Experimental/Interchange/Runtime/Source/Import/Private/MaterialX/MaterialExpressionHsvToRgb.cpp
@@ -0,0 +1,54 @@
+// Copyright Epic Games, Inc. All Rights Reserved.
+
+#include "MaterialX/MaterialExpressionHsvToRgb.h"
+#include "MaterialCompiler.h"
+#include "Materials/MaterialExpressionCustom.h"
+
+#define LOCTEXT_NAMESPACE "MaterialExpressionHsvToRgb"
+
+UMaterialExpressionHsvToRgb::UMaterialExpressionHsvToRgb(const FObjectInitializer& ObjectInitializer)
+	: Super(ObjectInitializer)
+{
+	// Structure to hold one-time initialization
+	struct FConstructorStatics
+	{
+		FText NAME_ImageAdjustment;
+		FConstructorStatics()
+			: NAME_ImageAdjustment(LOCTEXT("Image Adjustment", "Image Adjustment"))
+		{}
+	};
+	static FConstructorStatics ConstructorStatics;
+
+#if WITH_EDITORONLY_DATA
+	MenuCategories.Add(ConstructorStatics.NAME_ImageAdjustment);
+#endif
+}
+
+#if WITH_EDITOR
+int32 UMaterialExpressionHsvToRgb::Compile(FMaterialCompiler* Compiler, int32 OutputIndex)
+{
+	if(!Input.GetTracedInput().Expression)
+	{
+		return Compiler->Errorf(TEXT("Missing RGBtoHSV input"));
+	}
+
+	UMaterialExpressionCustom* MaterialExpressionCustom = NewObject<UMaterialExpressionCustom>();
+	MaterialExpressionCustom->Inputs[0].InputName = TEXT("c");
+	MaterialExpressionCustom->Inputs[0].Input = Input;
+	MaterialExpressionCustom->OutputType = ECustomMaterialOutputType::CMOT_Float3;
+	MaterialExpressionCustom->Code = TEXT(R"(
+    float4 K = float4(1.0, 2.0 / 3.0, 1.0 / 3.0, 3.0);
+    float3 p = abs(frac(c.xxx + K.xyz) * 6.0 - K.www);
+    return c.z * lerp(K.xxx, saturate(p - K.xxx), c.y);)");
+
+	TArray<int32> Inputs{ Input.Compile(Compiler) };
+	return Compiler->CustomExpression(MaterialExpressionCustom, 0, Inputs);
+}
+
+void UMaterialExpressionHsvToRgb::GetCaption(TArray<FString>& OutCaptions) const
+{
+	OutCaptions.Add(TEXT("HSVToRGB"));
+}
+#endif
+
+#undef LOCTEXT_NAMESPACE 
diff --git a/Engine/Plugins/Experimental/Interchange/Runtime/Source/Import/Private/MaterialX/MaterialExpressionRgbToHsv.cpp b/Engine/Plugins/Experimental/Interchange/Runtime/Source/Import/Private/MaterialX/MaterialExpressionRgbToHsv.cpp
new file mode 100644
--- /dev/null
+++ b/Engine/Plugins/Experimental/Interchange/Runtime/Source/Import/Private/MaterialX/MaterialExpressionRgbToHsv.cpp
@@ -0,0 +1,58 @@
+// Copyright Epic Games, Inc. All Rights Reserved.
+
+#include "MaterialX/MaterialExpressionRgbToHsv.h"
+#include "MaterialCompiler.h"
+#include "Materials/MaterialExpressionCustom.h"
+
+#define LOCTEXT_NAMESPACE "MaterialExpressionRgbToHsv"
+
+UMaterialExpressionRgbToHsv::UMaterialExpressionRgbToHsv(const FObjectInitializer& ObjectInitializer)
+	: Super(ObjectInitializer)
+{
+	// Structure to hold one-time initialization
+	struct FConstructorStatics
+	{
+		FText NAME_ImageAdjustment;
+		FConstructorStatics()
+			: NAME_ImageAdjustment(LOCTEXT("Image Adjustment", "Image Adjustment"))
+		{}
+	};
+	static FConstructorStatics ConstructorStatics;
+
+#if WITH_EDITORONLY_DATA
+	MenuCategories.Add(ConstructorStatics.NAME_ImageAdjustment);
+#endif
+}
+
+#if WITH_EDITOR
+int32 UMaterialExpressionRgbToHsv::Compile(FMaterialCompiler* Compiler, int32 OutputIndex)
+{
+	if(!Input.GetTracedInput().Expression)
+	{
+		return Compiler->Errorf(TEXT("Missing RGBtoHSV input"));
+	}
+
+	UMaterialExpressionCustom* MaterialExpressionCustom = NewObject<UMaterialExpressionCustom>();
+	MaterialExpressionCustom->Inputs[0].InputName = TEXT("c");
+	MaterialExpressionCustom->Inputs[0].Input = Input;
+	MaterialExpressionCustom->OutputType = ECustomMaterialOutputType::CMOT_Float3;
+	MaterialExpressionCustom->Code = TEXT(R"(
+    float4 K = float4(0.0, -1.0 / 3.0, 2.0 / 3.0, -1.0);
+    float4 p = lerp(float4(c.bg, K.wz), float4(c.gb, K.xy), step(c.b, c.g));
+    float4 q = lerp(float4(p.xyw, c.r), float4(c.r, p.yzx), step(p.x, c.r));
+
+    float d = q.x - min(q.w, q.y);
+    float e = 1.0e-10;
+    return float3(abs(q.z + (q.w - q.y) / (6.0 * d + e)), d / (q.x + e), q.x);)");
+		
+	TArray<int32> Inputs{ Input.Compile(Compiler) };
+	return Compiler->CustomExpression(MaterialExpressionCustom, 0, Inputs);
+}
+
+void UMaterialExpressionRgbToHsv::GetCaption(TArray<FString>& OutCaptions) const
+{
+	OutCaptions.Add(TEXT("RGBToHSV"));
+}
+#endif
+
+#undef LOCTEXT_NAMESPACE 
\ No newline at end of file
diff --git a/Engine/Plugins/Experimental/Interchange/Runtime/Source/Import/Public/MaterialX/MaterialExpressionHsvToRgb.h b/Engine/Plugins/Experimental/Interchange/Runtime/Source/Import/Public/MaterialX/MaterialExpressionHsvToRgb.h
new file mode 100644
--- /dev/null
+++ b/Engine/Plugins/Experimental/Interchange/Runtime/Source/Import/Public/MaterialX/MaterialExpressionHsvToRgb.h
@@ -0,0 +1,23 @@
+// Copyright Epic Games, Inc. All Rights Reserved.
+#pragma once
+
+#include "Materials/MaterialExpression.h"
+
+#include "MaterialExpressionHsvToRgb.generated.h"
+
+UCLASS(collapsecategories, hidecategories = Object, MinimalAPI)
+class UMaterialExpressionHsvToRgb : public UMaterialExpression
+{
+	GENERATED_UCLASS_BODY()
+
+	UPROPERTY()
+	FExpressionInput Input;
+
+	//~ Begin UMaterialExpression Interface
+#if WITH_EDITOR
+	virtual int32 Compile(class FMaterialCompiler* Compiler, int32 OutputIndex) override;
+	virtual void GetCaption(TArray<FString>& OutCaptions) const override;
+#endif
+	//~ End UMaterialExpression Interface
+};
+
diff --git a/Engine/Plugins/Experimental/Interchange/Runtime/Source/Import/Public/MaterialX/MaterialExpressionRgbToHsv.h b/Engine/Plugins/Experimental/Interchange/Runtime/Source/Import/Public/MaterialX/MaterialExpressionRgbToHsv.h
new file mode 100644
--- /dev/null
+++ b/Engine/Plugins/Experimental/Interchange/Runtime/Source/Import/Public/MaterialX/MaterialExpressionRgbToHsv.h
@@ -0,0 +1,23 @@
+// Copyright Epic Games, Inc. All Rights Reserved.
+#pragma once
+
+#include "Materials/MaterialExpression.h"
+
+#include "MaterialExpressionRgbToHsv.generated.h"
+
+UCLASS(collapsecategories, hidecategories = Object, MinimalAPI)
+class UMaterialExpressionRgbToHsv : public UMaterialExpression
+{
+	GENERATED_UCLASS_BODY()
+
+	UPROPERTY()
+	FExpressionInput Input;
+
+	//~ Begin UMaterialExpression Interface
+#if WITH_EDITOR
+	virtual int32 Compile(class FMaterialCompiler* Compiler, int32 OutputIndex) override;
+	virtual void GetCaption(TArray<FString>& OutCaptions) const override;
+#endif
+	//~ End UMaterialExpression Interface
+};
+
