Interchange material pipeline:
- Track texture usage and warn if the same texture is used by incompatible inputs (ie: color and normal).
- Remove green channel flipping if a texture isn't used exclusively for vector inputs. This isn't a valid setup and the warning mentionned previously will be raised.
- Don't consider the specular color of the Phong shading model to be a linear input, since it represents a color and not a scalar value.
- Fixed an issue where a SRGB texture would be flagged as linear if its alpha channel was used in a scalar material input (ie: opacity).

diff --git a/Engine/Plugins/Experimental/Interchange/Runtime/Source/Pipelines/Private/InterchangeGenericMaterialPipeline.cpp b/Engine/Plugins/Experimental/Interchange/Runtime/Source/Pipelines/Private/InterchangeGenericMaterialPipeline.cpp
--- a/Engine/Plugins/Experimental/Interchange/Runtime/Source/Pipelines/Private/InterchangeGenericMaterialPipeline.cpp
+++ b/Engine/Plugins/Experimental/Interchange/Runtime/Source/Pipelines/Private/InterchangeGenericMaterialPipeline.cpp
@@ -42,6 +42,26 @@
 #include "UObject/Object.h"
 #include "UObject/ObjectMacros.h"
 
+#define LOCTEXT_NAMESPACE "InterchangeGenericMaterialPipeline"
+
+FString LexToString(UInterchangeGenericMaterialPipeline::EMaterialInputType Value)
+{
+	switch(Value)
+	{
+	case UInterchangeGenericMaterialPipeline::EMaterialInputType::Unknown:
+		return TEXT("Unknown");
+	case UInterchangeGenericMaterialPipeline::EMaterialInputType::Color:
+		return TEXT("Color");
+	case UInterchangeGenericMaterialPipeline::EMaterialInputType::Vector:
+		return TEXT("Vector");
+	case UInterchangeGenericMaterialPipeline::EMaterialInputType::Scalar:
+		return TEXT("Scalar");
+	default:
+		ensure(false);
+		return FString();
+	}
+}
+
 void UInterchangeGenericMaterialPipeline::AdjustSettingsForContext(EInterchangePipelineContext ImportType, TObjectPtr<UObject> ReimportAsset)
 {
 	Super::AdjustSettingsForContext(ImportType, ReimportAsset);
@@ -296,10 +316,8 @@ bool UInterchangeGenericMaterialPipeline::HandlePhongModel(const UInterchangeSha
 			}
 		}
 
-		// Specular
+		// Specular Color
 		{
-			TGuardValue<bool> ParsingForLinearInputGuard(bParsingForLinearInput, true);
-
 			TTuple<UInterchangeMaterialExpressionFactoryNode*, FString> SpecularExpressionFactoryNode =
 				CreateMaterialExpressionForInput(MaterialFactoryNode, ShaderGraphNode, Parameters::SpecularColor.ToString(), FunctionCallExpression->GetUniqueID());
 
@@ -315,7 +333,7 @@ bool UInterchangeGenericMaterialPipeline::HandlePhongModel(const UInterchangeSha
 			const bool bHasShininessInput = UInterchangeShaderPortsAPI::HasInput(ShaderGraphNode, Parameters::Shininess);
 			if (bHasShininessInput)
 			{
-				TGuardValue<bool> ParsingForLinearInputGuard(bParsingForLinearInput, true);
+				TGuardValue<EMaterialInputType> InputTypeBeingProcessedGuard(MaterialCreationContext.InputTypeBeingProcessed, EMaterialInputType::Scalar);
 
 				TTuple<UInterchangeMaterialExpressionFactoryNode*, FString> ShininessExpressionFactoryNode =
 					CreateMaterialExpressionForInput(MaterialFactoryNode, ShaderGraphNode, Parameters::Shininess.ToString(), MaterialFactoryNode->GetUniqueID());
@@ -402,7 +420,7 @@ bool UInterchangeGenericMaterialPipeline::HandlePBRModel(const UInterchangeShade
 
 		if (bHasInput)
 		{
-			TGuardValue<bool> ParsingForLinearInputGuard(bParsingForLinearInput, true);
+			TGuardValue<EMaterialInputType> InputTypeBeingProcessedGuard(MaterialCreationContext.InputTypeBeingProcessed, EMaterialInputType::Scalar);
 
 			TTuple<UInterchangeMaterialExpressionFactoryNode*, FString> ExpressionFactoryNode =
 				CreateMaterialExpressionForInput(MaterialFactoryNode, ShaderGraphNode, Parameters::Metallic.ToString(), MaterialFactoryNode->GetUniqueID());
@@ -422,7 +440,7 @@ bool UInterchangeGenericMaterialPipeline::HandlePBRModel(const UInterchangeShade
 
 		if (bHasInput)
 		{
-			TGuardValue<bool> ParsingForLinearInputGuard(bParsingForLinearInput, true);
+			TGuardValue<EMaterialInputType> InputTypeBeingProcessedGuard(MaterialCreationContext.InputTypeBeingProcessed, EMaterialInputType::Scalar);
 
 			TTuple<UInterchangeMaterialExpressionFactoryNode*, FString> ExpressionFactoryNode =
 				CreateMaterialExpressionForInput(MaterialFactoryNode, ShaderGraphNode, Parameters::Specular.ToString(), MaterialFactoryNode->GetUniqueID());
@@ -442,7 +460,7 @@ bool UInterchangeGenericMaterialPipeline::HandlePBRModel(const UInterchangeShade
 
 		if (bHasInput)
 		{
-			TGuardValue<bool> ParsingForLinearInputGuard(bParsingForLinearInput, true);
+			TGuardValue<EMaterialInputType> InputTypeBeingProcessedGuard(MaterialCreationContext.InputTypeBeingProcessed, EMaterialInputType::Scalar);
 
 			TTuple<UInterchangeMaterialExpressionFactoryNode*, FString> ExpressionFactoryNode =
 				CreateMaterialExpressionForInput(MaterialFactoryNode, ShaderGraphNode, Parameters::Roughness.ToString(), MaterialFactoryNode->GetUniqueID());
@@ -480,10 +498,9 @@ bool UInterchangeGenericMaterialPipeline::HandleStandardSurfaceModel(const UInte
 
 		using namespace UE::Interchange::Materials;
 
-		auto ConnectMaterialExpressionOutputToInput = [&](const FString & InputName, bool bLinearParsingValue, bool bNormalParsingValue)
+		auto ConnectMaterialExpressionOutputToInput = [&](const FString & InputName, EMaterialInputType InputType)
 		{
-			TGuardValue<bool> ParsingForLinearInputGuard(bParsingForLinearInput, bLinearParsingValue);
-			TGuardValue<bool> ParsingForNormalInputGuard(bParsingForNormalInput, bNormalParsingValue);
+			TGuardValue<EMaterialInputType> InputTypeBeingProcessedGuard(MaterialCreationContext.InputTypeBeingProcessed, InputType);
 
 			TTuple<UInterchangeMaterialExpressionFactoryNode*, FString> ExpressionFactoryNode =
 				CreateMaterialExpressionForInput(MaterialFactoryNode, ShaderGraphNode, InputName, FunctionCallExpression->GetUniqueID());
@@ -495,42 +512,41 @@ bool UInterchangeGenericMaterialPipeline::HandleStandardSurfaceModel(const UInte
 			}
 		};
 
-		using ArgumentsType = TTuple<FString, bool, bool>;
+		using ArgumentsType = TTuple<FString, EMaterialInputType>;
 		
 		static ArgumentsType ArgumentsPerInputs[]
 		{
-			ArgumentsType{StandardSurface::Parameters::Base.ToString(), true, false},				// Base
-			ArgumentsType{StandardSurface::Parameters::BaseColor.ToString(), false, false},			// BaseColor
-			ArgumentsType{StandardSurface::Parameters::DiffuseRoughness.ToString(), false, false},	// DiffuseRoughness
-			ArgumentsType{StandardSurface::Parameters::Metalness.ToString(), true, false},			// Metalness
-			ArgumentsType{StandardSurface::Parameters::Specular.ToString(), true, false},			// Specular
-			ArgumentsType{StandardSurface::Parameters::SpecularRoughness.ToString(), true, false},	// SpecularRoughness
-			ArgumentsType{StandardSurface::Parameters::SpecularIOR.ToString(), true, false},        // SpecularIOR
-			ArgumentsType{StandardSurface::Parameters::SpecularAnisotropy.ToString(), true, false},	// SpecularAnisotropy
-			ArgumentsType{StandardSurface::Parameters::SpecularRotation.ToString(), true, false},	// SpecularRotation
-			ArgumentsType{StandardSurface::Parameters::Subsurface.ToString(), true, false},			// Subsurface
-			ArgumentsType{StandardSurface::Parameters::SubsurfaceColor.ToString(), true, false},	// SubsurfaceColor
-			ArgumentsType{StandardSurface::Parameters::SubsurfaceRadius.ToString(), true, false},	// SubsurfaceRadius
-			ArgumentsType{StandardSurface::Parameters::SubsurfaceScale.ToString(), true, false},	// SubsurfaceScale
-			ArgumentsType{StandardSurface::Parameters::Sheen.ToString(), true, false},				// Sheen
-			ArgumentsType{StandardSurface::Parameters::SheenColor.ToString(), true, false},			// SheenColor
-			ArgumentsType{StandardSurface::Parameters::SheenRoughness.ToString(), true, false},		// SheenRoughness
-			ArgumentsType{StandardSurface::Parameters::Coat.ToString(), true, false},				// Coat
-			ArgumentsType{StandardSurface::Parameters::CoatColor.ToString(), false, false},			// Coat Color
-			ArgumentsType{StandardSurface::Parameters::CoatRoughness.ToString(), true, false},		// CoatRoughness
-			ArgumentsType{StandardSurface::Parameters::ThinFilmThickness.ToString(), true, false},	// ThinFilmThickness
-			ArgumentsType{StandardSurface::Parameters::Emission.ToString(), true, false},			// Emission
-			ArgumentsType{StandardSurface::Parameters::EmissionColor.ToString(), false, false},		// EmissionColor
-			ArgumentsType{StandardSurface::Parameters::Normal.ToString(), false, true},				// Normal
-			ArgumentsType{StandardSurface::Parameters::Tangent.ToString(), false, true}				// Tangent
+			ArgumentsType{StandardSurface::Parameters::Base.ToString(), EMaterialInputType::Scalar},					// Base
+			ArgumentsType{StandardSurface::Parameters::BaseColor.ToString(), EMaterialInputType::Color},				// BaseColor
+			ArgumentsType{StandardSurface::Parameters::DiffuseRoughness.ToString(),  EMaterialInputType::Scalar},		// DiffuseRoughness
+			ArgumentsType{StandardSurface::Parameters::Metalness.ToString(), EMaterialInputType::Scalar},			// Metalness
+			ArgumentsType{StandardSurface::Parameters::Specular.ToString(), EMaterialInputType::Scalar},				// Specular
+			ArgumentsType{StandardSurface::Parameters::SpecularRoughness.ToString(), EMaterialInputType::Scalar},	// SpecularRoughness
+			ArgumentsType{StandardSurface::Parameters::SpecularIOR.ToString(), EMaterialInputType::Scalar},			// SpecularIOR
+			ArgumentsType{StandardSurface::Parameters::SpecularAnisotropy.ToString(), EMaterialInputType::Scalar},	// SpecularAnisotropy
+			ArgumentsType{StandardSurface::Parameters::SpecularRotation.ToString(), EMaterialInputType::Scalar},		// SpecularRotation
+			ArgumentsType{StandardSurface::Parameters::Subsurface.ToString(), EMaterialInputType::Scalar},			// Subsurface
+			ArgumentsType{StandardSurface::Parameters::SubsurfaceColor.ToString(), EMaterialInputType::Scalar},		// SubsurfaceColor
+			ArgumentsType{StandardSurface::Parameters::SubsurfaceRadius.ToString(), EMaterialInputType::Scalar},		// SubsurfaceRadius
+			ArgumentsType{StandardSurface::Parameters::SubsurfaceScale.ToString(), EMaterialInputType::Scalar},		// SubsurfaceScale
+			ArgumentsType{StandardSurface::Parameters::Sheen.ToString(), EMaterialInputType::Scalar},				// Sheen
+			ArgumentsType{StandardSurface::Parameters::SheenColor.ToString(), EMaterialInputType::Color},			// SheenColor
+			ArgumentsType{StandardSurface::Parameters::SheenRoughness.ToString(), EMaterialInputType::Scalar},		// SheenRoughness
+			ArgumentsType{StandardSurface::Parameters::Coat.ToString(), EMaterialInputType::Scalar},					// Coat
+			ArgumentsType{StandardSurface::Parameters::CoatColor.ToString(),  EMaterialInputType::Color},			// Coat Color
+			ArgumentsType{StandardSurface::Parameters::CoatRoughness.ToString(), EMaterialInputType::Scalar},		// CoatRoughness
+			ArgumentsType{StandardSurface::Parameters::ThinFilmThickness.ToString(), EMaterialInputType::Scalar},	// ThinFilmThickness
+			ArgumentsType{StandardSurface::Parameters::Emission.ToString(), EMaterialInputType::Scalar},				// Emission
+			ArgumentsType{StandardSurface::Parameters::EmissionColor.ToString(), EMaterialInputType::Color},			// EmissionColor
+			ArgumentsType{StandardSurface::Parameters::Normal.ToString(), EMaterialInputType::Vector},				// Normal
+			ArgumentsType{StandardSurface::Parameters::Tangent.ToString(), EMaterialInputType::Vector}				// Tangent
 		};
 
 		for(const ArgumentsType & Arguments : ArgumentsPerInputs)
 		{
 			const FString& InputName = Arguments.Get<0>();
-			bool bLinearParsingValue = Arguments.Get<1>();
-			bool bNormalParsingValue = Arguments.Get<2>();
-			ConnectMaterialExpressionOutputToInput(InputName, bLinearParsingValue, bNormalParsingValue);
+			EMaterialInputType InputType = Arguments.Get<1>();
+			ConnectMaterialExpressionOutputToInput(InputName, InputType);
 		}
 		
 		MaterialFactoryNode->ConnectOutputToBaseColor(FunctionCallExpressionUid,  PBR::Parameters::BaseColor.ToString());
@@ -586,7 +602,7 @@ bool UInterchangeGenericMaterialPipeline::HandleClearCoat(const UInterchangeShad
 
 		if (bHasInput)
 		{
-			TGuardValue<bool> ParsingForLinearInputGuard(bParsingForLinearInput, true);
+			TGuardValue<EMaterialInputType> InputTypeBeingProcessedGuard(MaterialCreationContext.InputTypeBeingProcessed, EMaterialInputType::Scalar);
 
 			TTuple<UInterchangeMaterialExpressionFactoryNode*, FString> ExpressionFactoryNode =
 				CreateMaterialExpressionForInput(MaterialFactoryNode, ShaderGraphNode, Parameters::ClearCoat.ToString(), MaterialFactoryNode->GetUniqueID());
@@ -606,7 +622,7 @@ bool UInterchangeGenericMaterialPipeline::HandleClearCoat(const UInterchangeShad
 
 		if (bHasInput)
 		{
-			TGuardValue<bool> ParsingForLinearInputGuard(bParsingForLinearInput, true);
+			TGuardValue<EMaterialInputType> InputTypeBeingProcessedGuard(MaterialCreationContext.InputTypeBeingProcessed, EMaterialInputType::Scalar);
 
 			TTuple<UInterchangeMaterialExpressionFactoryNode*, FString> ExpressionFactoryNode =
 				CreateMaterialExpressionForInput(MaterialFactoryNode, ShaderGraphNode, Parameters::ClearCoatRoughness.ToString(), MaterialFactoryNode->GetUniqueID());
@@ -626,7 +642,7 @@ bool UInterchangeGenericMaterialPipeline::HandleClearCoat(const UInterchangeShad
 
 		if (bHasInput)
 		{
-			TGuardValue<bool> ParsingForNormalInputGuard(bParsingForNormalInput, true);
+			TGuardValue<EMaterialInputType> InputTypeBeingProcessedGuard(MaterialCreationContext.InputTypeBeingProcessed, EMaterialInputType::Vector);
 
 			TTuple<UInterchangeMaterialExpressionFactoryNode*, FString> ExpressionFactoryNode =
 				CreateMaterialExpressionForInput(MaterialFactoryNode, ShaderGraphNode, Parameters::ClearCoatNormal.ToString(), MaterialFactoryNode->GetUniqueID());
@@ -804,7 +820,7 @@ void UInterchangeGenericMaterialPipeline::HandleCommonParameters(const UIntercha
 
 		if (bHasNormalInput)
 		{
-			TGuardValue<bool> ParsingForNormalInputGuard(bParsingForNormalInput, true);
+			TGuardValue<EMaterialInputType> InputTypeBeingProcessedGuard(MaterialCreationContext.InputTypeBeingProcessed, EMaterialInputType::Vector);
 
 			TTuple<UInterchangeMaterialExpressionFactoryNode*, FString> ExpressionFactoryNode =
 				CreateMaterialExpressionForInput(MaterialFactoryNode, ShaderGraphNode, Parameters::Normal.ToString(), MaterialFactoryNode->GetUniqueID());
@@ -818,7 +834,7 @@ void UInterchangeGenericMaterialPipeline::HandleCommonParameters(const UIntercha
 
 	// Opacity
 	{
-		TGuardValue<bool> ParsingForLinearInputGuard(bParsingForLinearInput, true);
+		TGuardValue<EMaterialInputType> InputTypeBeingProcessedGuard(MaterialCreationContext.InputTypeBeingProcessed, EMaterialInputType::Scalar);
 
 		const bool bHasOpacityInput = UInterchangeShaderPortsAPI::HasInput(ShaderGraphNode, Parameters::Opacity);
 
@@ -871,7 +887,7 @@ void UInterchangeGenericMaterialPipeline::HandleCommonParameters(const UIntercha
 
 	// Ambient Occlusion
 	{
-		TGuardValue<bool> ParsingForLinearInputGuard(bParsingForLinearInput, true);
+		TGuardValue<EMaterialInputType> InputTypeBeingProcessedGuard(MaterialCreationContext.InputTypeBeingProcessed, EMaterialInputType::Scalar);
 
 		const bool bHasOcclusionInput = UInterchangeShaderPortsAPI::HasInput(ShaderGraphNode, Parameters::Occlusion);
 
@@ -890,7 +906,7 @@ void UInterchangeGenericMaterialPipeline::HandleCommonParameters(const UIntercha
 	// Index of Refraction (IOR)
 	// We'll lerp between Air IOR (1) and the IOR from the shader graph based on a fresnel, as per UE doc on refraction.
 	{
-		TGuardValue<bool> ParsingForLinearInputGuard(bParsingForLinearInput, true);
+		TGuardValue<EMaterialInputType> InputTypeBeingProcessedGuard(MaterialCreationContext.InputTypeBeingProcessed, EMaterialInputType::Scalar);
 
 		const bool bHasIorInput = UInterchangeShaderPortsAPI::HasInput(ShaderGraphNode, Parameters::IndexOfRefraction);
 
@@ -1033,22 +1049,49 @@ void UInterchangeGenericMaterialPipeline::HandleTextureSampleNode(const UInterch
 		TextureSampleFactoryNode->SetCustomExpressionClassName(ExpressionClassName);
 		TextureSampleFactoryNode->AddStringAttribute(UInterchangeShaderPortsAPI::MakeInputValueKey(Inputs::Texture.ToString()), TextureFactoryUid);
 
-		if (bParsingForNormalInput)
+		if (UInterchangeTextureFactoryNode* TextureFactoryNode = Cast<UInterchangeTextureFactoryNode>(BaseNodeContainer->GetFactoryNode(TextureFactoryUid)))
 		{
-			if (UInterchangeTextureFactoryNode* TextureFactoryNode = Cast<UInterchangeTextureFactoryNode>(BaseNodeContainer->GetFactoryNode(TextureFactoryUid)))
+			EMaterialInputType TextureUsage = EMaterialInputType::Unknown;
+			TextureFactoryNode->GetAttribute(TEXT("TextureUsage"), TextureUsage);
+
+			const bool bIsOutputLinear = MaterialExpressionCreationContextStack.Top().OutputName.Equals(Outputs::A.ToString());
+			const EMaterialInputType DesiredTextureUsage = MaterialCreationContext.InputTypeBeingProcessed == EMaterialInputType::Scalar && bIsOutputLinear ?
+																		EMaterialInputType::Unknown : // Alpha channels are always in linear space so ignore them when determining texture usage
+																		MaterialCreationContext.InputTypeBeingProcessed;
+
+			if (TextureUsage == EMaterialInputType::Unknown)
 			{
-				TextureFactoryNode->SetCustomCompressionSettings(TextureCompressionSettings::TC_Normalmap);
+				if (DesiredTextureUsage == EMaterialInputType::Vector)
+				{
+					TextureFactoryNode->SetCustomCompressionSettings(TextureCompressionSettings::TC_Normalmap);
+				}
+				else if (DesiredTextureUsage == EMaterialInputType::Scalar)
+				{
+					bool bSRGB;
+					if (!TextureNode->GetCustomSRGB(bSRGB))
+					{
+						// Only set CustomSRGB if it wasn't set by the InterchangeGenericTexturePipeline before
+						TextureFactoryNode->SetCustomSRGB(false);
+					}
+				}
+
+				TextureFactoryNode->SetAttribute(TEXT("TextureUsage"), DesiredTextureUsage);
 			}
-		}
-		else if (bParsingForLinearInput)
-		{
-			if (UInterchangeTextureFactoryNode* TextureFactoryNode = Cast<UInterchangeTextureFactoryNode>(BaseNodeContainer->GetFactoryNode(TextureFactoryUid)))
+			else if (TextureUsage != DesiredTextureUsage && DesiredTextureUsage != EMaterialInputType::Unknown)
 			{
-				bool bSRGB;
-				if (!TextureNode->GetCustomSRGB(bSRGB))
+				UInterchangeResultWarning_Generic* TextureUsageWarning = AddMessage<UInterchangeResultWarning_Generic>();
+				TextureUsageWarning->DestinationAssetName = TextureFactoryNode->GetAssetName();
+				TextureUsageWarning->AssetType = TextureFactoryNode->GetObjectClass();
+
+				TextureUsageWarning->Text = FText::Format(LOCTEXT("TextureUsageMismatch", "{0} is being used as both {1} and {2} which aren't compatible."),
+					FText::FromString(TextureFactoryNode->GetAssetName()), FText::FromString(LexToString(TextureUsage)), FText::FromString(LexToString(DesiredTextureUsage)));
+
+				// Flipping the green channel only makes sense for vector data as it's used to compensate for different handedness.
+				// Clear it if we're not gonna be used only as a vector map. This normally happens when a normal map is also used as a color map.
+				bool bFlipGreenChannel;
+				if (TextureFactoryNode->GetCustombFlipGreenChannel(bFlipGreenChannel))
 				{
-					//Only set CustomSRGB if it wasn't set by the InterchangeGenericTexturePipeline before
-					TextureFactoryNode->SetCustomSRGB(false);
+					TextureFactoryNode->SetCustombFlipGreenChannel(false);
 				}
 			}
 		}
@@ -1493,9 +1536,10 @@ TTuple<UInterchangeMaterialExpressionFactoryNode*, FString> UInterchangeGenericM
 
 	UInterchangeMaterialExpressionFactoryNode* MaterialExpressionFactoryNode = nullptr;
 
+	FMaterialExpressionCreationContext& ExpressionCreationContext = MaterialExpressionCreationContextStack.AddDefaulted_GetRef();
+
 	FString ConnectedShaderNodeUid;
-	FString OutputName;
-	if (UInterchangeShaderPortsAPI::GetInputConnection(ShaderNode, InputName, ConnectedShaderNodeUid, OutputName))
+	if (UInterchangeShaderPortsAPI::GetInputConnection(ShaderNode, InputName, ConnectedShaderNodeUid, ExpressionCreationContext.OutputName))
 	{
 		if (const UInterchangeShaderNode* ConnectedShaderNode = Cast<const UInterchangeShaderNode>(BaseNodeContainer->GetNode(ConnectedShaderNodeUid)))
 		{
@@ -1518,7 +1562,10 @@ TTuple<UInterchangeMaterialExpressionFactoryNode*, FString> UInterchangeGenericM
 		}
 	}
 
-	return TTuple<UInterchangeMaterialExpressionFactoryNode*, FString>{MaterialExpressionFactoryNode, OutputName};
+	TTuple<UInterchangeMaterialExpressionFactoryNode*, FString> Result {MaterialExpressionFactoryNode, ExpressionCreationContext.OutputName};
+	MaterialExpressionCreationContextStack.Pop();
+
+	return Result;
 }
 
 UInterchangeMaterialFactoryNode* UInterchangeGenericMaterialPipeline::CreateMaterialFactoryNode(const UInterchangeShaderGraphNode* ShaderGraphNode)
@@ -1951,3 +1998,5 @@ UInterchangeMaterialFunctionFactoryNode* UInterchangeGenericMaterialPipeline::Cr
 
 	return FactoryNode;
 }
+
+#undef LOCTEXT_NAMESPACE
diff --git a/Engine/Plugins/Experimental/Interchange/Runtime/Source/Pipelines/Public/InterchangeGenericMaterialPipeline.h b/Engine/Plugins/Experimental/Interchange/Runtime/Source/Pipelines/Public/InterchangeGenericMaterialPipeline.h
--- a/Engine/Plugins/Experimental/Interchange/Runtime/Source/Pipelines/Public/InterchangeGenericMaterialPipeline.h
+++ b/Engine/Plugins/Experimental/Interchange/Runtime/Source/Pipelines/Public/InterchangeGenericMaterialPipeline.h
@@ -133,8 +133,26 @@ private:
 	TVariant<FString, FLinearColor, float> VisitTextureSampleNode(const UInterchangeShaderNode* ShaderNode) const;
 
 private:
-	bool bParsingForNormalInput;
-	bool bParsingForLinearInput; // True when parsing non-color inputs (metallic, roughness, specular, etc.)
+	enum class EMaterialInputType : uint8
+	{
+		Unknown,
+		Color,
+		Vector,
+		Scalar
+	};
+	friend FString LexToString(UInterchangeGenericMaterialPipeline::EMaterialInputType);
+
+	struct FMaterialCreationContext
+	{
+		EMaterialInputType InputTypeBeingProcessed = EMaterialInputType::Color;
+	} MaterialCreationContext;
+
+	struct FMaterialExpressionCreationContext
+	{
+		FString OutputName; // The name of the output we will be connecting from
+	};
+
+	TArray<FMaterialExpressionCreationContext> MaterialExpressionCreationContextStack;
 };
 
 
