PR #9375: Common typo: doens't -> doesn't (Contributed by [REDACTED_PERSON])

diff --git a/Engine/Plugins/2D/Paper2D/Source/PaperTiledImporter/Private/PaperTiledImporterFactory.cpp b/Engine/Plugins/2D/Paper2D/Source/PaperTiledImporter/Private/PaperTiledImporterFactory.cpp
--- a/Engine/Plugins/2D/Paper2D/Source/PaperTiledImporter/Private/PaperTiledImporterFactory.cpp
+++ b/Engine/Plugins/2D/Paper2D/Source/PaperTiledImporter/Private/PaperTiledImporterFactory.cpp
@@ -952,7 +952,7 @@ bool FTiledTileInfo::ParseTileInfoFromJSON(int32 TileIndex, TSharedPtr<FJsonObje
 	}
 
 	// Try reading the per-tile collision data if present
-	// Note: This is really an entire fake objectgroup layer, but only the objects array matters; Tiled doens't even provide a way to edit the rest of the data.
+	// Note: This is really an entire fake objectgroup layer, but only the objects array matters; Tiled doesn't even provide a way to edit the rest of the data.
 	const TSharedPtr<FJsonObject>* ObjectGroupSubobject;
 	if (Tree->TryGetObjectField(TEXT("objectgroup"), /*out*/ ObjectGroupSubobject))
 	{
diff --git a/Engine/Plugins/Animation/ControlRig/Source/ControlRigEditor/Private/Editor/SRigCurveContainer.cpp b/Engine/Plugins/Animation/ControlRig/Source/ControlRigEditor/Private/Editor/SRigCurveContainer.cpp
--- a/Engine/Plugins/Animation/ControlRig/Source/ControlRigEditor/Private/Editor/SRigCurveContainer.cpp
+++ b/Engine/Plugins/Animation/ControlRig/Source/ControlRigEditor/Private/Editor/SRigCurveContainer.cpp
@@ -634,7 +634,7 @@ void SRigCurveContainer::CreateImportMenu(FMenuBuilder& MenuBuilder)
 			SNew(STextBlock)
 			.Font(FAppStyle::GetFontStyle("ControlRig.Curve.Menu"))
 			.Text(LOCTEXT("ImportMesh_Title", "Select Mesh"))
-			.ToolTipText(LOCTEXT("ImportMesh_Tooltip", "Select Mesh to import Curve from... It will only import if the node doens't exists in the current Curve."))
+			.ToolTipText(LOCTEXT("ImportMesh_Tooltip", "Select Mesh to import Curve from... It will only import if the node doesn't exist in the current Curve."))
 		]
 		+ SVerticalBox::Slot()
 		.AutoHeight()
diff --git a/Engine/Plugins/Animation/ControlRig/Source/ControlRigEditor/Private/Editor/SRigHierarchy.cpp b/Engine/Plugins/Animation/ControlRig/Source/ControlRigEditor/Private/Editor/SRigHierarchy.cpp
--- a/Engine/Plugins/Animation/ControlRig/Source/ControlRigEditor/Private/Editor/SRigHierarchy.cpp
+++ b/Engine/Plugins/Animation/ControlRig/Source/ControlRigEditor/Private/Editor/SRigHierarchy.cpp
@@ -1392,7 +1392,7 @@ void SRigHierarchy::CreateImportMenu(FMenuBuilder& MenuBuilder)
 			SNew(STextBlock)
 			.Font(FAppStyle::GetFontStyle("ControlRig.Hierarchy.Menu"))
 			.Text(LOCTEXT("ImportMesh_Title", "Select Mesh"))
-			.ToolTipText(LOCTEXT("ImportMesh_Tooltip", "Select Mesh to import hierarchy from... It will only import if the node doens't exists in the current hierarchy."))
+			.ToolTipText(LOCTEXT("ImportMesh_Tooltip", "Select Mesh to import hierarchy from... It will only import if the node doesn't exist in the current hierarchy."))
 		]
 		+ SVerticalBox::Slot()
 		.AutoHeight()
diff --git a/Engine/Plugins/Experimental/Interchange/Runtime/Source/FactoryNodes/Public/InterchangeAnimSequenceFactoryNode.h b/Engine/Plugins/Experimental/Interchange/Runtime/Source/FactoryNodes/Public/InterchangeAnimSequenceFactoryNode.h
--- a/Engine/Plugins/Experimental/Interchange/Runtime/Source/FactoryNodes/Public/InterchangeAnimSequenceFactoryNode.h
+++ b/Engine/Plugins/Experimental/Interchange/Runtime/Source/FactoryNodes/Public/InterchangeAnimSequenceFactoryNode.h
@@ -128,7 +128,7 @@ public:
 	/**
 	 * Get the custom attribute DoNotImportCurveWithZero, return false if the attribute is not set.
 	 * 
-	 * Note - If value is true, do not import if it doens't have any value other than zero. This is to avoid adding extra curves to evaluate.
+	 * Note - If value is true, do not import if it doesn't have any value other than zero. This is to avoid adding extra curves to evaluate.
 	 */
 	UFUNCTION(BlueprintCallable, Category = "Interchange | Node | AnimSequence")
 	bool GetCustomDoNotImportCurveWithZero(bool& AttributeValue) const;
@@ -136,7 +136,7 @@ public:
 	/**
 	 * Set the custom attribute DoNotImportCurveWithZero. Return false if the attribute cannot be set.
 	 * 
-	 * Note - If value is true, do not import if it doens't have any value other than zero. This is to avoid adding extra curves to evaluate.
+	 * Note - If value is true, do not import if it doesn't have any value other than zero. This is to avoid adding extra curves to evaluate.
 	 */
 	UFUNCTION(BlueprintCallable, Category = "Interchange | Node | AnimSequence")
 	bool SetCustomDoNotImportCurveWithZero(const bool& AttributeValue);
diff --git a/Engine/Plugins/Experimental/Interchange/Runtime/Source/Pipelines/Public/InterchangeGenericAnimationPipeline.h b/Engine/Plugins/Experimental/Interchange/Runtime/Source/Pipelines/Public/InterchangeGenericAnimationPipeline.h
--- a/Engine/Plugins/Experimental/Interchange/Runtime/Source/Pipelines/Public/InterchangeGenericAnimationPipeline.h
+++ b/Engine/Plugins/Experimental/Interchange/Runtime/Source/Pipelines/Public/InterchangeGenericAnimationPipeline.h
@@ -81,7 +81,7 @@ public:
 	UPROPERTY(EditAnywhere, BlueprintReadWrite, Category = "Animations", meta = (EditCondition = "bImportCustomAttribute", DisplayName = "Remove Redundant Keys"))
 	bool bRemoveCurveRedundantKeys = false;
 
-	/** When importing custom attribute or morphtarget as curve, do not import if it doens't have any value other than zero. This is to avoid adding extra curves to evaluate */
+	/** When importing custom attribute or morphtarget as curve, do not import if it doesn't have any value other than zero. This is to avoid adding extra curves to evaluate */
 	UPROPERTY(EditAnywhere, BlueprintReadWrite, Category = "Animations", meta = (EditCondition = "bImportCustomAttribute", DisplayName = "Do not import curves with only 0 values"))
 	bool bDoNotImportCurveWithZero = false;
 
diff --git a/Engine/Source/Developer/MeshBoneReduction/Private/MeshBoneReduction.cpp b/Engine/Source/Developer/MeshBoneReduction/Private/MeshBoneReduction.cpp
--- a/Engine/Source/Developer/MeshBoneReduction/Private/MeshBoneReduction.cpp
+++ b/Engine/Source/Developer/MeshBoneReduction/Private/MeshBoneReduction.cpp
@@ -168,7 +168,7 @@ public:
 					}
 					else
 					{
-						// if parent doens't exists, we have to add one
+						// if parent doesn't exist, we have to add one
 						// this doesn't change bone map size 
 						Section.BoneMap.RemoveAt(BoneMapIndex);
 						ParentBoneMapIndex = Section.BoneMap.Add(*ParentBoneIndex);
diff --git a/Engine/Source/Editor/MovieSceneTools/Private/AutomatedLevelSequenceCapture.cpp b/Engine/Source/Editor/MovieSceneTools/Private/AutomatedLevelSequenceCapture.cpp
--- a/Engine/Source/Editor/MovieSceneTools/Private/AutomatedLevelSequenceCapture.cpp
+++ b/Engine/Source/Editor/MovieSceneTools/Private/AutomatedLevelSequenceCapture.cpp
@@ -577,7 +577,7 @@ void UAutomatedLevelSequenceCapture::SetupFrameRange()
 				// This is a fun hack... Due to the fragility of this code (which this makes more fragile admittedly...) the original audio implementation
 				// just ran the entire process twice, starting all the way back at the Setup loop so that everything would be re-initialized like it was
 				// for the video capture. Unfortunately, it looks like when we switch from a Fixed Timestep clock to the Platform clock (needed for audio
-				// as audio is realtime only) this allows the sequence player to get out of sync with the warmup frame counter. Audio recording doens't start
+				// as audio is realtime only) this allows the sequence player to get out of sync with the warmup frame counter. Audio recording doesn't start
 				// until the warmup time has passed, but because the sequence is playing at realtime it can start playing the part of the sequence you wanted
 				// the audio for before the recorder ever kicks in!
 				// To minimize changes, we're just going to override any warmup/delay times to zero here in the event that this is an audio pass, so that
diff --git a/Engine/Source/Editor/SkeletalMeshEditor/Private/SControlRigMappingWindow.cpp b/Engine/Source/Editor/SkeletalMeshEditor/Private/SControlRigMappingWindow.cpp
--- a/Engine/Source/Editor/SkeletalMeshEditor/Private/SControlRigMappingWindow.cpp
+++ b/Engine/Source/Editor/SkeletalMeshEditor/Private/SControlRigMappingWindow.cpp
@@ -422,7 +422,7 @@ void SControlRigMappingWindow::CreateBoneMappingList(const FString& SearchText,
 
 				if (bDoFiltering)
 				{
-					// make sure it doens't fit any of them
+					// make sure it doesn't fit any of them
 					if (!DisplayName.Contains(SearchText) && (BoneName && !(*BoneName).ToString().Contains(SearchText)))
 					{
 						continue; // Skip items that don't match our filter
diff --git a/Engine/Source/Editor/UnrealEd/Classes/Factories/FbxAnimSequenceImportData.h b/Engine/Source/Editor/UnrealEd/Classes/Factories/FbxAnimSequenceImportData.h
--- a/Engine/Source/Editor/UnrealEd/Classes/Factories/FbxAnimSequenceImportData.h
+++ b/Engine/Source/Editor/UnrealEd/Classes/Factories/FbxAnimSequenceImportData.h
@@ -102,7 +102,7 @@ class UNREALED_API UFbxAnimSequenceImportData : public UFbxAssetImportData
 	UPROPERTY(EditAnywhere, AdvancedDisplay, config, Category = ImportSettings)
 	bool bDeleteExistingMorphTargetCurves;
 
-	/** When importing custom attribute or morphtarget as curve, do not import if it doens't have any value other than zero. This is to avoid adding extra curves to evaluate */
+	/** When importing custom attribute or morphtarget as curve, do not import if it doesn't have any value other than zero. This is to avoid adding extra curves to evaluate */
 	UPROPERTY(EditAnywhere, AdvancedDisplay, config, Category = ImportSettings, meta = (DisplayName = "Do not import curves with only 0 values"))
 	bool bDoNotImportCurveWithZero;
 
diff --git a/Engine/Source/Runtime/AnimGraphRuntime/Private/AnimSequencerInstanceProxy.cpp b/Engine/Source/Runtime/AnimGraphRuntime/Private/AnimSequencerInstanceProxy.cpp
--- a/Engine/Source/Runtime/AnimGraphRuntime/Private/AnimSequencerInstanceProxy.cpp
+++ b/Engine/Source/Runtime/AnimGraphRuntime/Private/AnimSequencerInstanceProxy.cpp
@@ -218,7 +218,7 @@ void FAnimSequencerInstanceProxy::InitAnimTrack(UAnimSequenceBase* InAnimSequenc
 // where we could just clear one sequence id. We just clear all the weights before update. 
 // once they go out of range, they don't get called anymore, so there is no good point of tearing down
 // there is multiple tear down point but we couldn't find where only happens once activated and once getting out
-// because sequencer finds the nearest point, not exact point, it doens't have good point of tearing down
+// because sequencer finds the nearest point, not exact point, it doesn't have good point of tearing down
 void FAnimSequencerInstanceProxy::TermAnimTrack(int32 SequenceId)
 {
 	FSequencerPlayerState* PlayerState = FindPlayer(SequenceId);
diff --git a/Engine/Source/Runtime/Core/Private/Misc/ConfigCacheIni.cpp b/Engine/Source/Runtime/Core/Private/Misc/ConfigCacheIni.cpp
--- a/Engine/Source/Runtime/Core/Private/Misc/ConfigCacheIni.cpp
+++ b/Engine/Source/Runtime/Core/Private/Misc/ConfigCacheIni.cpp
@@ -3298,7 +3298,7 @@ FString FConfigCacheIni::GetDestIniFilename(const TCHAR* BaseIniName, const TCHA
 	{
 		FString Name(PlatformName ? PlatformName : ANSI_TO_TCHAR(FPlatformProperties::PlatformName()));
 
-		// if the BaseIniName doens't contain the config dir, put it all together
+		// if the BaseIniName doesn't contain the config dir, put it all together
 		if (FCString::Stristr(BaseIniName, GeneratedConfigDir) != nullptr)
 		{
 			IniFilename = BaseIniName;
diff --git a/Engine/Source/Runtime/Engine/Classes/Engine/World.h b/Engine/Source/Runtime/Engine/Classes/Engine/World.h
--- a/Engine/Source/Runtime/Engine/Classes/Engine/World.h
+++ b/Engine/Source/Runtime/Engine/Classes/Engine/World.h
@@ -1345,7 +1345,7 @@ public:
 	 */
 	uint8 bKismetScriptError:1;
 
-	// Kismet debugging flags - they can be only editor only, but they're uint32, so it doens't make much difference
+	// Kismet debugging flags - they can be only editor only, but they're uint32, so it doesn't make much difference
 	uint8 bDebugPauseExecution:1;
 
 	/** When set, camera is potentially moveable even when paused */
diff --git a/Engine/Source/Runtime/Engine/Private/Animation/AnimationRuntime.cpp b/Engine/Source/Runtime/Engine/Private/Animation/AnimationRuntime.cpp
--- a/Engine/Source/Runtime/Engine/Private/Animation/AnimationRuntime.cpp
+++ b/Engine/Source/Runtime/Engine/Private/Animation/AnimationRuntime.cpp
@@ -485,7 +485,7 @@ void FAnimationRuntime::BlendTwoPosesTogetherPerBone(const FAnimationPoseData& S
 		{
 			OutPose[BoneIndex] = SourcePoseTwo[BoneIndex];
 		}
-		// if it doens't have weight, take source pose 1
+		// if it doesn't have weight, take source pose 1
 		else if (FAnimationRuntime::HasWeight(BlendWeight))
 		{
 			BlendTransform<ETransformBlendMode::Overwrite>(SourcePoseOne[BoneIndex], OutPose[BoneIndex], 1.f - BlendWeight);
diff --git a/Engine/Source/Runtime/Engine/Private/AnimationAsset.cpp b/Engine/Source/Runtime/Engine/Private/AnimationAsset.cpp
--- a/Engine/Source/Runtime/Engine/Private/AnimationAsset.cpp
+++ b/Engine/Source/Runtime/Engine/Private/AnimationAsset.cpp
@@ -489,7 +489,7 @@ bool UAnimationAsset::ReplaceSkeleton(USkeleton* NewSkeleton, bool bConvertSpace
 
 bool UAnimationAsset::GetAllAnimationSequencesReferred(TArray<UAnimationAsset*>& AnimationSequences, bool bRecursive /*= true*/) 
 {
-	//@todo:@fixme: this doens't work for retargeting because postload gets called after duplication, mixing up the mapping table
+	//@todo:@fixme: this doesn't work for retargeting because postload gets called after duplication, mixing up the mapping table
 	// because skeleton changes, for now we don't support retargeting for parent asset, it will disconnect, and just duplicate everything else
 // 	if (ParentAsset)
 // 	{
@@ -516,7 +516,7 @@ void UAnimationAsset::HandleAnimReferenceCollection(TArray<UAnimationAsset*>& An
 
 void UAnimationAsset::ReplaceReferredAnimations(const TMap<UAnimationAsset*, UAnimationAsset*>& ReplacementMap)
 {
-	//@todo:@fixme: this doens't work for retargeting because postload gets called after duplication, mixing up the mapping table
+	//@todo:@fixme: this doesn't work for retargeting because postload gets called after duplication, mixing up the mapping table
 	// because skeleton changes, for now we don't support retargeting for parent asset, it will disconnect, and just duplicate everything else
 	if (ParentAsset)
 	{
diff --git a/Engine/Source/Runtime/Engine/Private/Collision/CollisionProfile.cpp b/Engine/Source/Runtime/Engine/Private/Collision/CollisionProfile.cpp
--- a/Engine/Source/Runtime/Engine/Private/Collision/CollisionProfile.cpp
+++ b/Engine/Source/Runtime/Engine/Private/Collision/CollisionProfile.cpp
@@ -279,7 +279,7 @@ void UCollisionProfile::LoadProfileConfig(bool bForceInit)
 	// before any op, verify if profiles contains invalid name - such as Custom profile name - remove all of them
 	for (auto Iter=Profiles.CreateConstIterator(); Iter; ++Iter)
 	{
-		// make sure it doens't have any 
+		// make sure it doesn't have any 
 		if (Iter->Name == CustomCollisionProfileName)
 		{
 			UE_LOG(LogCollisionProfile, Error, TEXT("Profiles contain invalid name : %s is reserved for internal use"), *CustomCollisionProfileName.ToString());
diff --git a/Engine/Source/Runtime/SlateCore/Private/FastUpdate/SlateInvalidationRoot.cpp b/Engine/Source/Runtime/SlateCore/Private/FastUpdate/SlateInvalidationRoot.cpp
--- a/Engine/Source/Runtime/SlateCore/Private/FastUpdate/SlateInvalidationRoot.cpp
+++ b/Engine/Source/Runtime/SlateCore/Private/FastUpdate/SlateInvalidationRoot.cpp
@@ -1054,8 +1054,8 @@ void FSlateInvalidationRoot::ProcessAttributeUpdate()
 					AttributeItt.Advance();
 
 #if WITH_SLATE_DEBUGGING
-					ensureMsgf(PreviousVisibility == InvalidationWidget.Visibility, TEXT("The visibility of widget '%s' doens't match the previous visibility after the attribute update."), *FReflectionMetaData::GetWidgetDebugInfo(WidgetPtr));
-					ensureMsgf(PreviousLeafMostChildIndex == InvalidationWidget.LeafMostChildIndex, TEXT("The number of child of widget '%s' doens't match the previous count after the attribute update."), *FReflectionMetaData::GetWidgetDebugInfo(WidgetPtr));
+					ensureMsgf(PreviousVisibility == InvalidationWidget.Visibility, TEXT("The visibility of widget '%s' doesn't match the previous visibility after the attribute update."), *FReflectionMetaData::GetWidgetDebugInfo(WidgetPtr));
+					ensureMsgf(PreviousLeafMostChildIndex == InvalidationWidget.LeafMostChildIndex, TEXT("The number of child of widget '%s' doesn't match the previous count after the attribute update."), *FReflectionMetaData::GetWidgetDebugInfo(WidgetPtr));
 #endif
 				}
 				else
