MaterialX Translator:
- Fixing wrong inputting of Lerp, MaterialX inverts the inputs, the formula is Fx + G(1-y) while UE is the opposite
- Adding support of nodes from the math nodegroup:
  - Normalize
  - Sign
  - Sqrt
  - Magnitude
  - Dotproduct
  - Crossproduct
  - Transformpoint
  - Transformvector
  - Transformnormal
  - Exp
  - Tan
  - Acos
  - Asin
  - Atan2
  - Absval
  - Invert
  - Modulo
  - Rotate2d
  - Rotate3d
- Refactoring of the method MakeLinearColorFrom into GetLinearColor (only handles colors)
- Adding a function GetVector for vec2/3/4, because unlike a color, a vector doesn't have a colorspace
- Adding different material expression:
  - Append3Vector
  - Append4Vector
  - Length for magnitude
  - Exponential for exp node
- Adding a function CreateNode, that allows the creation of a node in the nodegraph, in case of inserting a node between 2 nodes, useful when we an operation has to be handled in several pass but is already handled by smaller operators
- Adding the node convert from the channel nodegroup

Generic material pipeline Standard Surface:
- Connecting coat_normal
- removing LinearParsingValue, now it is handled by the Texture node

diff --git a/Engine/Plugins/Experimental/Interchange/Runtime/Source/Import/Private/MaterialX/InterchangeMaterialXDefinitions.h b/Engine/Plugins/Experimental/Interchange/Runtime/Source/Import/Private/MaterialX/InterchangeMaterialXDefinitions.h
--- a/Engine/Plugins/Experimental/Interchange/Runtime/Source/Import/Private/MaterialX/InterchangeMaterialXDefinitions.h
+++ b/Engine/Plugins/Experimental/Interchange/Runtime/Source/Import/Private/MaterialX/InterchangeMaterialXDefinitions.h
@@ -149,34 +149,71 @@ namespace MaterialX
 
 	namespace Category
 	{
+		// Math nodes
 		static constexpr const char* Add = "add";
 		static constexpr const char* Sub = "subtract";
 		static constexpr const char* Multiply = "multiply";
+		static constexpr const char* Divide = "divide";
+		static constexpr const char* Modulo = "modulo";
+		static constexpr const char* Invert = "invert";
+		static constexpr const char* Absval = "absval";
+		static constexpr const char* Floor = "floor";
+		static constexpr const char* Ceil = "ceil";
+		static constexpr const char* Power = "power";
+		static constexpr const char* Sin = "sin";
+		static constexpr const char* Cos = "cos";
+		static constexpr const char* Tan = "tan";
+		static constexpr const char* Asin = "asin";
+		static constexpr const char* Acos = "acos";
+		static constexpr const char* Atan2 = "atan2";
+		static constexpr const char* Sqrt = "sqrt";
+		static constexpr const char* Ln = "ln";
+		static constexpr const char* Exp = "exp";
+		static constexpr const char* Sign = "sign";
 		static constexpr const char* Clamp = "clamp";
-		static constexpr const char* Max = "max";
 		static constexpr const char* Min = "min";
-		static constexpr const char* Mix = "mix";
+		static constexpr const char* Max = "max";
+		static constexpr const char* Normalize = "normalize";
+		static constexpr const char* Magnitude = "magnitude";
 		static constexpr const char* DotProduct = "dotproduct";
+		static constexpr const char* CrossProduct = "crossproduct";
+		static constexpr const char* TransformPoint = "transformpoint";
+		static constexpr const char* TransformVector = "transformvector";
+		static constexpr const char* TransformNormal = "transformnormal";
+		static constexpr const char* TransformMatrix = "transformmatrix";
+		static constexpr const char* NormalMap = "normalmap";
+		static constexpr const char* Transpose = "transpose";
+		static constexpr const char* Determinant = "determinant";
+		static constexpr const char* InvertMatrix = "invertmatrix";
+		static constexpr const char* Rotate2D = "rotate2d";
+		static constexpr const char* Rotate3D = "rotate3d";
+		static constexpr const char* ArrayAppend = "arrayappend";
+		// Compositing nodes
+		static constexpr const char* Mix = "mix";
+		// Channel nodes
+		static constexpr const char* Convert = "convert";
 		static constexpr const char* Combine2 = "combine2";
-		static constexpr const char* Combine3 = "combine2";
+		static constexpr const char* Combine3 = "combine3";
 		static constexpr const char* Combine4 = "combine4";
+		static constexpr const char* Extract = "extract";
+		// Procedural nodes 
 		static constexpr const char* Constant = "constant";
-		static constexpr const char* Cos = "cos";
-		static constexpr const char* Sin = "sin";
+		// Organization nodes 
 		static constexpr const char* Dot = "dot";
+		// Texture nodes 
 		static constexpr const char* Image = "image";
 		static constexpr const char* TiledImage = "tiledimage";
-		static constexpr const char* NormalMap = "normalmap";
-		static constexpr const char* Extract = "extract";
+		// Geometric nodes 
 		static constexpr const char* TexCoord = "texcoord";
+		// Light nodes
 		static constexpr const char* PointLight = "point_light";
 		static constexpr const char* DirectionalLight = "directional_light";
 		static constexpr const char* SpotLight = "spot_light";
-		static constexpr const char* StandardSurface = "standard_surface";
-		//static constexpr const char* Convert = "convert";
-		static constexpr const char* Power = "power";
+		// Adjustment nodes
 		static constexpr const char* RgbToHsv = "rgbtohsv";
 		static constexpr const char* HsvToRgb = "hsvtorgb";
+		// PBR
+		static constexpr const char* StandardSurface = "standard_surface";
 	}
 
 	namespace NodeDefinition
diff --git a/Engine/Plugins/Experimental/Interchange/Runtime/Source/Import/Private/MaterialX/InterchangeMaterialXTranslator.cpp b/Engine/Plugins/Experimental/Interchange/Runtime/Source/Import/Private/MaterialX/InterchangeMaterialXTranslator.cpp
--- a/Engine/Plugins/Experimental/Interchange/Runtime/Source/Import/Private/MaterialX/InterchangeMaterialXTranslator.cpp
+++ b/Engine/Plugins/Experimental/Interchange/Runtime/Source/Import/Private/MaterialX/InterchangeMaterialXTranslator.cpp
@@ -23,6 +23,7 @@ static FAutoConsoleVariableRef CCvarInterchangeEnableMaterialXImport(
 	TEXT("Whether MaterialX support is enabled."),
 	ECVF_Default);
 
+
 namespace mx = MaterialX;
 
 namespace UE::Interchange::MaterialX
@@ -60,32 +61,85 @@ UInterchangeMaterialXTranslator::UInterchangeMaterialXTranslator()
 #if WITH_EDITOR
 	:
 InputNamesMaterialX2UE{
-	{{TEXT(""),						 TEXT("in")},		TEXT("Input")},
-	{{TEXT(""),						 TEXT("in1")},		TEXT("A")},
-	{{TEXT(""),						 TEXT("in2")},		TEXT("B")},
-	{{TEXT(""),						 TEXT("fg")},		TEXT("A")},
-	{{TEXT(""),						 TEXT("bg")},		TEXT("B")},
-	{{TEXT(""),						 TEXT("mix")},		TEXT("Factor")},
-	{{TEXT(""),						 TEXT("low")},		TEXT("Min")},
-	{{TEXT(""),						 TEXT("high")},		TEXT("Max")},
-	{{TEXT(""),						 TEXT("texcoord")}, TEXT("Coordinates")},
-	{{MaterialX::Category::Power,    TEXT("in1")},		TEXT("Base")},
-	{{MaterialX::Category::Power,	 TEXT("in2")},		TEXT("Exponent")}},
+	{{TEXT(""),                         TEXT("bg")},		TEXT("B")},
+	{{TEXT(""),                         TEXT("fg")},		TEXT("A")},
+	{{TEXT(""),                         TEXT("high")},		TEXT("Max")},
+	{{TEXT(""),                         TEXT("in")},		TEXT("Input")},
+	{{TEXT(""),                         TEXT("in1")},		TEXT("A")},
+	{{TEXT(""),                         TEXT("in2")},		TEXT("B")},
+	{{TEXT(""),                         TEXT("in3")},		TEXT("C")},
+	{{TEXT(""),                         TEXT("in4")},		TEXT("D")},
+	{{TEXT(""),                         TEXT("low")},		TEXT("Min")},
+	{{TEXT(""),                         TEXT("mix")},		TEXT("Factor")},
+	{{TEXT(""),                         TEXT("texcoord")},  TEXT("Coordinates")},
+	{{MaterialX::Category::Atan2,       TEXT("in1")},       TEXT("Y")},
+	{{MaterialX::Category::Atan2,       TEXT("in2")},       TEXT("X")},
+	{{MaterialX::Category::Magnitude,   TEXT("in")},        TEXT("A")},
+	{{MaterialX::Category::Mix,         TEXT("fg")},        TEXT("B")},
+	{{MaterialX::Category::Mix,         TEXT("bg")},        TEXT("A")},
+	{{MaterialX::Category::Normalize,   TEXT("in")},        TEXT("VectorInput")},
+	{{MaterialX::Category::Power,       TEXT("in1")},       TEXT("Base")},
+	{{MaterialX::Category::Power,	    TEXT("in2")},       TEXT("Exponent")},
+	{{MaterialX::Category::Invert,      TEXT("amount")},    TEXT("A")},
+	{{MaterialX::Category::Invert,      TEXT("in")},        TEXT("B")},
+	{{MaterialX::Category::Rotate2D,    TEXT("amount")},    TEXT("RotationAngle")},
+	{{MaterialX::Category::Rotate2D,    TEXT("in")},        TEXT("Position")},
+	{{MaterialX::Category::Rotate3D,    TEXT("amount")},    TEXT("RotationAngle")},
+	{{MaterialX::Category::Rotate3D,    TEXT("axis")},		TEXT("NormalizedRotationAxis")},
+	{{MaterialX::Category::Rotate3D,    TEXT("in")},        TEXT("Position")} },
 NodeNamesMaterialX2UE{
-	{MaterialX::Category::Add,      TEXT("Add")},
-	{MaterialX::Category::Sub,      TEXT("Subtract")},
-	{MaterialX::Category::Multiply, TEXT("Multiply")},
-	{MaterialX::Category::Sin,      TEXT("Sine")},
-	{MaterialX::Category::Cos,      TEXT("Cosine")},
-	{MaterialX::Category::Clamp,    TEXT("Clamp")},
-	{MaterialX::Category::Mix,      TEXT("Lerp")},
-	{MaterialX::Category::Max,      TEXT("Max")},
-	{MaterialX::Category::Min,      TEXT("Min")},
-	{MaterialX::Category::Combine2, TEXT("AppendVector")},
-	{MaterialX::Category::Combine3, TEXT("AppendVector")},
-	{MaterialX::Category::Combine4, TEXT("AppendVector")},
-	{MaterialX::Category::Power,    TEXT("Power")}},
-UEInputs{ TEXT("A"), TEXT("B"), TEXT("Input"), TEXT("Factor"), TEXT("Min"), TEXT("Max"), TEXT("Base"), TEXT("Exponent") }
+	// Math nodes
+	{MaterialX::Category::Absval,       TEXT("Abs")},
+	{MaterialX::Category::Add,          TEXT("Add")},
+	{MaterialX::Category::Acos,         TEXT("Arccosine")},
+	{MaterialX::Category::Asin,         TEXT("Arcsine")},
+	{MaterialX::Category::Atan2,        TEXT("Arctangent2")},
+	{MaterialX::Category::Ceil,         TEXT("Ceil")},
+	{MaterialX::Category::Clamp,        TEXT("Clamp")},
+	{MaterialX::Category::Cos,          TEXT("Cosine")},
+	{MaterialX::Category::CrossProduct, TEXT("Crossproduct")},
+	{MaterialX::Category::Divide,       TEXT("Divide")},
+	{MaterialX::Category::DotProduct,   TEXT("Dotproduct")},
+	{MaterialX::Category::Exp,          TEXT("Exponential")},
+	{MaterialX::Category::Floor,        TEXT("Floor")},
+	{MaterialX::Category::Invert,       TEXT("Subtract")},
+	{MaterialX::Category::Magnitude,    TEXT("Length")},
+	{MaterialX::Category::Max,          TEXT("Max")},
+	{MaterialX::Category::Min,          TEXT("Min")},
+	{MaterialX::Category::Modulo,       TEXT("Fmod")},
+	{MaterialX::Category::Multiply,     TEXT("Multiply")},
+	{MaterialX::Category::Normalize,    TEXT("Normalize")},
+	{MaterialX::Category::Power,        TEXT("Power")},
+	{MaterialX::Category::Sign,         TEXT("Sign")},
+	{MaterialX::Category::Sin,          TEXT("Sine")},
+	{MaterialX::Category::Sqrt,         TEXT("SquareRoot")},
+	{MaterialX::Category::Sub,          TEXT("Subtract")},
+	{MaterialX::Category::Tan,          TEXT("Tangent")},
+	// Compositing nodes
+	{MaterialX::Category::Mix,          TEXT("Lerp")},
+	// Channel nodes
+	{MaterialX::Category::Combine2,     TEXT("AppendVector")},
+	{MaterialX::Category::Combine3,     TEXT("Append3Vector")},
+	{MaterialX::Category::Combine4,     TEXT("Append4Vector")},
+	// Geometric nodes 
+	{MaterialX::Category::TexCoord,		TEXT("TextureCoordinate")},
+	// Adjustment nodes,
+	{MaterialX::Category::HsvToRgb,		TEXT("HsvToRgb")},
+	{MaterialX::Category::RgbToHsv,		TEXT("RgbToHsv")}},
+UEInputs{
+    TEXT("A"),
+	TEXT("B"),
+    TEXT("Base"),
+    TEXT("Exponent"),
+    TEXT("C"),
+    TEXT("D"),
+    TEXT("Factor"),
+    TEXT("Input"),
+    TEXT("Max"),
+    TEXT("Min"),
+    TEXT("VectorInput")
+    TEXT("X"),
+    TEXT("Y")}
 #endif // WITH_EDITOR
 {
 }
@@ -266,7 +320,7 @@ void UInterchangeMaterialXTranslator::ProcessStandardSurface(UInterchangeBaseNod
 
 	TMap<FString, UInterchangeShaderNode*> NamesToShaderNodes;
 
-	UInterchangeShaderGraphNode* ShaderGraphNode = CreateShaderNode<UInterchangeShaderGraphNode>(StandardSurfaceNode->getName().c_str(), StandardSurface::Name.ToString(), TEXT(""), NamesToShaderNodes, NodeContainer);
+	UInterchangeShaderGraphNode* ShaderGraphNode = CreateShaderNode<UInterchangeShaderGraphNode>(StandardSurfaceNode->getName().c_str(), StandardSurface::Name.ToString(), NamesToShaderNodes, NodeContainer);
 
 	//Base
 	{
@@ -576,7 +630,7 @@ void UInterchangeMaterialXTranslator::ProcessLightShader(UInterchangeBaseNodeCon
 	// Color
 	{
 		mx::InputPtr LightColor = LightShaderNode->getInput(mx::Lights::Input::Color);
-		const FLinearColor Color = MakeLinearColorFromColor3(LightColor);
+		const FLinearColor Color = GetLinearColor(LightColor);
 		LightNode->SetCustomLightColor(Color);
 	}
 
@@ -587,7 +641,7 @@ void UInterchangeMaterialXTranslator::ProcessLightShader(UInterchangeBaseNodeCon
 	}
 }
 
-static void GetLightDirection(mx::InputPtr DirectionInput, FTransform & Transform)
+static void GetLightDirection(mx::InputPtr DirectionInput, FTransform& Transform)
 {
 	mx::Vector3 Direction = mx::fromValueString<mx::Vector3>(DirectionInput->getValueString());
 
@@ -740,124 +794,40 @@ bool UInterchangeMaterialXTranslator::ConnectNodeGraphOutputToInput(MaterialX::I
 				}
 				else if(UpstreamNode->getCategory() == mx::Category::Constant)
 				{
-					if(mx::InputPtr InputConstant = UpstreamNode->getInput("value"); !AddAttribute(InputConstant, InputChannelName, ParentShaderNode))
-					{
-						UInterchangeResultWarning_Generic* Message = AddMessage<UInterchangeResultWarning_Generic>();
-						Message->Text = FText::Format(LOCTEXT("InputTypeNotSupported", "<{0}>: \"{1}\" is not supported yet"),
-													  FText::FromString(GetInputName(InputConstant)),
-													  FText::FromString(InputConstant->getType().c_str()));
-					}
-
+					ConnectConstantInputToOutput(UpstreamNode, ParentShaderNode, InputChannelName, NamesToShaderNodes, NodeContainer);
 				}
 				else if(UpstreamNode->getCategory() == mx::Category::Extract)
 				{
-					UInterchangeShaderNode* MaskShaderNode = CreateShaderNode<UInterchangeShaderNode>(UpstreamNode->getName().c_str(), Mask::Name.ToString(), ParentShaderNode->GetUniqueID(), NamesToShaderNodes, NodeContainer);
-
-					if(mx::InputPtr InputIndex = UpstreamNode->getInput("index"))
-					{
-						const int32 Index = mx::fromValueString<int>(InputIndex->getValueString());
-						switch(Index)
-						{
-						case 0: MaskShaderNode->AddBooleanAttribute(Mask::Attributes::R, true); break;
-						case 1: MaskShaderNode->AddBooleanAttribute(Mask::Attributes::G, true); break;
-						case 2: MaskShaderNode->AddBooleanAttribute(Mask::Attributes::B, true); break;
-						case 3: MaskShaderNode->AddBooleanAttribute(Mask::Attributes::A, true); break;
-						default:
-							UInterchangeResultError_Generic* Message = AddMessage<UInterchangeResultError_Generic>();
-							Message->Text = FText::FromString(TEXT("Wrong index number for extract node, values are from [0-3]"));
-							break;
-						}
-					}
-
-					UInterchangeShaderPortsAPI::ConnectDefaultOuputToInput(ParentShaderNode, InputChannelName, MaskShaderNode->GetUniqueID());
+					ConnectExtractInputToOutput(UpstreamNode, ParentShaderNode, InputChannelName, NamesToShaderNodes, NodeContainer);
 				}
 				//dot means identity, input == output
 				else if(UpstreamNode->getCategory() == mx::Category::Dot || UpstreamNode->getCategory() == mx::Category::NormalMap)
 				{
-					if(mx::InputPtr Input = GetInputFromOriginalName(UpstreamNode, "in"))
-					{
-						RenameInput(Input, TCHAR_TO_UTF8(*InputChannelName)); //let's take the parent node's input name
-						NamesToShaderNodes.FindOrAdd(UpstreamNode->getName().c_str(), ParentShaderNode);
-					}
+					ConnectDotInputToOutput(UpstreamNode, ParentShaderNode, InputChannelName, NamesToShaderNodes, NodeContainer);
+				}
+				else if(UpstreamNode->getCategory() == mx::Category::TransformPoint)
+				{
+					ConnectTransformPositionInputToOutput(UpstreamNode, ParentShaderNode, InputChannelName, NamesToShaderNodes, NodeContainer);
+				}
+				else if(UpstreamNode->getCategory() == mx::Category::TransformVector || UpstreamNode->getCategory() == mx::Category::TransformNormal)
+				{
+					ConnectTransformVectorInputToOutput(UpstreamNode, ParentShaderNode, InputChannelName, NamesToShaderNodes, NodeContainer);
+				}
+				else if(UpstreamNode->getCategory() == mx::Category::Rotate2D)
+				{
+					ConnectRotate2DInputToOutput(UpstreamNode, ParentShaderNode, InputChannelName, NamesToShaderNodes, NodeContainer);
+				}
+				else if(UpstreamNode->getCategory() == mx::Category::Rotate3D)
+				{
+					ConnectRotate3DInputToOutput(UpstreamNode, ParentShaderNode, InputChannelName, NamesToShaderNodes, NodeContainer);
 				}
 				else if(UpstreamNode->getCategory() == mx::Category::Image || UpstreamNode->getCategory() == mx::Category::TiledImage)
 				{
-					if(UInterchangeTextureNode* TextureNode = CreateTextureNode(UpstreamNode, NodeContainer))
-					{
-						//By default set the output of a texture to RGB, if the type is float then it's either up to an extract node or the Material Input to handle it
-						FString OutputChannel{ TEXT("RGB") };
-
-						if(UpstreamNode->getType() == mx::Type::Vector4 || UpstreamNode->getType() == mx::Type::Color4)
-							OutputChannel = TEXT("RGBA");
-
-						UInterchangeShaderNode* TextureShaderNode = CreateShaderNode<UInterchangeShaderNode>(UpstreamNode->getName().c_str(), TextureSample::Name.ToString(), ParentShaderNode->GetUniqueID(), NamesToShaderNodes, NodeContainer);
-						TextureShaderNode->AddStringAttribute(UInterchangeShaderPortsAPI::MakeInputValueKey(TextureSample::Inputs::Texture.ToString()), TextureNode->GetUniqueID());
-						UInterchangeShaderPortsAPI::ConnectOuputToInput(ParentShaderNode, InputChannelName, TextureShaderNode->GetUniqueID(), OutputChannel);
-
-						UInterchangeShaderNode* ImageNode = TextureShaderNode;
-						FString ImageNodeInputName{ TextureSample::Inputs::Coordinates.ToString() };
-
-						auto ConnectUVTransformToOutput = [this, UpstreamNode, &NodeContainer, &NamesToShaderNodes](UInterchangeShaderNode*& ImageNode, FString& ImageNodeInputName, const FString& ShaderType, const char* InputName)
-						{
-							if(mx::InputPtr Input = UpstreamNode->getInput(InputName))
-							{
-								const FString ShaderNodeName{ UpstreamNode->getName().c_str() + FString(TEXT("_")) + InputName };
-								UInterchangeShaderNode* ShaderNode = CreateShaderNode<UInterchangeShaderNode>(ShaderNodeName, ShaderType, ImageNode->GetUniqueID(), NamesToShaderNodes, NodeContainer);
-								UInterchangeShaderPortsAPI::ConnectDefaultOuputToInput(ImageNode, ImageNodeInputName, ShaderNode->GetUniqueID());
-
-								//Vec2
-								if(Input->hasValueString())
-								{
-									const FString MaskNodeName{ UpstreamNode->getName().c_str() + FString(TEXT("_")) + ShaderType+TEXT("MaskNode") };
-									UInterchangeShaderNode* MaskShaderNode = CreateShaderNode<UInterchangeShaderNode>(MaskNodeName, Mask::Name.ToString(), ShaderNode->GetUniqueID(), NamesToShaderNodes, NodeContainer);
-									MaskShaderNode->AddBooleanAttribute(Mask::Attributes::R, true);
-									MaskShaderNode->AddBooleanAttribute(Mask::Attributes::G, true);
-
-									mx::Vector2 Vec2 = mx::fromValueString<mx::Vector2>(Input->getValueString());
-									MaskShaderNode->AddLinearColorAttribute(UInterchangeShaderPortsAPI::MakeInputValueKey(Mask::Inputs::Input.ToString()), FLinearColor(Vec2[0], Vec2[1], 0));
-									UInterchangeShaderPortsAPI::ConnectDefaultOuputToInput(ShaderNode, TEXT("B"), MaskShaderNode->GetUniqueID());
-								}
-								else
-								{
-									RenameInput(Input, "B");
-								}
-
-								//The node will replace the texture node for the input
-								NamesToShaderNodes.Add(UpstreamNode->getName().c_str(), ShaderNode);
-
-								//We also need to replace the name of texcoord input by one of the inputs of this node
-								if(mx::InputPtr InputTexCoord = GetInputFromOriginalName(UpstreamNode, mx::TiledImage::Inputs::TexCoord))
-								{
-									RenameInput(InputTexCoord, "A");
-								}
-								else
-								{
-									//Let's reuse the same texture coordinate node for further use (no need to set a parent)
-									const FString TextureCoordinateName{ FPaths::GetBaseFilename(UpstreamNode->getActiveSourceUri().c_str()) + FString(TEXT("_texcoord")) };
-									UInterchangeShaderNode* TextureCoordinateNode = CreateShaderNode<UInterchangeShaderNode>(TextureCoordinateName, TextureCoordinate::Name.ToString(), TEXT(""), NamesToShaderNodes, NodeContainer);
-									UInterchangeShaderPortsAPI::ConnectDefaultOuputToInput(ShaderNode, TEXT("A"), TextureCoordinateNode->GetUniqueID());
-								}
-
-								ImageNode = ShaderNode;
-								ImageNodeInputName = TEXT("A");
-							}
-						};
-
-						// UV offset (MaterialX defines it as a substraction)
-						// MaterialX spec: the offset for the given image along the U and V axes. Mathematically
-						// equivalent to subtracting the given vector value from the incoming texture coordinates.
-						ConnectUVTransformToOutput(ImageNode, ImageNodeInputName, TEXT("Subtract"), mx::TiledImage::Inputs::UVOffset);
-						ConnectUVTransformToOutput(ImageNode, ImageNodeInputName, TEXT("Multiply"), mx::TiledImage::Inputs::UVTiling);
-					}
-					else
-					{
-						AddAttribute(UpstreamNode->getInput(mx::NodeGroup::Texture2D::Inputs::Default), InputChannelName, ParentShaderNode);
-					}
+					ConnectImageInputToOutput(UpstreamNode, ParentShaderNode, InputChannelName, NamesToShaderNodes, NodeContainer);
 				}
-				else if(UpstreamNode->getCategory() == mx::Category::TexCoord)
+				else if(UpstreamNode->getCategory() == mx::Category::Convert)
 				{
-					UInterchangeShaderNode* TextureCoordinateNode = CreateShaderNode<UInterchangeShaderNode>(UpstreamNode->getName().c_str(), TextureCoordinate::Name.ToString(), ParentShaderNode->GetUniqueID(), NamesToShaderNodes, NodeContainer);
-					UInterchangeShaderPortsAPI::ConnectDefaultOuputToInput(ParentShaderNode, InputChannelName, TextureCoordinateNode->GetUniqueID());
+					ConnectConvertInputToOutput(UpstreamNode, ParentShaderNode, InputChannelName, NamesToShaderNodes, NodeContainer);
 				}
 				else
 				{
@@ -880,7 +850,7 @@ bool UInterchangeMaterialXTranslator::ConnectNodeOutputToInput(MaterialX::NodePt
 	{
 		UInterchangeShaderNode* OperatorNode = nullptr;
 
-		OperatorNode = CreateShaderNode<UInterchangeShaderNode>(Node->getName().c_str(), *ShaderType, ParentShaderNode->GetUniqueID(), NamesToShaderNodes, NodeContainer);
+		OperatorNode = CreateShaderNode<UInterchangeShaderNode>(Node->getName().c_str(), *ShaderType, NamesToShaderNodes, NodeContainer);
 
 		for(mx::InputPtr Input : Node->getInputs())
 		{
@@ -897,7 +867,7 @@ bool UInterchangeMaterialXTranslator::ConnectNodeOutputToInput(MaterialX::NodePt
 				if(InputInterface->hasValue())
 				{
 					// We need to take the input name from the original Input not the Interface
-					if(const FString* InputNameFound = UEInputs.Find(GetInputName( Input)))
+					if(const FString* InputNameFound = UEInputs.Find(GetInputName(Input)))
 					{
 						AddAttribute(InputInterface, *InputNameFound, OperatorNode);
 					}
@@ -1037,7 +1007,6 @@ FString UInterchangeMaterialXTranslator::GetInputName(MaterialX::InputPtr Input)
 		Name.replace(0, OriginalName.size(), "");
 	}
 
-
 	return Name.c_str();
 }
 
@@ -1097,23 +1066,18 @@ bool UInterchangeMaterialXTranslator::AddAttribute(MaterialX::InputPtr Input, co
 		{
 			return ShaderNode->AddFloatAttribute(UInterchangeShaderPortsAPI::MakeInputValueKey(InputChannelName), mx::fromValueString<float>(Input->getValueString()));
 		}
-		else if(Input->getType() == mx::Type::Integer)
-		{
-			printf("");
-		}
-
-		FLinearColor LinearColor;
 
 		bool bIsColor = false;
+		FLinearColor LinearColor;
 
-		if(bool bIsColor3 = (Input->getType() == mx::Type::Color3 || Input->getType() == mx::Type::Vector3))
+		if(Input->getType() == mx::Type::Color3 || Input->getType() == mx::Type::Color4)
 		{
-			LinearColor = MakeLinearColorFromColor3(Input);
+			LinearColor = GetLinearColor(Input);
 			bIsColor = true;
 		}
-		else if(bool bIsColor4 = (Input->getType() == mx::Type::Color4|| Input->getType() == mx::Type::Vector4))
+		else if(Input->getType() == mx::Type::Vector2 || Input->getType() == mx::Type::Vector3 || Input->getType() == mx::Type::Vector4)
 		{
-			LinearColor = MakeLinearColorFromColor4(Input);
+			LinearColor = GetVector(Input);
 			bIsColor = true;
 		}
 
@@ -1150,7 +1114,7 @@ bool UInterchangeMaterialXTranslator::AddLinearColorAttribute(MaterialX::InputPt
 	{
 		if(Input->hasValueString())
 		{
-			const FLinearColor Value = MakeLinearColorFromColor3(Input);
+			const FLinearColor Value = GetLinearColor(Input);
 
 			if(!Value.Equals(DefaultValue))
 			{
@@ -1200,21 +1164,35 @@ FString UInterchangeMaterialXTranslator::GetColorSpace(MaterialX::ElementPtr Ele
 	return ColorSpace;
 }
 
-FLinearColor UInterchangeMaterialXTranslator::MakeLinearColorFromColor3(MaterialX::InputPtr Input) const
+FLinearColor UInterchangeMaterialXTranslator::GetLinearColor(MaterialX::InputPtr Input) const
 {
-	mx::Color3 Color = mx::fromValueString<mx::Color3>(Input->getValueString());
-
 	//we assume that the default color space is linear
-	FLinearColor LinearColor(Color[0], Color[1], Color[2]);
+	FLinearColor LinearColor;
+
+	if(Input->getType() == mx::Type::Color3)
+	{
+		mx::Color3 Color = mx::fromValueString<mx::Color3>(Input->getValueString());
+		LinearColor = FLinearColor{ Color[0], Color[1], Color[2] };
+	}
+	else if(Input->getType() == mx::Type::Color4)
+	{
+		mx::Color4 Color = mx::fromValueString<mx::Color4>(Input->getValueString());
+		LinearColor = FLinearColor{ Color[0], Color[1], Color[2], Color[3] };
+	}
+	else
+	{
+		ensureMsgf(false, TEXT("input type can only be color3 or color4"));
+	}
+
 	const FString ColorSpace = GetColorSpace(Input);
-	
+
 	if(ColorSpace.IsEmpty() || ColorSpace == TEXT("lin_rec709") || ColorSpace == TEXT("none"))
 	{
 		;//noop
 	}
 	else if(ColorSpace == TEXT("gamma22"))
 	{
-		LinearColor = FLinearColor::FromPow22Color(FColor(Color[0]/255.f, Color[1]/255.f, Color[2]/255.f));
+		LinearColor = FLinearColor::FromPow22Color(FColor(LinearColor.R / 255.f, LinearColor.G / 255.f, LinearColor.B / 255.f, LinearColor.A / 255.f));
 	}
 	else
 	{
@@ -1228,33 +1206,452 @@ FLinearColor UInterchangeMaterialXTranslator::MakeLinearColorFromColor3(Material
 	return LinearColor;
 }
 
-FLinearColor UInterchangeMaterialXTranslator::MakeLinearColorFromColor4(MaterialX::InputPtr Input) const
+void UInterchangeMaterialXTranslator::ConnectConstantInputToOutput(MaterialX::NodePtr UpstreamNode, UInterchangeShaderNode* ParentShaderNode, const FString& InputChannelName, TMap<FString, UInterchangeShaderNode*>& NamesToShaderNodes, UInterchangeBaseNodeContainer& NodeContainer) const
+{
+	if(mx::InputPtr InputConstant = UpstreamNode->getInput("value"); !AddAttribute(InputConstant, InputChannelName, ParentShaderNode))
+	{
+		UInterchangeResultWarning_Generic* Message = AddMessage<UInterchangeResultWarning_Generic>();
+		Message->Text = FText::Format(LOCTEXT("InputTypeNotSupported", "<{0}>: \"{1}\" is not supported yet"),
+									  FText::FromString(GetInputName(InputConstant)),
+									  FText::FromString(InputConstant->getType().c_str()));
+	}
+}
+
+void UInterchangeMaterialXTranslator::ConnectExtractInputToOutput(MaterialX::NodePtr UpstreamNode, UInterchangeShaderNode* ParentShaderNode, const FString& InputChannelName, TMap<FString, UInterchangeShaderNode*>& NamesToShaderNodes, UInterchangeBaseNodeContainer& NodeContainer) const
 {
-	mx::Color4 Color = mx::fromValueString<mx::Color4>(Input->getValueString());
+	using namespace UE::Interchange::Materials::Standard::Nodes;
 
-	FLinearColor LinearColor(Color[0], Color[1], Color[2], Color[3]);
-	const FString ColorSpace = GetColorSpace(Input);
+	UInterchangeShaderNode* MaskShaderNode = CreateShaderNode<UInterchangeShaderNode>(UpstreamNode->getName().c_str(), Mask::Name.ToString(), NamesToShaderNodes, NodeContainer);
+
+	if(mx::InputPtr InputIndex = UpstreamNode->getInput("index"))
+	{
+		const int32 Index = mx::fromValueString<int>(InputIndex->getValueString());
+		switch(Index)
+		{
+		case 0: MaskShaderNode->AddBooleanAttribute(Mask::Attributes::R, true); break;
+		case 1: MaskShaderNode->AddBooleanAttribute(Mask::Attributes::G, true); break;
+		case 2: MaskShaderNode->AddBooleanAttribute(Mask::Attributes::B, true); break;
+		case 3: MaskShaderNode->AddBooleanAttribute(Mask::Attributes::A, true); break;
+		default:
+			UInterchangeResultError_Generic* Message = AddMessage<UInterchangeResultError_Generic>();
+			Message->Text = FText::FromString(TEXT("Wrong index number for extract node, values are from [0-3]"));
+			break;
+		}
+	}
+
+	UInterchangeShaderPortsAPI::ConnectDefaultOuputToInput(ParentShaderNode, InputChannelName, MaskShaderNode->GetUniqueID());
+}
+
+void UInterchangeMaterialXTranslator::ConnectDotInputToOutput(MaterialX::NodePtr UpstreamNode, UInterchangeShaderNode* ParentShaderNode, const FString& InputChannelName, TMap<FString, UInterchangeShaderNode*>& NamesToShaderNodes, UInterchangeBaseNodeContainer& NodeContainer) const
+{
+	if(mx::InputPtr Input = GetInputFromOriginalName(UpstreamNode, "in"))
+	{
+		RenameInput(Input, TCHAR_TO_UTF8(*InputChannelName)); //let's take the parent node's input name
+		NamesToShaderNodes.FindOrAdd(UpstreamNode->getName().c_str(), ParentShaderNode);
+	}
+}
+
+void UInterchangeMaterialXTranslator::ConnectTransformPositionInputToOutput(MaterialX::NodePtr UpstreamNode, UInterchangeShaderNode* ParentShaderNode, const FString& InputChannelName, TMap<FString, UInterchangeShaderNode*>& NamesToShaderNodes, UInterchangeBaseNodeContainer& NodeContainer) const
+{
+	UInterchangeShaderNode* TransformNode = CreateShaderNode<UInterchangeShaderNode>(UpstreamNode->getName().c_str(), TEXT("TransformPosition"), NamesToShaderNodes, NodeContainer);
+	if(mx::InputPtr Input = GetInputFromOriginalName(UpstreamNode, "in"); Input->hasValue())
+	{
+		FLinearColor Vector = GetVector(Input);
+		TransformNode->AddLinearColorAttribute(UInterchangeShaderPortsAPI::MakeInputValueKey(TEXT("Input")), Vector);
+	}
+	UInterchangeShaderPortsAPI::ConnectDefaultOuputToInput(ParentShaderNode, InputChannelName, TransformNode->GetUniqueID());
+}
+
+void UInterchangeMaterialXTranslator::ConnectTransformVectorInputToOutput(MaterialX::NodePtr UpstreamNode, UInterchangeShaderNode* ParentShaderNode, const FString& InputChannelName, TMap<FString, UInterchangeShaderNode*>& NamesToShaderNodes, UInterchangeBaseNodeContainer& NodeContainer) const
+{
+	UInterchangeShaderNode* TransformNode = CreateShaderNode<UInterchangeShaderNode>(UpstreamNode->getName().c_str(), TEXT("Transform"), NamesToShaderNodes, NodeContainer);
+	if(mx::InputPtr Input = GetInputFromOriginalName(UpstreamNode, "in"); Input->hasValue())
+	{
+		FLinearColor Vector = GetVector(Input);
+		TransformNode->AddLinearColorAttribute(UInterchangeShaderPortsAPI::MakeInputValueKey(TEXT("Input")), Vector);
+	}
+	UInterchangeShaderPortsAPI::ConnectDefaultOuputToInput(ParentShaderNode, InputChannelName, TransformNode->GetUniqueID());
+}
+
+void UInterchangeMaterialXTranslator::ConnectRotate2DInputToOutput(mx::NodePtr UpstreamNode, UInterchangeShaderNode* ParentShaderNode, const FString& InputChannelName, TMap<FString, UInterchangeShaderNode*>& NamesToShaderNodes, UInterchangeBaseNodeContainer& NodeContainer) const
+{
+	UInterchangeShaderNode* Rotate2DNode = CreateShaderNode<UInterchangeShaderNode>(UpstreamNode->getName().c_str(), TEXT("RotateAboutAxis"), NamesToShaderNodes, NodeContainer);
+	Rotate2DNode->AddLinearColorAttribute(UInterchangeShaderPortsAPI::MakeInputValueKey(TEXT("NormalizedRotationAxis")), FLinearColor(0, 0, 1));
+	Rotate2DNode->AddLinearColorAttribute(UInterchangeShaderPortsAPI::MakeInputValueKey(TEXT("PivotPoint")), FLinearColor(0.5, 0.5, 0));
+
+	if(mx::InputPtr Input = GetInputFromOriginalName(UpstreamNode, "in"); Input->hasValue())
+	{
+		FLinearColor Vector = GetVector(Input);
+		Rotate2DNode->AddLinearColorAttribute(UInterchangeShaderPortsAPI::MakeInputValueKey(TEXT("Position")), Vector);
+	}
+
+	if(mx::InputPtr Input = GetInputFromOriginalName(UpstreamNode, "amount"); Input->hasValue())
+	{
+		//MaterialX angle is in degrees, in UE it's in the range [0,1]
+		float Amount = mx::fromValueString<float>(Input->getValueString()) / 360.;
+		Rotate2DNode->AddFloatAttribute(UInterchangeShaderPortsAPI::MakeInputValueKey(TEXT("RotationAngle")), Amount);
+	}
+	else // we need to create a Divide node
+	{
+		//copy "in" input into "in1" and create "amount" input as "in2" with the value 360 because UE's angle is a value between [0,1]
+		mx::NodePtr NewDivideNode = CreateNode(UpstreamNode->getParent()->asA<mx::NodeGraph>(),
+											   UpstreamNode->getName().c_str(),
+											   mx::Category::Divide,
+											   { {"in1", Input} }, // rename input to match <divide> input "in1"
+											   { {"in2", FAttributeValueArray{{"type", "float"}, {"value", "360"}}} });
+
+		// Input now points to the new node
+		Input->setNodeName(NewDivideNode->getName());
+	}
+
+	UInterchangeShaderPortsAPI::ConnectDefaultOuputToInput(ParentShaderNode, InputChannelName, Rotate2DNode->GetUniqueID());
+}
+
+void UInterchangeMaterialXTranslator::ConnectRotate3DInputToOutput(mx::NodePtr UpstreamNode, UInterchangeShaderNode* ParentShaderNode, const FString& InputChannelName, TMap<FString, UInterchangeShaderNode*>& NamesToShaderNodes, UInterchangeBaseNodeContainer& NodeContainer) const
+{
+	UInterchangeShaderNode* Rotate3DNode = CreateShaderNode<UInterchangeShaderNode>(UpstreamNode->getName().c_str(), TEXT("RotateAboutAxis"), NamesToShaderNodes, NodeContainer);
+	Rotate3DNode->AddLinearColorAttribute(UInterchangeShaderPortsAPI::MakeInputValueKey(TEXT("PivotPoint")), FLinearColor(0.5, 0.5, 0));
+
+	if(mx::InputPtr Input = GetInputFromOriginalName(UpstreamNode, "in"); Input->hasValue())
+	{
+		FLinearColor Vector = GetVector(Input);
+		Rotate3DNode->AddLinearColorAttribute(UInterchangeShaderPortsAPI::MakeInputValueKey(TEXT("Position")), Vector);
+	}
+
+	if(mx::InputPtr Input = GetInputFromOriginalName(UpstreamNode, "axis"); Input->hasValue())
+	{
+		FLinearColor Vector = GetVector(Input);
+		Rotate3DNode->AddLinearColorAttribute(UInterchangeShaderPortsAPI::MakeInputValueKey(TEXT("NormalizedRotationAxis")), Vector);
+	}
+
+	if(mx::InputPtr Input = GetInputFromOriginalName(UpstreamNode, "amount"); Input->hasValue())
+	{
+		//MaterialX angle is in degrees, in UE it's in the range [0,1]
+		float Amount = mx::fromValueString<float>(Input->getValueString()) / 360.;
+		Rotate3DNode->AddFloatAttribute(UInterchangeShaderPortsAPI::MakeInputValueKey(TEXT("RotationAngle")), Amount);
+	}
+	else // we need to create a Divide node
+	{
+		//copy "in" input into "in1" and create "amount" input as "in2" with the value 360 because UE's angle is a value between [0,1]
+		mx::NodePtr NewDivideNode = CreateNode(UpstreamNode->getParent()->asA<mx::NodeGraph>(),
+											   UpstreamNode->getName().c_str(),
+											   mx::Category::Divide,
+											   { {"in1", Input} }, // rename input to match <divide> input "in1"
+											   { {"in2", FAttributeValueArray{{"type", "float"}, {"value", "360"}}} });
+
+		// Input now points to the new node
+		Input->setNodeName(NewDivideNode->getName());
+	}
+
+	UInterchangeShaderPortsAPI::ConnectDefaultOuputToInput(ParentShaderNode, InputChannelName, Rotate3DNode->GetUniqueID());
+}
+
+void UInterchangeMaterialXTranslator::ConnectImageInputToOutput(MaterialX::NodePtr UpstreamNode, UInterchangeShaderNode* ParentShaderNode, const FString& InputChannelName, TMap<FString, UInterchangeShaderNode*>& NamesToShaderNodes, UInterchangeBaseNodeContainer& NodeContainer) const
+{
+	using namespace UE::Interchange::Materials::Standard::Nodes;
+
+	if(UInterchangeTextureNode* TextureNode = CreateTextureNode(UpstreamNode, NodeContainer))
+	{
+		//By default set the output of a texture to RGB, if the type is float then it's either up to an extract node or the Material Input to handle it
+		FString OutputChannel{ TEXT("RGB") };
 
-	if(ColorSpace == TEXT("") || ColorSpace == TEXT("lin_rec709") || ColorSpace == TEXT("none"))
+		if(UpstreamNode->getType() == mx::Type::Vector4 || UpstreamNode->getType() == mx::Type::Color4)
+			OutputChannel = TEXT("RGBA");
+
+		UInterchangeShaderNode* TextureShaderNode = CreateShaderNode<UInterchangeShaderNode>(UpstreamNode->getName().c_str(), TextureSample::Name.ToString(), NamesToShaderNodes, NodeContainer);
+		TextureShaderNode->AddStringAttribute(UInterchangeShaderPortsAPI::MakeInputValueKey(TextureSample::Inputs::Texture.ToString()), TextureNode->GetUniqueID());
+		UInterchangeShaderPortsAPI::ConnectOuputToInput(ParentShaderNode, InputChannelName, TextureShaderNode->GetUniqueID(), OutputChannel);
+
+		UInterchangeShaderNode* ImageNode = TextureShaderNode;
+		FString ImageNodeInputName{ TextureSample::Inputs::Coordinates.ToString() };
+
+		auto ConnectUVTransformToOutput = [this, UpstreamNode, &NodeContainer, &NamesToShaderNodes](UInterchangeShaderNode*& ImageNode, FString& ImageNodeInputName, const FString& ShaderType, const char* InputName)
+		{
+			if(mx::InputPtr Input = UpstreamNode->getInput(InputName))
+			{
+				const FString ShaderNodeName{ UpstreamNode->getName().c_str() + FString(TEXT("_")) + InputName };
+				UInterchangeShaderNode* ShaderTransformNode = CreateShaderNode<UInterchangeShaderNode>(ShaderNodeName, ShaderType, NamesToShaderNodes, NodeContainer);
+				UInterchangeShaderPortsAPI::ConnectDefaultOuputToInput(ImageNode, ImageNodeInputName, ShaderTransformNode->GetUniqueID());
+
+				//Vec2
+				if(Input->hasValueString())
+				{
+					const FString MaskNodeName{ UpstreamNode->getName().c_str() + FString(TEXT("_")) + ShaderType + TEXT("MaskNode") };
+					UInterchangeShaderNode* MaskShaderNode = CreateShaderNode<UInterchangeShaderNode>(MaskNodeName, Mask::Name.ToString(), NamesToShaderNodes, NodeContainer);
+					MaskShaderNode->AddBooleanAttribute(Mask::Attributes::R, true);
+					MaskShaderNode->AddBooleanAttribute(Mask::Attributes::G, true);
+
+					mx::Vector2 Vec2 = mx::fromValueString<mx::Vector2>(Input->getValueString());
+					MaskShaderNode->AddLinearColorAttribute(UInterchangeShaderPortsAPI::MakeInputValueKey(Mask::Inputs::Input.ToString()), FLinearColor(Vec2[0], Vec2[1], 0));
+					UInterchangeShaderPortsAPI::ConnectDefaultOuputToInput(ShaderTransformNode, TEXT("B"), MaskShaderNode->GetUniqueID());
+				}
+				else
+				{
+					RenameInput(Input, "B");
+				}
+
+				//The node will replace the texture node for the input
+				NamesToShaderNodes.Add(UpstreamNode->getName().c_str(), ShaderTransformNode);
+
+				//We also need to replace the name of texcoord input by one of the inputs of this node
+				if(mx::InputPtr InputTexCoord = GetInputFromOriginalName(UpstreamNode, mx::TiledImage::Inputs::TexCoord))
+				{
+					RenameInput(InputTexCoord, "A");
+				}
+				else
+				{
+					//Let's reuse the same texture coordinate node for further use (no need to set a parent)
+					const FString TextureCoordinateName{ FPaths::GetBaseFilename(UpstreamNode->getActiveSourceUri().c_str()) + FString(TEXT("_texcoord")) };
+					UInterchangeShaderNode* TextureCoordinateNode = CreateShaderNode<UInterchangeShaderNode>(TextureCoordinateName, TextureCoordinate::Name.ToString(), NamesToShaderNodes, NodeContainer);
+					UInterchangeShaderPortsAPI::ConnectDefaultOuputToInput(ShaderTransformNode, TEXT("A"), TextureCoordinateNode->GetUniqueID());
+				}
+
+				ImageNode = ShaderTransformNode;
+				ImageNodeInputName = TEXT("A");
+			}
+		};
+
+		// UV offset (MaterialX defines it as a substraction)
+		// MaterialX spec: the offset for the given image along the U and V axes. Mathematically
+		// equivalent to subtracting the given vector value from the incoming texture coordinates.
+		ConnectUVTransformToOutput(ImageNode, ImageNodeInputName, TEXT("Subtract"), mx::TiledImage::Inputs::UVOffset);
+		ConnectUVTransformToOutput(ImageNode, ImageNodeInputName, TEXT("Multiply"), mx::TiledImage::Inputs::UVTiling);
+	}
+	else
 	{
-		;//no op
+		AddAttribute(UpstreamNode->getInput(mx::NodeGroup::Texture2D::Inputs::Default), InputChannelName, ParentShaderNode);
 	}
-	if(ColorSpace == TEXT("gamma22"))
+}
+
+void UInterchangeMaterialXTranslator::ConnectConvertInputToOutput(MaterialX::NodePtr UpstreamNode, UInterchangeShaderNode* ParentShaderNode, const FString& InputChannelName, TMap<FString, UInterchangeShaderNode*>& NamesToShaderNodes, UInterchangeBaseNodeContainer& NodeContainer) const
+{
+	using namespace UE::Interchange::Materials::Standard::Nodes;
+
+	//In case of an upwards conversion, let's do an append, downwards a mask, otherwise leave as is
+	mx::InputPtr Input = GetInputFromOriginalName(UpstreamNode, "in");
+	const std::string& NodeType = UpstreamNode->getType();
+	const std::string& InputType = Input->getType();
+
+	//Ensure that the types are supported
+	const bool bIsNodeTypeSupported =
+		NodeType == mx::Type::Color4 ||
+		NodeType == mx::Type::Color3 ||
+		NodeType == mx::Type::Vector4 ||
+		NodeType == mx::Type::Vector3 ||
+		NodeType == mx::Type::Vector2 ||
+		NodeType == mx::Type::Float ||
+		NodeType == mx::Type::Integer ||
+		NodeType == mx::Type::Boolean;
+
+	const bool bIsInputTypeSupported =
+		InputType == mx::Type::Color4 ||
+		InputType == mx::Type::Color3 ||
+		InputType == mx::Type::Vector4 ||
+		InputType == mx::Type::Vector3 ||
+		InputType == mx::Type::Vector2 ||
+		InputType == mx::Type::Float ||
+		InputType == mx::Type::Integer ||
+		InputType == mx::Type::Boolean;
+
+	if(!bIsNodeTypeSupported || !bIsInputTypeSupported)
 	{
-		LinearColor = FLinearColor::FromPow22Color(FColor(Color[0] / 255.f, Color[1] / 255.f, Color[2] / 255.f, Color[3] / 255.f));
+		UInterchangeResultError_Generic* Message = AddMessage<UInterchangeResultError_Generic>();
+		Message->Text = FText::FromString("<convert> node has non supported types");
+		return;
+	}
+
+	char NodeN = NodeType.back();
+	char InputN = InputType.back();
+
+	constexpr auto Remap = [](auto Value, auto Low1, auto High1, auto Low2, auto High2)
+	{
+		return Low2 + (Value - Low1) * (High2 - Low2) / (High1 - Low1);
+	};
+
+	//Remap NodeN and InputN in a range lower than '2' in case of integer/boolean/float
+	if(NodeN > '4')
+	{
+		NodeN = Remap(NodeN, 'a', 'z', 0, '1');
+	}
+	if(InputN > '4')
+	{
+		InputN = Remap(InputN, 'a', 'z', 0, '1');
+	}
+
+	if(InputN > NodeN) // Mask
+	{
+		UInterchangeShaderNode* MaskShaderNode = nullptr;
+		if(NodeN == '3')
+		{
+			MaskShaderNode = CreateMaskShaderNode(0b1110, UpstreamNode->getName().c_str(), NamesToShaderNodes, NodeContainer);
+		}
+		else if(NodeN == '2')
+		{
+			MaskShaderNode = CreateMaskShaderNode(0b1100, UpstreamNode->getName().c_str(), NamesToShaderNodes, NodeContainer);
+		}
+		else
+		{
+			MaskShaderNode = CreateMaskShaderNode(0b1000, UpstreamNode->getName().c_str(), NamesToShaderNodes, NodeContainer);
+		}
+
+		if(Input->hasValue())
+		{
+			AddAttribute(Input, TEXT("Input"), MaskShaderNode);
+		}
+
+		UInterchangeShaderPortsAPI::ConnectDefaultOuputToInput(ParentShaderNode, InputChannelName, MaskShaderNode->GetUniqueID());
+	}
+	else // Append 
+	{
+		//same as dot, just connect the next output to this parent input
+		RenameInput(Input, TCHAR_TO_UTF8(*InputChannelName)); //let's take the parent node's input name
+		NamesToShaderNodes.FindOrAdd(UpstreamNode->getName().c_str(), ParentShaderNode);
+
+		//No need to create a node, since the input and the node have the same channels, we just check if there's a value
+		if(NodeN == InputN || (NodeN < '2' && InputN < '2'))
+		{
+			if(Input->hasValue())
+			{
+				AddAttribute(Input, InputChannelName, ParentShaderNode);
+			}
+			return;
+		}
+		std::string Category;
+		TArray<FInputToCopy> InputsToCopy;
+		TArray<FInputToCreate> InputsToCreate;
+
+		// float to N
+		if(InputN < '2')
+		{
+			if(NodeN == '2')
+			{
+				Category = mx::Category::Combine2;
+				InputsToCopy.Add({ "in1", Input });
+				InputsToCopy.Add({ "in2", Input });
+			}
+			else if(NodeN == '3')
+			{
+				Category = mx::Category::Combine3;
+				InputsToCopy.Add({ "in1", Input });
+				InputsToCopy.Add({ "in2", Input });
+				InputsToCopy.Add({ "in3", Input });
+			}
+			else if(NodeN == '4')
+			{
+				Category = mx::Category::Combine4;
+				InputsToCopy.Add({ "in1", Input });
+				InputsToCopy.Add({ "in2", Input });
+				InputsToCopy.Add({ "in3", Input });
+				InputsToCopy.Add({ "in4", Input });
+			}
+		}
+		else if((InputN == '2' && NodeN == '3') || (InputN == '3' && NodeN == '4'))
+		{
+			Category = mx::Category::Combine2;
+			InputsToCopy.Add({ "in1", Input });
+			InputsToCreate.Add({ "in2", FAttributeValueArray{{"type","float"}, {"value","1"}}});
+		}
+
+		//copy "in" input into "in1" and create "amount" input as "in2" with the value 360 because UE's angle is a value between [0,1]
+		mx::NodePtr CombineNode = CreateNode(UpstreamNode->getParent()->asA<mx::NodeGraph>(),
+											 UpstreamNode->getName().c_str(),
+											 Category.c_str(),
+											 InputsToCopy, // rename input to match <divide> input "in1"
+											 InputsToCreate);
+
+		// Input now points to the new node
+		Input->setNodeName(CombineNode->getName());
+	}
+}
+
+UInterchangeShaderNode* UInterchangeMaterialXTranslator::CreateMaskShaderNode(uint8 RGBA, const FString& NodeName, TMap<FString, UInterchangeShaderNode*>& NamesToShaderNodes, UInterchangeBaseNodeContainer& NodeContainer) const
+{
+	bool bR = (0b1000 & RGBA) >> 3;
+	bool bG = (0b0100 & RGBA) >> 2;
+	bool bB = (0b0010 & RGBA) >> 1;
+	bool bA = (0b0001 & RGBA) >> 0;
+	using namespace UE::Interchange::Materials::Standard::Nodes;
+	UInterchangeShaderNode* MaskShaderNode = CreateShaderNode<UInterchangeShaderNode>(NodeName, Mask::Name.ToString(), NamesToShaderNodes, NodeContainer);
+	MaskShaderNode->AddBooleanAttribute(Mask::Attributes::R, bR);
+	MaskShaderNode->AddBooleanAttribute(Mask::Attributes::G, bG);
+	MaskShaderNode->AddBooleanAttribute(Mask::Attributes::B, bB);
+	MaskShaderNode->AddBooleanAttribute(Mask::Attributes::A, bA);
+
+	return MaskShaderNode;
+}
+
+FLinearColor UInterchangeMaterialXTranslator::GetVector(MaterialX::InputPtr Input) const
+{
+	FLinearColor LinearColor;
+
+	if(Input->getType() == mx::Type::Vector2)
+	{
+		mx::Vector2 Color = mx::fromValueString<mx::Vector2>(Input->getValueString());
+		LinearColor = FLinearColor{ Color[0], Color[1], 0 };
+	}
+	else if(Input->getType() == mx::Type::Vector3)
+	{
+		mx::Vector3 Color = mx::fromValueString<mx::Vector3>(Input->getValueString());
+		LinearColor = FLinearColor{ Color[0], Color[1], Color[2] };
+	}
+	else if(Input->getType() == mx::Type::Vector4)
+	{
+		mx::Vector4 Color = mx::fromValueString<mx::Vector4>(Input->getValueString());
+		LinearColor = FLinearColor{ Color[0], Color[1], Color[2], Color[3] };
 	}
 	else
 	{
-		UInterchangeResultWarning_Generic* Message = AddMessage<UInterchangeResultWarning_Generic>();
-		Message->Text = FText::Format(LOCTEXT("ColorSpaceNotSupported", "<{0}>-<{1}>: Colorspace {2} is not supported yet, falling back to linear"),
-									  FText::FromString(Input->getParent()->getName().c_str()),
-									  FText::FromString(Input->getName().c_str()),
-									  FText::FromString(ColorSpace));
+		ensureMsgf(false, TEXT("input type can only be a vectorN"));
 	}
 
 	return LinearColor;
 }
 
+MaterialX::NodePtr UInterchangeMaterialXTranslator::CreateNode(MaterialX::NodeGraphPtr NodeGraph, const char* NodeName, const char* Category, TArray<FInputToCopy> InputsToCopy, TArray<FInputToCreate> InputsToCreate) const
+{
+	std::string UniqueNodeName{ NodeName + std::string{"_"} + Category };
+
+	mx::NodePtr Node = NodeGraph->getNode(UniqueNodeName);
+
+	if(!Node)
+	{
+		Node = NodeGraph->addChildOfCategory(Category, UniqueNodeName)->asA<mx::Node>();
+
+		for(const FInputToCopy& Pair : InputsToCopy)
+		{
+			const char* NewInputName = Pair.Get<0>();
+			mx::InputPtr Input = Pair.Get<1>();
+
+			mx::InputPtr InputCopy = Node->addInput();
+			InputCopy->copyContentFrom(Input);
+			//remove the attribute OriginalName since when the input will be copied, the name will be the one from the spec (before a renaming)
+			InputCopy->removeAttribute(mx::Attributes::OriginalName);
+
+			if(NewInputName)
+			{
+				InputCopy->setName(NewInputName);
+			}
+		}
+
+		for(const FInputToCreate& Input : InputsToCreate)
+		{
+			const char* InputName = Input.Get<0>();
+			const FAttributeValueArray& Attributes = Input.Get<1>();
+
+			mx::InputPtr NewInput = Node->addInput();
+			NewInput->setName(InputName);
+			for(const FAttributeValue& AttributeValue : Attributes)
+			{
+				const char* Attribute = AttributeValue.Get<0>();
+				const char* Value = AttributeValue.Get<1>();
+				NewInput->setAttribute(Attribute, Value);
+			}
+		}
+	}
+
+	return Node;
+}
+
 #endif //WITH_EDITOR
 
 #undef LOCTEXT_NAMESPACE
\ No newline at end of file
diff --git a/Engine/Plugins/Experimental/Interchange/Runtime/Source/Import/Private/MaterialX/MaterialExpressionAppend3Vector.cpp b/Engine/Plugins/Experimental/Interchange/Runtime/Source/Import/Private/MaterialX/MaterialExpressionAppend3Vector.cpp
new file mode 100644
--- /dev/null
+++ b/Engine/Plugins/Experimental/Interchange/Runtime/Source/Import/Private/MaterialX/MaterialExpressionAppend3Vector.cpp
@@ -0,0 +1,60 @@
+// Copyright Epic Games, Inc. All Rights Reserved.
+
+#include "MaterialX/MaterialExpressionAppend3Vector.h"
+#include "MaterialCompiler.h"
+
+#define LOCTEXT_NAMESPACE "MaterialExpressionAppend3Vector"
+
+UMaterialExpressionAppend3Vector::UMaterialExpressionAppend3Vector(const FObjectInitializer& ObjectInitializer)
+	: Super(ObjectInitializer)
+{
+	// Structure to hold one-time initialization
+	struct FConstructorStatics
+	{
+		FText NAME_Math;
+		FText NAME_VectorOps;
+		FConstructorStatics()
+			: NAME_Math(LOCTEXT("Math", "Math"))
+			, NAME_VectorOps(LOCTEXT("VectorOps", "VectorOps"))
+		{}
+	};
+	static FConstructorStatics ConstructorStatics;
+
+#if WITH_EDITORONLY_DATA
+	MenuCategories.Add(ConstructorStatics.NAME_Math);
+	MenuCategories.Add(ConstructorStatics.NAME_VectorOps);
+#endif
+}
+
+#if WITH_EDITOR
+int32 UMaterialExpressionAppend3Vector::Compile(FMaterialCompiler* Compiler, int32 OutputIndex)
+{
+	if(!A.GetTracedInput().Expression)
+	{
+		return Compiler->Errorf(TEXT("Missing AppendVector input A"));
+	}
+	else if(!B.GetTracedInput().Expression)
+	{
+		return Compiler->Errorf(TEXT("Missing AppendVector input B"));
+	}
+	else if(!C.GetTracedInput().Expression)
+	{
+		return Compiler->Errorf(TEXT("Missing AppendVector input C"));
+	}
+	else
+	{
+		int32 Arg1 = A.Compile(Compiler);
+		int32 Arg2 = B.Compile(Compiler);
+		int32 Arg3 = C.Compile(Compiler);
+		return Compiler->AppendVector(Arg1,
+			   Compiler->AppendVector(Arg2, Arg3));
+	}
+}
+
+void UMaterialExpressionAppend3Vector::GetCaption(TArray<FString>& OutCaptions) const
+{
+	OutCaptions.Add(TEXT("Append3"));
+}
+#endif
+
+#undef LOCTEXT_NAMESPACE 
\ No newline at end of file
diff --git a/Engine/Plugins/Experimental/Interchange/Runtime/Source/Import/Private/MaterialX/MaterialExpressionAppend4Vector.cpp b/Engine/Plugins/Experimental/Interchange/Runtime/Source/Import/Private/MaterialX/MaterialExpressionAppend4Vector.cpp
new file mode 100644
--- /dev/null
+++ b/Engine/Plugins/Experimental/Interchange/Runtime/Source/Import/Private/MaterialX/MaterialExpressionAppend4Vector.cpp
@@ -0,0 +1,66 @@
+// Copyright Epic Games, Inc. All Rights Reserved.
+
+#include "MaterialX/MaterialExpressionAppend4Vector.h"
+#include "MaterialCompiler.h"
+
+#define LOCTEXT_NAMESPACE "MaterialExpressionAppend4Vector"
+
+UMaterialExpressionAppend4Vector::UMaterialExpressionAppend4Vector(const FObjectInitializer& ObjectInitializer)
+	: Super(ObjectInitializer)
+{
+	// Structure to hold one-time initialization
+	struct FConstructorStatics
+	{
+		FText NAME_Math;
+		FText NAME_VectorOps;
+		FConstructorStatics()
+			: NAME_Math(LOCTEXT("Math", "Math"))
+			, NAME_VectorOps(LOCTEXT("VectorOps", "VectorOps"))
+		{}
+	};
+	static FConstructorStatics ConstructorStatics;
+
+#if WITH_EDITORONLY_DATA
+	MenuCategories.Add(ConstructorStatics.NAME_Math);
+	MenuCategories.Add(ConstructorStatics.NAME_VectorOps);
+#endif
+}
+
+#if WITH_EDITOR
+int32 UMaterialExpressionAppend4Vector::Compile(FMaterialCompiler* Compiler, int32 OutputIndex)
+{
+	if(!A.GetTracedInput().Expression)
+	{
+		return Compiler->Errorf(TEXT("Missing AppendVector input A"));
+	}
+	else if(!B.GetTracedInput().Expression)
+	{
+		return Compiler->Errorf(TEXT("Missing AppendVector input B"));
+	}
+	else if(!C.GetTracedInput().Expression)
+	{
+		return Compiler->Errorf(TEXT("Missing AppendVector input C"));
+	}
+	else if(!D.GetTracedInput().Expression)
+	{
+		return Compiler->Errorf(TEXT("Missing AppendVector input D"));
+	}
+	else
+	{
+		int32 Arg1 = A.Compile(Compiler);
+		int32 Arg2 = B.Compile(Compiler);
+		int32 Arg3 = C.Compile(Compiler);
+		int32 Arg4 = D.Compile(Compiler);
+		return Compiler->AppendVector(Arg1,
+			   Compiler->AppendVector(Arg2,
+			   Compiler->AppendVector(Arg3, Arg4)));
+	}
+}
+
+void UMaterialExpressionAppend4Vector::GetCaption(TArray<FString>& OutCaptions) const
+{
+	OutCaptions.Add(TEXT("Append4"));
+}
+#endif
+
+#undef LOCTEXT_NAMESPACE
\ No newline at end of file
diff --git a/Engine/Plugins/Experimental/Interchange/Runtime/Source/Import/Private/MaterialX/MaterialExpressionExponential.cpp b/Engine/Plugins/Experimental/Interchange/Runtime/Source/Import/Private/MaterialX/MaterialExpressionExponential.cpp
new file mode 100644
--- /dev/null
+++ b/Engine/Plugins/Experimental/Interchange/Runtime/Source/Import/Private/MaterialX/MaterialExpressionExponential.cpp
@@ -0,0 +1,46 @@
+// Copyright Epic Games, Inc. All Rights Reserved.
+#include "MaterialX/MaterialExpressionExponential.h"
+
+#include "MaterialCompiler.h"
+#include "MaterialHLSLGenerator.h"
+
+#define LOCTEXT_NAMESPACE "MaterialExpressionExponential"
+
+UMaterialExpressionExponential::UMaterialExpressionExponential(const FObjectInitializer& ObjectInitializer)
+	: Super(ObjectInitializer)
+{
+	// Structure to hold one-time initialization
+	struct FConstructorStatics
+	{
+		FText NAME_Math;
+		FConstructorStatics()
+			: NAME_Math(LOCTEXT("Math", "Math"))
+		{}
+	};
+	static FConstructorStatics ConstructorStatics;
+
+#if WITH_EDITORONLY_DATA
+	MenuCategories.Add(ConstructorStatics.NAME_Math);
+#endif
+}
+
+#if WITH_EDITOR
+int32 UMaterialExpressionExponential::Compile(FMaterialCompiler* Compiler, int32 OutputIndex)
+{
+	if(!Input.GetTracedInput().Expression)
+	{
+		return Compiler->Errorf(TEXT("Missing Exponential input"));
+	}
+
+	int32 Euler = Compiler->Constant(2.7182818);
+
+	return Compiler->Power(Euler, Input.Compile(Compiler));
+}
+
+void UMaterialExpressionExponential::GetCaption(TArray<FString>& OutCaptions) const
+{
+	OutCaptions.Add(TEXT("Exponential"));
+}
+#endif
+
+#undef LOCTEXT_NAMESPACE 
\ No newline at end of file
diff --git a/Engine/Plugins/Experimental/Interchange/Runtime/Source/Import/Private/MaterialX/MaterialExpressionLength.cpp b/Engine/Plugins/Experimental/Interchange/Runtime/Source/Import/Private/MaterialX/MaterialExpressionLength.cpp
new file mode 100644
--- /dev/null
+++ b/Engine/Plugins/Experimental/Interchange/Runtime/Source/Import/Private/MaterialX/MaterialExpressionLength.cpp
@@ -0,0 +1,42 @@
+// Copyright Epic Games, Inc. All Rights Reserved.
+#include "MaterialX/MaterialExpressionLength.h"
+#include "MaterialCompiler.h"
+
+#define LOCTEXT_NAMESPACE "MaterialExpressionLength"
+
+UMaterialExpressionLength::UMaterialExpressionLength(const FObjectInitializer& ObjectInitializer)
+	: Super(ObjectInitializer)
+{
+	// Structure to hold one-time initialization
+	struct FConstructorStatics
+	{
+		FText NAME_Utility;
+		FConstructorStatics()
+			: NAME_Utility(LOCTEXT("Utility", "Utility"))
+		{}
+	};
+	static FConstructorStatics ConstructorStatics;
+
+#if WITH_EDITORONLY_DATA
+	MenuCategories.Add(ConstructorStatics.NAME_Utility);
+#endif
+}
+
+#if WITH_EDITOR
+int32 UMaterialExpressionLength::Compile(FMaterialCompiler* Compiler, int32 OutputIndex)
+{
+	if(!Input.GetTracedInput().Expression)
+	{
+		return Compiler->Errorf(TEXT("Missing Length input"));
+	}
+
+	return Compiler->Length(Input.Compile(Compiler));
+}
+
+void UMaterialExpressionLength::GetCaption(TArray<FString>& OutCaptions) const
+{
+	OutCaptions.Add(TEXT("Length"));
+}
+#endif
+
+#undef LOCTEXT_NAMESPACE 
\ No newline at end of file
diff --git a/Engine/Plugins/Experimental/Interchange/Runtime/Source/Import/Public/MaterialX/InterchangeMaterialXTranslator.h b/Engine/Plugins/Experimental/Interchange/Runtime/Source/Import/Public/MaterialX/InterchangeMaterialXTranslator.h
--- a/Engine/Plugins/Experimental/Interchange/Runtime/Source/Import/Public/MaterialX/InterchangeMaterialXTranslator.h
+++ b/Engine/Plugins/Experimental/Interchange/Runtime/Source/Import/Public/MaterialX/InterchangeMaterialXTranslator.h
@@ -130,20 +130,19 @@ protected:
 	 * 
 	 * @param NodeName - The name of the shader node
 	 * @param ShaderType - The shader node's type we want to create
-	 * @param ParentNode - The parent node of the created node
 	 * @param NamesToShaderNodes - Map of the shader nodes already created
 	 * @param The container where to add the translated Interchange nodes.
 	 * 
 	 * @return The shader node that was created
 	 */
 	template<typename ShaderNodeType>
-	ShaderNodeType* CreateShaderNode(const FString& NodeName, const FString& ShaderType, const FString & ParentNodeUID, TMap<FString, UInterchangeShaderNode*>& NamesToShaderNodes, UInterchangeBaseNodeContainer& NodeContainer) const
+	ShaderNodeType* CreateShaderNode(const FString& NodeName, const FString& ShaderType, TMap<FString, UInterchangeShaderNode*>& NamesToShaderNodes, UInterchangeBaseNodeContainer& NodeContainer) const
 	{
 		static_assert(std::is_convertible_v<ShaderNodeType*, UInterchangeShaderNode*>, "CreateShaderNode only accepts type that derived from UInterchangeShaderNode");
 
 		ShaderNodeType* Node;
 
-		const FString NodeUID = UInterchangeShaderNode::MakeNodeUid(NodeName, ParentNodeUID);
+		const FString NodeUID = UInterchangeShaderNode::MakeNodeUid(NodeName, FStringView{});
 		
 		//Test directly in the NodeContainer, because the NamesToShaderNodes can be altered during the node graph either by the parent (dot/normalmap),
 		//or by putting an intermediary node between the child and the parent (tiledimage)
@@ -154,7 +153,7 @@ protected:
 			NodeContainer.AddNode(Node);
 			if constexpr(std::is_same_v<ShaderNodeType, UInterchangeShaderGraphNode>)
 			{
-				NodeContainer.SetNodeParentUid(NodeUID, ParentNodeUID);
+				NodeContainer.SetNodeParentUid(NodeUID, FString{});
 			}
 			Node->SetCustomShaderType(ShaderType);
 
@@ -326,22 +325,70 @@ protected:
 	FString GetColorSpace(MaterialX::ElementPtr Element) const;
 
 	/**
-	 * Helper function that returns a color after a color space conversion, the function makes no assumption on the input, and it should have a value of Color3 type
+	 * Helper function that returns a linear color after a color space conversion, the function makes no assumption on the input, and it should have a value of Color3-4 type
 	 * 
-	 * @param Input - The input that has a Color3 value in it
+	 * @param Input - The input that has a Color3-4 value in it
 	 * 
 	 * @return The linear color after color space conversion
 	 */
-	FLinearColor MakeLinearColorFromColor3(MaterialX::InputPtr Input) const;
+	FLinearColor GetLinearColor(MaterialX::InputPtr Input) const;
+
+	void ConnectConstantInputToOutput(MaterialX::NodePtr UpstreamNode, UInterchangeShaderNode* ParentShaderNode, const FString& InputChannelName, TMap<FString, UInterchangeShaderNode*>& NamesToShaderNodes, UInterchangeBaseNodeContainer& NodeContainer) const;
+	void ConnectExtractInputToOutput(MaterialX::NodePtr UpstreamNode, UInterchangeShaderNode* ParentShaderNode, const FString& InputChannelName, TMap<FString, UInterchangeShaderNode*>& NamesToShaderNodes, UInterchangeBaseNodeContainer& NodeContainer) const;
+	void ConnectDotInputToOutput(MaterialX::NodePtr UpstreamNode, UInterchangeShaderNode* ParentShaderNode, const FString& InputChannelName, TMap<FString, UInterchangeShaderNode*>& NamesToShaderNodes, UInterchangeBaseNodeContainer& NodeContainer) const;
+	void ConnectTransformPositionInputToOutput(MaterialX::NodePtr UpstreamNode, UInterchangeShaderNode* ParentShaderNode, const FString& InputChannelName, TMap<FString, UInterchangeShaderNode*>& NamesToShaderNodes, UInterchangeBaseNodeContainer& NodeContainer) const;
+	void ConnectTransformVectorInputToOutput(MaterialX::NodePtr UpstreamNode, UInterchangeShaderNode* ParentShaderNode, const FString& InputChannelName, TMap<FString, UInterchangeShaderNode*>& NamesToShaderNodes, UInterchangeBaseNodeContainer& NodeContainer) const;
+	void ConnectRotate2DInputToOutput(MaterialX::NodePtr UpstreamNode, UInterchangeShaderNode* ParentShaderNode, const FString& InputChannelName, TMap<FString, UInterchangeShaderNode*>& NamesToShaderNodes, UInterchangeBaseNodeContainer& NodeContainer) const;
+	void ConnectRotate3DInputToOutput(MaterialX::NodePtr UpstreamNode, UInterchangeShaderNode* ParentShaderNode, const FString& InputChannelName, TMap<FString, UInterchangeShaderNode*>& NamesToShaderNodes, UInterchangeBaseNodeContainer& NodeContainer) const;
+	void ConnectImageInputToOutput(MaterialX::NodePtr UpstreamNode, UInterchangeShaderNode* ParentShaderNode, const FString& InputChannelName, TMap<FString, UInterchangeShaderNode*>& NamesToShaderNodes, UInterchangeBaseNodeContainer& NodeContainer) const;
+	void ConnectConvertInputToOutput(MaterialX::NodePtr UpstreamNode, UInterchangeShaderNode* ParentShaderNode, const FString& InputChannelName, TMap<FString, UInterchangeShaderNode*>& NamesToShaderNodes, UInterchangeBaseNodeContainer& NodeContainer) const;
+
+	UInterchangeShaderNode* CreateMaskShaderNode(uint8 RGBA, const FString & NodeName, TMap<FString, UInterchangeShaderNode*>& NamesToShaderNodes, UInterchangeBaseNodeContainer& NodeContainer) const;
 
 	/**
-	 * Helper function that returns a color after a color space conversion, the function makes no assumption on the input, and it should have a value of Color3 type
+	 * Helper function that returns a vector, the function makes no assumption on the input, and it should have a value of vectorN type
 	 *
-	 * @param Input - The input that has a Color3 value in it
+	 * @param Input - The input that has a vectorN value in it
 	 *
-	 * @return The linear color after color space conversion
+	 * @return The vector
+	 */
+	FLinearColor GetVector(MaterialX::InputPtr Input) const;
+
+	/** 
+	 * @param const char * - New name of the Input to copy
+	 * @param MaterialX::InputPtr - Input to copy
+	 */
+	using FInputToCopy = TPair<const char*, MaterialX::InputPtr>; 
+
+	/**
+	 * @param const char * - Attribute
+	 * @param const char * - Value
+	 */
+	using FAttributeValue = TPair<const char*, const char*>;
+
+	/**
+	 * @param TArray<FAttributeValueType> - Array of attribute/value
+	 */
+	using FAttributeValueArray = TArray<FAttributeValue>;
+
+	/**
+	 * @param TArray<FAttributeValueType> - Array of attribute/value
+	 */
+	using FInputToCreate = TPair<const char*, FAttributeValueArray>; // Name - Attributes
+
+	/**
+	 * Helper function that create a MaterialX node, useful to insert a node between 2 nodes especially when an operation cannot be done in one pass
+	 * If the node already exists in the graph, this one is returned
+	 *
+	 * @param NodeGraph - The node graph to insert the new node into
+	 * @param NodeName - The name of the node
+	 * @param Category - The category of the node to create
+	 * @param InputsToCopy - The inputs from another node/nodegraph to copy in the new node, can be empty
+	 * @param InputsToCreate - The inputs to create from scratch and insert in the new node, can be empty
+	 * 
+	 * @return The newly created node
 	 */
-	FLinearColor MakeLinearColorFromColor4(MaterialX::InputPtr Input) const;
+	MaterialX::NodePtr CreateNode(MaterialX::NodeGraphPtr NodeGraph, const char * NodeName, const char * Category, TArray<FInputToCopy> InputsToCopy, TArray<FInputToCreate> InputsToCreate) const;
 
 private:
 
diff --git a/Engine/Plugins/Experimental/Interchange/Runtime/Source/Import/Public/MaterialX/MaterialExpressionAppend3Vector.h b/Engine/Plugins/Experimental/Interchange/Runtime/Source/Import/Public/MaterialX/MaterialExpressionAppend3Vector.h
new file mode 100644
--- /dev/null
+++ b/Engine/Plugins/Experimental/Interchange/Runtime/Source/Import/Public/MaterialX/MaterialExpressionAppend3Vector.h
@@ -0,0 +1,33 @@
+// Copyright Epic Games, Inc. All Rights Reserved.
+#pragma once
+
+#include "CoreMinimal.h"
+#include "Materials/MaterialExpression.h"
+
+#include "MaterialExpressionAppend3Vector.generated.h"
+
+/**
+ * A material expression that allows combining 3 channels together to create a vector with more channel than the original
+ */
+UCLASS(MinimalAPI, collapsecategories, hidecategories = Object)
+class UMaterialExpressionAppend3Vector : public UMaterialExpression
+{
+	GENERATED_UCLASS_BODY()
+
+	UPROPERTY()
+	FExpressionInput A;
+
+	UPROPERTY()
+	FExpressionInput B;
+
+	UPROPERTY()
+	FExpressionInput C;
+
+	//~ Begin UMaterialExpression Interface
+#if WITH_EDITOR
+	virtual int32 Compile(class FMaterialCompiler* Compiler, int32 OutputIndex) override;
+	virtual void GetCaption(TArray<FString>& OutCaptions) const override;
+#endif
+	//~ End UMaterialExpression Interface
+};
+
diff --git a/Engine/Plugins/Experimental/Interchange/Runtime/Source/Import/Public/MaterialX/MaterialExpressionAppend4Vector.h b/Engine/Plugins/Experimental/Interchange/Runtime/Source/Import/Public/MaterialX/MaterialExpressionAppend4Vector.h
new file mode 100644
--- /dev/null
+++ b/Engine/Plugins/Experimental/Interchange/Runtime/Source/Import/Public/MaterialX/MaterialExpressionAppend4Vector.h
@@ -0,0 +1,38 @@
+// Copyright Epic Games, Inc. All Rights Reserved.
+#pragma once
+
+#include "CoreMinimal.h"
+#include "UObject/ObjectMacros.h"
+#include "MaterialExpressionIO.h"
+#include "Materials/MaterialExpression.h"
+
+#include "MaterialExpressionAppend4Vector.generated.h"
+
+/**
+ * A material expression that allows combining 4 channels together to create a vector with more channel than the original
+ */
+UCLASS(MinimalAPI, collapsecategories, hidecategories = Object)
+class UMaterialExpressionAppend4Vector : public UMaterialExpression
+{
+	GENERATED_UCLASS_BODY()
+
+	UPROPERTY()
+	FExpressionInput A;
+	
+	UPROPERTY()
+	FExpressionInput B;
+
+	UPROPERTY()
+	FExpressionInput C;
+	
+	UPROPERTY()
+	FExpressionInput D;
+
+	//~ Begin UMaterialExpression Interface
+#if WITH_EDITOR
+	virtual int32 Compile(class FMaterialCompiler* Compiler, int32 OutputIndex) override;
+	virtual void GetCaption(TArray<FString>& OutCaptions) const override;
+#endif
+	//~ End UMaterialExpression Interface
+};
+
diff --git a/Engine/Plugins/Experimental/Interchange/Runtime/Source/Import/Public/MaterialX/MaterialExpressionExponential.h b/Engine/Plugins/Experimental/Interchange/Runtime/Source/Import/Public/MaterialX/MaterialExpressionExponential.h
new file mode 100644
--- /dev/null
+++ b/Engine/Plugins/Experimental/Interchange/Runtime/Source/Import/Public/MaterialX/MaterialExpressionExponential.h
@@ -0,0 +1,27 @@
+// Copyright Epic Games, Inc. All Rights Reserved.
+#pragma once
+
+#include "CoreMinimal.h"
+#include "Materials/MaterialExpression.h"
+
+#include "MaterialExpressionExponential.generated.h"
+
+/**
+ * A material expression that computes the exponential of an input
+ */
+UCLASS(collapsecategories, hidecategories = Object, MinimalAPI)
+class UMaterialExpressionExponential : public UMaterialExpression
+{
+	GENERATED_UCLASS_BODY()
+
+	UPROPERTY()
+	FExpressionInput Input;
+
+	//~ Begin UMaterialExpression Interface
+#if WITH_EDITOR
+	virtual int32 Compile(class FMaterialCompiler* Compiler, int32 OutputIndex) override;
+	virtual void GetCaption(TArray<FString>& OutCaptions) const override;
+#endif
+	//~ End UMaterialExpression Interface
+};
+
diff --git a/Engine/Plugins/Experimental/Interchange/Runtime/Source/Import/Public/MaterialX/MaterialExpressionLength.h b/Engine/Plugins/Experimental/Interchange/Runtime/Source/Import/Public/MaterialX/MaterialExpressionLength.h
new file mode 100644
--- /dev/null
+++ b/Engine/Plugins/Experimental/Interchange/Runtime/Source/Import/Public/MaterialX/MaterialExpressionLength.h
@@ -0,0 +1,27 @@
+// Copyright Epic Games, Inc. All Rights Reserved.
+#pragma once
+
+#include "CoreMinimal.h"
+#include "Materials/MaterialExpression.h"
+
+#include "MaterialExpressionLength.generated.h"
+
+/**
+ * A material expression that computes the length (magnitude) of a vector
+ */
+UCLASS(collapsecategories, hidecategories = Object, MinimalAPI)
+class UMaterialExpressionLength : public UMaterialExpression
+{
+	GENERATED_UCLASS_BODY()
+
+	UPROPERTY()
+	FExpressionInput Input;
+
+	//~ Begin UMaterialExpression Interface
+#if WITH_EDITOR
+	virtual int32 Compile(class FMaterialCompiler* Compiler, int32 OutputIndex) override;
+	virtual void GetCaption(TArray<FString>& OutCaptions) const override;
+#endif
+	//~ End UMaterialExpression Interface
+};
+
diff --git a/Engine/Plugins/Experimental/Interchange/Runtime/Source/Pipelines/Private/InterchangeGenericMaterialPipeline.cpp b/Engine/Plugins/Experimental/Interchange/Runtime/Source/Pipelines/Private/InterchangeGenericMaterialPipeline.cpp
--- a/Engine/Plugins/Experimental/Interchange/Runtime/Source/Pipelines/Private/InterchangeGenericMaterialPipeline.cpp
+++ b/Engine/Plugins/Experimental/Interchange/Runtime/Source/Pipelines/Private/InterchangeGenericMaterialPipeline.cpp
@@ -516,30 +516,31 @@ bool UInterchangeGenericMaterialPipeline::HandleStandardSurfaceModel(const UInte
 		
 		static ArgumentsType ArgumentsPerInputs[]
 		{
-			ArgumentsType{StandardSurface::Parameters::Base.ToString(), EMaterialInputType::Scalar},					// Base
-			ArgumentsType{StandardSurface::Parameters::BaseColor.ToString(), EMaterialInputType::Color},				// BaseColor
-			ArgumentsType{StandardSurface::Parameters::DiffuseRoughness.ToString(),  EMaterialInputType::Scalar},		// DiffuseRoughness
-			ArgumentsType{StandardSurface::Parameters::Metalness.ToString(), EMaterialInputType::Scalar},			// Metalness
-			ArgumentsType{StandardSurface::Parameters::Specular.ToString(), EMaterialInputType::Scalar},				// Specular
-			ArgumentsType{StandardSurface::Parameters::SpecularRoughness.ToString(), EMaterialInputType::Scalar},	// SpecularRoughness
-			ArgumentsType{StandardSurface::Parameters::SpecularIOR.ToString(), EMaterialInputType::Scalar},			// SpecularIOR
-			ArgumentsType{StandardSurface::Parameters::SpecularAnisotropy.ToString(), EMaterialInputType::Scalar},	// SpecularAnisotropy
-			ArgumentsType{StandardSurface::Parameters::SpecularRotation.ToString(), EMaterialInputType::Scalar},		// SpecularRotation
-			ArgumentsType{StandardSurface::Parameters::Subsurface.ToString(), EMaterialInputType::Scalar},			// Subsurface
-			ArgumentsType{StandardSurface::Parameters::SubsurfaceColor.ToString(), EMaterialInputType::Scalar},		// SubsurfaceColor
-			ArgumentsType{StandardSurface::Parameters::SubsurfaceRadius.ToString(), EMaterialInputType::Scalar},		// SubsurfaceRadius
-			ArgumentsType{StandardSurface::Parameters::SubsurfaceScale.ToString(), EMaterialInputType::Scalar},		// SubsurfaceScale
-			ArgumentsType{StandardSurface::Parameters::Sheen.ToString(), EMaterialInputType::Scalar},				// Sheen
-			ArgumentsType{StandardSurface::Parameters::SheenColor.ToString(), EMaterialInputType::Color},			// SheenColor
-			ArgumentsType{StandardSurface::Parameters::SheenRoughness.ToString(), EMaterialInputType::Scalar},		// SheenRoughness
-			ArgumentsType{StandardSurface::Parameters::Coat.ToString(), EMaterialInputType::Scalar},					// Coat
-			ArgumentsType{StandardSurface::Parameters::CoatColor.ToString(),  EMaterialInputType::Color},			// Coat Color
-			ArgumentsType{StandardSurface::Parameters::CoatRoughness.ToString(), EMaterialInputType::Scalar},		// CoatRoughness
-			ArgumentsType{StandardSurface::Parameters::ThinFilmThickness.ToString(), EMaterialInputType::Scalar},	// ThinFilmThickness
-			ArgumentsType{StandardSurface::Parameters::Emission.ToString(), EMaterialInputType::Scalar},				// Emission
-			ArgumentsType{StandardSurface::Parameters::EmissionColor.ToString(), EMaterialInputType::Color},			// EmissionColor
-			ArgumentsType{StandardSurface::Parameters::Normal.ToString(), EMaterialInputType::Vector},				// Normal
-			ArgumentsType{StandardSurface::Parameters::Tangent.ToString(), EMaterialInputType::Vector}				// Tangent
+			ArgumentsType{StandardSurface::Parameters::Base.ToString(), EMaterialInputType::Scalar},               // Base
+			ArgumentsType{StandardSurface::Parameters::BaseColor.ToString(), EMaterialInputType::Color},           // BaseColor
+			ArgumentsType{StandardSurface::Parameters::DiffuseRoughness.ToString(),  EMaterialInputType::Scalar},  // DiffuseRoughness
+			ArgumentsType{StandardSurface::Parameters::Metalness.ToString(), EMaterialInputType::Scalar},          // Metalness
+			ArgumentsType{StandardSurface::Parameters::Specular.ToString(), EMaterialInputType::Scalar},           // Specular
+			ArgumentsType{StandardSurface::Parameters::SpecularRoughness.ToString(), EMaterialInputType::Scalar},  // SpecularRoughness
+			ArgumentsType{StandardSurface::Parameters::SpecularIOR.ToString(), EMaterialInputType::Scalar},        // SpecularIOR
+			ArgumentsType{StandardSurface::Parameters::SpecularAnisotropy.ToString(), EMaterialInputType::Scalar}, // SpecularAnisotropy
+			ArgumentsType{StandardSurface::Parameters::SpecularRotation.ToString(), EMaterialInputType::Scalar},   // SpecularRotation
+			ArgumentsType{StandardSurface::Parameters::Subsurface.ToString(), EMaterialInputType::Scalar},         // Subsurface
+			ArgumentsType{StandardSurface::Parameters::SubsurfaceColor.ToString(), EMaterialInputType::Scalar},    // SubsurfaceColor
+			ArgumentsType{StandardSurface::Parameters::SubsurfaceRadius.ToString(), EMaterialInputType::Scalar},   // SubsurfaceRadius
+			ArgumentsType{StandardSurface::Parameters::SubsurfaceScale.ToString(), EMaterialInputType::Scalar},    // SubsurfaceScale
+			ArgumentsType{StandardSurface::Parameters::Sheen.ToString(), EMaterialInputType::Scalar},              // Sheen
+			ArgumentsType{StandardSurface::Parameters::SheenColor.ToString(), EMaterialInputType::Color},          // SheenColor
+			ArgumentsType{StandardSurface::Parameters::SheenRoughness.ToString(), EMaterialInputType::Scalar},     // SheenRoughness
+			ArgumentsType{StandardSurface::Parameters::Coat.ToString(), EMaterialInputType::Scalar},               // Coat
+			ArgumentsType{StandardSurface::Parameters::CoatColor.ToString(),  EMaterialInputType::Color},          // CoatColor
+			ArgumentsType{StandardSurface::Parameters::CoatNormal.ToString(),  EMaterialInputType::Vector},        // CoatNormal
+			ArgumentsType{StandardSurface::Parameters::CoatRoughness.ToString(), EMaterialInputType::Scalar},      // CoatRoughness
+			ArgumentsType{StandardSurface::Parameters::ThinFilmThickness.ToString(), EMaterialInputType::Scalar},  // ThinFilmThickness
+			ArgumentsType{StandardSurface::Parameters::Emission.ToString(), EMaterialInputType::Scalar},           // Emission
+			ArgumentsType{StandardSurface::Parameters::EmissionColor.ToString(), EMaterialInputType::Color},       // EmissionColor
+			ArgumentsType{StandardSurface::Parameters::Normal.ToString(), EMaterialInputType::Vector},             // Normal
+			ArgumentsType{StandardSurface::Parameters::Tangent.ToString(), EMaterialInputType::Vector}             // Tangent
 		};
 
 		for(const ArgumentsType & Arguments : ArgumentsPerInputs)
