Fixing a crash due to the fact of taking a dangling reference on the MaterialExpressionCreationContextStack array after a realloc.

diff --git a/Engine/Plugins/Experimental/Interchange/Runtime/Source/Pipelines/Private/InterchangeGenericMaterialPipeline.cpp b/Engine/Plugins/Experimental/Interchange/Runtime/Source/Pipelines/Private/InterchangeGenericMaterialPipeline.cpp
--- a/Engine/Plugins/Experimental/Interchange/Runtime/Source/Pipelines/Private/InterchangeGenericMaterialPipeline.cpp
+++ b/Engine/Plugins/Experimental/Interchange/Runtime/Source/Pipelines/Private/InterchangeGenericMaterialPipeline.cpp
@@ -1537,10 +1537,10 @@ TTuple<UInterchangeMaterialExpressionFactoryNode*, FString> UInterchangeGenericM
 
 	UInterchangeMaterialExpressionFactoryNode* MaterialExpressionFactoryNode = nullptr;
 
-	FMaterialExpressionCreationContext& ExpressionCreationContext = MaterialExpressionCreationContextStack.AddDefaulted_GetRef();
+	int32 ExpressionContextIndex = MaterialExpressionCreationContextStack.AddDefaulted();
 
 	FString ConnectedShaderNodeUid;
-	if (UInterchangeShaderPortsAPI::GetInputConnection(ShaderNode, InputName, ConnectedShaderNodeUid, ExpressionCreationContext.OutputName))
+	if (UInterchangeShaderPortsAPI::GetInputConnection(ShaderNode, InputName, ConnectedShaderNodeUid, MaterialExpressionCreationContextStack[ExpressionContextIndex].OutputName))
 	{
 		if (const UInterchangeShaderNode* ConnectedShaderNode = Cast<const UInterchangeShaderNode>(BaseNodeContainer->GetNode(ConnectedShaderNodeUid)))
 		{
@@ -1563,7 +1563,7 @@ TTuple<UInterchangeMaterialExpressionFactoryNode*, FString> UInterchangeGenericM
 		}
 	}
 
-	TTuple<UInterchangeMaterialExpressionFactoryNode*, FString> Result {MaterialExpressionFactoryNode, ExpressionCreationContext.OutputName};
+	TTuple<UInterchangeMaterialExpressionFactoryNode*, FString> Result {MaterialExpressionFactoryNode, MaterialExpressionCreationContextStack[ExpressionContextIndex].OutputName};
 	MaterialExpressionCreationContextStack.Pop();
 
 	return Result;
