GLTF/Interchange Adding Non supported extension handling (failing out)

diff --git a/Engine/Plugins/Editor/GLTFImporter/Source/GLTFCore/Private/GLTF/GLTFAsset.cpp b/Engine/Plugins/Editor/GLTFImporter/Source/GLTFCore/Private/GLTF/GLTFAsset.cpp
--- a/Engine/Plugins/Editor/GLTFImporter/Source/GLTFCore/Private/GLTF/GLTFAsset.cpp
+++ b/Engine/Plugins/Editor/GLTFImporter/Source/GLTFCore/Private/GLTF/GLTFAsset.cpp
@@ -59,6 +59,7 @@ namespace GLTF
 		Materials.Empty();
 
 		ExtensionsUsed.Empty((int)EExtension::Count);
+		RequiredExtensions.Empty();
 		Metadata.GeneratorName.Empty();
 		Metadata.Extras.Empty();
 
diff --git a/Engine/Plugins/Editor/GLTFImporter/Source/GLTFCore/Private/GLTF/GLTFReader.cpp b/Engine/Plugins/Editor/GLTFImporter/Source/GLTFCore/Private/GLTF/GLTFReader.cpp
--- a/Engine/Plugins/Editor/GLTFImporter/Source/GLTFCore/Private/GLTF/GLTFReader.cpp
+++ b/Engine/Plugins/Editor/GLTFImporter/Source/GLTFCore/Private/GLTF/GLTFReader.cpp
@@ -809,6 +809,8 @@ namespace GLTF
 		const uint32 TextureCount  = ArraySize(*JsonRoot, TEXT("textures"));
 		const uint32 MaterialCount = ArraySize(*JsonRoot, TEXT("materials"));
 
+		const uint32 ExtensionsRequiredCount = ArraySize(*JsonRoot, TEXT("extensionsRequired"));
+
 		{
 			// cleanup and reserve
 			OutAsset.Buffers.Empty(BufferCount);
@@ -826,6 +828,7 @@ namespace GLTF
 			OutAsset.Textures.Empty(TextureCount);
 			OutAsset.Materials.Empty(MaterialCount);
 			OutAsset.ExtensionsUsed.Empty((int)EExtension::Count);
+			OutAsset.RequiredExtensions.Empty();
 		}
 
 		// allocate asset mapped data for images and buffers
@@ -852,6 +855,23 @@ namespace GLTF
 		SetupObjects(TextureCount, TEXT("textures"), [this](const FJsonObject& Object) { SetupTexture(Object); });
 		SetupObjects(MaterialCount, TEXT("materials"), [this](const FJsonObject& Object) { SetupMaterial(Object); });
 
+		SetupObjects(MaterialCount, TEXT("materials"), [this](const FJsonObject& Object) { SetupMaterial(Object); });
+
+
+		/*
+		* Returns true if all required extensions are supported.
+		* According to gltf specification checking only the top-level extensionsRequired property in order to decide
+		* if the model import is supported should be sufficient as the documentation states:
+		* "All glTF extensions required to load and/or render an asset MUST be listed in the top-level extensionsRequired array"
+		*/
+		{
+			for (const TSharedPtr<FJsonValue>& Extension : JsonRoot->GetArrayField(TEXT("extensionsRequired")))
+			{
+				FString ExtensionString = Extension->AsString();
+				OutAsset.RequiredExtensions.Add(ExtensionString);
+			}
+		}
+
 		SetupNodesType();
 		ExtensionsHandler->SetupAssetExtensions(*JsonRoot);
 	}
diff --git a/Engine/Plugins/Editor/GLTFImporter/Source/GLTFCore/Public/GLTFAsset.h b/Engine/Plugins/Editor/GLTFImporter/Source/GLTFCore/Public/GLTFAsset.h
--- a/Engine/Plugins/Editor/GLTFImporter/Source/GLTFCore/Public/GLTFAsset.h
+++ b/Engine/Plugins/Editor/GLTFImporter/Source/GLTFCore/Public/GLTFAsset.h
@@ -86,6 +86,7 @@ namespace GLTF
 		TArray<FString>   Variants;
 
 		TSet<EExtension> ExtensionsUsed;
+		TArray<FString>	 RequiredExtensions;
 		FMetadata        Metadata;
 
 		/**
diff --git a/Engine/Plugins/Experimental/Interchange/Runtime/Source/Import/Private/Gltf/InterchangeGltfTranslator.cpp b/Engine/Plugins/Experimental/Interchange/Runtime/Source/Import/Private/Gltf/InterchangeGltfTranslator.cpp
--- a/Engine/Plugins/Experimental/Interchange/Runtime/Source/Import/Private/Gltf/InterchangeGltfTranslator.cpp
+++ b/Engine/Plugins/Experimental/Interchange/Runtime/Source/Import/Private/Gltf/InterchangeGltfTranslator.cpp
@@ -33,6 +33,17 @@ static FAutoConsoleVariableRef CCvarInterchangeEnableGLTFImport(
 	TEXT("Whether glTF support is enabled."),
 	ECVF_Default);
 
+// is not supported yet
+static const TArray<FString> SupportedExtensions = {/* Lights */
+														//TEXT("KHR_lights_punctual"), TEXT("KHR_lights"), //not supported yet
+													/* Variants */
+														TEXT("KHR_materials_variants"), 
+													/* Materials */
+														//TEXT("KHR_materials_unlit"), //not supported yet
+														TEXT("KHR_materials_ior"), TEXT("KHR_materials_clearcoat"), TEXT("KHR_materials_transmission"), TEXT("KHR_materials_sheen"),
+														TEXT("KHR_materials_specular"), TEXT("KHR_materials_pbrSpecularGlossiness"), 
+														TEXT("MSFT_packing_occlusionRoughnessMetallic"), TEXT("MSFT_packing_normalRoughnessMetallic") };
+
 namespace UE::Interchange::Gltf::Private
 {
 	EInterchangeTextureWrapMode ConvertWrap(GLTF::FSampler::EWrap Wrap)
@@ -839,7 +850,28 @@ bool UInterchangeGltfTranslator::Translate( UInterchangeBaseNodeContainer& NodeC
 
 	const bool bLoadImageData = false;
 	const bool bLoadMetaData = false;
-	GltfFileReader.ReadFile( FilePath, bLoadImageData, bLoadMetaData, const_cast< UInterchangeGltfTranslator* >( this )->GltfAsset );	
+	GltfFileReader.ReadFile( FilePath, bLoadImageData, bLoadMetaData, const_cast< UInterchangeGltfTranslator* >( this )->GltfAsset );
+
+	if (GltfAsset.RequiredExtensions.Num() != 0)
+	{
+		FString NotSupportedExtensions;
+		for (const FString& RequiredExtension: GltfAsset.RequiredExtensions)
+		{
+			if (SupportedExtensions.Find(RequiredExtension) == INDEX_NONE)
+			{
+				if (NotSupportedExtensions.Len() > 0)
+				{
+					NotSupportedExtensions += ", ";
+				}
+				NotSupportedExtensions += RequiredExtension;
+			}
+		}
+		if (NotSupportedExtensions.Len() > 0)
+		{
+			UE_LOG(LogInterchangeImport, Error, TEXT("Not All Required Extensions are supported. (Unsupported extensions: %s)"), *NotSupportedExtensions);
+			return false;
+		}
+	}
 
 	const FString FileName = FPaths::GetBaseFilename(FilePath);
 	const_cast< UInterchangeGltfTranslator* >( this )->GltfAsset.GenerateNames(FileName);
