Add support for Rigged animation [Interchange]

GLTF/Interchange skeleton import (including recursive skeletons, single mesh with multiple skins)
GLTF/Interchange Rigged (Single) Animation support
GLTF Core fixes regarding Joint weight acquisitions

diff --git a/Engine/Plugins/Editor/GLTFImporter/Source/GLTFCore/Private/GLTF/GLTFAccessor.cpp b/Engine/Plugins/Editor/GLTFImporter/Source/GLTFCore/Private/GLTF/GLTFAccessor.cpp
--- a/Engine/Plugins/Editor/GLTFImporter/Source/GLTFCore/Private/GLTF/GLTFAccessor.cpp
+++ b/Engine/Plugins/Editor/GLTFImporter/Source/GLTFCore/Private/GLTF/GLTFAccessor.cpp
@@ -328,13 +328,13 @@ namespace GLTF
 					case EComponentType::U8:
 						for (int i = 0; i < 4; ++i)
 						{
-							Values[i] = *(const uint8*)ValuePtr;
+							Values[i] = ((uint8*)ValuePtr)[i];
 						}
 						return;
 					case EComponentType::U16:
 						for (int i = 0; i < 4; ++i)
 						{
-							Values[i] = *(const uint16*)ValuePtr;
+							Values[i] = ((uint16*)ValuePtr)[i];
 						}
 						return;
 					default:
diff --git a/Engine/Plugins/Editor/GLTFImporter/Source/GLTFCore/Private/GLTF/GLTFReader.cpp b/Engine/Plugins/Editor/GLTFImporter/Source/GLTFCore/Private/GLTF/GLTFReader.cpp
--- a/Engine/Plugins/Editor/GLTFImporter/Source/GLTFCore/Private/GLTF/GLTFReader.cpp
+++ b/Engine/Plugins/Editor/GLTFImporter/Source/GLTFCore/Private/GLTF/GLTFReader.cpp
@@ -276,6 +276,8 @@ namespace GLTF
 			{
 				const int32 NodeIndex = Value->AsNumber();
 				Scene.Nodes.Add(NodeIndex);
+
+				BuildParentIndices(-1, NodeIndex);
 			}
 		}
 
@@ -842,8 +844,8 @@ namespace GLTF
 		SetupObjects(AccessorCount, TEXT("accessors"), [this](const FJsonObject& Object) { SetupAccessor(Object); });
 
 		SetupObjects(MeshCount, TEXT("meshes"), [this](const FJsonObject& Object) { SetupMesh(Object); });
-		SetupObjects(SceneCount, TEXT("scenes"), [this](const FJsonObject& Object) { SetupScene(Object); });
 		SetupObjects(NodeCount, TEXT("nodes"), [this](const FJsonObject& Object) { SetupNode(Object); });
+		SetupObjects(SceneCount, TEXT("scenes"), [this](const FJsonObject& Object) { SetupScene(Object); });
 		SetupObjects(CameraCount, TEXT("cameras"), [this](const FJsonObject& Object) { SetupCamera(Object); });
 		SetupObjects(SkinCount, TEXT("skins"), [this](const FJsonObject& Object) { SetupSkin(Object); });
 		SetupObjects(AnimationsCount, TEXT("animations"), [this](const FJsonObject& Object) { SetupAnimation(Object); });
@@ -928,4 +930,19 @@ namespace GLTF
 		}
 	}
 
+	void FFileReader::BuildParentIndices(int32 ParentNodeIndex, int32 CurrentNodeIndex) const
+	{
+		if (!Asset->Nodes.IsValidIndex(CurrentNodeIndex))
+		{
+			return;
+		}
+		GLTF::FNode& Node = Asset->Nodes[CurrentNodeIndex];
+		Node.ParentIndex = ParentNodeIndex;
+
+		for (const int32 ChildNodeIndex : Node.Children)
+		{
+			BuildParentIndices(CurrentNodeIndex, ChildNodeIndex);
+		}
+	}
+
 }  // namespace GLTF
diff --git a/Engine/Plugins/Editor/GLTFImporter/Source/GLTFCore/Private/GLTFMeshFactory.cpp b/Engine/Plugins/Editor/GLTFImporter/Source/GLTFCore/Private/GLTFMeshFactory.cpp
--- a/Engine/Plugins/Editor/GLTFImporter/Source/GLTFCore/Private/GLTFMeshFactory.cpp
+++ b/Engine/Plugins/Editor/GLTFImporter/Source/GLTFCore/Private/GLTFMeshFactory.cpp
@@ -81,13 +81,11 @@ namespace GLTF
 			Vector4dBufferCount = 3,
 		};
 
-		using FIndexVertexIdMap = TMap<int32, FVertexID>;
-
 		float                ImportUniformScale;
 
 		TSet<int32>                  MaterialIndicesUsed;
 		TMap<int32, FPolygonGroupID> MaterialIndexToPolygonGroupID;
-		TArray<FIndexVertexIdMap>    PositionIndexToVertexIdPerPrim;
+		TArray<FMeshFactory::FIndexVertexIdMap>    PositionIndexToVertexIdPerPrim;
 
 		TArray<FVector2f>                       Vector2dBuffers[MAX_MESH_TEXTURE_COORDS_MD + 1];
 		TArray<FVector3f>                       VectorBuffers[VectorBufferCount];
@@ -197,7 +195,7 @@ namespace GLTF
 			TArray<FVector3f>& Positions = GetVectorBuffer(PositionBufferIndex);
 			Primitive.GetPositions(Positions);
 
-			FIndexVertexIdMap& PositionIndexToVertexId = PositionIndexToVertexIdPerPrim[Index];
+			FMeshFactory::FIndexVertexIdMap& PositionIndexToVertexId = PositionIndexToVertexIdPerPrim[Index];
 			PositionIndexToVertexId.Empty(Positions.Num());
 			for (int32 PositionIndex = 0; PositionIndex < Positions.Num(); ++PositionIndex)
 			{
@@ -472,4 +470,9 @@ namespace GLTF
 	{
 		Impl->CleanUp();
 	}
+
+	TArray<FMeshFactory::FIndexVertexIdMap>& FMeshFactory::GetPositionIndexToVertexIdPerPrim() const
+	{
+		return Impl->PositionIndexToVertexIdPerPrim;
+	}
 } //namespace GLTF
diff --git a/Engine/Plugins/Editor/GLTFImporter/Source/GLTFCore/Public/GLTFMeshFactory.h b/Engine/Plugins/Editor/GLTFImporter/Source/GLTFCore/Public/GLTFMeshFactory.h
--- a/Engine/Plugins/Editor/GLTFImporter/Source/GLTFCore/Public/GLTFMeshFactory.h
+++ b/Engine/Plugins/Editor/GLTFImporter/Source/GLTFCore/Public/GLTFMeshFactory.h
@@ -4,6 +4,7 @@
 
 #include "CoreMinimal.h"
 #include "GLTFLogger.h"
+#include "MeshTypes.h"
 
 struct FMeshDescription;
 
@@ -18,6 +19,8 @@ namespace GLTF
 		FMeshFactory();
 		~FMeshFactory();
 
+		using FIndexVertexIdMap = TMap<int32, FVertexID>;
+
 		void FillMeshDescription(const GLTF::FMesh &Mesh, FMeshDescription* MeshDescription);
 
 		float GetUniformScale() const;
@@ -27,6 +30,8 @@ namespace GLTF
 
 		void SetReserveSize(uint32 Size);
 
+		TArray<FMeshFactory::FIndexVertexIdMap>& GetPositionIndexToVertexIdPerPrim() const;
+
 		void CleanUp();
 
 	private:
diff --git a/Engine/Plugins/Editor/GLTFImporter/Source/GLTFCore/Public/GLTFNode.h b/Engine/Plugins/Editor/GLTFImporter/Source/GLTFCore/Public/GLTFNode.h
--- a/Engine/Plugins/Editor/GLTFImporter/Source/GLTFCore/Public/GLTFNode.h
+++ b/Engine/Plugins/Editor/GLTFImporter/Source/GLTFCore/Public/GLTFNode.h
@@ -37,12 +37,15 @@ namespace GLTF
 		int32 CameraIndex;  // index into FAsset::Cameras
 		int32 LightIndex;   // index into FAsset::Lights
 
+		int32 ParentIndex;
+
 		FNode()
 		    : Type(EType::None)
 		    , MeshIndex(INDEX_NONE)
 		    , Skindex(INDEX_NONE)
 		    , CameraIndex(INDEX_NONE)
 		    , LightIndex(INDEX_NONE)
+			, ParentIndex(INDEX_NONE)
 		{
 		}
 	};
diff --git a/Engine/Plugins/Editor/GLTFImporter/Source/GLTFCore/Public/GLTFReader.h b/Engine/Plugins/Editor/GLTFImporter/Source/GLTFCore/Public/GLTFReader.h
--- a/Engine/Plugins/Editor/GLTFImporter/Source/GLTFCore/Public/GLTFReader.h
+++ b/Engine/Plugins/Editor/GLTFImporter/Source/GLTFCore/Public/GLTFReader.h
@@ -60,6 +60,8 @@ namespace GLTF
 		void SetupObjects(uint32 ObjectCount, const TCHAR* FieldName, SetupFunc Func) const;
 		void SetupNodesType() const;
 
+		void BuildParentIndices(int32 ParentNodeIndex, int32 CurrentNodeIndex) const;
+
 	private:
 		uint32 BufferCount;
 		uint32 BufferViewCount;
diff --git a/Engine/Plugins/Experimental/Interchange/Runtime/Source/Import/Private/Gltf/InterchangeGltfAnimation.cpp b/Engine/Plugins/Experimental/Interchange/Runtime/Source/Import/Private/Gltf/InterchangeGltfAnimation.cpp
new file mode 100644
--- /dev/null
+++ b/Engine/Plugins/Experimental/Interchange/Runtime/Source/Import/Private/Gltf/InterchangeGltfAnimation.cpp
@@ -0,0 +1,464 @@
+// Copyright Epic Games, Inc. All Rights Reserved. 
+
+#include "Gltf/InterchangeGltfPrivate.h"
+
+#include "Animation/InterchangeAnimationPayload.h"
+#include "InterchangeImportLog.h"
+
+#include "GLTFAsset.h"
+
+#include "Animation/AnimTypes.h"
+
+namespace UE::Interchange::Gltf::Private
+{
+	/*
+	* According to gltf specification Seconds acquired from Samplers.input:
+	* "The values represent time in seconds with time[0] >= 0.0, and strictly increasing values, Index.e., time[n + 1] > time[n]."
+	* Meaning the incoming array(Seconds) is sorted.
+	*/
+	TTuple<int32, int32> GetNeighbourIndices(TArray<float>& Seconds, float SampleTime)
+	{
+		int32 Lower = 0;
+		int32 Higher = 0;
+
+		if (Seconds.Num() == 0)
+		{
+			return TTuple<int32, int32>(0, 0);
+		}
+		else if (Seconds.Num() == 1)
+		{
+			return TTuple<int32, int32>(Lower, Higher);
+		}
+
+		if (SampleTime < Seconds[0])
+		{
+			Lower = Higher = 0;
+		}
+		else if (SampleTime >= Seconds[Seconds.Num() - 1])
+		{
+			Lower = Higher = Seconds.Num() - 1;
+		}
+		else
+		{
+			for (size_t Index = 0; Index < Seconds.Num() - 1; Index++)
+			{
+				if ((Seconds[Index] <= SampleTime) && (SampleTime < Seconds[Index + 1]))
+				{
+					Lower = Index;
+					Higher = Index + 1;
+					break;
+				}
+			}
+		}
+
+		return TTuple<int32, int32>(Lower, Higher);
+	}
+
+	template <typename T>
+	T InterpolateValue(TArray<float>& Seconds /*input*/, TArray<T>& SegmentValues /*output*/, int32 kIndex /*lower*/, int32 kp1Index /*higher*/, double CurrentRequestedTimeStamp, GLTF::FAnimation::EInterpolation Interpolation, bool IsRotation)
+	{
+		if (SegmentValues.Num() <= kIndex || SegmentValues.Num() <= kp1Index || kIndex < 0 || kp1Index < 0)
+		{
+			UE_LOG(LogInterchangeImport, Warning, TEXT("Interpolation Failed due to invalid kIndices"));
+			return T();
+		}
+
+		if (kIndex == kp1Index)
+		{
+			if (Interpolation == GLTF::FAnimation::EInterpolation::CubicSpline)
+			{
+				return SegmentValues[kIndex * 3 + 1];
+			}
+			return SegmentValues[kp1Index];
+		}
+
+		// To match variables against https://registry.khronos.org/glTF/specs/2.0/glTF-2.0.html#appendix-c-interpolation
+		// Time := Time
+		// TimeDelta := td
+		double TimeDelta = Seconds[kp1Index] - Seconds[kIndex];
+		double Time = (CurrentRequestedTimeStamp - Seconds[kIndex]) / TimeDelta;
+
+		T Value;
+		switch (Interpolation)
+		{
+		case GLTF::FAnimation::EInterpolation::Linear:
+		{
+			if (IsRotation)
+			{
+				//Spherical Linear Interpolation
+				// a=arccos(|vk⋅vk+1|)
+				// s=vk⋅vk+1 / |vk⋅vk+1|
+
+				T vk = SegmentValues[kIndex];
+				T vkp1 = SegmentValues[kp1Index];
+
+				double DotProduct = FMath::Clamp(double(vk | vkp1), -1.0, 1.0);
+				double a = acos(abs(DotProduct));
+
+				if (FMath::IsNearlyZero(a))
+				{
+					//Implementation Note
+					//	When a is close to zero, spherical linear interpolation turns into regular linear interpolation.
+					Value = (1.0 - Time) * vk + Time * vkp1;
+				}
+				else
+				{
+					double s = DotProduct / abs(DotProduct);
+
+					// vt=(sin(a(1−t))/sin(a))∗vk + s∗(sin(at)/sin(a))∗vk+1
+
+					Value = (sin(a * (1.0 - Time)) / sin(a)) * vk + s * (sin(a * Time) / sin(a)) * vkp1;
+				}
+			}
+			else
+			{
+				//Normal Linear Interpolation
+				// vt=(1−t)∗vk + t∗vk+1
+
+				T vk = SegmentValues[kIndex];
+				T vkp1 = SegmentValues[kp1Index];
+
+				Value = (1.0 - Time) * vk + Time * vkp1;
+			}
+			break;
+		}
+		case GLTF::FAnimation::EInterpolation::Step:
+		{
+			//Step Interpolation
+			// vt=vk
+
+			T vk = SegmentValues[kIndex];
+			T vkp1 = SegmentValues[kp1Index];
+			if (Time < 0.5)
+			{
+				Value = vk;
+			}
+			else
+			{
+				Value = vkp1;
+			}
+			break;
+		}
+		case GLTF::FAnimation::EInterpolation::CubicSpline:
+		{
+			//Cubic Spline Interpolation
+			// For each timestamp stored in the animation sampler, there are three associated keyframe values: in-tangent, property Value, and out-tangent.
+			// ak, vk, and bk be the in-tangent, the property Value, and the out-tangent of the k-th frame respectively.
+			// The interpolated sampler Value vt at the timestamp tc is computed as follows.
+			// vt = (2t^3−3t^2 + 1)∗vk 
+			//		+ td(t^3−2t^2 + t)∗bk 
+			//		+ (−2t^3 + 3t^2)∗vk + 1 
+			//		+ td(t^3−t^2)∗ak + 1
+
+			T ak = SegmentValues[kIndex * 3 + 0];
+			T vk = SegmentValues[kIndex * 3 + 1];
+			T bk = SegmentValues[kIndex * 3 + 2];
+			T akp1 = SegmentValues[kp1Index * 3 + 0];
+			T vkp1 = SegmentValues[kp1Index * 3 + 1];
+			T bkp1 = SegmentValues[kp1Index * 3 + 2];
+
+			Value = (2 * Time * Time * Time - 3 * Time * Time + 1) * vk 
+				+ TimeDelta * (Time * Time * Time - 2 * Time * Time + Time) * bk 
+				+ (-2 * Time * Time * Time + 3 * Time * Time) * vkp1 
+				+ TimeDelta * (Time * Time * Time - Time * Time) * akp1;
+
+			break;
+		}
+
+		default:
+			break;
+		}
+
+		return Value;
+	}
+
+	bool ParsePayLoadKey(const GLTF::FAsset& GltfAsset, const FString& PayLoadKey, int32& AnimationIndex, TArray<int32>& ChannelIndices)
+	{
+		TArray<FString> PayLoadValues;
+
+		PayLoadKey.ParseIntoArray(PayLoadValues, TEXT(":"));
+
+		LexFromString(AnimationIndex, *PayLoadValues[0]);
+
+		if (!GltfAsset.Animations.IsValidIndex(AnimationIndex))
+		{
+			return false;
+		}
+
+		for (size_t Index = 1; Index < PayLoadValues.Num(); Index++)
+		{
+			int32 ChannelIndex;
+			LexFromString(ChannelIndex, *PayLoadValues[Index]);
+
+			if (GltfAsset.Animations[AnimationIndex].Channels.IsValidIndex(ChannelIndex))
+			{
+				ChannelIndices.Add(ChannelIndex);
+			}
+		}
+
+		return ChannelIndices.Num() != 0;
+	}
+
+	bool GetAnimationTransformPayloadData(const FString& PayLoadKey, const GLTF::FAsset& GltfAsset, FAnimationCurvePayloadData& OutPayloadData)
+	{
+		//translation, rotation, scale time 3 channels:
+		OutPayloadData.Curves.SetNum(9);
+
+		TArray<int32> ChannelIndices;
+		int32 AnimationIndex;
+
+		if (!ParsePayLoadKey(GltfAsset, PayLoadKey, AnimationIndex, ChannelIndices))
+		{
+			return false;
+		}
+		const GLTF::FAnimation& Animation = GltfAsset.Animations[AnimationIndex];
+		
+		for (const int32 ChannelIndex : ChannelIndices)
+		{
+			const GLTF::FAnimation::FChannel& Channel = Animation.Channels[ChannelIndex];
+			const GLTF::FAnimation::FSampler& Sampler = Animation.Samplers[Channel.Sampler];
+
+			TArray<float> FrameTimeBuffer;
+			TArray<float> FrameDataBuffer;
+
+			Sampler.Input.GetFloatArray(FrameTimeBuffer);
+
+			ERichCurveInterpMode InterpolationMode = static_cast<ERichCurveInterpMode>(Sampler.Interpolation);
+
+			TFunction<void(TFunction<FVector3f(const float*)>, int32, int32, bool)> ConvertCurve = [InterpolationMode, &FrameTimeBuffer, &FrameDataBuffer, &OutPayloadData](TFunction<FVector3f(const float*)> ConvertSample, int32 CurvesIndexOffset, int32 SampleSize, bool bUnwindRotation)
+			{
+				int32 ActualSampleSize = SampleSize;
+				if (InterpolationMode == ERichCurveInterpMode::RCIM_Cubic)
+				{
+					ActualSampleSize *= 3;
+				}
+				const int32 SampleCount = FrameDataBuffer.Num() / ActualSampleSize;
+
+				FRichCurve& Curve0 = OutPayloadData.Curves[CurvesIndexOffset + 0];
+				FRichCurve& Curve1 = OutPayloadData.Curves[CurvesIndexOffset + 1];
+				FRichCurve& Curve2 = OutPayloadData.Curves[CurvesIndexOffset + 2];
+
+				for (int32 SampleIndex = 0; SampleIndex < SampleCount; ++SampleIndex)
+				{
+
+
+					if (InterpolationMode == ERichCurveInterpMode::RCIM_Cubic)
+					{
+						const float* RawSampleArrive = FrameDataBuffer.GetData() + (SampleIndex * ActualSampleSize) + 0 * SampleSize;
+						const float* RawSampleValue = FrameDataBuffer.GetData() + (SampleIndex * ActualSampleSize) + 1 * SampleSize;
+						const float* RawSampleLeave = FrameDataBuffer.GetData() + (SampleIndex * ActualSampleSize) + 2 * SampleSize;
+
+						FVector3f SampleArrive = ConvertSample(RawSampleArrive);
+						FVector3f SampleValue = ConvertSample(RawSampleValue);
+						FVector3f SampleLeave = ConvertSample(RawSampleLeave);
+
+						FKeyHandle Key0 = Curve0.AddKey(FrameTimeBuffer[SampleIndex], SampleValue[0], bUnwindRotation);
+						FRichCurveKey& CurveKey0 = Curve0.GetKey(Key0);
+						CurveKey0.ArriveTangent = SampleArrive[0];
+						CurveKey0.LeaveTangent = SampleLeave[0];
+
+						FKeyHandle Key1 = Curve1.AddKey(FrameTimeBuffer[SampleIndex], SampleValue[1], bUnwindRotation);
+						FRichCurveKey& CurveKey1 = Curve1.GetKey(Key1);
+						CurveKey1.ArriveTangent = SampleArrive[1];
+						CurveKey1.LeaveTangent = SampleLeave[1];
+
+						FKeyHandle Key2 = Curve2.AddKey(FrameTimeBuffer[SampleIndex], SampleValue[2], bUnwindRotation);
+						FRichCurveKey& CurveKey2 = Curve2.GetKey(Key2);
+						CurveKey2.ArriveTangent = SampleArrive[2];
+						CurveKey2.LeaveTangent = SampleLeave[2];
+
+						Curve0.SetKeyInterpMode(Key0, InterpolationMode);
+						Curve1.SetKeyInterpMode(Key1, InterpolationMode);
+						Curve2.SetKeyInterpMode(Key2, InterpolationMode);
+					}
+					else
+					{
+						const float* RawSample = FrameDataBuffer.GetData() + (SampleIndex * ActualSampleSize);
+						FVector3f Sample = ConvertSample(RawSample);
+
+						FKeyHandle Key0 = Curve0.AddKey(FrameTimeBuffer[SampleIndex], Sample[0], bUnwindRotation);
+						FKeyHandle Key1 = Curve1.AddKey(FrameTimeBuffer[SampleIndex], Sample[1], bUnwindRotation);
+						FKeyHandle Key2 = Curve2.AddKey(FrameTimeBuffer[SampleIndex], Sample[2], bUnwindRotation);
+
+						Curve0.SetKeyInterpMode(Key0, InterpolationMode);
+						Curve1.SetKeyInterpMode(Key1, InterpolationMode);
+						Curve2.SetKeyInterpMode(Key2, InterpolationMode);
+					}
+				}
+			};
+
+			switch (Channel.Target.Path)
+			{
+			case GLTF::FAnimation::EPath::Translation:
+			{
+				FrameDataBuffer.SetNumUninitialized(Sampler.Output.Count * 3);
+				Sampler.Output.GetCoordArray(reinterpret_cast<FVector3f*>(FrameDataBuffer.GetData()));
+				const float Scale = GltfUnitConversionMultiplier; // Convert m to cm
+				ConvertCurve([&Scale](const float* VecData)
+					{
+						return FVector3f(VecData[0] * Scale, VecData[1] * Scale, VecData[2] * Scale);
+					}, 0, 3, false);
+				break;
+			}
+			case GLTF::FAnimation::EPath::Rotation:
+			{
+				FrameDataBuffer.SetNumUninitialized(Sampler.Output.Count * 4);
+				Sampler.Output.GetQuatArray(reinterpret_cast<FVector4f*>(FrameDataBuffer.GetData()));
+				ConvertCurve([](const float* QuatData)
+					{
+						const FQuat4f Quat(QuatData[0], QuatData[1], QuatData[2], QuatData[3]);
+						return Quat.Euler();
+					}, 3, 4, true);
+				break;
+			}
+			case GLTF::FAnimation::EPath::Scale:
+			{
+				FrameDataBuffer.SetNumUninitialized(Sampler.Output.Count * 3);
+				Sampler.Output.GetCoordArray(reinterpret_cast<FVector3f*>(FrameDataBuffer.GetData()));
+				ConvertCurve([](const float* VecData)
+					{
+						return FVector3f(VecData[0], VecData[1], VecData[2]);
+					}, 6, 3, false);
+				break;
+			}
+			case GLTF::FAnimation::EPath::Weights:
+			{
+				UE_LOG(LogInterchangeImport, Warning, TEXT("Morph Animation type not supported yet."));
+				break;
+			}
+			default:
+				UE_LOG(LogInterchangeImport, Warning, TEXT("Animation type not supported"));
+				break;
+			}
+		}
+		return true;
+	}
+
+	bool GetBakedAnimationTransformPayloadData(const FString& PayLoadKey, const GLTF::FAsset& GltfAsset, FAnimationBakeTransformPayloadData& PayloadData)
+	{
+		TArray<int32> ChannelIndices;
+		int32 AnimationIndex;
+
+		if (!ParsePayLoadKey(GltfAsset, PayLoadKey, AnimationIndex, ChannelIndices))
+		{
+			return false;
+		}
+		const GLTF::FAnimation& GltfAnimation = GltfAsset.Animations[AnimationIndex];
+		
+		const double BakeInterval = 1.0 / PayloadData.BakeFrequency;
+		const double SequenceLength = FMath::Max<double>(PayloadData.RangeEndTime - PayloadData.RangeStartTime, MINIMUM_ANIMATION_LENGTH);
+		int32 BakeKeyCount = (SequenceLength / BakeInterval) + 1;
+
+		//buffers to use for final Transform:
+		TArray<FVector3f> TranslationData;
+		TArray<FQuat4f> RotationData;
+		TArray<FVector3f> ScaleData;
+
+		//initial transforms:
+		//channels are targetting the same node:
+		bool BaseTransformsAcquired = false;
+
+		for (const int32& ChannelIndex : ChannelIndices)
+		{
+			const GLTF::FAnimation::FChannel& Channel = GltfAnimation.Channels[ChannelIndex];
+			const GLTF::FAnimation::FSampler& Sampler = GltfAnimation.Samplers[Channel.Sampler];
+
+			if (!BaseTransformsAcquired)
+			{
+				//on the first valid channel id we set the targeted nodes default transforms
+				//the default transforms will be overwritten by the provided animation data, IF they are provided
+				BaseTransformsAcquired = true;
+				FVector3f JointBaseTranslation(Channel.Target.Node.Transform.GetTranslation());
+				FQuat4f JointBaseRotation(Channel.Target.Node.Transform.GetRotation());
+				FVector3f JointBaseScale(Channel.Target.Node.Transform.GetScale3D());
+
+				TranslationData.Init(JointBaseTranslation, BakeKeyCount);
+				RotationData.Init(JointBaseRotation, BakeKeyCount);
+				ScaleData.Init(JointBaseScale, BakeKeyCount);
+			}
+
+			TArray<float> Seconds;
+			Sampler.Input.GetFloatArray(Seconds);
+
+			if (Seconds.Num() == 0)
+			{
+				continue;
+			}
+
+			switch (Channel.Target.Path)
+			{
+			case GLTF::FAnimation::EPath::Translation:
+			{
+				TArray<FVector3f> TranslationBuffer;
+				TranslationBuffer.SetNumUninitialized(Sampler.Output.Count);
+				Sampler.Output.GetCoordArray(TranslationBuffer.GetData());
+				const float Scale = GltfUnitConversionMultiplier; // Convert m to cm
+				for (FVector3f& Position : TranslationBuffer)
+				{
+					Position *= Scale;
+				}
+
+				for (size_t Index = 0; Index < BakeKeyCount; Index++)
+				{
+					double CurrentRequestedTimeStamp = Index * BakeInterval;
+					TTuple<int32, int32> LowerHigherIndices = GetNeighbourIndices(Seconds, CurrentRequestedTimeStamp);
+					TranslationData[Index] = InterpolateValue(Seconds, TranslationBuffer, LowerHigherIndices.Key, LowerHigherIndices.Value, CurrentRequestedTimeStamp, Sampler.Interpolation, false);
+				}
+				break;
+			}
+			case GLTF::FAnimation::EPath::Rotation:
+			{
+				TArray<FQuat4f> RotationBuffer;
+				TArray<FVector4f> RotationBufferTemp;
+				RotationBufferTemp.SetNumUninitialized(Sampler.Output.Count);
+				RotationBuffer.SetNumUninitialized(Sampler.Output.Count);
+
+				Sampler.Output.GetQuatArray(RotationBufferTemp.GetData());
+
+				for (size_t Index = 0; Index < Sampler.Output.Count; Index++)
+				{
+					RotationBuffer[Index] = FQuat4f(RotationBufferTemp[Index][0], RotationBufferTemp[Index][1], RotationBufferTemp[Index][2], RotationBufferTemp[Index][3]);
+				}
+
+				//When the animation sampler targets a node’s rotation property, the interpolated quaternion MUST be normalized before applying the result to the node’s rotation.
+				for (size_t Index = 0; Index < BakeKeyCount; Index++)
+				{
+					double CurrentRequestedTimeStamp = Index * BakeInterval;
+					TTuple<int32, int32> LowerHigherIndices = GetNeighbourIndices(Seconds, CurrentRequestedTimeStamp);
+					RotationData[Index] = InterpolateValue(Seconds, RotationBuffer, LowerHigherIndices.Key, LowerHigherIndices.Value, CurrentRequestedTimeStamp, Sampler.Interpolation, true);
+					RotationData[Index].Normalize();
+				}
+				break;
+			}
+			case GLTF::FAnimation::EPath::Scale:
+			{
+				TArray<FVector3f> ScaleBuffer;
+				ScaleBuffer.SetNumUninitialized(Sampler.Output.Count);
+				Sampler.Output.GetCoordArray(ScaleBuffer.GetData());
+
+				for (size_t Index = 0; Index < BakeKeyCount; Index++)
+				{
+					double CurrentRequestedTimeStamp = Index * BakeInterval;
+					TTuple<int32, int32> LowerHigherIndices = GetNeighbourIndices(Seconds, CurrentRequestedTimeStamp);
+					ScaleData[Index] = InterpolateValue(Seconds, ScaleBuffer, LowerHigherIndices.Key, LowerHigherIndices.Value, CurrentRequestedTimeStamp, Sampler.Interpolation, false);
+				}
+
+				break;
+			}
+			default:
+				UE_LOG(LogInterchangeImport, Warning, TEXT("Animation type not supported"));
+				break;
+			}
+		}
+
+		for (size_t Index = 0; Index < BakeKeyCount; Index++)
+		{
+			PayloadData.Transforms.Add(FTransform(
+				FQuat4d(RotationData[Index]),
+				FVector3d(TranslationData[Index]),
+				FVector3d(ScaleData[Index])));
+		}
+
+		return true;
+	}
+}
\ No newline at end of file
diff --git a/Engine/Plugins/Experimental/Interchange/Runtime/Source/Import/Private/Gltf/InterchangeGltfMesh.cpp b/Engine/Plugins/Experimental/Interchange/Runtime/Source/Import/Private/Gltf/InterchangeGltfMesh.cpp
new file mode 100644
--- /dev/null
+++ b/Engine/Plugins/Experimental/Interchange/Runtime/Source/Import/Private/Gltf/InterchangeGltfMesh.cpp
@@ -0,0 +1,287 @@
+// Copyright Epic Games, Inc. All Rights Reserved. 
+
+#include "Gltf/InterchangeGltfPrivate.h"
+
+#include "GLTFAsset.h"
+#include "GLTFMeshFactory.h"
+
+#include "StaticMeshAttributes.h"
+#include "SkeletalMeshAttributes.h"
+
+#include "StaticMeshOperations.h"
+#include "SkeletalMeshOperations.h"
+
+#include "Mesh/InterchangeStaticMeshPayload.h"
+
+namespace UE::Interchange::Gltf::Private
+{
+	/*
+	* For a given node index calculates the Joint Transformations
+	*/
+	void CalculateJointTransformations(const TArray<GLTF::FNode>& Nodes, int32 CurrentIndex, FTransform& CurrentTransform)
+	{
+		while (Nodes.IsValidIndex(CurrentIndex) && Nodes[CurrentIndex].Type == GLTF::FNode::EType::Joint)
+		{
+			CurrentTransform *= Nodes[CurrentIndex].Transform;
+			CurrentIndex = Nodes[CurrentIndex].ParentIndex;
+		}
+	}
+
+	int32 GetRootNodeIndex(const GLTF::FAsset& GltfAsset, const TArray<int32>& NodeIndices)
+	{
+		int32 ClosestIndex = -1;
+		int32 ClosestRouteLength = INT32_MAX;
+		for (int32 NodeIndex : NodeIndices)
+		{
+			int32 CurrentNodeIndexTracker = NodeIndex;
+			int32 CurrentRouteLength = 0;
+			while (GltfAsset.Nodes[CurrentNodeIndexTracker].ParentIndex != INDEX_NONE)
+			{
+				CurrentRouteLength++;
+				CurrentNodeIndexTracker = GltfAsset.Nodes[CurrentNodeIndexTracker].ParentIndex;
+			}
+			if (CurrentRouteLength < ClosestRouteLength)
+			{
+				ClosestRouteLength = CurrentRouteLength;
+				ClosestIndex = NodeIndex;
+			}
+		}
+		return ClosestIndex;
+	}
+
+	void PatchPolygonGroups(FMeshDescription& MeshDescription, const GLTF::FAsset& GltfAsset)
+	{
+		// Patch polygon groups material slot names to match Interchange expectations (rename material slots from indices to material names)
+		{
+			FStaticMeshAttributes StaticMeshAttributes(MeshDescription);
+
+			for (int32 MaterialSlotIndex = 0; MaterialSlotIndex < StaticMeshAttributes.GetPolygonGroupMaterialSlotNames().GetNumElements(); ++MaterialSlotIndex)
+			{
+				int32 MaterialIndex = 0;
+				LexFromString(MaterialIndex, *StaticMeshAttributes.GetPolygonGroupMaterialSlotNames()[MaterialSlotIndex].ToString());
+
+				if (GltfAsset.Materials.IsValidIndex(MaterialIndex))
+				{
+					const FString MaterialName = GltfAsset.Materials[MaterialIndex].Name;
+					StaticMeshAttributes.GetPolygonGroupMaterialSlotNames()[MaterialSlotIndex] = *MaterialName;
+				}
+			}
+		}
+	}
+
+	bool GetSkeletalMeshDescriptionForPayLoadKey(const GLTF::FAsset& GltfAsset, const FString& PayLoadKey, FMeshDescription& MeshDescription, TArray<FString>* OutJointUniqueNames)
+	{
+		TArray<FString> PayLoadKeys;
+		PayLoadKey.ParseIntoArray(PayLoadKeys, TEXT(":"));
+		TMap<int32, TArray<int32>> MeshIndexToSkinIndicesMap;
+
+		int32 MeshIndex = 0;
+		for (const FString& SubPayLoadKey : PayLoadKeys)
+		{
+			int32 MeshAndSkinIndex = 0;
+			LexFromString(MeshAndSkinIndex, *SubPayLoadKey);
+
+			MeshIndex = MeshAndSkinIndex & 0xFFFF;
+			const int32 SkinIndex = MeshAndSkinIndex >> 16;
+
+			if (!GltfAsset.Meshes.IsValidIndex(MeshIndex) || !GltfAsset.Skins.IsValidIndex(SkinIndex))
+			{
+				continue;
+			}
+
+			TArray<int32>& SkinIndices = MeshIndexToSkinIndicesMap.FindOrAdd(MeshIndex);
+			SkinIndices.Add(SkinIndex);
+		}
+		if (MeshIndexToSkinIndicesMap.Array().Num() > 1)
+		{
+			//invalid scenario, because its indicating that the skeletal mesh payload key contains multiple mesh targets.
+			return false;
+		}
+
+		if (MeshIndexToSkinIndicesMap.Array().Num() == 0)
+		{
+			return false;
+		}
+
+		FMeshDescription BaseMeshDescription;
+
+		const GLTF::FMesh& GltfMesh = GltfAsset.Meshes[MeshIndex];
+		GLTF::FMeshFactory MeshFactory;
+		MeshFactory.SetUniformScale(GltfUnitConversionMultiplier); // GLTF is in meters while UE is in centimeters
+		MeshFactory.FillMeshDescription(GltfMesh, &BaseMeshDescription);
+
+		PatchPolygonGroups(BaseMeshDescription, GltfAsset);
+
+		TArray<FMeshDescription> SkinnedMeshDescriptions;
+		TMap<int32, TArray<FString>> SkinVsLocalJointNames;
+		const TArray<int32>& SkinIndices = MeshIndexToSkinIndicesMap[MeshIndex];
+
+		for (size_t Index = 0; Index < SkinIndices.Num(); Index++)
+		{
+			FMeshDescription& SkinnedMeshDescription = SkinnedMeshDescriptions.Add_GetRef(BaseMeshDescription);
+			const int32 SkinIndex = SkinIndices[Index];
+
+			if (SkinIndex == INDEX_NONE)
+			{
+				continue;
+			}
+			const GLTF::FSkinInfo& Skin = GltfAsset.Skins[SkinIndex];
+			FStaticMeshAttributes StaticMeshAttributes(SkinnedMeshDescription);
+
+			//for instanced meshes we need to bake the transforms of the joints:
+			if (SkinIndices.Num() > 1 && GltfAsset.Skins[SkinIndex].Joints.Num() > 0)
+			{
+				FTransform TransformLocalToWorld3d(FTransform::Identity);
+
+				CalculateJointTransformations(GltfAsset.Nodes, GetRootNodeIndex(GltfAsset, GltfAsset.Skins[SkinIndex].Joints), TransformLocalToWorld3d);
+
+				FTransform3f TransformLocalToWorld(TransformLocalToWorld3d);
+				TVertexAttributesRef<FVector3f> VertexPositions = StaticMeshAttributes.GetVertexPositions();
+				FVertexArray VertexIndices = SkinnedMeshDescription.Vertices();
+				for (FVertexID VertexID : SkinnedMeshDescription.Vertices().GetElementIDs())
+				{
+					FVector3f VertexPosition = VertexPositions[VertexID];
+
+					VertexPosition = TransformLocalToWorld.TransformPosition(VertexPosition);
+
+					VertexPositions[VertexID] = VertexPosition;
+				}
+			}
+
+			FSkeletalMeshAttributes SkeletalMeshAttributes(SkinnedMeshDescription);
+			SkeletalMeshAttributes.Register();
+
+			using namespace UE::AnimationCore;
+			TMap<FVertexID, TArray<FBoneWeight>> RawBoneWeights;
+
+			//Add the influence data in the skeletalmesh description
+			FSkinWeightsVertexAttributesRef VertexSkinWeights = SkeletalMeshAttributes.GetVertexSkinWeights();
+
+			const TArray<GLTF::FMeshFactory::FIndexVertexIdMap>& PositionIndexToVertexIdPerPrim = MeshFactory.GetPositionIndexToVertexIdPerPrim();
+
+			TArray<FString>& LocalJointNames = SkinVsLocalJointNames.FindOrAdd(SkinIndex);
+			for (size_t JoinIndex = 0; JoinIndex < Skin.Joints.Num(); JoinIndex++)
+			{
+				OutJointUniqueNames->Add(GltfAsset.Nodes[Skin.Joints[JoinIndex]].Name);
+				LocalJointNames.Add(GltfAsset.Nodes[Skin.Joints[JoinIndex]].Name);
+			}
+
+			for (int32 PrimIndex = 0; PrimIndex < GltfMesh.Primitives.Num(); ++PrimIndex)
+			{
+				const GLTF::FPrimitive& Prim = GltfMesh.Primitives[PrimIndex];
+
+				TArray<GLTF::FJointInfluence> JointInfluences;
+				Prim.GetJointInfluences(JointInfluences);
+
+				const GLTF::FMeshFactory::FIndexVertexIdMap& PositionIndexToVertexId = PositionIndexToVertexIdPerPrim[PrimIndex];
+
+				for (int32 PositionIndex = 0; PositionIndex < JointInfluences.Num(); ++PositionIndex)
+				{
+					const FVertexID* VertexID = PositionIndexToVertexId.Find(PositionIndex);
+					if (!VertexID || !ensure(SkinnedMeshDescription.IsVertexValid(*VertexID)))
+					{
+						continue;
+					}
+
+					for (int32 BoneIndex = 0; BoneIndex < 4; ++BoneIndex)
+					{
+						const int32 JointIndex = JointInfluences[PositionIndex].ID[BoneIndex];
+						const int32 NodeIndex = Skin.Joints[JointIndex];
+						const float Weight = JointInfluences[PositionIndex].Weight[BoneIndex];
+
+						const GLTF::FNode* Bone = &GltfAsset.Nodes[NodeIndex];
+
+						TArray<FBoneWeight>* BoneWeights = RawBoneWeights.Find(*VertexID);
+						if (BoneWeights)
+						{
+							// Do we already have a weight for this bone?
+							bool bShouldAdd = true;
+							for (int32 WeightIndex = 0; WeightIndex < BoneWeights->Num(); WeightIndex++)
+							{
+								FBoneWeight& BoneWeight = (*BoneWeights)[WeightIndex];
+								if (BoneWeight.GetBoneIndex() == JointIndex)
+								{
+									if (BoneWeight.GetWeight() < Weight)
+									{
+										BoneWeight.SetWeight(Weight);
+									}
+									bShouldAdd = false;
+									break;
+								}
+							}
+							if (bShouldAdd)
+							{
+								BoneWeights->Add(FBoneWeight(JointIndex, Weight));
+							}
+						}
+						else
+						{
+							RawBoneWeights.Add(*VertexID).Add(FBoneWeight(JointIndex, Weight));
+						}
+					}
+				}
+			}
+
+			// Add all the raw bone weights. This will cause the weights to be sorted and re-normalized after culling to max influences.
+			for (const TTuple<FVertexID, TArray<FBoneWeight>>& Item : RawBoneWeights)
+			{
+				VertexSkinWeights.Set(Item.Key, Item.Value);
+			}
+		}
+
+		FStaticMeshOperations::FAppendSettings AppendSettings;
+		for (int32 ChannelIdx = 0; ChannelIdx < FStaticMeshOperations::FAppendSettings::MAX_NUM_UV_CHANNELS; ++ChannelIdx)
+		{
+			AppendSettings.bMergeUVChannels[ChannelIdx] = true;
+		}
+		for (size_t Index = 1; Index < SkinnedMeshDescriptions.Num(); Index++)
+		{
+			const int32 VertexOffset = SkinnedMeshDescriptions[0].Vertices().Num();
+			FSkeletalMeshOperations::FSkeletalMeshAppendSettings SkeletalMeshAppendSettings;
+			SkeletalMeshAppendSettings.SourceVertexIDOffset = VertexOffset;
+			const int32 LocalJointCount = SkinVsLocalJointNames[SkinIndices[Index]].Num();
+			SkeletalMeshAppendSettings.SourceRemapBoneIndex.AddZeroed(LocalJointCount);
+			for (int32 LocalJointIndex = 0; LocalJointIndex < LocalJointCount; ++LocalJointIndex)
+			{
+				SkeletalMeshAppendSettings.SourceRemapBoneIndex[LocalJointIndex] = LocalJointIndex;
+				const FString& LocalJointName = SkinVsLocalJointNames[SkinIndices[Index]][LocalJointIndex];
+				for (int32 GlobalJointNamesIndex = 0; GlobalJointNamesIndex < OutJointUniqueNames->Num(); ++GlobalJointNamesIndex)
+				{
+					FString temp = (*OutJointUniqueNames)[GlobalJointNamesIndex];
+					if (temp.Equals(LocalJointName))
+					{
+						SkeletalMeshAppendSettings.SourceRemapBoneIndex[LocalJointIndex] = GlobalJointNamesIndex;
+						break;
+					}
+				}
+			}
+
+			FStaticMeshOperations::AppendMeshDescription(SkinnedMeshDescriptions[Index], SkinnedMeshDescriptions[0], AppendSettings);
+			FSkeletalMeshOperations::AppendSkinWeight(SkinnedMeshDescriptions[Index], SkinnedMeshDescriptions[0], SkeletalMeshAppendSettings);
+		}
+
+		MeshDescription = SkinnedMeshDescriptions[0];
+
+		return true;
+	}
+
+	bool GetStaticMeshPayloadDataForPayLoadKey(const GLTF::FAsset& GltfAsset, const FString& PayLoadKey, FStaticMeshPayloadData& StaticMeshPayloadData)
+	{
+		int32 MeshIndex = 0;
+		LexFromString(MeshIndex, *PayLoadKey);
+
+		if (!GltfAsset.Meshes.IsValidIndex(MeshIndex))
+		{
+			return false;
+		}
+
+		const GLTF::FMesh& GltfMesh = GltfAsset.Meshes[MeshIndex];
+		GLTF::FMeshFactory MeshFactory;
+		MeshFactory.SetUniformScale(100.f); // GLTF is in meters while UE is in centimeters
+		MeshFactory.FillMeshDescription(GltfMesh, &StaticMeshPayloadData.MeshDescription);
+
+		PatchPolygonGroups(StaticMeshPayloadData.MeshDescription, GltfAsset);
+
+		return true;
+	}
+}
\ No newline at end of file
diff --git a/Engine/Plugins/Experimental/Interchange/Runtime/Source/Import/Private/Gltf/InterchangeGltfPrivate.h b/Engine/Plugins/Experimental/Interchange/Runtime/Source/Import/Private/Gltf/InterchangeGltfPrivate.h
new file mode 100644
--- /dev/null
+++ b/Engine/Plugins/Experimental/Interchange/Runtime/Source/Import/Private/Gltf/InterchangeGltfPrivate.h
@@ -0,0 +1,37 @@
+// Copyright Epic Games, Inc. All Rights Reserved.
+#pragma once
+
+#include "CoreMinimal.h"
+
+struct FMeshDescription;
+
+namespace GLTF
+{
+	struct FAsset;
+}
+
+namespace UE::Interchange
+{
+	struct FAnimationCurvePayloadData;
+	struct FAnimationBakeTransformPayloadData;
+	struct FStaticMeshPayloadData;
+
+	namespace Gltf::Private
+	{
+		static float GltfUnitConversionMultiplier = 100.f;
+
+		// Animation related functions
+		bool GetAnimationTransformPayloadData(const FString& PayLoadKey, const GLTF::FAsset& GltfAsset, FAnimationCurvePayloadData& OutPayloadData);
+		bool GetBakedAnimationTransformPayloadData(const FString& PayLoadKey, const GLTF::FAsset& GltfAsset, FAnimationBakeTransformPayloadData& PayloadData);
+		// 
+
+		// Mesh related functions
+		int32 GetRootNodeIndex(const GLTF::FAsset& GltfAsset, const TArray<int32>& NodeIndices);
+
+		bool GetSkeletalMeshDescriptionForPayLoadKey(const GLTF::FAsset& GltfAsset, const FString& PayLoadKey,
+			FMeshDescription& MeshDescription, TArray<FString>* OutJointUniqueNames);
+
+		bool GetStaticMeshPayloadDataForPayLoadKey(const GLTF::FAsset& GltfAsset, const FString& PayLoadKey, FStaticMeshPayloadData& StaticMeshPayloadData);
+		//
+	}
+}
\ No newline at end of file
diff --git a/Engine/Plugins/Experimental/Interchange/Runtime/Source/Import/Private/Gltf/InterchangeGltfTranslator.cpp b/Engine/Plugins/Experimental/Interchange/Runtime/Source/Import/Private/Gltf/InterchangeGltfTranslator.cpp
--- a/Engine/Plugins/Experimental/Interchange/Runtime/Source/Import/Private/Gltf/InterchangeGltfTranslator.cpp
+++ b/Engine/Plugins/Experimental/Interchange/Runtime/Source/Import/Private/Gltf/InterchangeGltfTranslator.cpp
@@ -17,6 +17,9 @@
 #include "InterchangeShaderGraphNode.h"
 #include "InterchangeTexture2DNode.h"
 #include "InterchangeVariantSetNode.h"
+#include "Nodes/InterchangeSourceNode.h"
+
+#include "Nodes/InterchangeAnimationAPI.h"
 
 #include "Sections/MovieScene3DTransformSection.h"
 #include "Texture/InterchangeImageWrapperTranslator.h"
@@ -24,8 +27,15 @@
 #include "Algo/Find.h"
 #include "Async/Async.h"
 #include "StaticMeshAttributes.h"
+#include "SkeletalMeshAttributes.h"
 #include "UObject/GCObjectScopeGuard.h"
 
+#include "StaticMeshOperations.h"
+#include "SkeletalMeshOperations.h"
+
+#include "Gltf/InterchangeGltfPrivate.h"
+
+
 static bool GInterchangeEnableGLTFImport = true;
 static FAutoConsoleVariableRef CCvarInterchangeEnableGLTFImport(
 	TEXT("Interchange.FeatureFlags.Import.GLTF"),
@@ -33,7 +43,6 @@ static FAutoConsoleVariableRef CCvarInterchangeEnableGLTFImport(
 	TEXT("Whether glTF support is enabled."),
 	ECVF_Default);
 
-// is not supported yet
 static const TArray<FString> SupportedExtensions = {/* Lights */
 														//TEXT("KHR_lights_punctual"), TEXT("KHR_lights"), //not supported yet
 													/* Variants */
@@ -62,82 +71,6 @@ namespace UE::Interchange::Gltf::Private
 		}
 	}
 
-	void GetAnimationTransformPayloadData(const GLTF::FAnimation& Animation, int32 ChannelIndex, UE::Interchange::FAnimationCurvePayloadData& OutPayloadData)
-	{
-		const GLTF::FAnimation::FChannel& Channel = Animation.Channels[ChannelIndex];
-		const GLTF::FAnimation::FSampler& Sampler = Animation.Samplers[Channel.Sampler];
-
-		TArray<float> FrameTimeBuffer;
-		TArray<float> FrameDataBuffer;
-
-		Sampler.Input.GetFloatArray(FrameTimeBuffer);
-
-		ERichCurveInterpMode InterpolationMode = static_cast<ERichCurveInterpMode>(Sampler.Interpolation);
-
-		TFunction<void(TFunction<FVector3f(const float*)>, int32, int32, bool)> ConvertCurve = [InterpolationMode, &FrameTimeBuffer, &FrameDataBuffer, &OutPayloadData](TFunction<FVector3f(const float*)> ConvertSample, int32 CurvesIndexOffset, int32 SampleSize, bool bUnwindRotation)
-		{
-			const int32 SampleCount = FrameDataBuffer.Num() / SampleSize;
-
-			FRichCurve& Curve0 = OutPayloadData.Curves[CurvesIndexOffset + 0];
-			FRichCurve& Curve1 = OutPayloadData.Curves[CurvesIndexOffset + 1];
-			FRichCurve& Curve2 = OutPayloadData.Curves[CurvesIndexOffset + 2];
-
-			for (int32 SampleIndex = 0; SampleIndex < SampleCount; ++SampleIndex)
-			{
-				const float* RawSample = FrameDataBuffer.GetData() + (SampleIndex * SampleSize);
-
-				FVector3f Sample = ConvertSample(RawSample);
-
-				FKeyHandle Key0 = Curve0.AddKey(FrameTimeBuffer[SampleIndex], Sample[0], bUnwindRotation);
-				FKeyHandle Key1 = Curve1.AddKey(FrameTimeBuffer[SampleIndex], Sample[1], bUnwindRotation);
-				FKeyHandle Key2 = Curve2.AddKey(FrameTimeBuffer[SampleIndex], Sample[2], bUnwindRotation);
-
-				Curve0.SetKeyInterpMode(Key0, InterpolationMode);
-				Curve1.SetKeyInterpMode(Key1, InterpolationMode);
-				Curve2.SetKeyInterpMode(Key2, InterpolationMode);
-			}
-		};
-
-		switch (Channel.Target.Path)
-		{
-			case GLTF::FAnimation::EPath::Translation:
-			{
-				FrameDataBuffer.SetNumUninitialized(Sampler.Output.Count * 3);
-				Sampler.Output.GetCoordArray(reinterpret_cast<FVector3f*>(FrameDataBuffer.GetData()));
-				const float Scale = 100.f; // Convert m to cm
-				ConvertCurve([&Scale](const float* VecData)
-					{
-						return FVector3f(VecData[0] * Scale, VecData[1] * Scale, VecData[2] * Scale);
-					}, 0, 3, false);
-				break;
-			}
-			case GLTF::FAnimation::EPath::Rotation:
-			{
-				FrameDataBuffer.SetNumUninitialized(Sampler.Output.Count * 4);
-				Sampler.Output.GetQuatArray(reinterpret_cast<FVector4f*>(FrameDataBuffer.GetData()));
-				ConvertCurve([](const float* QuatData)
-					{
-						const FQuat4f Quat(QuatData[0], QuatData[1], QuatData[2], QuatData[3]);
-						return Quat.Euler();
-					}, 3, 4, true);
-				break;
-			}
-			case GLTF::FAnimation::EPath::Scale:
-			{
-				FrameDataBuffer.SetNumUninitialized(Sampler.Output.Count * 3);
-				Sampler.Output.GetCoordArray(reinterpret_cast<FVector3f*>(FrameDataBuffer.GetData()));
-				ConvertCurve([](const float* VecData)
-					{
-						return FVector3f(VecData[0], VecData[1], VecData[2]);
-					}, 6, 3, false);
-				break;
-			}
-			default:
-				UE_LOG(LogInterchangeImport, Warning, TEXT("Animation type not supported"));
-				break;
-		}
-	}
-
 	bool CheckForVariants(const GLTF::FMesh& Mesh, int32 VariantCount, int32 MaterialCount)
 	{
 		for (const GLTF::FPrimitive& Primitive : Mesh.Primitives)
@@ -159,9 +92,17 @@ namespace UE::Interchange::Gltf::Private
 
 		return false;
 	}
+
+	void ScaleNodeTranslations(TArray<GLTF::FNode>& Nodes, float Scale)
+	{
+		for (GLTF::FNode& Node : Nodes)
+		{
+			Node.Transform.SetTranslation(Node.Transform.GetTranslation() * Scale);
+		}
+	}
 }
 
-void UInterchangeGltfTranslator::HandleGltfNode( UInterchangeBaseNodeContainer& NodeContainer, const GLTF::FNode& GltfNode, const FString& ParentNodeUid, const int32 NodeIndex, bool &bHasVariants ) const
+void UInterchangeGltfTranslator::HandleGltfNode( UInterchangeBaseNodeContainer& NodeContainer, const GLTF::FNode& GltfNode, const FString& ParentNodeUid, const int32 NodeIndex, bool &bHasVariants, TArray<int32>& SkinnedMeshNodes ) const
 {
 	using namespace UE::Interchange::Gltf::Private;
 
@@ -177,11 +118,27 @@ void UInterchangeGltfTranslator::HandleGltfNode( UInterchangeBaseNodeContainer&
 
 	FTransform Transform = GltfNode.Transform;
 
-	constexpr float MetersToCentimeters = 100.f;
-	Transform.SetTranslation( Transform.GetTranslation() * MetersToCentimeters );
+	Transform.SetTranslation( Transform.GetTranslation());
 
 	switch ( GltfNode.Type )
 	{
+		case GLTF::FNode::EType::MeshSkinned:
+		{
+			SkinnedMeshNodes.Add(NodeIndex);
+
+			if (!bHasVariants && GltfAsset.Variants.Num() > 0)
+			{
+				bHasVariants |= CheckForVariants(GltfAsset.Meshes[GltfNode.MeshIndex], GltfAsset.Variants.Num(), GltfAsset.Materials.Num());
+			}
+			break;
+		}
+
+		case GLTF::FNode::EType::Joint:
+		{
+			InterchangeSceneNode->AddSpecializedType(UE::Interchange::FSceneNodeStaticData::GetJointSpecializeTypeString());
+			break;
+		}
+
 		case GLTF::FNode::EType::Mesh:
 		{
 			if ( GltfAsset.Meshes.IsValidIndex( GltfNode.MeshIndex ) )
@@ -240,7 +197,7 @@ void UInterchangeGltfTranslator::HandleGltfNode( UInterchangeBaseNodeContainer&
 	{
 		if ( GltfAsset.Nodes.IsValidIndex( ChildIndex ) )
 		{
-			HandleGltfNode( NodeContainer, GltfAsset.Nodes[ ChildIndex ], NodeUid, ChildIndex, bHasVariants );
+			HandleGltfNode( NodeContainer, GltfAsset.Nodes[ ChildIndex ], NodeUid, ChildIndex, bHasVariants, SkinnedMeshNodes );
 		}
 	}
 }
@@ -852,6 +809,8 @@ bool UInterchangeGltfTranslator::Translate( UInterchangeBaseNodeContainer& NodeC
 	const bool bLoadMetaData = false;
 	GltfFileReader.ReadFile( FilePath, bLoadImageData, bLoadMetaData, const_cast< UInterchangeGltfTranslator* >( this )->GltfAsset );
 
+	ScaleNodeTranslations(const_cast<UInterchangeGltfTranslator*>(this)->GltfAsset.Nodes, GltfUnitConversionMultiplier);
+
 	if (GltfAsset.RequiredExtensions.Num() != 0)
 	{
 		FString NotSupportedExtensions;
@@ -913,15 +872,19 @@ bool UInterchangeGltfTranslator::Translate( UInterchangeBaseNodeContainer& NodeC
 
 			MeshNode->InitializeNode( MeshNodeUid, GltfMesh.Name, EInterchangeNodeContainerType::TranslatedAsset );
 			MeshNode->SetPayLoadKey( LexToString( MeshIndex++ ) );
+
 			NodeContainer.AddNode( MeshNode );
 
-			// Assign materials
-			for ( const GLTF::FPrimitive& Primitive : GltfMesh.Primitives )
+			
+			for (size_t PrimitiveCounter = 0; PrimitiveCounter < GltfMesh.Primitives.Num(); PrimitiveCounter++)
 			{
-				if ( GltfAsset.Materials.IsValidIndex( Primitive.MaterialIndex ) )
+				const GLTF::FPrimitive& Primitive = GltfMesh.Primitives[PrimitiveCounter];
+
+				// Assign materials
+				if (GltfAsset.Materials.IsValidIndex(Primitive.MaterialIndex))
 				{
 					const FString MaterialName = GltfAsset.Materials[Primitive.MaterialIndex].Name;
-					const FString ShaderGraphNodeUid = UInterchangeShaderGraphNode::MakeNodeUid( TEXT("Material_") + MaterialName);
+					const FString ShaderGraphNodeUid = UInterchangeShaderGraphNode::MakeNodeUid(TEXT("Material_") + MaterialName);
 					MeshNode->SetSlotMaterialDependencyUid(MaterialName, ShaderGraphNodeUid);
 				}
 			}
@@ -974,21 +937,27 @@ bool UInterchangeGltfTranslator::Translate( UInterchangeBaseNodeContainer& NodeC
 			SceneNode->InitializeNode( SceneNodeUid, SceneName, EInterchangeNodeContainerType::TranslatedScene );
 			NodeContainer.AddNode( SceneNode );
 
+			TArray<int32> SkinnedMeshNodes;
+
 			for ( const int32 NodeIndex : GltfScene.Nodes )
 			{
 				if ( GltfAsset.Nodes.IsValidIndex( NodeIndex ) )
 				{
-					HandleGltfNode( NodeContainer, GltfAsset.Nodes[ NodeIndex ], SceneNodeUid, NodeIndex, bHasVariants );
+					HandleGltfNode( NodeContainer, GltfAsset.Nodes[ NodeIndex ], SceneNodeUid, NodeIndex, bHasVariants, SkinnedMeshNodes );
 				}
 			}
+
+			// Skeletons:
+			HandleGltfSkeletons(NodeContainer, SceneNodeUid, SkinnedMeshNodes);
 		}
 	}
 
 	// Animations
 	{
+		TArray<FString> JointsAlreadyUsedInAnAnimation;
 		for (int32 AnimationIndex = 0; AnimationIndex < GltfAsset.Animations.Num(); ++AnimationIndex)
 		{
-			HandleGltfAnimation(NodeContainer, AnimationIndex);
+			HandleGltfAnimation(NodeContainer, AnimationIndex, JointsAlreadyUsedInAnAnimation);
 		}
 	}
 
@@ -1003,46 +972,16 @@ bool UInterchangeGltfTranslator::Translate( UInterchangeBaseNodeContainer& NodeC
 
 TFuture< TOptional< UE::Interchange::FStaticMeshPayloadData > > UInterchangeGltfTranslator::GetStaticMeshPayloadData( const FString& PayLoadKey ) const
 {
-	using namespace UE::Interchange::Gltf::Private;
-
-	TPromise< TOptional< UE::Interchange::FStaticMeshPayloadData > > MeshPayloadDataPromise;
-
-	int32 MeshIndex = 0;
-	LexFromString( MeshIndex, *PayLoadKey );
-
-	if ( !GltfAsset.Meshes.IsValidIndex( MeshIndex ) )
-	{
-		MeshPayloadDataPromise.SetValue(TOptional< UE::Interchange::FStaticMeshPayloadData >());
-		return MeshPayloadDataPromise.GetFuture();
-	}
-
-	UE::Interchange::FStaticMeshPayloadData StaticMeshPayloadData;
-
-	const GLTF::FMesh& GltfMesh = GltfAsset.Meshes[ MeshIndex ];
-	GLTF::FMeshFactory MeshFactory;
-	MeshFactory.SetUniformScale( 100.f ); // GLTF is in meters while UE is in centimeters
-	MeshFactory.FillMeshDescription( GltfMesh, &StaticMeshPayloadData.MeshDescription );
-
-	// Patch polygon groups material slot names to match Interchange expectations (rename material slots from indices to material names)
-	{
-		FStaticMeshAttributes StaticMeshAttributes( StaticMeshPayloadData.MeshDescription );
-
-		for ( int32 MaterialSlotIndex = 0; MaterialSlotIndex < StaticMeshAttributes.GetPolygonGroupMaterialSlotNames().GetNumElements(); ++MaterialSlotIndex )
+	return Async(EAsyncExecution::TaskGraph, [this, PayLoadKey]
 		{
-			int32 MaterialIndex = 0;
-			LexFromString( MaterialIndex, *StaticMeshAttributes.GetPolygonGroupMaterialSlotNames()[ MaterialSlotIndex ].ToString() );
-
-			if ( GltfAsset.Materials.IsValidIndex( MaterialIndex ) )
-			{
-				const FString MaterialName = GltfAsset.Materials[MaterialIndex].Name;
-				StaticMeshAttributes.GetPolygonGroupMaterialSlotNames()[ MaterialSlotIndex ] = *MaterialName;
-			}
-		}
-	}
+			UE::Interchange::FStaticMeshPayloadData StaticMeshPayloadData;
+			UE::Interchange::Gltf::Private::GetStaticMeshPayloadDataForPayLoadKey(GltfAsset, PayLoadKey, StaticMeshPayloadData);
 
-	MeshPayloadDataPromise.SetValue( StaticMeshPayloadData );
+			TOptional<UE::Interchange::FStaticMeshPayloadData> Result;
+			Result.Emplace(StaticMeshPayloadData);
 
-	return MeshPayloadDataPromise.GetFuture();
+			return Result;
+		});
 }
 
 TOptional< UE::Interchange::FImportImage > UInterchangeGltfTranslator::GetTexturePayloadData( const UInterchangeSourceData* InSourceData, const FString& PayLoadKey ) const
@@ -1086,68 +1025,20 @@ TOptional< UE::Interchange::FImportImage > UInterchangeGltfTranslator::GetTextur
 		return TextureTranslator->GetTexturePayloadData(PayloadSourceData, TextureFilePath);
 	}
 }
+
 TFuture<TOptional<UE::Interchange::FAnimationCurvePayloadData>> UInterchangeGltfTranslator::GetAnimationCurvePayloadData(const FString& PayLoadKey) const
 {
-	using namespace UE::Interchange;
-	using namespace UE::Interchange::Gltf::Private;
-
-	TSharedPtr<TPromise<TOptional<UE::Interchange::FAnimationCurvePayloadData>>> Promise = MakeShared<TPromise<TOptional<UE::Interchange::FAnimationCurvePayloadData>>>();
-	Promise->SetValue(TOptional<UE::Interchange::FAnimationCurvePayloadData>());
-
-	TArray<FString> PayloadTokens;
-
-	// We need at minimum animation index + one channel index
-	if (2 > PayLoadKey.ParseIntoArray(PayloadTokens, TEXT(":")))
-	{
-		// Invalid payload
-		return Promise->GetFuture();
-	}
-
-	TArray<int32> Tokens;
-	Tokens.Reserve(PayloadTokens.Num());
-
-	for (const FString& Token : PayloadTokens)
-	{
-		int32 IntToken = INDEX_NONE;
-		LexFromString(IntToken, *Token);
-
-		if (IntToken != INDEX_NONE)
-		{
-			Tokens.Add(IntToken);
-		}
-	}
-
-	const int32 AnimationIndex = Tokens[0];
-
-	if (!ensure(GltfAsset.Animations.IsValidIndex(AnimationIndex)))
-	{
-		// Invalid payload
-		return Promise->GetFuture();
-	}
-
-	return Async(EAsyncExecution::TaskGraph, [this, AnimationIndex, Tokens = MoveTemp(Tokens)]
+	return Async(EAsyncExecution::TaskGraph, [this, PayLoadKey]
 		{
-			TOptional<FAnimationCurvePayloadData> Result;
 
-			const GLTF::FAnimation& GltfAnimation = GltfAsset.Animations[AnimationIndex];
+			TOptional<UE::Interchange::FAnimationCurvePayloadData> Result;
+			UE::Interchange::FAnimationCurvePayloadData AnimationCurvePayloadData;
 
-			UE::Interchange::FAnimationCurvePayloadData TransformPayloadData;
-
-			TransformPayloadData.Curves.SetNum(9);
-
-			for (int32 TokenIndex = 1; TokenIndex < Tokens.Num(); ++TokenIndex)
+			if (UE::Interchange::Gltf::Private::GetAnimationTransformPayloadData(PayLoadKey, GltfAsset, AnimationCurvePayloadData))
 			{
-				const int32 ChannelIndex = Tokens[TokenIndex];
-				if (!ensure(GltfAnimation.Channels.IsValidIndex(ChannelIndex)))
-				{
-					continue;
-				}
-
-				GetAnimationTransformPayloadData(GltfAnimation, ChannelIndex, TransformPayloadData);
+				Result.Emplace(AnimationCurvePayloadData);
 			}
 
-			Result.Emplace(MoveTemp(TransformPayloadData));
-
 			return Result;
 		}
 	);
@@ -1164,25 +1055,153 @@ TFuture<TOptional<UE::Interchange::FAnimationStepCurvePayloadData>> UInterchange
 
 TFuture<TOptional<UE::Interchange::FAnimationBakeTransformPayloadData>> UInterchangeGltfTranslator::GetAnimationBakeTransformPayloadData(const FString& PayLoadKey, const double BakeFrequency, const double RangeStartSecond, const double RangeStopSecond) const
 {
-	TSharedPtr<TPromise<TOptional<UE::Interchange::FAnimationBakeTransformPayloadData>>> Promise = MakeShared<TPromise<TOptional<UE::Interchange::FAnimationBakeTransformPayloadData>>>();
-	Promise->SetValue(TOptional<UE::Interchange::FAnimationBakeTransformPayloadData>());
-	return Promise->GetFuture();
+	return Async(EAsyncExecution::TaskGraph, [this, PayLoadKey, BakeFrequency, RangeStartSecond, RangeStopSecond]
+		{
+			TOptional<UE::Interchange::FAnimationBakeTransformPayloadData> Result;
+
+			UE::Interchange::FAnimationBakeTransformPayloadData PayloadData;
+
+			PayloadData.BakeFrequency = BakeFrequency;
+			PayloadData.RangeStartTime = RangeStartSecond;
+			PayloadData.RangeEndTime = RangeStopSecond;
+
+			if (UE::Interchange::Gltf::Private::GetBakedAnimationTransformPayloadData(PayLoadKey, GltfAsset, PayloadData))
+			{
+				Result.Emplace(PayloadData);
+			}
+
+			return Result;
+		});
 }
 
-void UInterchangeGltfTranslator::HandleGltfAnimation(UInterchangeBaseNodeContainer& NodeContainer, int32 AnimationIndex) const
+void UInterchangeGltfTranslator::HandleGltfAnimation(UInterchangeBaseNodeContainer& NodeContainer, int32 AnimationIndex, TArray<FString>& JointsAlreadyUsedInAnAnimation /*Currently the system only supports 1 animation per skeleton*/) const
 {
 	const GLTF::FAnimation& GltfAnimation = GltfAsset.Animations[AnimationIndex];
 
 	TMap< const GLTF::FNode*, TArray<int32> > NodeChannelsMap;
 
+	TMap<FString, TArray<int32>> JointNodeUidsToChannelsUsed;
+
 	for (int32 ChannelIndex = 0; ChannelIndex < GltfAnimation.Channels.Num(); ++ChannelIndex)
 	{
-		const GLTF::FNode* AnimatedNode = &GltfAnimation.Channels[ChannelIndex].Target.Node;
+		const GLTF::FAnimation::FChannel& Channel = GltfAnimation.Channels[ChannelIndex];
+		const GLTF::FNode* AnimatedNode = &Channel.Target.Node;
+
+		if (AnimatedNode->Type == GLTF::FNode::EType::Joint)
+		{
+			const FString* JointNodeUid = NodeUidMap.Find(AnimatedNode);
+			if (ensure(JointNodeUid))
+			{
+				TArray<int32>& ChannelsUsed = JointNodeUidsToChannelsUsed.FindOrAdd(*JointNodeUid);
+				ChannelsUsed.Add(ChannelIndex);
+			}
+			
+			continue;
+		}
 
 		TArray<int32>& NodeChannels = NodeChannelsMap.FindOrAdd(AnimatedNode);
 		NodeChannels.Add(ChannelIndex);
 	}
 
+	//setup rigged animations:
+	{
+		for (TTuple<FString,TArray<int32>>& JointNodeUidToChannelsUsed : JointNodeUidsToChannelsUsed)
+		{
+			if (JointsAlreadyUsedInAnAnimation.Contains(JointNodeUidToChannelsUsed.Key))
+			{
+				//this joint (part of skeleton) is already part of an animation
+				//due to the Interchange system only supporting a single animation per skeleton at the moment we skip the setup of this joint and animatino pair
+				continue;
+			}
+			
+			//@ AllAnimationStart = 0;
+			//from gltf documentation: https://registry.khronos.org/glTF/specs/2.0/glTF-2.0.html#animations
+			//Implementation Note
+			//	For example, if the earliest sampler input for an animation is t = 10, a client implementation must begin playback of that animation channel at t = 0 with output clamped to the first available output value.
+			double AnimationStart = 0.0;
+			double AnimationStop = 0.0;
+			double FrameRate = 30.0;
+			double SingleFrameDuration = 1.0 / FrameRate;
+			int32 KeyCount = 0;
+			bool CorruptSamplerFound = false;
+
+			//check channel length:
+			{
+				for (int32 ChannelIndex : JointNodeUidToChannelsUsed.Value)
+				{
+					const GLTF::FAnimation::FChannel& Channel = GltfAnimation.Channels[ChannelIndex];
+					const GLTF::FAnimation::FSampler& Sampler = GltfAnimation.Samplers[Channel.Sampler];
+					TArray<float> Seconds;
+					Sampler.Input.GetFloatArray(Seconds);					
+
+					if (Sampler.Interpolation == GLTF::FAnimation::EInterpolation::CubicSpline)
+					{
+						if (Sampler.Input.Count != 3 * Sampler.Output.Count)
+						{
+							CorruptSamplerFound = true;
+							UE_LOG(LogInterchangeImport, Warning, TEXT("GLTF Sampler Corrupt. Input and Output not meeting expectations."));
+							break;
+						}
+					}
+					else
+					{
+						if (Sampler.Input.Count != Sampler.Output.Count)
+						{
+							CorruptSamplerFound = true;
+							UE_LOG(LogInterchangeImport, Warning, TEXT("GLTF Sampler Corrupt. Input and Output not meeting expectations."));
+							break;
+						}
+					}
+
+					float CurrentAnimationStop = 0;
+					int CurrentKeyCount = 0;
+
+					if (Seconds.Num() > 0)
+					{
+						//calculate keycount and currentAnimationStop:
+						CurrentAnimationStop = Seconds[Seconds.Num()-1];
+
+						float CurrentKeyCountCandidate = CurrentAnimationStop / SingleFrameDuration;
+						CurrentKeyCount = int(CurrentKeyCountCandidate);
+						if (int(CurrentKeyCountCandidate) < CurrentKeyCountCandidate)
+						{
+							CurrentKeyCount++;
+						}
+						CurrentAnimationStop = CurrentKeyCount * SingleFrameDuration;
+					}
+					
+					if (KeyCount < CurrentKeyCount)
+					{
+						KeyCount = CurrentKeyCount;
+						AnimationStop = KeyCount * SingleFrameDuration;
+					}
+				}
+			}
+
+			if (CorruptSamplerFound)
+			{
+				break;
+			}
+
+			if (UInterchangeSceneNode* AnimatedInterchangeSceneNode = const_cast<UInterchangeSceneNode*>(Cast< UInterchangeSceneNode >(NodeContainer.GetNode(*JointNodeUidToChannelsUsed.Key))))
+			{
+				//add the joint to the of already set up joints to track it for handling single animation only.
+				JointsAlreadyUsedInAnAnimation.Add(JointNodeUidToChannelsUsed.Key);
+
+				FString Payload = LexToString(AnimationIndex);
+				for (int32 ChannelIndex : JointNodeUidToChannelsUsed.Value)
+				{
+					Payload += ":" + LexToString(ChannelIndex);
+				}
+
+				UInterchangeAnimationAPI::SetCustomNodeTransformPayloadKey(AnimatedInterchangeSceneNode, Payload);
+				UInterchangeAnimationAPI::SetCustomIsNodeTransformAnimated(AnimatedInterchangeSceneNode, true);
+				UInterchangeAnimationAPI::SetCustomNodeTransformAnimationStartTime(AnimatedInterchangeSceneNode, AnimationStart);
+				UInterchangeAnimationAPI::SetCustomNodeTransformAnimationEndTime(AnimatedInterchangeSceneNode, AnimationStop);
+			}
+		}
+	}
+
 	if (NodeChannelsMap.IsEmpty())
 	{
 		return;
@@ -1506,3 +1525,153 @@ bool UInterchangeGltfTranslator::GetVariantSetPayloadData(UE::Interchange::FVari
 
 	return true;
 }
+
+TFuture<TOptional<UE::Interchange::FSkeletalMeshLodPayloadData>> UInterchangeGltfTranslator::GetSkeletalMeshLodPayloadData(const FString& PayLoadKey) const
+{
+	return Async(EAsyncExecution::TaskGraph, [this, PayLoadKey]
+		{
+			TOptional<UE::Interchange::FSkeletalMeshLodPayloadData> Result;
+			UE::Interchange::FSkeletalMeshLodPayloadData SkeletalMeshPayloadData;
+			
+			if (UE::Interchange::Gltf::Private::GetSkeletalMeshDescriptionForPayLoadKey(GltfAsset, PayLoadKey, SkeletalMeshPayloadData.LodMeshDescription, &SkeletalMeshPayloadData.JointNames))
+			{
+				Result.Emplace(SkeletalMeshPayloadData);
+			}
+
+			return Result;
+		});
+}
+
+TFuture<TOptional<UE::Interchange::FSkeletalMeshMorphTargetPayloadData>> UInterchangeGltfTranslator::GetSkeletalMeshMorphTargetPayloadData(const FString& PayLoadKey) const
+{
+	TSharedPtr<TPromise<TOptional<UE::Interchange::FSkeletalMeshMorphTargetPayloadData>>> Promise = MakeShared<TPromise<TOptional<UE::Interchange::FSkeletalMeshMorphTargetPayloadData>>>();
+	Promise->SetValue(TOptional<UE::Interchange::FSkeletalMeshMorphTargetPayloadData>{});
+	return Promise->GetFuture();
+}
+
+int32 UInterchangeGltfTranslator::FindRootJointNodeIndex(int32 CurrentIndex) const
+{
+	if (GltfAsset.Nodes[CurrentIndex].Type != GLTF::FNode::EType::Joint)
+	{
+		return CurrentIndex;
+	}
+
+	while (GltfAsset.Nodes.IsValidIndex(GltfAsset.Nodes[CurrentIndex].ParentIndex) && GltfAsset.Nodes[GltfAsset.Nodes[CurrentIndex].ParentIndex].Type == GLTF::FNode::EType::Joint)
+	{
+		CurrentIndex = GltfAsset.Nodes[CurrentIndex].ParentIndex;
+	}
+	return CurrentIndex;
+}
+
+void UInterchangeGltfTranslator::HandleGltfSkeletons(UInterchangeBaseNodeContainer& NodeContainer, const FString& SceneNodeUid, const TArray<int32>& SkinnedMeshNodes) const
+{
+	//if skeletons present set them as a root joint for the join hierarchy as per documentation:
+	//The skeleton property (if present) points to the node that is the common root of a joints hierarchy or to a direct or indirect parent node of the common root.
+	for (const GLTF::FSkinInfo& Skin : GltfAsset.Skins)
+	{
+		if (GltfAsset.Nodes.IsValidIndex(Skin.Skeleton))
+		{
+			const FString* SkeletonNodeUid = NodeUidMap.Find(&GltfAsset.Nodes[Skin.Skeleton]);
+			if (SkeletonNodeUid)
+			{
+				NodeContainer.SetNodeParentUid(*SkeletonNodeUid, SceneNodeUid);
+			}
+		}
+	}
+
+
+	TMap<int32, TMap<int32, TArray<int32>>> MeshIndexToRootJointGroupedSkinnedMeshNodesMap;
+
+	//group SkinnedMeshNodes based on Joint Root Parents and Mesh indices
+	//this is needed in order to figure out how many duplications do we need for a given mesh
+	for (int32 SkinnedMeshNodeIndex : SkinnedMeshNodes)
+	{
+		const GLTF::FNode& SkinnedMeshNode = GltfAsset.Nodes[SkinnedMeshNodeIndex];
+
+		TMap<int32, TArray<int32>>& RootJointGroupedSkinnedMeshNodes = MeshIndexToRootJointGroupedSkinnedMeshNodesMap.FindOrAdd(SkinnedMeshNode.MeshIndex);
+
+		//get the SkinnedMeshNode's skin's first joint as the starting ground and find the top most root joint for it:
+		int32 RootJointIndex = FindRootJointNodeIndex(GltfAsset.Skins[SkinnedMeshNode.Skindex].Joints[0]);
+
+		//basedon that root joint group the SkinnedMeshNodes:
+		TArray<int32>& GroupedSkinnedMeshNodes = RootJointGroupedSkinnedMeshNodes.FindOrAdd(RootJointIndex);
+		GroupedSkinnedMeshNodes.Add(SkinnedMeshNodeIndex);
+	}
+
+	for (const TTuple<int32, TMap<int32, TArray<int32>>>& MeshIndexRootJointGroups : MeshIndexToRootJointGroupedSkinnedMeshNodesMap)
+	{
+		const TMap<int32, TArray<int32>>& RootJointGroupedSkinnedMeshNodes = MeshIndexRootJointGroups.Value;
+
+		//Get original MeshNode:
+		const FString OriginalMeshNodeUid = TEXT("\\Mesh\\") + GltfAsset.Meshes[MeshIndexRootJointGroups.Key].Name;
+		if (const UInterchangeMeshNode* ConstOriginalMeshNode = Cast< UInterchangeMeshNode >(NodeContainer.GetNode(OriginalMeshNodeUid)))
+		{
+			UInterchangeMeshNode* OriginalMeshNode = const_cast<UInterchangeMeshNode*>(ConstOriginalMeshNode);
+
+			if (!ensure(OriginalMeshNode))
+			{
+				continue;
+			}
+
+			//iterate through the groups:
+			//rootjoint , array<skinnedMeshNodes>
+			for (const TTuple<int32, TArray<int32>>& RootJointToSkinnedMeshNodes : RootJointGroupedSkinnedMeshNodes)
+			{
+				//Duplicate MeshNode for each group:
+				UInterchangeMeshNode* DuplicatedMeshNode = NewObject< UInterchangeMeshNode >(&NodeContainer);
+				FString DuplicatedMeshNodeUid = OriginalMeshNode->GetUniqueID() + TEXT("_") + LexToString(RootJointToSkinnedMeshNodes.Key);
+				FString DuplicatedMeshDisplayLabel = OriginalMeshNode->GetDisplayLabel() + TEXT("_") + LexToString(RootJointToSkinnedMeshNodes.Key);
+				DuplicatedMeshNode->InitializeNode(DuplicatedMeshNodeUid, DuplicatedMeshDisplayLabel, EInterchangeNodeContainerType::TranslatedAsset);
+				TMap<FString, FString> MaterialDependencies;
+				OriginalMeshNode->GetSlotMaterialDependencies(MaterialDependencies);
+				for (const TTuple<FString, FString>& MaterialDependency : MaterialDependencies)
+				{
+					DuplicatedMeshNode->SetSlotMaterialDependencyUid(MaterialDependency.Key, MaterialDependency.Value);
+				}
+
+				DuplicatedMeshNode->SetSkinnedMesh(true);
+
+				//generate payload key:
+				//of template:
+				//"LexToString(SkinnedMeshNode.MeshIndex | (SkinnedMeshNode.Skindex << 16))":"LexToString(SkinnedMeshNode.MeshIndex | (SkinnedMeshNode.Skindex << 16))".....
+				FString Payload = "";
+				for (int32 SkinnedMeshIndex : RootJointToSkinnedMeshNodes.Value)
+				{
+					const GLTF::FNode& SkinnedMeshNode = GltfAsset.Nodes[SkinnedMeshIndex];
+					if (Payload.Len() > 0)
+					{
+						Payload += ":";
+					}
+					Payload += LexToString(SkinnedMeshNode.MeshIndex | (SkinnedMeshNode.Skindex << 16));
+				}
+				DuplicatedMeshNode->SetPayLoadKey(Payload);
+
+				//set the root joint node as the skeletonDependency:
+				int32 RootJointNodeIndex = RootJointToSkinnedMeshNodes.Key;
+				const GLTF::FNode& RootJointNode = GltfAsset.Nodes[RootJointNodeIndex];
+				const FString* SkeletonNodeUid = NodeUidMap.Find(&RootJointNode);
+				if (ensure(SkeletonNodeUid))
+				{
+					DuplicatedMeshNode->SetSkeletonDependencyUid(*SkeletonNodeUid);
+				}
+				
+
+				//set the mesh node's custom asset instance uid to the new duplicated mesh
+				//if there are more than one skins, then choose the topmost (root node of the collection, top most in a hierarchical tree term) occurance of SkinnedMeshIndex
+				int32 MeshNodeIndex = UE::Interchange::Gltf::Private::GetRootNodeIndex(GltfAsset, RootJointToSkinnedMeshNodes.Value);
+				const GLTF::FNode& MeshNode = GltfAsset.Nodes[MeshNodeIndex];
+				const FString* SceneMeshNodeUid = NodeUidMap.Find(&MeshNode);
+				if (const UInterchangeSceneNode* ConstSceneMeshNode = Cast< UInterchangeSceneNode >(NodeContainer.GetNode(*SceneMeshNodeUid)))
+				{
+					UInterchangeSceneNode* SceneMeshNode = const_cast<UInterchangeSceneNode*>(ConstSceneMeshNode);
+					if (ensure(SceneMeshNode))
+					{
+						SceneMeshNode->SetCustomAssetInstanceUid(DuplicatedMeshNodeUid);
+					}
+				}
+
+				NodeContainer.AddNode(DuplicatedMeshNode);
+			}
+		}
+	}
+}
\ No newline at end of file
diff --git a/Engine/Plugins/Experimental/Interchange/Runtime/Source/Import/Public/Gltf/InterchangeGltfPrivate.h b/Engine/Plugins/Experimental/Interchange/Runtime/Source/Import/Public/Gltf/InterchangeGltfPrivate.h
new file mode 100644
--- /dev/null
+++ b/Engine/Plugins/Experimental/Interchange/Runtime/Source/Import/Public/Gltf/InterchangeGltfPrivate.h
@@ -0,0 +1,38 @@
+// Copyright Epic Games, Inc. All Rights Reserved.
+#pragma once
+
+#include "CoreMinimal.h"
+
+struct FMeshDescription;
+
+namespace GLTF
+{
+	struct FAsset;
+}
+
+namespace UE::Interchange
+{
+	struct FAnimationCurvePayloadData;
+	struct FAnimationBakeTransformPayloadData;
+	struct FStaticMeshPayloadData;
+
+	namespace Gltf::Private
+	{
+		static float GltfUnitConversionMultiplier = 100.f;
+
+		// Animation related functions
+		bool GetAnimationTransformPayloadData(const FString& PayLoadKey, const GLTF::FAsset& GltfAsset, FAnimationCurvePayloadData& OutPayloadData);
+		bool GetBakedAnimationTransformPayloadData(const FString& PayLoadKey, const GLTF::FAsset& GltfAsset, FAnimationBakeTransformPayloadData& PayloadData);
+		// 
+
+		// Mesh related functions
+		int32 GetRootNodeIndex(const TMap<int32, int32>& CurrentToParentNodeIdxMap, const TArray<int32>& NodeIndices);
+
+		bool GetSkeletalMeshDescriptionForPayLoadKey(const GLTF::FAsset& GltfAsset, const TMap<int32, int32>& CurrentToParentNodeIdxMap,
+			const FString& PayLoadKey,
+			FMeshDescription& MeshDescription, TArray<FString>* OutJointUniqueNames);
+
+		bool GetStaticMeshPayloadDataForPayLoadKey(const GLTF::FAsset& GltfAsset, const FString& PayLoadKey, FStaticMeshPayloadData& StaticMeshPayloadData);
+		//
+	}
+}
\ No newline at end of file
diff --git a/Engine/Plugins/Experimental/Interchange/Runtime/Source/Import/Public/Gltf/InterchangeGltfTranslator.h b/Engine/Plugins/Experimental/Interchange/Runtime/Source/Import/Public/Gltf/InterchangeGltfTranslator.h
--- a/Engine/Plugins/Experimental/Interchange/Runtime/Source/Import/Public/Gltf/InterchangeGltfTranslator.h
+++ b/Engine/Plugins/Experimental/Interchange/Runtime/Source/Import/Public/Gltf/InterchangeGltfTranslator.h
@@ -6,6 +6,8 @@
 #include "Animation/InterchangeAnimationPayloadInterface.h"
 #include "Mesh/InterchangeStaticMeshPayload.h"
 #include "Mesh/InterchangeStaticMeshPayloadInterface.h"
+#include "Mesh/InterchangeSkeletalMeshPayload.h"
+#include "Mesh/InterchangeSkeletalMeshPayloadInterface.h"
 #include "Nodes/InterchangeBaseNodeContainer.h"
 #include "Scene/InterchangeVariantSetPayloadInterface.h"
 #include "Texture/InterchangeTexturePayloadInterface.h"
@@ -17,6 +19,7 @@
 class UInterchangeShaderGraphNode;
 class UInterchangeShaderNode;
 class UInterchangeVariantSetNode;
+class UInterchangeMeshNode;
 
 /* Gltf translator class support import of texture, material, static mesh, skeletal mesh, */
 
@@ -26,6 +29,7 @@ class UInterchangeGltfTranslator : public UInterchangeTranslatorBase
 	, public IInterchangeTexturePayloadInterface
 	, public IInterchangeAnimationPayloadInterface
 	, public IInterchangeVariantSetPayloadInterface
+	, public IInterchangeSkeletalMeshPayloadInterface
 {
 	GENERATED_BODY()
 
@@ -59,14 +63,20 @@ public:
 	virtual TFuture<TOptional<UE::Interchange::FVariantSetPayloadData>> GetVariantSetPayloadData(const FString& PayloadKey) const override;
 	/* IInterchangeVariantSetPayloadInterface End */
 
+	/* IInterchangeSkeletalMeshPayloadInterface Begin */
+	virtual TFuture<TOptional<UE::Interchange::FSkeletalMeshLodPayloadData>> GetSkeletalMeshLodPayloadData(const FString& PayLoadKey) const override;
+	virtual TFuture<TOptional<UE::Interchange::FSkeletalMeshMorphTargetPayloadData>> GetSkeletalMeshMorphTargetPayloadData(const FString& PayLoadKey) const override;
+	/* IInterchangeSkeletalMeshPayloadInterface End */
+
 protected:
 	using FNodeUidMap = TMap<const GLTF::FNode*, FString>;
 
-	void HandleGltfNode( UInterchangeBaseNodeContainer& NodeContainer, const GLTF::FNode& GltfNode, const FString& ParentNodeUid, const int32 NodeIndex, bool &bHasVariants ) const;
+	void HandleGltfSkeletons( UInterchangeBaseNodeContainer& NodeContainer, const FString& SceneNodeUid, const TArray<int32>& SkinnedMeshNodes ) const;
+	void HandleGltfNode( UInterchangeBaseNodeContainer& NodeContainer, const GLTF::FNode& GltfNode, const FString& ParentNodeUid, const int32 NodeIndex, bool &bHasVariants, TArray<int32>& SkinnedMeshNodes ) const;
 	void HandleGltfMaterial( UInterchangeBaseNodeContainer& NodeContainer, const GLTF::FMaterial& GltfMaterial, UInterchangeShaderGraphNode& ShaderGraphNode ) const;
 	void HandleGltfMaterialParameter( UInterchangeBaseNodeContainer& NodeContainer, const GLTF::FTextureMap& TextureMap, UInterchangeShaderNode& ShaderNode,
 		const FString& MapName, const TVariant< FLinearColor, float >& MapFactor, const FString& OutputChannel, const bool bInverse = false, const bool bIsNormal = false ) const;
-	void HandleGltfAnimation(UInterchangeBaseNodeContainer& NodeContainer, int32 AnimationIndex) const;
+	void HandleGltfAnimation(UInterchangeBaseNodeContainer& NodeContainer, int32 AnimationIndex, TArray<FString>& JointsAlreadyUsedInAnAnimation) const;
 	void HandleGltfVariants(UInterchangeBaseNodeContainer& NodeContainer, const FString& FileName) const;
 
 	/** Support for KHR_materials_clearcoat */
@@ -80,10 +90,13 @@ protected:
 	void HandleGltfTextureTransform( UInterchangeBaseNodeContainer& NodeContainer, const GLTF::FTextureTransform& TextureTransform, const int32 TexCoordIndex, UInterchangeShaderNode& ShaderNode ) const;
 
 	bool GetVariantSetPayloadData(UE::Interchange::FVariantSetPayloadData& PayloadData) const;
+
 private:
 	void SetTextureSRGB(UInterchangeBaseNodeContainer& NodeContainer, const GLTF::FTextureMap& TextureMap) const;
 	void SetTextureFlipGreenChannel(UInterchangeBaseNodeContainer& NodeContainer, const GLTF::FTextureMap& TextureMap) const;
 
+	int32 FindRootJointNodeIndex(int32 CurrentIndex) const;
+
 	GLTF::FAsset GltfAsset;
 	mutable FNodeUidMap NodeUidMap;
 };
