[Interchange] Added support for unlit material to GLTF translator

diff --git a/Engine/Plugins/Experimental/Interchange/Runtime/Source/Import/Private/Gltf/InterchangeGltfTranslator.cpp b/Engine/Plugins/Experimental/Interchange/Runtime/Source/Import/Private/Gltf/InterchangeGltfTranslator.cpp
--- a/Engine/Plugins/Experimental/Interchange/Runtime/Source/Import/Private/Gltf/InterchangeGltfTranslator.cpp
+++ b/Engine/Plugins/Experimental/Interchange/Runtime/Source/Import/Private/Gltf/InterchangeGltfTranslator.cpp
@@ -44,14 +44,23 @@ static FAutoConsoleVariableRef CCvarInterchangeEnableGLTFImport(
 	ECVF_Default);
 
 static const TArray<FString> SupportedExtensions = {/* Lights */
-														//TEXT("KHR_lights_punctual"), TEXT("KHR_lights"), //not supported yet
+													TEXT("KHR_lights_punctual"),
+													TEXT("KHR_lights"),
 													/* Variants */
-														TEXT("KHR_materials_variants"), 
+													TEXT("KHR_materials_variants"), 
 													/* Materials */
-														//TEXT("KHR_materials_unlit"), //not supported yet
-														TEXT("KHR_materials_ior"), TEXT("KHR_materials_clearcoat"), TEXT("KHR_materials_transmission"), TEXT("KHR_materials_sheen"),
-														TEXT("KHR_materials_specular"), TEXT("KHR_materials_pbrSpecularGlossiness"), 
-														TEXT("MSFT_packing_occlusionRoughnessMetallic"), TEXT("MSFT_packing_normalRoughnessMetallic") };
+													TEXT("KHR_materials_unlit"),
+													TEXT("KHR_materials_ior"),
+													TEXT("KHR_materials_clearcoat"),
+													TEXT("KHR_materials_transmission"),
+													TEXT("KHR_materials_sheen"),
+													TEXT("KHR_materials_specular"),
+													TEXT("KHR_materials_pbrSpecularGlossiness"),
+													TEXT("MSFT_packing_occlusionRoughnessMetallic"),
+													TEXT("MSFT_packing_normalRoughnessMetallic"),
+													/* Textures */
+													TEXT("KHR_texture_transform")
+};
 
 namespace UE::Interchange::Gltf::Private
 {
@@ -339,6 +348,20 @@ void UInterchangeGltfTranslator::HandleGltfMaterial( UInterchangeBaseNodeContain
 	//(based on [REDACTED_PROFILE_URL]/glTF-Sample-Models/tree/master/2.0/NormalTangentTest#problem-flipped-y-axis-or-flipped-green-channel)
 	SetTextureFlipGreenChannel(NodeContainer, GltfMaterial.Normal);
 	SetTextureFlipGreenChannel(NodeContainer, GltfMaterial.ClearCoat.NormalMap);
+
+	if (GltfMaterial.bIsUnlitShadingModel)
+	{
+		// Base Color
+		{
+			TVariant< FLinearColor, float > BaseColorFactor;
+			BaseColorFactor.Set< FLinearColor >(FLinearColor(GltfMaterial.BaseColorFactor));
+
+			HandleGltfMaterialParameter(NodeContainer, GltfMaterial.BaseColor, ShaderGraphNode, Unlit::Parameters::UnlitColor.ToString(),
+				BaseColorFactor, Standard::Nodes::TextureSample::Outputs::RGB.ToString());
+		}
+
+		return;
+	}
 	
 	//if there is a clearcoatnormal map then we want to swap it with the normals map
 	//as the Interchange pipeline will connect the clearcoatnormal map to ClearCoatBottomNormalMap
diff --git a/Engine/Plugins/Experimental/Interchange/Runtime/Source/Nodes/Public/InterchangeMaterialDefinitions.h b/Engine/Plugins/Experimental/Interchange/Runtime/Source/Nodes/Public/InterchangeMaterialDefinitions.h
--- a/Engine/Plugins/Experimental/Interchange/Runtime/Source/Nodes/Public/InterchangeMaterialDefinitions.h
+++ b/Engine/Plugins/Experimental/Interchange/Runtime/Source/Nodes/Public/InterchangeMaterialDefinitions.h
@@ -277,6 +277,14 @@ namespace Interchange
 				const FName Tangent = TEXT("tangent");
 			}
 		}
+
+		namespace Unlit
+		{
+			namespace Parameters
+			{
+				const FName UnlitColor = TEXT("UnlitColor"); // Type: linear color
+			}
+		}
 	}
 }
 }
diff --git a/Engine/Plugins/Experimental/Interchange/Runtime/Source/Pipelines/Private/InterchangeGenericMaterialPipeline.cpp b/Engine/Plugins/Experimental/Interchange/Runtime/Source/Pipelines/Private/InterchangeGenericMaterialPipeline.cpp
--- a/Engine/Plugins/Experimental/Interchange/Runtime/Source/Pipelines/Private/InterchangeGenericMaterialPipeline.cpp
+++ b/Engine/Plugins/Experimental/Interchange/Runtime/Source/Pipelines/Private/InterchangeGenericMaterialPipeline.cpp
@@ -1585,6 +1585,12 @@ UInterchangeMaterialFactoryNode* UInterchangeGenericMaterialPipeline::CreateMate
 		return MaterialFactoryNode;
 	}
 
+	if (HandleUnlitModel(ShaderGraphNode, MaterialFactoryNode))
+	{
+		// No need to proceed any further
+		return MaterialFactoryNode;
+	}
+
 	if (!HandlePhongModel(ShaderGraphNode, MaterialFactoryNode))
 	{
 		HandleLambertModel(ShaderGraphNode, MaterialFactoryNode);
@@ -1647,6 +1653,10 @@ UInterchangeMaterialInstanceFactoryNode* UInterchangeGenericMaterialPipeline::Cr
 	{
 		MaterialInstanceFactoryNode->SetCustomParent(TEXT("Material'/Interchange/Materials/LambertSurfaceMaterial.LambertSurfaceMaterial'"));
 	}
+	else if (IsUnlitModel(ShaderGraphNode))
+	{
+		MaterialInstanceFactoryNode->SetCustomParent(TEXT("Material'/Interchange/Materials/UnlitMaterial.UnlitMaterial'"));
+	}
 	else
 	{
 		// Default to PBR
@@ -2000,4 +2010,45 @@ UInterchangeMaterialFunctionFactoryNode* UInterchangeGenericMaterialPipeline::Cr
 	return FactoryNode;
 }
 
+bool UInterchangeGenericMaterialPipeline::IsUnlitModel(const UInterchangeShaderGraphNode* ShaderGraphNode) const
+{
+	using namespace UE::Interchange::Materials::Unlit;
+
+	const bool bHasUnlitColorInput = UInterchangeShaderPortsAPI::HasInput(ShaderGraphNode, Parameters::UnlitColor);
+
+	return bHasUnlitColorInput;
+}
+
+bool UInterchangeGenericMaterialPipeline::HandleUnlitModel(const UInterchangeShaderGraphNode* ShaderGraphNode, UInterchangeMaterialFactoryNode* MaterialFactoryNode)
+{
+	using namespace UE::Interchange::Materials;
+
+	bool bShadingModelHandled = false;
+
+	// Unlit Color
+	{
+		const bool bHasInput = UInterchangeShaderPortsAPI::HasInput(ShaderGraphNode, Unlit::Parameters::UnlitColor);
+
+		if (bHasInput)
+		{
+			TTuple<UInterchangeMaterialExpressionFactoryNode*, FString> ExpressionFactoryNode =
+				CreateMaterialExpressionForInput(MaterialFactoryNode, ShaderGraphNode, Unlit::Parameters::UnlitColor.ToString(), MaterialFactoryNode->GetUniqueID());
+
+			if (ExpressionFactoryNode.Get<0>())
+			{
+				MaterialFactoryNode->ConnectOutputToEmissiveColor(ExpressionFactoryNode.Get<0>()->GetUniqueID(), ExpressionFactoryNode.Get<1>());
+			}
+
+			bShadingModelHandled = true;
+		}
+	}
+
+	if (bShadingModelHandled)
+	{
+		MaterialFactoryNode->SetCustomShadingModel(EMaterialShadingModel::MSM_Unlit);
+	}
+
+	return bShadingModelHandled;
+}
+
 #undef LOCTEXT_NAMESPACE
diff --git a/Engine/Plugins/Experimental/Interchange/Runtime/Source/Pipelines/Public/InterchangeGenericMaterialPipeline.h b/Engine/Plugins/Experimental/Interchange/Runtime/Source/Pipelines/Public/InterchangeGenericMaterialPipeline.h
--- a/Engine/Plugins/Experimental/Interchange/Runtime/Source/Pipelines/Public/InterchangeGenericMaterialPipeline.h
+++ b/Engine/Plugins/Experimental/Interchange/Runtime/Source/Pipelines/Public/InterchangeGenericMaterialPipeline.h
@@ -88,6 +88,9 @@ private:
 	/** True if the shader graph has the standard surface's shader type name */
 	bool IsStandardSurfaceModel(const UInterchangeShaderGraphNode* ShaderGraphNode) const;
 
+	/** True if the shader graph has an unlit color input. */
+	bool IsUnlitModel(const UInterchangeShaderGraphNode* ShaderGraphNode) const;
+
 	bool HandlePhongModel(const UInterchangeShaderGraphNode* ShaderGraphNode, UInterchangeMaterialFactoryNode* MaterialFactoryNode);
 	bool HandleLambertModel(const UInterchangeShaderGraphNode* ShaderGraphNode, UInterchangeMaterialFactoryNode* MaterialFactoryNode);
 	bool HandlePBRModel(const UInterchangeShaderGraphNode* ShaderGraphNode, UInterchangeMaterialFactoryNode* MaterialFactoryNode);
@@ -98,6 +101,7 @@ private:
 	bool HandleThinTranslucent(const UInterchangeShaderGraphNode* ShaderGraphNode, UInterchangeMaterialFactoryNode* MaterialFactoryNode);
 	void HandleCommonParameters(const UInterchangeShaderGraphNode* ShaderGraphNode, UInterchangeMaterialFactoryNode* MaterialFactoryNode);
 	bool HandleBxDFInput(const UInterchangeShaderGraphNode* ShaderGraphNode, UInterchangeMaterialFactoryNode* MaterialFactoryNode);
+	bool HandleUnlitModel(const UInterchangeShaderGraphNode* ShaderGraphNode, UInterchangeMaterialFactoryNode* MaterialFactoryNode);
 
 	void HandleFlattenNormalNode(const UInterchangeShaderNode* ShaderNode, UInterchangeBaseMaterialFactoryNode* MaterialFactoryNode, UInterchangeMaterialExpressionFactoryNode* FlattenNormalFactoryNode);
 	void HandleMakeFloat3Node(const UInterchangeShaderNode* ShaderNode, UInterchangeBaseMaterialFactoryNode* MaterialFactoryNode, UInterchangeMaterialExpressionFactoryNode* MakeFloat3FactoryNode);
