MaterialX Translator:
-  Substitute all the nodes that are defined by subgraphs in a nodegraphin by their own nodegraph, in order to handle nodes that are defined with a node graph (such as tiledimage2d or place2d)
-  Refactor the renaming of inputs, add an attribute instead of renaming the input, otherwise the nodedef and the nodegraph are confused
-  Add vector2 attribute when needed (useful for texcoord transform)

diff --git a/Engine/Plugins/Experimental/Interchange/Runtime/Source/Import/Private/MaterialX/InterchangeMaterialXDefinitions.h b/Engine/Plugins/Experimental/Interchange/Runtime/Source/Import/Private/MaterialX/InterchangeMaterialXDefinitions.h
--- a/Engine/Plugins/Experimental/Interchange/Runtime/Source/Import/Private/MaterialX/InterchangeMaterialXDefinitions.h
+++ b/Engine/Plugins/Experimental/Interchange/Runtime/Source/Import/Private/MaterialX/InterchangeMaterialXDefinitions.h
@@ -145,6 +145,7 @@ namespace MaterialX
 	{
 		static constexpr const char* IsVisited = "UE:IsVisited";
 		static constexpr const char* OriginalName = "UE:OriginalName";
+		static constexpr const char* NewName = "UE:NewName";
 	}
 
 	namespace Category
@@ -187,6 +188,7 @@ namespace MaterialX
 		static constexpr const char* InvertMatrix = "invertmatrix";
 		static constexpr const char* Rotate2D = "rotate2d";
 		static constexpr const char* Rotate3D = "rotate3d";
+		static constexpr const char* Place2D = "place2d";
 		static constexpr const char* ArrayAppend = "arrayappend";
 		// Compositing nodes
 		static constexpr const char* Mix = "mix";
diff --git a/Engine/Plugins/Experimental/Interchange/Runtime/Source/Import/Private/MaterialX/InterchangeMaterialXTranslator.cpp b/Engine/Plugins/Experimental/Interchange/Runtime/Source/Import/Private/MaterialX/InterchangeMaterialXTranslator.cpp
--- a/Engine/Plugins/Experimental/Interchange/Runtime/Source/Import/Private/MaterialX/InterchangeMaterialXTranslator.cpp
+++ b/Engine/Plugins/Experimental/Interchange/Runtime/Source/Import/Private/MaterialX/InterchangeMaterialXTranslator.cpp
@@ -765,6 +765,16 @@ bool UInterchangeMaterialXTranslator::ConnectNodeGraphOutputToInput(MaterialX::I
 			return false;
 		}
 
+		//Let's substitute all the sub graphs, so we don't have to handle them ourselves, but they will be handled by smaller operations		
+		//extract and separate nodes will be handled by ourselves, since they are defined with a swizzle node, which is not present in UE
+		auto FilterNodes = [](mx::NodePtr Node) ->bool
+		{
+			const std::string& Category= Node->getCategory();
+			return Category != "extract" && Category != "separate2" && Category != "separate3" && Category != "separate4";
+		};
+
+		Output->getParent()->asA<mx::NodeGraph>()->flattenSubgraphs(mx::EMPTY_STRING, FilterNodes);
+
 		for(mx::Edge Edge : Output->traverseGraph())
 		{
 			if(mx::NodePtr UpstreamNode = Edge.getUpstreamElement()->asA<mx::Node>())
@@ -821,7 +831,7 @@ bool UInterchangeMaterialXTranslator::ConnectNodeGraphOutputToInput(MaterialX::I
 				{
 					ConnectRotate3DInputToOutput(UpstreamNode, ParentShaderNode, InputChannelName, NamesToShaderNodes, NodeContainer);
 				}
-				else if(UpstreamNode->getCategory() == mx::Category::Image || UpstreamNode->getCategory() == mx::Category::TiledImage)
+				else if(UpstreamNode->getCategory() == mx::Category::Image)
 				{
 					ConnectImageInputToOutput(UpstreamNode, ParentShaderNode, InputChannelName, NamesToShaderNodes, NodeContainer);
 				}
@@ -863,8 +873,9 @@ bool UInterchangeMaterialXTranslator::ConnectNodeOutputToInput(MaterialX::NodePt
 			}
 			else if(Input->hasInterfaceName())
 			{
+				auto NodeDef = Node->getNodeDef();
 				mx::InputPtr InputInterface = Input->getInterfaceInput();
-				if(InputInterface->hasValue())
+				if(InputInterface && InputInterface->hasValue())
 				{
 					// We need to take the input name from the original Input not the Interface
 					if(const FString* InputNameFound = UEInputs.Find(GetInputName(Input)))
@@ -964,50 +975,19 @@ void UInterchangeMaterialXTranslator::RenameNodeInputs(MaterialX::NodePtr Node)
 
 void UInterchangeMaterialXTranslator::RenameInput(MaterialX::InputPtr Input, const char* NewName) const
 {
-	std::string OriginalName;
-
-	if(!Input->hasAttribute(mx::Attributes::OriginalName))
-	{
-		OriginalName = Input->getName();
-		Input->setAttribute(mx::Attributes::OriginalName, OriginalName); // keep the original name for further processing
-	}
-	else
-	{
-		OriginalName = Input->getAttribute(mx::Attributes::OriginalName);
-	}
-
-	Input->setName(OriginalName + "_" + NewName);
+	Input->setAttribute(mx::Attributes::NewName, NewName);
 }
 
-MaterialX::InputPtr UInterchangeMaterialXTranslator::GetInputFromOriginalName(MaterialX::NodePtr Node, const char* OriginalNameAttribute) const
+FString UInterchangeMaterialXTranslator::GetInputName(MaterialX::InputPtr Input) const
 {
-	mx::InputPtr InputRes{ nullptr };
-
-	for(mx::InputPtr Input : Node->getInputs())
+	if(Input->hasAttribute(mx::Attributes::NewName))
 	{
-		if(Input->getAttribute(mx::Attributes::OriginalName) == OriginalNameAttribute)
-		{
-			InputRes = Input;
-			break;
-		}
+		return Input->getAttribute(mx::Attributes::NewName).c_str();
 	}
-
-	return InputRes;
-}
-
-FString UInterchangeMaterialXTranslator::GetInputName(MaterialX::InputPtr Input) const
-{
-	std::string Name = Input->getName();
-
-	if(Input->hasAttribute(mx::Attributes::OriginalName))
+	else
 	{
-		std::string OriginalName;
-		OriginalName = Input->getAttribute(mx::Attributes::OriginalName);
-		OriginalName += "_";
-		Name.replace(0, OriginalName.size(), "");
+		return Input->getName().c_str();
 	}
-
-	return Name.c_str();
 }
 
 MaterialX::InputPtr UInterchangeMaterialXTranslator::GetStandardSurfaceInput(MaterialX::NodePtr StandardSurface, const char* InputName, MaterialX::DocumentPtr Document) const
@@ -1067,23 +1047,20 @@ bool UInterchangeMaterialXTranslator::AddAttribute(MaterialX::InputPtr Input, co
 			return ShaderNode->AddFloatAttribute(UInterchangeShaderPortsAPI::MakeInputValueKey(InputChannelName), mx::fromValueString<float>(Input->getValueString()));
 		}
 
-		bool bIsColor = false;
-		FLinearColor LinearColor;
-
 		if(Input->getType() == mx::Type::Color3 || Input->getType() == mx::Type::Color4)
 		{
-			LinearColor = GetLinearColor(Input);
-			bIsColor = true;
+			FLinearColor LinearColor = GetLinearColor(Input);
+			return ShaderNode->AddLinearColorAttribute(UInterchangeShaderPortsAPI::MakeInputValueKey(InputChannelName), LinearColor);
 		}
-		else if(Input->getType() == mx::Type::Vector2 || Input->getType() == mx::Type::Vector3 || Input->getType() == mx::Type::Vector4)
+		else if(Input->getType() == mx::Type::Vector2)
 		{
-			LinearColor = GetVector(Input);
-			bIsColor = true;
+			FLinearColor Vector = GetVector(Input);
+			return ShaderNode->AddVector2Attribute(UInterchangeShaderPortsAPI::MakeInputValueKey(InputChannelName), FVector2f{ Vector.R, Vector.G });
 		}
-
-		if(bIsColor)
+		else if(Input->getType() == mx::Type::Vector3 || Input->getType() == mx::Type::Vector4)
 		{
-			return ShaderNode->AddLinearColorAttribute(UInterchangeShaderPortsAPI::MakeInputValueKey(InputChannelName), LinearColor);
+			FLinearColor Vector = GetVector(Input);
+			return ShaderNode->AddLinearColorAttribute(UInterchangeShaderPortsAPI::MakeInputValueKey(InputChannelName), Vector);
 		}
 	}
 
@@ -1244,7 +1221,7 @@ void UInterchangeMaterialXTranslator::ConnectExtractInputToOutput(MaterialX::Nod
 
 void UInterchangeMaterialXTranslator::ConnectDotInputToOutput(MaterialX::NodePtr UpstreamNode, UInterchangeShaderNode* ParentShaderNode, const FString& InputChannelName, TMap<FString, UInterchangeShaderNode*>& NamesToShaderNodes, UInterchangeBaseNodeContainer& NodeContainer) const
 {
-	if(mx::InputPtr Input = GetInputFromOriginalName(UpstreamNode, "in"))
+	if(mx::InputPtr Input = UpstreamNode->getInput("in"))
 	{
 		RenameInput(Input, TCHAR_TO_UTF8(*InputChannelName)); //let's take the parent node's input name
 		NamesToShaderNodes.FindOrAdd(UpstreamNode->getName().c_str(), ParentShaderNode);
@@ -1254,7 +1231,7 @@ void UInterchangeMaterialXTranslator::ConnectDotInputToOutput(MaterialX::NodePtr
 void UInterchangeMaterialXTranslator::ConnectTransformPositionInputToOutput(MaterialX::NodePtr UpstreamNode, UInterchangeShaderNode* ParentShaderNode, const FString& InputChannelName, TMap<FString, UInterchangeShaderNode*>& NamesToShaderNodes, UInterchangeBaseNodeContainer& NodeContainer) const
 {
 	UInterchangeShaderNode* TransformNode = CreateShaderNode<UInterchangeShaderNode>(UpstreamNode->getName().c_str(), TEXT("TransformPosition"), NamesToShaderNodes, NodeContainer);
-	if(mx::InputPtr Input = GetInputFromOriginalName(UpstreamNode, "in"); Input->hasValue())
+	if(mx::InputPtr Input = UpstreamNode->getInput("in"); Input->hasValue())
 	{
 		FLinearColor Vector = GetVector(Input);
 		TransformNode->AddLinearColorAttribute(UInterchangeShaderPortsAPI::MakeInputValueKey(TEXT("Input")), Vector);
@@ -1265,7 +1242,7 @@ void UInterchangeMaterialXTranslator::ConnectTransformPositionInputToOutput(Mate
 void UInterchangeMaterialXTranslator::ConnectTransformVectorInputToOutput(MaterialX::NodePtr UpstreamNode, UInterchangeShaderNode* ParentShaderNode, const FString& InputChannelName, TMap<FString, UInterchangeShaderNode*>& NamesToShaderNodes, UInterchangeBaseNodeContainer& NodeContainer) const
 {
 	UInterchangeShaderNode* TransformNode = CreateShaderNode<UInterchangeShaderNode>(UpstreamNode->getName().c_str(), TEXT("Transform"), NamesToShaderNodes, NodeContainer);
-	if(mx::InputPtr Input = GetInputFromOriginalName(UpstreamNode, "in"); Input->hasValue())
+	if(mx::InputPtr Input = UpstreamNode->getInput("in"); Input->hasValue())
 	{
 		FLinearColor Vector = GetVector(Input);
 		TransformNode->AddLinearColorAttribute(UInterchangeShaderPortsAPI::MakeInputValueKey(TEXT("Input")), Vector);
@@ -1279,13 +1256,13 @@ void UInterchangeMaterialXTranslator::ConnectRotate2DInputToOutput(mx::NodePtr U
 	Rotate2DNode->AddLinearColorAttribute(UInterchangeShaderPortsAPI::MakeInputValueKey(TEXT("NormalizedRotationAxis")), FLinearColor(0, 0, 1));
 	Rotate2DNode->AddLinearColorAttribute(UInterchangeShaderPortsAPI::MakeInputValueKey(TEXT("PivotPoint")), FLinearColor(0.5, 0.5, 0));
 
-	if(mx::InputPtr Input = GetInputFromOriginalName(UpstreamNode, "in"); Input->hasValue())
+	if(mx::InputPtr Input = UpstreamNode->getInput("in"); Input->hasValue())
 	{
 		FLinearColor Vector = GetVector(Input);
 		Rotate2DNode->AddLinearColorAttribute(UInterchangeShaderPortsAPI::MakeInputValueKey(TEXT("Position")), Vector);
 	}
 
-	if(mx::InputPtr Input = GetInputFromOriginalName(UpstreamNode, "amount"); Input->hasValue())
+	if(mx::InputPtr Input = UpstreamNode->getInput("amount"); Input->hasValue())
 	{
 		//MaterialX angle is in degrees, in UE it's in the range [0,1]
 		float Amount = mx::fromValueString<float>(Input->getValueString()) / 360.;
@@ -1312,19 +1289,19 @@ void UInterchangeMaterialXTranslator::ConnectRotate3DInputToOutput(mx::NodePtr U
 	UInterchangeShaderNode* Rotate3DNode = CreateShaderNode<UInterchangeShaderNode>(UpstreamNode->getName().c_str(), TEXT("RotateAboutAxis"), NamesToShaderNodes, NodeContainer);
 	Rotate3DNode->AddLinearColorAttribute(UInterchangeShaderPortsAPI::MakeInputValueKey(TEXT("PivotPoint")), FLinearColor(0.5, 0.5, 0));
 
-	if(mx::InputPtr Input = GetInputFromOriginalName(UpstreamNode, "in"); Input->hasValue())
+	if(mx::InputPtr Input = UpstreamNode->getInput("in"); Input->hasValue())
 	{
 		FLinearColor Vector = GetVector(Input);
 		Rotate3DNode->AddLinearColorAttribute(UInterchangeShaderPortsAPI::MakeInputValueKey(TEXT("Position")), Vector);
 	}
 
-	if(mx::InputPtr Input = GetInputFromOriginalName(UpstreamNode, "axis"); Input->hasValue())
+	if(mx::InputPtr Input = UpstreamNode->getInput("axis"); Input->hasValue())
 	{
 		FLinearColor Vector = GetVector(Input);
 		Rotate3DNode->AddLinearColorAttribute(UInterchangeShaderPortsAPI::MakeInputValueKey(TEXT("NormalizedRotationAxis")), Vector);
 	}
 
-	if(mx::InputPtr Input = GetInputFromOriginalName(UpstreamNode, "amount"); Input->hasValue())
+	if(mx::InputPtr Input = UpstreamNode->getInput("amount"); Input->hasValue())
 	{
 		//MaterialX angle is in degrees, in UE it's in the range [0,1]
 		float Amount = mx::fromValueString<float>(Input->getValueString()) / 360.;
@@ -1361,61 +1338,6 @@ void UInterchangeMaterialXTranslator::ConnectImageInputToOutput(MaterialX::NodeP
 		UInterchangeShaderNode* TextureShaderNode = CreateShaderNode<UInterchangeShaderNode>(UpstreamNode->getName().c_str(), TextureSample::Name.ToString(), NamesToShaderNodes, NodeContainer);
 		TextureShaderNode->AddStringAttribute(UInterchangeShaderPortsAPI::MakeInputValueKey(TextureSample::Inputs::Texture.ToString()), TextureNode->GetUniqueID());
 		UInterchangeShaderPortsAPI::ConnectOuputToInput(ParentShaderNode, InputChannelName, TextureShaderNode->GetUniqueID(), OutputChannel);
-
-		UInterchangeShaderNode* ImageNode = TextureShaderNode;
-		FString ImageNodeInputName{ TextureSample::Inputs::Coordinates.ToString() };
-
-		auto ConnectUVTransformToOutput = [this, UpstreamNode, &NodeContainer, &NamesToShaderNodes](UInterchangeShaderNode*& ImageNode, FString& ImageNodeInputName, const FString& ShaderType, const char* InputName)
-		{
-			if(mx::InputPtr Input = UpstreamNode->getInput(InputName))
-			{
-				const FString ShaderNodeName{ UpstreamNode->getName().c_str() + FString(TEXT("_")) + InputName };
-				UInterchangeShaderNode* ShaderTransformNode = CreateShaderNode<UInterchangeShaderNode>(ShaderNodeName, ShaderType, NamesToShaderNodes, NodeContainer);
-				UInterchangeShaderPortsAPI::ConnectDefaultOuputToInput(ImageNode, ImageNodeInputName, ShaderTransformNode->GetUniqueID());
-
-				//Vec2
-				if(Input->hasValueString())
-				{
-					const FString MaskNodeName{ UpstreamNode->getName().c_str() + FString(TEXT("_")) + ShaderType + TEXT("MaskNode") };
-					UInterchangeShaderNode* MaskShaderNode = CreateShaderNode<UInterchangeShaderNode>(MaskNodeName, Mask::Name.ToString(), NamesToShaderNodes, NodeContainer);
-					MaskShaderNode->AddBooleanAttribute(Mask::Attributes::R, true);
-					MaskShaderNode->AddBooleanAttribute(Mask::Attributes::G, true);
-
-					mx::Vector2 Vec2 = mx::fromValueString<mx::Vector2>(Input->getValueString());
-					MaskShaderNode->AddLinearColorAttribute(UInterchangeShaderPortsAPI::MakeInputValueKey(Mask::Inputs::Input.ToString()), FLinearColor(Vec2[0], Vec2[1], 0));
-					UInterchangeShaderPortsAPI::ConnectDefaultOuputToInput(ShaderTransformNode, TEXT("B"), MaskShaderNode->GetUniqueID());
-				}
-				else
-				{
-					RenameInput(Input, "B");
-				}
-
-				//The node will replace the texture node for the input
-				NamesToShaderNodes.Add(UpstreamNode->getName().c_str(), ShaderTransformNode);
-
-				//We also need to replace the name of texcoord input by one of the inputs of this node
-				if(mx::InputPtr InputTexCoord = GetInputFromOriginalName(UpstreamNode, mx::TiledImage::Inputs::TexCoord))
-				{
-					RenameInput(InputTexCoord, "A");
-				}
-				else
-				{
-					//Let's reuse the same texture coordinate node for further use (no need to set a parent)
-					const FString TextureCoordinateName{ FPaths::GetBaseFilename(UpstreamNode->getActiveSourceUri().c_str()) + FString(TEXT("_texcoord")) };
-					UInterchangeShaderNode* TextureCoordinateNode = CreateShaderNode<UInterchangeShaderNode>(TextureCoordinateName, TextureCoordinate::Name.ToString(), NamesToShaderNodes, NodeContainer);
-					UInterchangeShaderPortsAPI::ConnectDefaultOuputToInput(ShaderTransformNode, TEXT("A"), TextureCoordinateNode->GetUniqueID());
-				}
-
-				ImageNode = ShaderTransformNode;
-				ImageNodeInputName = TEXT("A");
-			}
-		};
-
-		// UV offset (MaterialX defines it as a substraction)
-		// MaterialX spec: the offset for the given image along the U and V axes. Mathematically
-		// equivalent to subtracting the given vector value from the incoming texture coordinates.
-		ConnectUVTransformToOutput(ImageNode, ImageNodeInputName, TEXT("Subtract"), mx::TiledImage::Inputs::UVOffset);
-		ConnectUVTransformToOutput(ImageNode, ImageNodeInputName, TEXT("Multiply"), mx::TiledImage::Inputs::UVTiling);
 	}
 	else
 	{
@@ -1428,7 +1350,7 @@ void UInterchangeMaterialXTranslator::ConnectConvertInputToOutput(MaterialX::Nod
 	using namespace UE::Interchange::Materials::Standard::Nodes;
 
 	//In case of an upwards conversion, let's do an append, downwards a mask, otherwise leave as is
-	mx::InputPtr Input = GetInputFromOriginalName(UpstreamNode, "in");
+	mx::InputPtr Input = UpstreamNode->getInput("in");
 	const std::string& NodeType = UpstreamNode->getType();
 	const std::string& InputType = Input->getType();
 
diff --git a/Engine/Plugins/Experimental/Interchange/Runtime/Source/Import/Public/MaterialX/InterchangeMaterialXTranslator.h b/Engine/Plugins/Experimental/Interchange/Runtime/Source/Import/Public/MaterialX/InterchangeMaterialXTranslator.h
--- a/Engine/Plugins/Experimental/Interchange/Runtime/Source/Import/Public/MaterialX/InterchangeMaterialXTranslator.h
+++ b/Engine/Plugins/Experimental/Interchange/Runtime/Source/Import/Public/MaterialX/InterchangeMaterialXTranslator.h
@@ -198,16 +198,6 @@ protected:
 	 */
 	void RenameInput(MaterialX::InputPtr Input, const char * NewName) const;
 
-	/**
-	 * Helper function to retrieve an Input in a Node from its original name (after a renaming)
-	 * 
-	 * @param Node - The node with the Inputs to look up
-	 * @param OriginalNameAttribute - The previous name of an Input
-	 * 
-	 * @return The found Input, nullptr otherwise
-	 */
-	MaterialX::InputPtr GetInputFromOriginalName(MaterialX::NodePtr Node, const char * OriginalNameAttribute) const;
-
 	/**
 	 * Get the input name, use this function instead of getName, because a renaming may have occured and we ensure to have the proper name that will be used by UE inputs
 	 * 
