[MaterialX Interchange] Add basic support for translucency

diff --git a/Engine/Plugins/Experimental/Interchange/Runtime/Source/Import/Private/MaterialX/InterchangeMaterialXTranslator.cpp b/Engine/Plugins/Experimental/Interchange/Runtime/Source/Import/Private/MaterialX/InterchangeMaterialXTranslator.cpp
--- a/Engine/Plugins/Experimental/Interchange/Runtime/Source/Import/Private/MaterialX/InterchangeMaterialXTranslator.cpp
+++ b/Engine/Plugins/Experimental/Interchange/Runtime/Source/Import/Private/MaterialX/InterchangeMaterialXTranslator.cpp
@@ -592,6 +592,79 @@ void UInterchangeMaterialXTranslator::ProcessStandardSurface(UInterchangeBaseNod
 			ConnectNodeGraphOutputToInput(InputTangent, ShaderGraphNode, StandardSurface::Parameters::Tangent.ToString(), NamesToShaderNodes, NodeContainer);
 		}
 	}
+
+	//Transmission
+	{
+		//Weight
+		{
+			mx::InputPtr InputTransmission = GetStandardSurfaceInput(StandardSurfaceNode, mx::StandardSurface::Input::Transmission, Document);
+
+			if (!ConnectNodeGraphOutputToInput(InputTransmission, ShaderGraphNode, StandardSurface::Parameters::Transmission.ToString(), NamesToShaderNodes, NodeContainer))
+			{
+				AddFloatAttribute(InputTransmission, StandardSurface::Parameters::Transmission.ToString(), ShaderGraphNode, mx::StandardSurface::DefaultValue::Float::Transmission);
+			}
+		}
+
+		//Color
+		{
+			mx::InputPtr InputTransmissionColor = GetStandardSurfaceInput(StandardSurfaceNode, mx::StandardSurface::Input::TransmissionColor, Document);
+
+			if (!ConnectNodeGraphOutputToInput(InputTransmissionColor, ShaderGraphNode, StandardSurface::Parameters::TransmissionColor.ToString(), NamesToShaderNodes, NodeContainer))
+			{
+				AddLinearColorAttribute(InputTransmissionColor, StandardSurface::Parameters::TransmissionColor.ToString(), ShaderGraphNode, mx::StandardSurface::DefaultValue::Color3::TransmissionColor);
+			}
+		}
+
+		//Depth
+		{
+			mx::InputPtr InputTransmissionDepth = GetStandardSurfaceInput(StandardSurfaceNode, mx::StandardSurface::Input::TransmissionDepth, Document);
+
+			if (!ConnectNodeGraphOutputToInput(InputTransmissionDepth, ShaderGraphNode, StandardSurface::Parameters::TransmissionDepth.ToString(), NamesToShaderNodes, NodeContainer))
+			{
+				AddFloatAttribute(InputTransmissionDepth, StandardSurface::Parameters::TransmissionDepth.ToString(), ShaderGraphNode, mx::StandardSurface::DefaultValue::Float::TransmissionDepth);
+			}
+		}
+
+		//Scatter
+		{
+			mx::InputPtr InputTransmissionScatter = GetStandardSurfaceInput(StandardSurfaceNode, mx::StandardSurface::Input::TransmissionScatter, Document);
+
+			if (!ConnectNodeGraphOutputToInput(InputTransmissionScatter, ShaderGraphNode, StandardSurface::Parameters::TransmissionScatter.ToString(), NamesToShaderNodes, NodeContainer))
+			{
+				AddLinearColorAttribute(InputTransmissionScatter, StandardSurface::Parameters::TransmissionScatter.ToString(), ShaderGraphNode, mx::StandardSurface::DefaultValue::Color3::TransmissionScatter);
+			}
+		}
+
+		//Scatter Anisotropy
+		{
+			mx::InputPtr InputTransmissionScatterAnisotropy = GetStandardSurfaceInput(StandardSurfaceNode, mx::StandardSurface::Input::TransmissionScatterAnisotropy, Document);
+
+			if (!ConnectNodeGraphOutputToInput(InputTransmissionScatterAnisotropy, ShaderGraphNode, StandardSurface::Parameters::TransmissionScatterAnisotropy.ToString(), NamesToShaderNodes, NodeContainer))
+			{
+				AddFloatAttribute(InputTransmissionScatterAnisotropy, StandardSurface::Parameters::TransmissionScatterAnisotropy.ToString(), ShaderGraphNode, mx::StandardSurface::DefaultValue::Float::TransmissionScatterAnisotropy);
+			}
+		}
+
+		//Dispersion
+		{
+			mx::InputPtr InputTransmissionDispersion = GetStandardSurfaceInput(StandardSurfaceNode, mx::StandardSurface::Input::TransmissionDispersion, Document);
+
+			if (!ConnectNodeGraphOutputToInput(InputTransmissionDispersion, ShaderGraphNode, StandardSurface::Parameters::TransmissionDispersion.ToString(), NamesToShaderNodes, NodeContainer))
+			{
+				AddFloatAttribute(InputTransmissionDispersion, StandardSurface::Parameters::TransmissionDispersion.ToString(), ShaderGraphNode, mx::StandardSurface::DefaultValue::Float::TransmissionDispersion);
+			}
+		}
+
+		//Extra Roughness
+		{
+			mx::InputPtr InputTransmissionExtraRoughness = GetStandardSurfaceInput(StandardSurfaceNode, mx::StandardSurface::Input::TransmissionExtraRoughness, Document);
+
+			if (!ConnectNodeGraphOutputToInput(InputTransmissionExtraRoughness, ShaderGraphNode, StandardSurface::Parameters::TransmissionExtraRoughness.ToString(), NamesToShaderNodes, NodeContainer))
+			{
+				AddFloatAttribute(InputTransmissionExtraRoughness, StandardSurface::Parameters::TransmissionExtraRoughness.ToString(), ShaderGraphNode, mx::StandardSurface::DefaultValue::Float::TransmissionExtraRoughness);
+			}
+		}
+	}
 }
 
 void UInterchangeMaterialXTranslator::ProcessLightShader(UInterchangeBaseNodeContainer& NodeContainer, MaterialX::NodePtr LightShaderNode, MaterialX::DocumentPtr Document) const
diff --git a/Engine/Plugins/Experimental/Interchange/Runtime/Source/Pipelines/Private/InterchangeGenericMaterialPipeline.cpp b/Engine/Plugins/Experimental/Interchange/Runtime/Source/Pipelines/Private/InterchangeGenericMaterialPipeline.cpp
--- a/Engine/Plugins/Experimental/Interchange/Runtime/Source/Pipelines/Private/InterchangeGenericMaterialPipeline.cpp
+++ b/Engine/Plugins/Experimental/Interchange/Runtime/Source/Pipelines/Private/InterchangeGenericMaterialPipeline.cpp
@@ -578,31 +578,38 @@ bool UInterchangeGenericMaterialPipeline::HandleStandardSurfaceModel(const UInte
 		
 		static ArgumentsType ArgumentsPerInputs[]
 		{
-			ArgumentsType{StandardSurface::Parameters::Base.ToString(), EMaterialInputType::Scalar},               // Base
-			ArgumentsType{StandardSurface::Parameters::BaseColor.ToString(), EMaterialInputType::Color},           // BaseColor
-			ArgumentsType{StandardSurface::Parameters::DiffuseRoughness.ToString(),  EMaterialInputType::Scalar},  // DiffuseRoughness
-			ArgumentsType{StandardSurface::Parameters::Metalness.ToString(), EMaterialInputType::Scalar},          // Metalness
-			ArgumentsType{StandardSurface::Parameters::Specular.ToString(), EMaterialInputType::Scalar},           // Specular
-			ArgumentsType{StandardSurface::Parameters::SpecularRoughness.ToString(), EMaterialInputType::Scalar},  // SpecularRoughness
-			ArgumentsType{StandardSurface::Parameters::SpecularIOR.ToString(), EMaterialInputType::Scalar},        // SpecularIOR
-			ArgumentsType{StandardSurface::Parameters::SpecularAnisotropy.ToString(), EMaterialInputType::Scalar}, // SpecularAnisotropy
-			ArgumentsType{StandardSurface::Parameters::SpecularRotation.ToString(), EMaterialInputType::Scalar},   // SpecularRotation
-			ArgumentsType{StandardSurface::Parameters::Subsurface.ToString(), EMaterialInputType::Scalar},         // Subsurface
-			ArgumentsType{StandardSurface::Parameters::SubsurfaceColor.ToString(), EMaterialInputType::Scalar},    // SubsurfaceColor
-			ArgumentsType{StandardSurface::Parameters::SubsurfaceRadius.ToString(), EMaterialInputType::Scalar},   // SubsurfaceRadius
-			ArgumentsType{StandardSurface::Parameters::SubsurfaceScale.ToString(), EMaterialInputType::Scalar},    // SubsurfaceScale
-			ArgumentsType{StandardSurface::Parameters::Sheen.ToString(), EMaterialInputType::Scalar},              // Sheen
-			ArgumentsType{StandardSurface::Parameters::SheenColor.ToString(), EMaterialInputType::Color},          // SheenColor
-			ArgumentsType{StandardSurface::Parameters::SheenRoughness.ToString(), EMaterialInputType::Scalar},     // SheenRoughness
-			ArgumentsType{StandardSurface::Parameters::Coat.ToString(), EMaterialInputType::Scalar},               // Coat
-			ArgumentsType{StandardSurface::Parameters::CoatColor.ToString(),  EMaterialInputType::Color},          // CoatColor
-			ArgumentsType{StandardSurface::Parameters::CoatNormal.ToString(),  EMaterialInputType::Vector},        // CoatNormal
-			ArgumentsType{StandardSurface::Parameters::CoatRoughness.ToString(), EMaterialInputType::Scalar},      // CoatRoughness
-			ArgumentsType{StandardSurface::Parameters::ThinFilmThickness.ToString(), EMaterialInputType::Scalar},  // ThinFilmThickness
-			ArgumentsType{StandardSurface::Parameters::Emission.ToString(), EMaterialInputType::Scalar},           // Emission
-			ArgumentsType{StandardSurface::Parameters::EmissionColor.ToString(), EMaterialInputType::Color},       // EmissionColor
-			ArgumentsType{StandardSurface::Parameters::Normal.ToString(), EMaterialInputType::Vector},             // Normal
-			ArgumentsType{StandardSurface::Parameters::Tangent.ToString(), EMaterialInputType::Vector}             // Tangent
+			ArgumentsType{StandardSurface::Parameters::Base.ToString(), EMaterialInputType::Scalar},                          // Base
+			ArgumentsType{StandardSurface::Parameters::BaseColor.ToString(), EMaterialInputType::Color},                      // BaseColor
+			ArgumentsType{StandardSurface::Parameters::DiffuseRoughness.ToString(),  EMaterialInputType::Scalar},             // DiffuseRoughness
+			ArgumentsType{StandardSurface::Parameters::Metalness.ToString(), EMaterialInputType::Scalar},                     // Metalness
+			ArgumentsType{StandardSurface::Parameters::Specular.ToString(), EMaterialInputType::Scalar},                      // Specular
+			ArgumentsType{StandardSurface::Parameters::SpecularRoughness.ToString(), EMaterialInputType::Scalar},             // SpecularRoughness
+			ArgumentsType{StandardSurface::Parameters::SpecularIOR.ToString(), EMaterialInputType::Scalar},                   // SpecularIOR
+			ArgumentsType{StandardSurface::Parameters::SpecularAnisotropy.ToString(), EMaterialInputType::Scalar},            // SpecularAnisotropy
+			ArgumentsType{StandardSurface::Parameters::SpecularRotation.ToString(), EMaterialInputType::Scalar},              // SpecularRotation
+			ArgumentsType{StandardSurface::Parameters::Subsurface.ToString(), EMaterialInputType::Scalar},                    // Subsurface
+			ArgumentsType{StandardSurface::Parameters::SubsurfaceColor.ToString(), EMaterialInputType::Scalar},               // SubsurfaceColor
+			ArgumentsType{StandardSurface::Parameters::SubsurfaceRadius.ToString(), EMaterialInputType::Scalar},              // SubsurfaceRadius
+			ArgumentsType{StandardSurface::Parameters::SubsurfaceScale.ToString(), EMaterialInputType::Scalar},               // SubsurfaceScale
+			ArgumentsType{StandardSurface::Parameters::Sheen.ToString(), EMaterialInputType::Scalar},                         // Sheen
+			ArgumentsType{StandardSurface::Parameters::SheenColor.ToString(), EMaterialInputType::Color},                     // SheenColor
+			ArgumentsType{StandardSurface::Parameters::SheenRoughness.ToString(), EMaterialInputType::Scalar},                // SheenRoughness
+			ArgumentsType{StandardSurface::Parameters::Coat.ToString(), EMaterialInputType::Scalar},                          // Coat
+			ArgumentsType{StandardSurface::Parameters::CoatColor.ToString(),  EMaterialInputType::Color},                     // CoatColor
+			ArgumentsType{StandardSurface::Parameters::CoatNormal.ToString(),  EMaterialInputType::Vector},                   // CoatNormal
+			ArgumentsType{StandardSurface::Parameters::CoatRoughness.ToString(), EMaterialInputType::Scalar},                 // CoatRoughness
+			ArgumentsType{StandardSurface::Parameters::ThinFilmThickness.ToString(), EMaterialInputType::Scalar},             // ThinFilmThickness
+			ArgumentsType{StandardSurface::Parameters::Emission.ToString(), EMaterialInputType::Scalar},                      // Emission
+			ArgumentsType{StandardSurface::Parameters::EmissionColor.ToString(), EMaterialInputType::Color},                  // EmissionColor
+			ArgumentsType{StandardSurface::Parameters::Normal.ToString(), EMaterialInputType::Vector},                        // Normal
+			ArgumentsType{StandardSurface::Parameters::Tangent.ToString(), EMaterialInputType::Vector},                       // Tangent
+			ArgumentsType{StandardSurface::Parameters::Transmission.ToString(),  EMaterialInputType::Scalar},                 // Transmission
+			ArgumentsType{StandardSurface::Parameters::TransmissionColor.ToString(), EMaterialInputType::Color},              // TransmissionColor
+			ArgumentsType{StandardSurface::Parameters::TransmissionDepth.ToString(), EMaterialInputType::Scalar},             // TransmissionDepth
+			ArgumentsType{StandardSurface::Parameters::TransmissionScatter.ToString(), EMaterialInputType::Color},            // TransmissionScatter
+			ArgumentsType{StandardSurface::Parameters::TransmissionScatterAnisotropy.ToString(), EMaterialInputType::Scalar}, // TransmissionScatterAnisotropy
+			ArgumentsType{StandardSurface::Parameters::TransmissionDispersion.ToString(), EMaterialInputType::Scalar},        // TransmissionDispersion
+			ArgumentsType{StandardSurface::Parameters::TransmissionExtraRoughness.ToString(), EMaterialInputType::Scalar},    // TransmissionExtraRoughness
 		};
 
 		for(const ArgumentsType & Arguments : ArgumentsPerInputs)
@@ -642,6 +649,14 @@ bool UInterchangeGenericMaterialPipeline::HandleStandardSurfaceModel(const UInte
 			MaterialFactoryNode->ConnectOutputToSubsurface(FunctionCallExpressionUid, Subsurface::Parameters::SubsurfaceColor.ToString());
 			MaterialFactoryNode->SetCustomShadingModel(MSM_Subsurface);
 		}
+		else if (UInterchangeShaderPortsAPI::HasInput(ShaderGraphNode, StandardSurface::Parameters::Transmission))
+		{
+			MaterialFactoryNode->ConnectOutputToTransmissionColor(FunctionCallExpressionUid, ThinTranslucent::Parameters::TransmissionColor.ToString());
+
+			MaterialFactoryNode->SetCustomBlendMode(EBlendMode::BLEND_Translucent);
+			MaterialFactoryNode->SetCustomShadingModel(EMaterialShadingModel::MSM_ThinTranslucent);
+			MaterialFactoryNode->SetCustomTranslucencyLightingMode(ETranslucencyLightingMode::TLM_SurfacePerPixelLighting);
+		}
 		else
 		{
 			MaterialFactoryNode->SetCustomShadingModel(MSM_DefaultLit);
