InterchangeImport: account for inheritance when retrieving inputs of a MaterialX NodeDef

The MaterialX specification includes support for inheritance where one
element can inherit from another. In particular with node definitions, one
definition can inherit from a generic, canonical definition of a node or
shader and add new inputs or change the values of existing inputs.
This functionality is already present in MaterialX 1.38.1 which is the
current engine version, although it is not used by Autodesk's standard
surface node definition in that version.

In MaterialX 1.38.3, this feature was used as a way to version the node
definition for standard surface. A new version of the node definition was
added that inherits from the older version and modifies the default values
for the "base" and "base_color" inputs. This means that the
"ND_standard_surface_surfaceshader" node definition now only has two
inputs of its own while all the rest come from the definition it inherits from,
"ND_standard_surface_surfaceshader_100":

[REDACTED_PROFILE_URL]/MaterialX/commit/24f6c2630bab3a4fe8ef80f4004745575b7deba2

In Interchange, there were a handful of call sites that used NodeDef's
getInput() method (which it inherits from InterfaceElement) in order to find
an input by name, but that method does not account for inheritance.

In order to account for the possibility of such inheritance correctly, those sites
should call getActiveInput() (also inherited from InterfaceElement) instead,
which does traverse up through the inheritance chain looking for any ancestor
that has the input if one cannot be found on the current interface element.

This fix can safely be applied based on the 1.38.1 version of MaterialX before
it becomes required when upgrading to 1.38.5, where the standard surface
node definition was modified to leverage inheritance. Otherwise, all inputs
on the standard surface node definition other than "base" and "base_color"
would seem to be missing from the perspective of Interchange.

diff --git a/Engine/Plugins/Experimental/Interchange/Runtime/Source/Import/Private/MaterialX/InterchangeMaterialXTranslator.cpp b/Engine/Plugins/Experimental/Interchange/Runtime/Source/Import/Private/MaterialX/InterchangeMaterialXTranslator.cpp
--- a/Engine/Plugins/Experimental/Interchange/Runtime/Source/Import/Private/MaterialX/InterchangeMaterialXTranslator.cpp
+++ b/Engine/Plugins/Experimental/Interchange/Runtime/Source/Import/Private/MaterialX/InterchangeMaterialXTranslator.cpp
@@ -1069,7 +1069,7 @@ MaterialX::InputPtr UInterchangeMaterialXTranslator::GetStandardSurfaceInput(Mat
 
 	if(!Input)
 	{
-		Input = Document->getNodeDef(mx::NodeDefinition::StandardSurface)->getInput(InputName);
+		Input = Document->getNodeDef(mx::NodeDefinition::StandardSurface)->getActiveInput(InputName);
 	}
 
 	return Input;
@@ -1081,7 +1081,7 @@ MaterialX::InputPtr UInterchangeMaterialXTranslator::GetPointLightInput(Material
 
 	if(!Input)
 	{
-		Input = Document->getNodeDef(mx::NodeDefinition::PointLight)->getInput(InputName);
+		Input = Document->getNodeDef(mx::NodeDefinition::PointLight)->getActiveInput(InputName);
 	}
 
 	return Input;
@@ -1093,7 +1093,7 @@ MaterialX::InputPtr UInterchangeMaterialXTranslator::GetDirectionalLightInput(Ma
 
 	if(!Input)
 	{
-		Input = Document->getNodeDef(mx::NodeDefinition::DirectionalLight)->getInput(InputName);
+		Input = Document->getNodeDef(mx::NodeDefinition::DirectionalLight)->getActiveInput(InputName);
 	}
 
 	return Input;
@@ -1105,7 +1105,7 @@ MaterialX::InputPtr UInterchangeMaterialXTranslator::GetSpotLightInput(MaterialX
 
 	if(!Input)
 	{
-		Input = Document->getNodeDef(mx::NodeDefinition::SpotLight)->getInput(InputName);
+		Input = Document->getNodeDef(mx::NodeDefinition::SpotLight)->getActiveInput(InputName);
 	}
 
 	return Input;
