Material expression:
- Luminance
- Remap
- Logarithm (Ln)
- RampLeftRight
- RampTopBottom
- SplitLeftRight
- SplitTopBottom
Translator
- Add nodes from above material expression
- Desaturation
- Smoothstep
- IfGreater
- IfGreaterEq
- IfEqual

diff --git a/Engine/Plugins/Experimental/Interchange/Runtime/Source/Import/Private/MaterialX/InterchangeMaterialXDefinitions.h b/Engine/Plugins/Experimental/Interchange/Runtime/Source/Import/Private/MaterialX/InterchangeMaterialXDefinitions.h
--- a/Engine/Plugins/Experimental/Interchange/Runtime/Source/Import/Private/MaterialX/InterchangeMaterialXDefinitions.h
+++ b/Engine/Plugins/Experimental/Interchange/Runtime/Source/Import/Private/MaterialX/InterchangeMaterialXDefinitions.h
@@ -151,55 +151,67 @@
 	namespace Category
 	{
 		// Math nodes
-		static constexpr const char* Add = "add";
-		static constexpr const char* Sub = "subtract";
-		static constexpr const char* Multiply = "multiply";
-		static constexpr const char* Divide = "divide";
-		static constexpr const char* Modulo = "modulo";
-		static constexpr const char* Invert = "invert";
 		static constexpr const char* Absval = "absval";
-		static constexpr const char* Floor = "floor";
-		static constexpr const char* Ceil = "ceil";
-		static constexpr const char* Power = "power";
-		static constexpr const char* Sin = "sin";
-		static constexpr const char* Cos = "cos";
-		static constexpr const char* Tan = "tan";
-		static constexpr const char* Asin = "asin";
 		static constexpr const char* Acos = "acos";
+		static constexpr const char* Add = "add";
+		static constexpr const char* ArrayAppend = "arrayappend";
+		static constexpr const char* Asin = "asin";
 		static constexpr const char* Atan2 = "atan2";
-		static constexpr const char* Sqrt = "sqrt";
-		static constexpr const char* Ln = "ln";
-		static constexpr const char* Exp = "exp";
-		static constexpr const char* Sign = "sign";
+		static constexpr const char* Ceil = "ceil";
 		static constexpr const char* Clamp = "clamp";
-		static constexpr const char* Min = "min";
-		static constexpr const char* Max = "max";
-		static constexpr const char* Normalize = "normalize";
-		static constexpr const char* Magnitude = "magnitude";
-		static constexpr const char* DotProduct = "dotproduct";
+		static constexpr const char* Cos = "cos";
 		static constexpr const char* CrossProduct = "crossproduct";
-		static constexpr const char* TransformPoint = "transformpoint";
-		static constexpr const char* TransformVector = "transformvector";
-		static constexpr const char* TransformNormal = "transformnormal";
-		static constexpr const char* TransformMatrix = "transformmatrix";
-		static constexpr const char* NormalMap = "normalmap";
-		static constexpr const char* Transpose = "transpose";
 		static constexpr const char* Determinant = "determinant";
+		static constexpr const char* Divide = "divide";
+		static constexpr const char* DotProduct = "dotproduct";
+		static constexpr const char* Exp = "exp";
+		static constexpr const char* Floor = "floor";
+		static constexpr const char* Invert = "invert";
 		static constexpr const char* InvertMatrix = "invertmatrix";
+		static constexpr const char* Ln = "ln";
+		static constexpr const char* Magnitude = "magnitude";
+		static constexpr const char* Max = "max";
+		static constexpr const char* Min = "min";
+		static constexpr const char* Modulo = "modulo";
+		static constexpr const char* Multiply = "multiply";
+		static constexpr const char* Normalize = "normalize";
+		static constexpr const char* NormalMap = "normalmap";
+		static constexpr const char* Place2D = "place2d";
+		static constexpr const char* Power = "power";
 		static constexpr const char* Rotate2D = "rotate2d";
 		static constexpr const char* Rotate3D = "rotate3d";
-		static constexpr const char* Place2D = "place2d";
-		static constexpr const char* ArrayAppend = "arrayappend";
+		static constexpr const char* Sign = "sign";
+		static constexpr const char* Sin = "sin";
+		static constexpr const char* Sqrt = "sqrt";
+		static constexpr const char* Sub = "subtract";
+		static constexpr const char* Tan = "tan";
+		static constexpr const char* TransformMatrix = "transformmatrix";
+		static constexpr const char* TransformNormal = "transformnormal";
+		static constexpr const char* TransformPoint = "transformpoint";
+		static constexpr const char* TransformVector = "transformvector";
+		static constexpr const char* Transpose = "transpose";
 		// Compositing nodes
 		static constexpr const char* Mix = "mix";
+		// Conditional nodes
+		static constexpr const char* IfGreater = "ifgreater";
+		static constexpr const char* IfGreaterEq = "ifgreatereq";
+		static constexpr const char* IfEqual = "ifequal";
 		// Channel nodes
 		static constexpr const char* Convert = "convert";
 		static constexpr const char* Combine2 = "combine2";
 		static constexpr const char* Combine3 = "combine3";
 		static constexpr const char* Combine4 = "combine4";
 		static constexpr const char* Extract = "extract";
+		static constexpr const char* Separate2 = "separate2";
+		static constexpr const char* Separate3 = "separate3";
+		static constexpr const char* Separate4 = "separate4";
 		// Procedural nodes 
 		static constexpr const char* Constant = "constant";
+		// Procedural2D nodes 
+		static constexpr const char* RampLR = "ramplr";
+		static constexpr const char* RampTB = "ramptb";
+		static constexpr const char* SplitLR = "splitlr";
+		static constexpr const char* SplitTB = "splittb";
 		// Organization nodes 
 		static constexpr const char* Dot = "dot";
 		// Texture nodes 
@@ -208,12 +220,16 @@
 		// Geometric nodes 
 		static constexpr const char* TexCoord = "texcoord";
 		// Light nodes
-		static constexpr const char* PointLight = "point_light";
 		static constexpr const char* DirectionalLight = "directional_light";
+		static constexpr const char* PointLight = "point_light";
 		static constexpr const char* SpotLight = "spot_light";
 		// Adjustment nodes
-		static constexpr const char* RgbToHsv = "rgbtohsv";
 		static constexpr const char* HsvToRgb = "hsvtorgb";
+		static constexpr const char* Luminance = "luminance";
+		static constexpr const char* Remap = "remap";
+		static constexpr const char* RgbToHsv = "rgbtohsv";
+		static constexpr const char* Saturate = "saturate";
+		static constexpr const char* Smoothstep = "smoothstep";
 		// PBR
 		static constexpr const char* StandardSurface = "standard_surface";
 	}
diff --git a/Engine/Plugins/Experimental/Interchange/Runtime/Source/Import/Private/MaterialX/InterchangeMaterialXTranslator.cpp b/Engine/Plugins/Experimental/Interchange/Runtime/Source/Import/Private/MaterialX/InterchangeMaterialXTranslator.cpp
--- a/Engine/Plugins/Experimental/Interchange/Runtime/Source/Import/Private/MaterialX/InterchangeMaterialXTranslator.cpp
+++ b/Engine/Plugins/Experimental/Interchange/Runtime/Source/Import/Private/MaterialX/InterchangeMaterialXTranslator.cpp
@@ -60,8 +60,10 @@ namespace UE::Interchange::MaterialX
 UInterchangeMaterialXTranslator::UInterchangeMaterialXTranslator()
 #if WITH_EDITOR
 	:
-InputNamesMaterialX2UE{
+InputNamesMaterialX2UE
+{
 	{{TEXT(""),                         TEXT("bg")},		TEXT("B")},
+	{{TEXT(""),                         TEXT("center")},		TEXT("Center")},
 	{{TEXT(""),                         TEXT("fg")},		TEXT("A")},
 	{{TEXT(""),                         TEXT("high")},		TEXT("Max")},
 	{{TEXT(""),                         TEXT("in")},		TEXT("Input")},
@@ -69,24 +71,46 @@ InputNamesMaterialX2UE{
 	{{TEXT(""),                         TEXT("in2")},		TEXT("B")},
 	{{TEXT(""),                         TEXT("in3")},		TEXT("C")},
 	{{TEXT(""),                         TEXT("in4")},		TEXT("D")},
+	{{TEXT(""),                         TEXT("inlow")},     TEXT("InputLow")},
+	{{TEXT(""),                         TEXT("inhigh")},    TEXT("InputHigh")},
 	{{TEXT(""),                         TEXT("low")},		TEXT("Min")},
+	{{TEXT(""),                         TEXT("lumacoeffs")},TEXT("LuminanceFactors")}, // for the moment not yet handled by Interchange, because of the attribute being an advanced pin
 	{{TEXT(""),                         TEXT("mix")},		TEXT("Factor")},
 	{{TEXT(""),                         TEXT("texcoord")},  TEXT("Coordinates")},
+	{{TEXT(""),                         TEXT("outlow")},    TEXT("TargetLow")},
+	{{TEXT(""),                         TEXT("outhigh")},   TEXT("TargetHigh")},
+	{{TEXT(""),                         TEXT("valuel")},    TEXT("A")},
+	{{TEXT(""),                         TEXT("valuer")},    TEXT("B")},
 	{{MaterialX::Category::Atan2,       TEXT("in1")},       TEXT("Y")},
 	{{MaterialX::Category::Atan2,       TEXT("in2")},       TEXT("X")},
+	{{MaterialX::Category::IfGreater,   TEXT("value1")},    TEXT("A")},
+	{{MaterialX::Category::IfGreater,   TEXT("value2")},    TEXT("B")},
+	{{MaterialX::Category::IfGreater,   TEXT("in1")},       TEXT("AGreaterThanB")},
+	{{MaterialX::Category::IfGreater,   TEXT("in2")},       TEXT("ALessThanB")}, //another input is added for the case equal see ConnectIfGreater
+	{{MaterialX::Category::IfGreaterEq, TEXT("value1")},    TEXT("A")},
+	{{MaterialX::Category::IfGreaterEq, TEXT("value2")},    TEXT("B")},
+	{{MaterialX::Category::IfGreaterEq, TEXT("in1")},       TEXT("AGreaterThanB")}, //another is added for the case equal see ConnectIfGreaterEq
+	{{MaterialX::Category::IfGreaterEq, TEXT("in2")},       TEXT("ALessThanB")},
+	{{MaterialX::Category::IfEqual,     TEXT("value1")},    TEXT("A")},
+	{{MaterialX::Category::IfEqual,     TEXT("value2")},    TEXT("B")},
+	{{MaterialX::Category::IfEqual,     TEXT("in1")},       TEXT("AEqualsB")},
+	{{MaterialX::Category::IfEqual,     TEXT("in2")},       TEXT("ALessThanB")}, //another input is added for the case greater see ConnectIfEqual
+	{{MaterialX::Category::Invert,      TEXT("amount")},    TEXT("A")},
+	{{MaterialX::Category::Invert,      TEXT("in")},        TEXT("B")},
 	{{MaterialX::Category::Magnitude,   TEXT("in")},        TEXT("A")},
 	{{MaterialX::Category::Mix,         TEXT("fg")},        TEXT("B")},
 	{{MaterialX::Category::Mix,         TEXT("bg")},        TEXT("A")},
 	{{MaterialX::Category::Normalize,   TEXT("in")},        TEXT("VectorInput")},
 	{{MaterialX::Category::Power,       TEXT("in1")},       TEXT("Base")},
 	{{MaterialX::Category::Power,	    TEXT("in2")},       TEXT("Exponent")},
-	{{MaterialX::Category::Invert,      TEXT("amount")},    TEXT("A")},
-	{{MaterialX::Category::Invert,      TEXT("in")},        TEXT("B")},
 	{{MaterialX::Category::Rotate2D,    TEXT("amount")},    TEXT("RotationAngle")},
 	{{MaterialX::Category::Rotate2D,    TEXT("in")},        TEXT("Position")},
 	{{MaterialX::Category::Rotate3D,    TEXT("amount")},    TEXT("RotationAngle")},
 	{{MaterialX::Category::Rotate3D,    TEXT("axis")},		TEXT("NormalizedRotationAxis")},
-	{{MaterialX::Category::Rotate3D,    TEXT("in")},        TEXT("Position")} },
+	{{MaterialX::Category::Rotate3D,    TEXT("in")},        TEXT("Position")},
+	{{MaterialX::Category::Saturate,    TEXT("amount")},    TEXT("Fraction")},
+	{{MaterialX::Category::Smoothstep,  TEXT("in")},        TEXT("Value")},
+},
 NodeNamesMaterialX2UE{
 	// Math nodes
 	{MaterialX::Category::Absval,       TEXT("Abs")},
@@ -103,6 +127,7 @@ NodeNamesMaterialX2UE{
 	{MaterialX::Category::Exp,          TEXT("Exponential")},
 	{MaterialX::Category::Floor,        TEXT("Floor")},
 	{MaterialX::Category::Invert,       TEXT("Subtract")},
+	{MaterialX::Category::Ln,           TEXT("Logarithm")},
 	{MaterialX::Category::Magnitude,    TEXT("Length")},
 	{MaterialX::Category::Max,          TEXT("Max")},
 	{MaterialX::Category::Min,          TEXT("Min")},
@@ -110,8 +135,12 @@ NodeNamesMaterialX2UE{
 	{MaterialX::Category::Multiply,     TEXT("Multiply")},
 	{MaterialX::Category::Normalize,    TEXT("Normalize")},
 	{MaterialX::Category::Power,        TEXT("Power")},
+	{MaterialX::Category::RampLR,       TEXT("RampLeftRight")},
+	{MaterialX::Category::RampTB,       TEXT("RampTopBottom")},
 	{MaterialX::Category::Sign,         TEXT("Sign")},
 	{MaterialX::Category::Sin,          TEXT("Sine")},
+	{MaterialX::Category::SplitLR,      TEXT("SplitLeftRight")},
+	{MaterialX::Category::SplitTB,      TEXT("SplitTopBottom")},
 	{MaterialX::Category::Sqrt,         TEXT("SquareRoot")},
 	{MaterialX::Category::Sub,          TEXT("Subtract")},
 	{MaterialX::Category::Tan,          TEXT("Tangent")},
@@ -125,21 +154,36 @@ NodeNamesMaterialX2UE{
 	{MaterialX::Category::TexCoord,		TEXT("TextureCoordinate")},
 	// Adjustment nodes,
 	{MaterialX::Category::HsvToRgb,		TEXT("HsvToRgb")},
-	{MaterialX::Category::RgbToHsv,		TEXT("RgbToHsv")}},
-UEInputs{
+	{MaterialX::Category::Luminance,    TEXT("Luminance")},
+	{MaterialX::Category::Remap,        TEXT("Remap")},
+	{MaterialX::Category::RgbToHsv,		TEXT("RgbToHsv")},
+	{MaterialX::Category::Saturate,		TEXT("Desaturation")},
+	{MaterialX::Category::Smoothstep,   TEXT("SmoothStep")}
+},
+UEInputs
+{
     TEXT("A"),
 	TEXT("B"),
     TEXT("Base"),
-    TEXT("Exponent"),
     TEXT("C"),
+	TEXT("Center"),
     TEXT("D"),
+	TEXT("Exponent"),
     TEXT("Factor"),
+	TEXT("Fraction"),
     TEXT("Input"),
+	TEXT("InputLow"),
+	TEXT("InputHigh"),
+	TEXT("LuminanceFactors"),
     TEXT("Max"),
     TEXT("Min"),
-    TEXT("VectorInput")
+	TEXT("TargetLow"),
+	TEXT("TargetHigh"),
+	TEXT("Value"),
+    TEXT("VectorInput"),
     TEXT("X"),
-    TEXT("Y")}
+    TEXT("Y")
+}
 #endif // WITH_EDITOR
 {
 }
@@ -838,12 +882,17 @@ bool UInterchangeMaterialXTranslator::ConnectNodeGraphOutputToInput(MaterialX::I
 			return false;
 		}
 
-		//Let's substitute all the sub graphs, so we don't have to handle them ourselves, but they will be handled by smaller operations		
+		//Let's substitue all the sub graphs, so we don't have to handle them ourselves, but they will be handled by smaller operations		
 		//extract and separate nodes will be handled by ourselves, since they are defined with a swizzle node, which is not present in UE
 		auto FilterNodes = [](mx::NodePtr Node) ->bool
 		{
 			const std::string& Category= Node->getCategory();
-			return Category != "extract" && Category != "separate2" && Category != "separate3" && Category != "separate4";
+			return 
+				Category != mx::Category::Extract &&
+				Category != mx::Category::Saturate &&
+				Category != mx::Category::Separate2 &&
+				Category != mx::Category::Separate3 &&
+				Category != mx::Category::Separate4;
 		};
 
 		Output->getParent()->asA<mx::NodeGraph>()->flattenSubgraphs(mx::EMPTY_STRING, FilterNodes);
@@ -912,6 +961,18 @@ bool UInterchangeMaterialXTranslator::ConnectNodeGraphOutputToInput(MaterialX::I
 				{
 					ConnectConvertInputToOutput(UpstreamNode, ParentShaderNode, InputChannelName, NamesToShaderNodes, NodeContainer);
 				}
+				else if(UpstreamNode->getCategory() == mx::Category::IfGreater)
+				{
+					ConnectIfGreaterInputToOutput(UpstreamNode, ParentShaderNode, InputChannelName, NamesToShaderNodes, NodeContainer);
+				}
+				else if(UpstreamNode->getCategory() == mx::Category::IfGreaterEq)
+				{
+					ConnectIfGreaterEqInputToOutput(UpstreamNode, ParentShaderNode, InputChannelName, NamesToShaderNodes, NodeContainer);
+				}
+				else if(UpstreamNode->getCategory() == mx::Category::IfEqual)
+				{
+					ConnectIfEqualInputToOutput(UpstreamNode, ParentShaderNode, InputChannelName, NamesToShaderNodes, NodeContainer);
+				}
 				else
 				{
 					UInterchangeResultWarning_Generic* Message = AddMessage<UInterchangeResultWarning_Generic>();
@@ -946,7 +1007,6 @@ bool UInterchangeMaterialXTranslator::ConnectNodeOutputToInput(MaterialX::NodePt
 			}
 			else if(Input->hasInterfaceName())
 			{
-				auto NodeDef = Node->getNodeDef();
 				mx::InputPtr InputInterface = Input->getInterfaceInput();
 				if(InputInterface && InputInterface->hasValue())
 				{
@@ -1406,7 +1466,9 @@ void UInterchangeMaterialXTranslator::ConnectImageInputToOutput(MaterialX::NodeP
 		FString OutputChannel{ TEXT("RGB") };
 
 		if(UpstreamNode->getType() == mx::Type::Vector4 || UpstreamNode->getType() == mx::Type::Color4)
+		{
 			OutputChannel = TEXT("RGBA");
+		}
 
 		UInterchangeShaderNode* TextureShaderNode = CreateShaderNode<UInterchangeShaderNode>(UpstreamNode->getName().c_str(), TextureSample::Name.ToString(), NamesToShaderNodes, NodeContainer);
 		TextureShaderNode->AddStringAttribute(UInterchangeShaderPortsAPI::MakeInputValueKey(TextureSample::Inputs::Texture.ToString()), TextureNode->GetUniqueID());
@@ -1559,6 +1621,123 @@ void UInterchangeMaterialXTranslator::ConnectConvertInputToOutput(MaterialX::Nod
 	}
 }
 
+void UInterchangeMaterialXTranslator::ConnectIfGreaterInputToOutput(MaterialX::NodePtr UpstreamNode, UInterchangeShaderNode* ParentShaderNode, const FString& InputChannelName, TMap<FString, UInterchangeShaderNode*>& NamesToShaderNodes, UInterchangeBaseNodeContainer& NodeContainer) const
+{
+	UInterchangeShaderNode* NodeIf = CreateShaderNode<UInterchangeShaderNode>(UpstreamNode->getName().c_str(), TEXT("If"), NamesToShaderNodes, NodeContainer);
+	UInterchangeShaderPortsAPI::ConnectDefaultOuputToInput(ParentShaderNode, InputChannelName, NodeIf->GetUniqueID());
+
+	if(mx::InputPtr Input = UpstreamNode->getInput("value1"); Input && Input->hasValueString())
+	{
+		AddAttribute(Input, GetInputName(Input), NodeIf);
+	}
+	
+	if(mx::InputPtr Input = UpstreamNode->getInput("value2"); Input && Input->hasValueString())
+	{
+		AddAttribute(Input, GetInputName(Input), NodeIf);
+	}
+	
+	if(mx::InputPtr Input = UpstreamNode->getInput("in1"); Input && Input->hasValueString())
+	{
+		AddAttribute(Input, GetInputName(Input), NodeIf);
+	}
+
+	//In that case we also need to add an attribute to AEqualsB
+	mx::InputPtr Input = UpstreamNode->getInput("in2");
+	if(Input)
+	{
+		if(Input->hasValueString())
+		{
+			AddAttribute(Input, GetInputName(Input), NodeIf);
+			AddAttribute(Input, TEXT("AEqualsB"), NodeIf);
+		}
+		else
+		{
+			//Let's add a new input that is a copy of in2 to connect it to the equal input
+			mx::InputPtr Input3 = UpstreamNode->addInput("in3");
+			Input3->copyContentFrom(Input);
+			RenameInput(Input3, "AEqualsB");
+		}
+	}
+}
+
+void UInterchangeMaterialXTranslator::ConnectIfGreaterEqInputToOutput(MaterialX::NodePtr UpstreamNode, UInterchangeShaderNode * ParentShaderNode, const FString & InputChannelName, TMap<FString, UInterchangeShaderNode*>&NamesToShaderNodes, UInterchangeBaseNodeContainer & NodeContainer) const
+{
+	UInterchangeShaderNode* NodeIf = CreateShaderNode<UInterchangeShaderNode>(UpstreamNode->getName().c_str(), TEXT("If"), NamesToShaderNodes, NodeContainer);
+	UInterchangeShaderPortsAPI::ConnectDefaultOuputToInput(ParentShaderNode, InputChannelName, NodeIf->GetUniqueID());
+
+	if(mx::InputPtr Input = UpstreamNode->getInput("value1"); Input && Input->hasValueString())
+	{
+		AddAttribute(Input, GetInputName(Input), NodeIf);
+	}
+
+	if(mx::InputPtr Input = UpstreamNode->getInput("value2"); Input && Input->hasValueString())
+	{
+		AddAttribute(Input, GetInputName(Input), NodeIf);
+	}
+
+	if(mx::InputPtr Input = UpstreamNode->getInput("in2"); Input && Input->hasValueString())
+	{
+		AddAttribute(Input, GetInputName(Input), NodeIf);
+	}
+
+	//In that case we also need to add an attribute to AEqualsB
+	mx::InputPtr Input = UpstreamNode->getInput("in1");
+	if(Input)
+	{
+		if(Input->hasValueString())
+		{
+			AddAttribute(Input, GetInputName(Input), NodeIf);
+			AddAttribute(Input, TEXT("AEqualsB"), NodeIf);
+		}
+		else
+		{
+			//Let's add a new input that is a copy of in2 to connect it to the equal input
+			mx::InputPtr Input3 = UpstreamNode->addInput("in3");
+			Input3->copyContentFrom(Input);
+			RenameInput(Input3, "AEqualsB");
+		}
+	}
+}
+
+void UInterchangeMaterialXTranslator::ConnectIfEqualInputToOutput(MaterialX::NodePtr UpstreamNode, UInterchangeShaderNode * ParentShaderNode, const FString & InputChannelName, TMap<FString, UInterchangeShaderNode*>&NamesToShaderNodes, UInterchangeBaseNodeContainer & NodeContainer) const
+{
+	UInterchangeShaderNode* NodeIf = CreateShaderNode<UInterchangeShaderNode>(UpstreamNode->getName().c_str(), TEXT("If"), NamesToShaderNodes, NodeContainer);
+	UInterchangeShaderPortsAPI::ConnectDefaultOuputToInput(ParentShaderNode, InputChannelName, NodeIf->GetUniqueID());
+
+	if(mx::InputPtr Input = UpstreamNode->getInput("value1"); Input && Input->hasValueString())
+	{
+		AddAttribute(Input, GetInputName(Input), NodeIf);
+	}
+
+	if(mx::InputPtr Input = UpstreamNode->getInput("value2"); Input && Input->hasValueString())
+	{
+		AddAttribute(Input, GetInputName(Input), NodeIf);
+	}
+
+	if(mx::InputPtr Input = UpstreamNode->getInput("in1"); Input && Input->hasValueString())
+	{
+		AddAttribute(Input, GetInputName(Input), NodeIf);
+	}
+
+	//In that case we also need to add an attribute to AGreaterThanB
+	mx::InputPtr Input = UpstreamNode->getInput("in2");
+	if(Input)
+	{
+		if(Input->hasValueString())
+		{
+			AddAttribute(Input, GetInputName(Input), NodeIf);
+			AddAttribute(Input, TEXT("AGreaterThanB"), NodeIf);
+		}
+		else
+		{
+			//Let's add a new input that is a copy of in2 to connect it to the equal input
+			mx::InputPtr Input3 = UpstreamNode->addInput("in3");
+			Input3->copyContentFrom(Input);
+			RenameInput(Input3, "AGreaterThanB");
+		}
+	}
+}
+
 UInterchangeShaderNode* UInterchangeMaterialXTranslator::CreateMaskShaderNode(uint8 RGBA, const FString& NodeName, TMap<FString, UInterchangeShaderNode*>& NamesToShaderNodes, UInterchangeBaseNodeContainer& NodeContainer) const
 {
 	bool bR = (0b1000 & RGBA) >> 3;
diff --git a/Engine/Plugins/Experimental/Interchange/Runtime/Source/Import/Private/MaterialX/MaterialExpressionAppend3Vector.h b/Engine/Plugins/Experimental/Interchange/Runtime/Source/Import/Private/MaterialX/MaterialExpressionAppend3Vector.h
new file mode 100644
--- /dev/null
+++ b/Engine/Plugins/Experimental/Interchange/Runtime/Source/Import/Private/MaterialX/MaterialExpressionAppend3Vector.h
@@ -0,0 +1,33 @@
+// Copyright Epic Games, Inc. All Rights Reserved.
+#pragma once
+
+#include "CoreMinimal.h"
+#include "Materials/MaterialExpression.h"
+
+#include "MaterialExpressionAppend3Vector.generated.h"
+
+/**
+ * A material expression that allows combining 3 channels together to create a vector with more channel than the original
+ */
+UCLASS(MinimalAPI, collapsecategories, hidecategories = Object)
+class UMaterialExpressionAppend3Vector : public UMaterialExpression
+{
+	GENERATED_UCLASS_BODY()
+
+	UPROPERTY()
+	FExpressionInput A;
+
+	UPROPERTY()
+	FExpressionInput B;
+
+	UPROPERTY()
+	FExpressionInput C;
+
+	//~ Begin UMaterialExpression Interface
+#if WITH_EDITOR
+	virtual int32 Compile(class FMaterialCompiler* Compiler, int32 OutputIndex) override;
+	virtual void GetCaption(TArray<FString>& OutCaptions) const override;
+#endif
+	//~ End UMaterialExpression Interface
+};
+
diff --git a/Engine/Plugins/Experimental/Interchange/Runtime/Source/Import/Private/MaterialX/MaterialExpressionAppend4Vector.h b/Engine/Plugins/Experimental/Interchange/Runtime/Source/Import/Private/MaterialX/MaterialExpressionAppend4Vector.h
new file mode 100644
--- /dev/null
+++ b/Engine/Plugins/Experimental/Interchange/Runtime/Source/Import/Private/MaterialX/MaterialExpressionAppend4Vector.h
@@ -0,0 +1,38 @@
+// Copyright Epic Games, Inc. All Rights Reserved.
+#pragma once
+
+#include "CoreMinimal.h"
+#include "UObject/ObjectMacros.h"
+#include "MaterialExpressionIO.h"
+#include "Materials/MaterialExpression.h"
+
+#include "MaterialExpressionAppend4Vector.generated.h"
+
+/**
+ * A material expression that allows combining 4 channels together to create a vector with more channel than the original
+ */
+UCLASS(MinimalAPI, collapsecategories, hidecategories = Object)
+class UMaterialExpressionAppend4Vector : public UMaterialExpression
+{
+	GENERATED_UCLASS_BODY()
+
+	UPROPERTY()
+	FExpressionInput A;
+	
+	UPROPERTY()
+	FExpressionInput B;
+
+	UPROPERTY()
+	FExpressionInput C;
+	
+	UPROPERTY()
+	FExpressionInput D;
+
+	//~ Begin UMaterialExpression Interface
+#if WITH_EDITOR
+	virtual int32 Compile(class FMaterialCompiler* Compiler, int32 OutputIndex) override;
+	virtual void GetCaption(TArray<FString>& OutCaptions) const override;
+#endif
+	//~ End UMaterialExpression Interface
+};
+
diff --git a/Engine/Plugins/Experimental/Interchange/Runtime/Source/Import/Private/MaterialX/MaterialExpressionExponential.h b/Engine/Plugins/Experimental/Interchange/Runtime/Source/Import/Private/MaterialX/MaterialExpressionExponential.h
new file mode 100644
--- /dev/null
+++ b/Engine/Plugins/Experimental/Interchange/Runtime/Source/Import/Private/MaterialX/MaterialExpressionExponential.h
@@ -0,0 +1,27 @@
+// Copyright Epic Games, Inc. All Rights Reserved.
+#pragma once
+
+#include "CoreMinimal.h"
+#include "Materials/MaterialExpression.h"
+
+#include "MaterialExpressionExponential.generated.h"
+
+/**
+ * A material expression that computes the exponential of an input
+ */
+UCLASS(collapsecategories, hidecategories = Object, MinimalAPI)
+class UMaterialExpressionExponential : public UMaterialExpression
+{
+	GENERATED_UCLASS_BODY()
+
+	UPROPERTY()
+	FExpressionInput Input;
+
+	//~ Begin UMaterialExpression Interface
+#if WITH_EDITOR
+	virtual int32 Compile(class FMaterialCompiler* Compiler, int32 OutputIndex) override;
+	virtual void GetCaption(TArray<FString>& OutCaptions) const override;
+#endif
+	//~ End UMaterialExpression Interface
+};
+
diff --git a/Engine/Plugins/Experimental/Interchange/Runtime/Source/Import/Private/MaterialX/MaterialExpressionHsvToRgb.h b/Engine/Plugins/Experimental/Interchange/Runtime/Source/Import/Private/MaterialX/MaterialExpressionHsvToRgb.h
new file mode 100644
--- /dev/null
+++ b/Engine/Plugins/Experimental/Interchange/Runtime/Source/Import/Private/MaterialX/MaterialExpressionHsvToRgb.h
@@ -0,0 +1,23 @@
+// Copyright Epic Games, Inc. All Rights Reserved.
+#pragma once
+
+#include "Materials/MaterialExpression.h"
+
+#include "MaterialExpressionHsvToRgb.generated.h"
+
+UCLASS(collapsecategories, hidecategories = Object, MinimalAPI)
+class UMaterialExpressionHsvToRgb : public UMaterialExpression
+{
+	GENERATED_UCLASS_BODY()
+
+	UPROPERTY()
+	FExpressionInput Input;
+
+	//~ Begin UMaterialExpression Interface
+#if WITH_EDITOR
+	virtual int32 Compile(class FMaterialCompiler* Compiler, int32 OutputIndex) override;
+	virtual void GetCaption(TArray<FString>& OutCaptions) const override;
+#endif
+	//~ End UMaterialExpression Interface
+};
+
diff --git a/Engine/Plugins/Experimental/Interchange/Runtime/Source/Import/Private/MaterialX/MaterialExpressionLength.h b/Engine/Plugins/Experimental/Interchange/Runtime/Source/Import/Private/MaterialX/MaterialExpressionLength.h
new file mode 100644
--- /dev/null
+++ b/Engine/Plugins/Experimental/Interchange/Runtime/Source/Import/Private/MaterialX/MaterialExpressionLength.h
@@ -0,0 +1,27 @@
+// Copyright Epic Games, Inc. All Rights Reserved.
+#pragma once
+
+#include "CoreMinimal.h"
+#include "Materials/MaterialExpression.h"
+
+#include "MaterialExpressionLength.generated.h"
+
+/**
+ * A material expression that computes the length (magnitude) of a vector
+ */
+UCLASS(collapsecategories, hidecategories = Object, MinimalAPI)
+class UMaterialExpressionLength : public UMaterialExpression
+{
+	GENERATED_UCLASS_BODY()
+
+	UPROPERTY()
+	FExpressionInput Input;
+
+	//~ Begin UMaterialExpression Interface
+#if WITH_EDITOR
+	virtual int32 Compile(class FMaterialCompiler* Compiler, int32 OutputIndex) override;
+	virtual void GetCaption(TArray<FString>& OutCaptions) const override;
+#endif
+	//~ End UMaterialExpression Interface
+};
+
diff --git a/Engine/Plugins/Experimental/Interchange/Runtime/Source/Import/Private/MaterialX/MaterialExpressionLogarithm.cpp b/Engine/Plugins/Experimental/Interchange/Runtime/Source/Import/Private/MaterialX/MaterialExpressionLogarithm.cpp
new file mode 100644
--- /dev/null
+++ b/Engine/Plugins/Experimental/Interchange/Runtime/Source/Import/Private/MaterialX/MaterialExpressionLogarithm.cpp
@@ -0,0 +1,46 @@
+// Copyright Epic Games, Inc. All Rights Reserved.
+#include "MaterialX/MaterialExpressionLogarithm.h"
+
+#include "MaterialCompiler.h"
+#include "MaterialHLSLGenerator.h"
+
+#define LOCTEXT_NAMESPACE "MaterialExpressionLogarithm"
+
+UMaterialExpressionLogarithm::UMaterialExpressionLogarithm(const FObjectInitializer& ObjectInitializer)
+	: Super(ObjectInitializer)
+{
+	// Structure to hold one-time initialization
+	struct FConstructorStatics
+	{
+		FText NAME_Math;
+		FConstructorStatics()
+			: NAME_Math(LOCTEXT("Math", "Math"))
+		{}
+	};
+	static FConstructorStatics ConstructorStatics;
+
+#if WITH_EDITORONLY_DATA
+	MenuCategories.Add(ConstructorStatics.NAME_Math);
+#endif
+}
+
+#if WITH_EDITOR
+int32 UMaterialExpressionLogarithm::Compile(FMaterialCompiler* Compiler, int32 OutputIndex)
+{
+	if(!Input.GetTracedInput().Expression)
+	{
+		return Compiler->Errorf(TEXT("Missing Ln input"));
+	}
+
+	const int32 Log10Euler = Compiler->Constant(0.4342945);
+
+	return Compiler->Div(Compiler->Logarithm10(Input.Compile(Compiler)), Log10Euler);
+}
+
+void UMaterialExpressionLogarithm::GetCaption(TArray<FString>& OutCaptions) const
+{
+	OutCaptions.Add(TEXT("Ln"));
+}
+#endif
+
+#undef LOCTEXT_NAMESPACE 
\ No newline at end of file
diff --git a/Engine/Plugins/Experimental/Interchange/Runtime/Source/Import/Private/MaterialX/MaterialExpressionLogarithm.h b/Engine/Plugins/Experimental/Interchange/Runtime/Source/Import/Private/MaterialX/MaterialExpressionLogarithm.h
new file mode 100644
--- /dev/null
+++ b/Engine/Plugins/Experimental/Interchange/Runtime/Source/Import/Private/MaterialX/MaterialExpressionLogarithm.h
@@ -0,0 +1,27 @@
+// Copyright Epic Games, Inc. All Rights Reserved.
+#pragma once
+
+#include "CoreMinimal.h"
+#include "Materials/MaterialExpression.h"
+
+#include "MaterialExpressionLogarithm.generated.h"
+
+/**
+ * A material expression that computes the natural logarithm
+ */
+UCLASS(collapsecategories, hidecategories = Object, MinimalAPI)
+class UMaterialExpressionLogarithm : public UMaterialExpression
+{
+	GENERATED_UCLASS_BODY()
+
+	UPROPERTY()
+	FExpressionInput Input;
+
+	//~ Begin UMaterialExpression Interface
+#if WITH_EDITOR
+	virtual int32 Compile(class FMaterialCompiler* Compiler, int32 OutputIndex) override;
+	virtual void GetCaption(TArray<FString>& OutCaptions) const override;
+#endif
+	//~ End UMaterialExpression Interface
+};
+
diff --git a/Engine/Plugins/Experimental/Interchange/Runtime/Source/Import/Private/MaterialX/MaterialExpressionLuminance.cpp b/Engine/Plugins/Experimental/Interchange/Runtime/Source/Import/Private/MaterialX/MaterialExpressionLuminance.cpp
new file mode 100644
--- /dev/null
+++ b/Engine/Plugins/Experimental/Interchange/Runtime/Source/Import/Private/MaterialX/MaterialExpressionLuminance.cpp
@@ -0,0 +1,82 @@
+// Copyright Epic Games, Inc. All Rights Reserved.
+
+#include "MaterialX/MaterialExpressionLuminance.h"
+#include "MaterialCompiler.h"
+
+#define LOCTEXT_NAMESPACE "MaterialExpressionLuminance"
+
+UMaterialExpressionLuminance::UMaterialExpressionLuminance(const FObjectInitializer& ObjectInitializer)
+	: Super(ObjectInitializer),
+	LuminanceFactors(0.2722287, 0.6740818, 0.0536895),
+	LuminanceMode(ELuminanceMode::ACEScg)
+{
+	// Structure to hold one-time initialization
+	struct FConstructorStatics
+	{
+		FText NAME_Color;
+		FText NAME_Utility;
+		FConstructorStatics()
+			: NAME_Color(LOCTEXT("Color", "Color"))
+			, NAME_Utility(LOCTEXT("Utility", "Utility"))
+		{}
+	};
+	static FConstructorStatics ConstructorStatics;
+
+#if WITH_EDITORONLY_DATA
+	MenuCategories.Add(ConstructorStatics.NAME_Color);
+	MenuCategories.Add(ConstructorStatics.NAME_Utility);
+#endif
+}
+
+#if WITH_EDITOR
+int32 UMaterialExpressionLuminance::Compile(FMaterialCompiler* Compiler, int32 OutputIndex)
+{
+	if(!Input.GetTracedInput().Expression)
+	{
+		return Compiler->Errorf(TEXT("Missing Luminance input"));
+	}
+
+	const int32 LuminanceFactorsIndex = Compiler->Constant3(LuminanceFactors.R, LuminanceFactors.G, LuminanceFactors.B);
+	auto Type = Compiler->GetParameterType(Input.Compile(Compiler));
+
+	return Compiler->Dot(Compiler->ComponentMask(Input.Compile(Compiler), true, true, true, false), LuminanceFactorsIndex);
+}
+
+void UMaterialExpressionLuminance::GetCaption(TArray<FString>& OutCaptions) const
+{
+	OutCaptions.Add(TEXT("Luminance"));
+}
+
+void UMaterialExpressionLuminance::PostEditChangeProperty(FPropertyChangedEvent& PropertyChangedEvent)
+{
+	FProperty* PropertyThatChanged = PropertyChangedEvent.MemberProperty;
+	const FString PropertyName = PropertyThatChanged ? PropertyThatChanged->GetName() : TEXT("");
+
+	if(PropertyName == GET_MEMBER_NAME_STRING_CHECKED(UMaterialExpressionLuminance, LuminanceMode))
+	{
+		switch(LuminanceMode)
+		{
+		case ELuminanceMode::ACEScg:
+			LuminanceFactors = FLinearColor(0.2722287, 0.6740818, 0.0536895);
+			break;
+
+		case ELuminanceMode::Rec709:
+			LuminanceFactors = FLinearColor(0.2126, 0.7152, 0.0722);
+			break;
+
+		case ELuminanceMode::Rec2020:
+			LuminanceFactors = FLinearColor(0.2627, 0.6780, 0.0593);
+			break;
+		}
+	}
+	else if(PropertyName == GET_MEMBER_NAME_STRING_CHECKED(UMaterialExpressionLuminance, LuminanceFactors))
+	{
+		LuminanceMode = ELuminanceMode::Custom;
+	}
+		
+	Super::PostEditChangeProperty(PropertyChangedEvent);
+}
+
+#endif
+
+#undef LOCTEXT_NAMESPACE 
\ No newline at end of file
diff --git a/Engine/Plugins/Experimental/Interchange/Runtime/Source/Import/Private/MaterialX/MaterialExpressionLuminance.h b/Engine/Plugins/Experimental/Interchange/Runtime/Source/Import/Private/MaterialX/MaterialExpressionLuminance.h
new file mode 100644
--- /dev/null
+++ b/Engine/Plugins/Experimental/Interchange/Runtime/Source/Import/Private/MaterialX/MaterialExpressionLuminance.h
@@ -0,0 +1,46 @@
+// Copyright Epic Games, Inc. All Rights Reserved.
+#pragma once
+
+#include "CoreMinimal.h"
+#include "Materials/MaterialExpression.h"
+
+#include "MaterialExpressionLuminance.generated.h"
+
+/**
+ * A material expression that outputs a grayscale image containing the luminance of the incoming RGB color in all color channels;
+    the alpha channel is left unchanged if present.
+ */
+
+UENUM()
+enum class ELuminanceMode : uint8
+{
+	ACEScg,
+	Rec709,
+	Rec2020,
+	Rec2100 = Rec2020,
+	Custom
+};
+
+UCLASS(collapsecategories, hidecategories = Object, MinimalAPI)
+class UMaterialExpressionLuminance : public UMaterialExpression
+{
+	GENERATED_UCLASS_BODY()
+
+	UPROPERTY()
+	FExpressionInput Input;
+
+	UPROPERTY(EditAnywhere, Category = MaterialExpressionLuminance, Meta = (ShowAsInputPin = "Advanced"))
+	FLinearColor LuminanceFactors;    // Color component factors for converting a color to greyscale.
+	
+	UPROPERTY(EditAnywhere, Category = MaterialExpressionLuminance, Meta = (ShowAsInputPin = "Advanced"))
+	ELuminanceMode LuminanceMode;
+	
+	//~ Begin UMaterialExpression Interface
+#if WITH_EDITOR
+	virtual int32 Compile(class FMaterialCompiler* Compiler, int32 OutputIndex) override;
+	virtual void GetCaption(TArray<FString>& OutCaptions) const override;
+	virtual void PostEditChangeProperty(FPropertyChangedEvent& PropertyChangedEvent) override;
+#endif
+	//~ End UMaterialExpression Interface
+};
+
diff --git a/Engine/Plugins/Experimental/Interchange/Runtime/Source/Import/Private/MaterialX/MaterialExpressionRampLeftRight.cpp b/Engine/Plugins/Experimental/Interchange/Runtime/Source/Import/Private/MaterialX/MaterialExpressionRampLeftRight.cpp
new file mode 100644
--- /dev/null
+++ b/Engine/Plugins/Experimental/Interchange/Runtime/Source/Import/Private/MaterialX/MaterialExpressionRampLeftRight.cpp
@@ -0,0 +1,49 @@
+// Copyright Epic Games, Inc. All Rights Reserved.
+#include "MaterialX/MaterialExpressionRampLeftRight.h"
+#include "MaterialCompiler.h"
+
+#define LOCTEXT_NAMESPACE "MaterialExpressionRampLeftRight"
+
+UMaterialExpressionRampLeftRight::UMaterialExpressionRampLeftRight(const FObjectInitializer& ObjectInitializer)
+	: Super(ObjectInitializer)
+{
+	// Structure to hold one-time initialization
+	struct FConstructorStatics
+	{
+		FText NAME_Utility;
+		FConstructorStatics()
+			: NAME_Utility(LOCTEXT("Utility", "Utility"))
+		{}
+	};
+	static FConstructorStatics ConstructorStatics;
+
+#if WITH_EDITORONLY_DATA
+	MenuCategories.Add(ConstructorStatics.NAME_Utility);
+#endif
+}
+
+#if WITH_EDITOR
+int32 UMaterialExpressionRampLeftRight::Compile(FMaterialCompiler* Compiler, int32 OutputIndex)
+{
+	if(!A.GetTracedInput().Expression)
+	{
+		return Compiler->Errorf(TEXT("Missing A input"));
+	}
+
+	if(!B.GetTracedInput().Expression)
+	{
+		return Compiler->Errorf(TEXT("Missing B input"));
+	}
+
+	int32 CoordinateIndex = Coordinates.GetTracedInput().Expression ? Coordinates.Compile(Compiler) : Compiler->TextureCoordinate(ConstCoordinate, false, false);
+
+	return Compiler->Lerp(A.Compile(Compiler), B.Compile(Compiler), Compiler->Saturate(Compiler->ComponentMask(CoordinateIndex, true, false, false, false)));
+}
+
+void UMaterialExpressionRampLeftRight::GetCaption(TArray<FString>& OutCaptions) const
+{
+	OutCaptions.Add(TEXT("RampLR"));
+}
+#endif
+
+#undef LOCTEXT_NAMESPACE 
\ No newline at end of file
diff --git a/Engine/Plugins/Experimental/Interchange/Runtime/Source/Import/Private/MaterialX/MaterialExpressionRampLeftRight.h b/Engine/Plugins/Experimental/Interchange/Runtime/Source/Import/Private/MaterialX/MaterialExpressionRampLeftRight.h
new file mode 100644
--- /dev/null
+++ b/Engine/Plugins/Experimental/Interchange/Runtime/Source/Import/Private/MaterialX/MaterialExpressionRampLeftRight.h
@@ -0,0 +1,37 @@
+// Copyright Epic Games, Inc. All Rights Reserved.
+#pragma once
+
+#include "CoreMinimal.h"
+#include "Materials/MaterialExpression.h"
+
+#include "MaterialExpressionRampLeftRight.generated.h"
+
+/**
+ * A material expression that computes a left-to-right bilinear value ramp.
+ */
+UCLASS(collapsecategories, hidecategories = Object, MinimalAPI)
+class UMaterialExpressionRampLeftRight : public UMaterialExpression
+{
+	GENERATED_UCLASS_BODY()
+
+	UPROPERTY(meta = (RequiredInput = "false", ToolTip = "Defaults to 'ConstCoordinate' if not specified"))
+	FExpressionInput Coordinates;
+
+	UPROPERTY()
+	FExpressionInput A;
+
+	UPROPERTY()
+	FExpressionInput B;
+
+	/** only used if Coordinates is not hooked up */
+	UPROPERTY(EditAnywhere, Category = MaterialExpressionTextureSample, meta = (OverridingInputProperty = "Coordinates"))
+	uint8 ConstCoordinate;
+
+	//~ Begin UMaterialExpression Interface
+#if WITH_EDITOR
+	virtual int32 Compile(class FMaterialCompiler* Compiler, int32 OutputIndex) override;
+	virtual void GetCaption(TArray<FString>& OutCaptions) const override;
+#endif
+	//~ End UMaterialExpression Interface
+};
+
diff --git a/Engine/Plugins/Experimental/Interchange/Runtime/Source/Import/Private/MaterialX/MaterialExpressionRampTopBottom.cpp b/Engine/Plugins/Experimental/Interchange/Runtime/Source/Import/Private/MaterialX/MaterialExpressionRampTopBottom.cpp
new file mode 100644
--- /dev/null
+++ b/Engine/Plugins/Experimental/Interchange/Runtime/Source/Import/Private/MaterialX/MaterialExpressionRampTopBottom.cpp
@@ -0,0 +1,49 @@
+// Copyright Epic Games, Inc. All Rights Reserved.
+#include "MaterialX/MaterialExpressionRampTopBottom.h"
+#include "MaterialCompiler.h"
+
+#define LOCTEXT_NAMESPACE "MaterialExpressionRampLeftRight"
+
+UMaterialExpressionRampTopBottom::UMaterialExpressionRampTopBottom(const FObjectInitializer& ObjectInitializer)
+	: Super(ObjectInitializer)
+{
+	// Structure to hold one-time initialization
+	struct FConstructorStatics
+	{
+		FText NAME_Utility;
+		FConstructorStatics()
+			: NAME_Utility(LOCTEXT("Utility", "Utility"))
+		{}
+	};
+	static FConstructorStatics ConstructorStatics;
+
+#if WITH_EDITORONLY_DATA
+	MenuCategories.Add(ConstructorStatics.NAME_Utility);
+#endif
+}
+
+#if WITH_EDITOR
+int32 UMaterialExpressionRampTopBottom::Compile(FMaterialCompiler* Compiler, int32 OutputIndex)
+{
+	if(!A.GetTracedInput().Expression)
+	{
+		return Compiler->Errorf(TEXT("Missing A input"));
+	}
+
+	if(!B.GetTracedInput().Expression)
+	{
+		return Compiler->Errorf(TEXT("Missing B input"));
+	}
+
+	int32 CoordinateIndex = Coordinates.GetTracedInput().Expression ? Coordinates.Compile(Compiler) : Compiler->TextureCoordinate(ConstCoordinate, false, false);
+
+	return Compiler->Lerp(A.Compile(Compiler), B.Compile(Compiler), Compiler->Saturate(Compiler->ComponentMask(CoordinateIndex, false, true, false, false)));
+}
+
+void UMaterialExpressionRampTopBottom::GetCaption(TArray<FString>& OutCaptions) const
+{
+	OutCaptions.Add(TEXT("RampTB"));
+}
+#endif
+
+#undef LOCTEXT_NAMESPACE 
\ No newline at end of file
diff --git a/Engine/Plugins/Experimental/Interchange/Runtime/Source/Import/Private/MaterialX/MaterialExpressionRampTopBottom.h b/Engine/Plugins/Experimental/Interchange/Runtime/Source/Import/Private/MaterialX/MaterialExpressionRampTopBottom.h
new file mode 100644
--- /dev/null
+++ b/Engine/Plugins/Experimental/Interchange/Runtime/Source/Import/Private/MaterialX/MaterialExpressionRampTopBottom.h
@@ -0,0 +1,37 @@
+// Copyright Epic Games, Inc. All Rights Reserved.
+#pragma once
+
+#include "CoreMinimal.h"
+#include "Materials/MaterialExpression.h"
+
+#include "MaterialExpressionRampTopBottom.generated.h"
+
+/**
+ * A material expression that computes a top-to-bottom bilinear value ramp.
+ */
+UCLASS(collapsecategories, hidecategories = Object, MinimalAPI)
+class UMaterialExpressionRampTopBottom : public UMaterialExpression
+{
+	GENERATED_UCLASS_BODY()
+
+		UPROPERTY(meta = (RequiredInput = "false", ToolTip = "Defaults to 'ConstCoordinate' if not specified"))
+		FExpressionInput Coordinates;
+
+	UPROPERTY()
+		FExpressionInput A;
+
+	UPROPERTY()
+		FExpressionInput B;
+
+	/** only used if Coordinates is not hooked up */
+	UPROPERTY(EditAnywhere, Category = MaterialExpressionTextureSample, meta = (OverridingInputProperty = "Coordinates"))
+		uint8 ConstCoordinate;
+
+	//~ Begin UMaterialExpression Interface
+#if WITH_EDITOR
+	virtual int32 Compile(class FMaterialCompiler* Compiler, int32 OutputIndex) override;
+	virtual void GetCaption(TArray<FString>& OutCaptions) const override;
+#endif
+	//~ End UMaterialExpression Interface
+};
+
diff --git a/Engine/Plugins/Experimental/Interchange/Runtime/Source/Import/Private/MaterialX/MaterialExpressionRemap.cpp b/Engine/Plugins/Experimental/Interchange/Runtime/Source/Import/Private/MaterialX/MaterialExpressionRemap.cpp
new file mode 100644
--- /dev/null
+++ b/Engine/Plugins/Experimental/Interchange/Runtime/Source/Import/Private/MaterialX/MaterialExpressionRemap.cpp
@@ -0,0 +1,58 @@
+// Copyright Epic Games, Inc. All Rights Reserved.
+#include "MaterialX/MaterialExpressionRemap.h"
+#include "MaterialCompiler.h"
+
+#define LOCTEXT_NAMESPACE "MaterialExpressionRemap"
+
+UMaterialExpressionRemap::UMaterialExpressionRemap(const FObjectInitializer& ObjectInitializer)
+	: Super(ObjectInitializer)
+{
+	// Structure to hold one-time initialization
+	struct FConstructorStatics
+	{
+		FText NAME_Math;
+		FConstructorStatics()
+			: NAME_Math(LOCTEXT("Math", "Math"))
+		{}
+	};
+	static FConstructorStatics ConstructorStatics;
+
+	InputLowDefault = 0.0f;
+	InputHighDefault = 1.0f;
+	TargetLowDefault = 0.0f;
+	TargetHighDefault = 1.0f;
+
+#if WITH_EDITORONLY_DATA
+	MenuCategories.Add(ConstructorStatics.NAME_Math);
+#endif
+}
+
+#if WITH_EDITOR
+int32 UMaterialExpressionRemap::Compile(FMaterialCompiler* Compiler, int32 OutputIndex)
+{
+	if(!Input.GetTracedInput().Expression)
+	{
+		return Compiler->Errorf(TEXT("Missing Remap input"));
+	}
+
+	// Min2 + (Input - Min1) x (Max2 - Min2) / (Max1 - Min1)
+	const int32 InputLowIndex = InputLow.GetTracedInput().Expression ? InputLow.Compile(Compiler) : Compiler->Constant(InputLowDefault);
+	const int32 InputHighIndex = InputHigh.GetTracedInput().Expression ? InputHigh.Compile(Compiler) : Compiler->Constant(InputHighDefault);
+
+	const int32 TargetLowIndex = TargetLow.GetTracedInput().Expression ? TargetLow.Compile(Compiler) : Compiler->Constant(TargetLowDefault);
+	const int32 TargetHighIndex = TargetHigh.GetTracedInput().Expression ? TargetHigh.Compile(Compiler) : Compiler->Constant(TargetHighDefault);
+
+	const int32 InputDelta = Compiler->Sub(InputHighIndex, InputLowIndex);
+	const int32 TargetDelta = Compiler->Sub(TargetHighIndex, TargetLowIndex);
+	const int32 Delta = Compiler->Div(TargetDelta, InputDelta);
+
+	return Compiler->Add(TargetLowIndex, Compiler->Mul(Compiler->Sub(Input.Compile(Compiler), InputLowIndex), Delta));
+}
+
+void UMaterialExpressionRemap::GetCaption(TArray<FString>& OutCaptions) const
+{
+	OutCaptions.Add(TEXT("Remap"));
+}
+#endif
+
+#undef LOCTEXT_NAMESPACE 
\ No newline at end of file
diff --git a/Engine/Plugins/Experimental/Interchange/Runtime/Source/Import/Private/MaterialX/MaterialExpressionRemap.h b/Engine/Plugins/Experimental/Interchange/Runtime/Source/Import/Private/MaterialX/MaterialExpressionRemap.h
new file mode 100644
--- /dev/null
+++ b/Engine/Plugins/Experimental/Interchange/Runtime/Source/Import/Private/MaterialX/MaterialExpressionRemap.h
@@ -0,0 +1,51 @@
+// Copyright Epic Games, Inc. All Rights Reserved.
+#pragma once
+
+#include "CoreMinimal.h"
+#include "Materials/MaterialExpression.h"
+
+#include "MaterialExpressionRemap.generated.h"
+
+/**
+ * A material expression that Remap a value from one range to another.
+ */
+UCLASS(collapsecategories, hidecategories = Object, MinimalAPI)
+class UMaterialExpressionRemap : public UMaterialExpression
+{
+	GENERATED_UCLASS_BODY()
+
+	UPROPERTY()
+	FExpressionInput Input;
+
+	UPROPERTY(meta = (RequiredInput = "false", ToolTip = "Defaults to 'MinDefault1' if not specified"))
+	FExpressionInput InputLow;
+
+	UPROPERTY(meta = (RequiredInput = "false", ToolTip = "Defaults to 'MaxDefault1' if not specified"))
+	FExpressionInput InputHigh;
+
+	UPROPERTY(meta = (RequiredInput = "false", ToolTip = "Defaults to 'MinDefault2' if not specified"))
+	FExpressionInput TargetLow;
+
+	UPROPERTY(meta = (RequiredInput = "false", ToolTip = "Defaults to 'MaxDefault1' if not specified"))
+	FExpressionInput TargetHigh;
+
+	UPROPERTY(EditAnywhere, Category = MaterialExpressionClamp, meta = (OverridingInputProperty = "InputLow"))
+	float InputLowDefault;
+
+	UPROPERTY(EditAnywhere, Category = MaterialExpressionClamp, meta = (OverridingInputProperty = "InputHigh"))
+	float InputHighDefault;
+
+	UPROPERTY(EditAnywhere, Category = MaterialExpressionClamp, meta = (OverridingInputProperty = "TargetLow"))
+	float TargetLowDefault;
+
+	UPROPERTY(EditAnywhere, Category = MaterialExpressionClamp, meta = (OverridingInputProperty = "TargetHigh"))
+	float TargetHighDefault;
+
+	//~ Begin UMaterialExpression Interface
+#if WITH_EDITOR
+	virtual int32 Compile(class FMaterialCompiler* Compiler, int32 OutputIndex) override;
+	virtual void GetCaption(TArray<FString>& OutCaptions) const override;
+#endif
+	//~ End UMaterialExpression Interface
+};
+
diff --git a/Engine/Plugins/Experimental/Interchange/Runtime/Source/Import/Private/MaterialX/MaterialExpressionRgbToHsv.h b/Engine/Plugins/Experimental/Interchange/Runtime/Source/Import/Private/MaterialX/MaterialExpressionRgbToHsv.h
new file mode 100644
--- /dev/null
+++ b/Engine/Plugins/Experimental/Interchange/Runtime/Source/Import/Private/MaterialX/MaterialExpressionRgbToHsv.h
@@ -0,0 +1,23 @@
+// Copyright Epic Games, Inc. All Rights Reserved.
+#pragma once
+
+#include "Materials/MaterialExpression.h"
+
+#include "MaterialExpressionRgbToHsv.generated.h"
+
+UCLASS(collapsecategories, hidecategories = Object, MinimalAPI)
+class UMaterialExpressionRgbToHsv : public UMaterialExpression
+{
+	GENERATED_UCLASS_BODY()
+
+	UPROPERTY()
+	FExpressionInput Input;
+
+	//~ Begin UMaterialExpression Interface
+#if WITH_EDITOR
+	virtual int32 Compile(class FMaterialCompiler* Compiler, int32 OutputIndex) override;
+	virtual void GetCaption(TArray<FString>& OutCaptions) const override;
+#endif
+	//~ End UMaterialExpression Interface
+};
+
diff --git a/Engine/Plugins/Experimental/Interchange/Runtime/Source/Import/Private/MaterialX/MaterialExpressionSplitLeftRight.cpp b/Engine/Plugins/Experimental/Interchange/Runtime/Source/Import/Private/MaterialX/MaterialExpressionSplitLeftRight.cpp
new file mode 100644
--- /dev/null
+++ b/Engine/Plugins/Experimental/Interchange/Runtime/Source/Import/Private/MaterialX/MaterialExpressionSplitLeftRight.cpp
@@ -0,0 +1,57 @@
+// Copyright Epic Games, Inc. All Rights Reserved.
+#include "MaterialX/MaterialExpressionSplitLeftRight.h"
+#include "MaterialCompiler.h"
+
+#define LOCTEXT_NAMESPACE "MaterialExpressionSplitLeftRight"
+
+UMaterialExpressionSplitLeftRight::UMaterialExpressionSplitLeftRight(const FObjectInitializer& ObjectInitializer)
+	: Super(ObjectInitializer),
+	ConstCenter(0.5f),
+	ConstCoordinate(0)
+{
+	// Structure to hold one-time initialization
+	struct FConstructorStatics
+	{
+		FText NAME_Utility;
+		FConstructorStatics()
+			: NAME_Utility(LOCTEXT("Utility", "Utility"))
+		{}
+	};
+	static FConstructorStatics ConstructorStatics;
+
+#if WITH_EDITORONLY_DATA
+	MenuCategories.Add(ConstructorStatics.NAME_Utility);
+#endif
+}
+
+#if WITH_EDITOR
+int32 UMaterialExpressionSplitLeftRight::Compile(FMaterialCompiler* Compiler, int32 OutputIndex)
+{
+	if(!A.GetTracedInput().Expression)
+	{
+		return Compiler->Errorf(TEXT("Missing A input"));
+	}
+
+	if(!B.GetTracedInput().Expression)
+	{
+		return Compiler->Errorf(TEXT("Missing B input"));
+	}
+
+	int32 CoordinateIndex = Coordinates.GetTracedInput().Expression ? Coordinates.Compile(Compiler) : Compiler->TextureCoordinate(ConstCoordinate, false, false);
+	int32 CenterIndex = Center.GetTracedInput().Expression ? Center.Compile(Compiler) : Compiler->Constant(ConstCenter);
+
+	int32 TexcoordXIndex = Compiler->ComponentMask(CoordinateIndex, true, false, false, false);
+	int32 InvSqrt2 = Compiler->Constant(0.70710678118654757);
+	int32 AFWidth = Compiler->Mul(Compiler->Length(Compiler->AppendVector(Compiler->DDX(TexcoordXIndex), Compiler->DDY(TexcoordXIndex))), InvSqrt2);
+	int32 AAStep = Compiler->SmoothStep(Compiler->Sub(CenterIndex, AFWidth), Compiler->Add(CenterIndex, AFWidth), TexcoordXIndex);
+
+	return Compiler->Lerp(A.Compile(Compiler), B.Compile(Compiler), AAStep);
+}
+
+void UMaterialExpressionSplitLeftRight::GetCaption(TArray<FString>& OutCaptions) const
+{
+	OutCaptions.Add(TEXT("SplitLR"));
+}
+#endif
+
+#undef LOCTEXT_NAMESPACE 
\ No newline at end of file
diff --git a/Engine/Plugins/Experimental/Interchange/Runtime/Source/Import/Private/MaterialX/MaterialExpressionSplitLeftRight.h b/Engine/Plugins/Experimental/Interchange/Runtime/Source/Import/Private/MaterialX/MaterialExpressionSplitLeftRight.h
new file mode 100644
--- /dev/null
+++ b/Engine/Plugins/Experimental/Interchange/Runtime/Source/Import/Private/MaterialX/MaterialExpressionSplitLeftRight.h
@@ -0,0 +1,44 @@
+// Copyright Epic Games, Inc. All Rights Reserved.
+#pragma once
+
+#include "CoreMinimal.h"
+#include "Materials/MaterialExpression.h"
+
+#include "MaterialExpressionSplitLeftRight.generated.h"
+
+/**
+ * A material expression that computes a left-right split matte, split at a specified u value.
+ */
+UCLASS(collapsecategories, hidecategories = Object, MinimalAPI)
+class UMaterialExpressionSplitLeftRight : public UMaterialExpression
+{
+	GENERATED_UCLASS_BODY()
+
+	UPROPERTY(meta = (RequiredInput = "false", ToolTip = "Defaults to 'ConstCoordinate' if not specified"))
+	FExpressionInput Coordinates;
+
+	UPROPERTY()
+	FExpressionInput A;
+
+	UPROPERTY()
+	FExpressionInput B;
+
+	UPROPERTY(meta = (RequiredInput = "false", ToolTip = "Defaults to 'ConstCenter' if not specified"))
+	FExpressionInput Center;
+
+	/** only used if A is not hooked up */
+	UPROPERTY(EditAnywhere, Category = MaterialExpressionMultiply, meta = (OverridingInputProperty = "Center"))
+	float ConstCenter;
+
+	/** only used if Coordinates is not hooked up */
+	UPROPERTY(EditAnywhere, Category = MaterialExpressionTextureSample, meta = (OverridingInputProperty = "Coordinates"))
+	uint8 ConstCoordinate;
+
+	//~ Begin UMaterialExpression Interface
+#if WITH_EDITOR
+	virtual int32 Compile(class FMaterialCompiler* Compiler, int32 OutputIndex) override;
+	virtual void GetCaption(TArray<FString>& OutCaptions) const override;
+#endif
+	//~ End UMaterialExpression Interface
+};
+
diff --git a/Engine/Plugins/Experimental/Interchange/Runtime/Source/Import/Private/MaterialX/MaterialExpressionSplitTopBottom.cpp b/Engine/Plugins/Experimental/Interchange/Runtime/Source/Import/Private/MaterialX/MaterialExpressionSplitTopBottom.cpp
new file mode 100644
--- /dev/null
+++ b/Engine/Plugins/Experimental/Interchange/Runtime/Source/Import/Private/MaterialX/MaterialExpressionSplitTopBottom.cpp
@@ -0,0 +1,57 @@
+// Copyright Epic Games, Inc. All Rights Reserved.
+#include "MaterialX/MaterialExpressionSplitTopBottom.h"
+#include "MaterialCompiler.h"
+
+#define LOCTEXT_NAMESPACE "MaterialExpressionSplitLeftRight"
+
+UMaterialExpressionSplitTopBottom::UMaterialExpressionSplitTopBottom(const FObjectInitializer& ObjectInitializer)
+	: Super(ObjectInitializer),
+	ConstCenter(0.5f),
+	ConstCoordinate(0)
+{
+	// Structure to hold one-time initialization
+	struct FConstructorStatics
+	{
+		FText NAME_Utility;
+		FConstructorStatics()
+			: NAME_Utility(LOCTEXT("Utility", "Utility"))
+		{}
+	};
+	static FConstructorStatics ConstructorStatics;
+
+#if WITH_EDITORONLY_DATA
+	MenuCategories.Add(ConstructorStatics.NAME_Utility);
+#endif
+}
+
+#if WITH_EDITOR
+int32 UMaterialExpressionSplitTopBottom::Compile(FMaterialCompiler* Compiler, int32 OutputIndex)
+{
+	if(!A.GetTracedInput().Expression)
+	{
+		return Compiler->Errorf(TEXT("Missing A input"));
+	}
+
+	if(!B.GetTracedInput().Expression)
+	{
+		return Compiler->Errorf(TEXT("Missing B input"));
+	}
+
+	int32 CoordinateIndex = Coordinates.GetTracedInput().Expression ? Coordinates.Compile(Compiler) : Compiler->TextureCoordinate(ConstCoordinate, false, false);
+	int32 CenterIndex = Center.GetTracedInput().Expression ? Center.Compile(Compiler) : Compiler->Constant(ConstCenter);
+
+	int32 TexcoordYIndex = Compiler->ComponentMask(CoordinateIndex, false, true, false, false);
+	int32 InvSqrt2 = Compiler->Constant(0.70710678118654757);
+	int32 AFWidth = Compiler->Mul(Compiler->Length(Compiler->AppendVector(Compiler->DDX(TexcoordYIndex), Compiler->DDY(TexcoordYIndex))), InvSqrt2);
+	int32 AAStep = Compiler->SmoothStep(Compiler->Sub(CenterIndex, AFWidth), Compiler->Add(CenterIndex, AFWidth), TexcoordYIndex);
+
+	return Compiler->Lerp(A.Compile(Compiler), B.Compile(Compiler), AAStep);
+}
+
+void UMaterialExpressionSplitTopBottom::GetCaption(TArray<FString>& OutCaptions) const
+{
+	OutCaptions.Add(TEXT("SplitTB"));
+}
+#endif
+
+#undef LOCTEXT_NAMESPACE 
\ No newline at end of file
diff --git a/Engine/Plugins/Experimental/Interchange/Runtime/Source/Import/Private/MaterialX/MaterialExpressionSplitTopBottom.h b/Engine/Plugins/Experimental/Interchange/Runtime/Source/Import/Private/MaterialX/MaterialExpressionSplitTopBottom.h
new file mode 100644
--- /dev/null
+++ b/Engine/Plugins/Experimental/Interchange/Runtime/Source/Import/Private/MaterialX/MaterialExpressionSplitTopBottom.h
@@ -0,0 +1,44 @@
+// Copyright Epic Games, Inc. All Rights Reserved.
+#pragma once
+
+#include "CoreMinimal.h"
+#include "Materials/MaterialExpression.h"
+
+#include "MaterialExpressionSplitTopBottom.generated.h"
+
+/**
+ * A material expression that computes a top-bottom split matte, split at a specified v value.
+ */
+UCLASS(collapsecategories, hidecategories = Object, MinimalAPI)
+class UMaterialExpressionSplitTopBottom : public UMaterialExpression
+{
+	GENERATED_UCLASS_BODY()
+
+	UPROPERTY(meta = (RequiredInput = "false", ToolTip = "Defaults to 'ConstCoordinate' if not specified"))
+	FExpressionInput Coordinates;
+
+	UPROPERTY()
+	FExpressionInput A;
+
+	UPROPERTY()
+	FExpressionInput B;
+
+	UPROPERTY(meta = (RequiredInput = "false", ToolTip = "Defaults to 'ConstCenter' if not specified"))
+	FExpressionInput Center;
+
+	/** only used if A is not hooked up */
+	UPROPERTY(EditAnywhere, Category = MaterialExpressionMultiply, meta = (OverridingInputProperty = "Center"))
+	float ConstCenter;
+
+	/** only used if Coordinates is not hooked up */
+	UPROPERTY(EditAnywhere, Category = MaterialExpressionTextureSample, meta = (OverridingInputProperty = "Coordinates"))
+	uint8 ConstCoordinate;
+
+	//~ Begin UMaterialExpression Interface
+#if WITH_EDITOR
+	virtual int32 Compile(class FMaterialCompiler* Compiler, int32 OutputIndex) override;
+	virtual void GetCaption(TArray<FString>& OutCaptions) const override;
+#endif
+	//~ End UMaterialExpression Interface
+};
+
diff --git a/Engine/Plugins/Experimental/Interchange/Runtime/Source/Import/Public/MaterialX/InterchangeMaterialXTranslator.h b/Engine/Plugins/Experimental/Interchange/Runtime/Source/Import/Public/MaterialX/InterchangeMaterialXTranslator.h
--- a/Engine/Plugins/Experimental/Interchange/Runtime/Source/Import/Public/MaterialX/InterchangeMaterialXTranslator.h
+++ b/Engine/Plugins/Experimental/Interchange/Runtime/Source/Import/Public/MaterialX/InterchangeMaterialXTranslator.h
@@ -332,6 +332,9 @@ protected:
 	void ConnectRotate3DInputToOutput(MaterialX::NodePtr UpstreamNode, UInterchangeShaderNode* ParentShaderNode, const FString& InputChannelName, TMap<FString, UInterchangeShaderNode*>& NamesToShaderNodes, UInterchangeBaseNodeContainer& NodeContainer) const;
 	void ConnectImageInputToOutput(MaterialX::NodePtr UpstreamNode, UInterchangeShaderNode* ParentShaderNode, const FString& InputChannelName, TMap<FString, UInterchangeShaderNode*>& NamesToShaderNodes, UInterchangeBaseNodeContainer& NodeContainer) const;
 	void ConnectConvertInputToOutput(MaterialX::NodePtr UpstreamNode, UInterchangeShaderNode* ParentShaderNode, const FString& InputChannelName, TMap<FString, UInterchangeShaderNode*>& NamesToShaderNodes, UInterchangeBaseNodeContainer& NodeContainer) const;
+	void ConnectIfGreaterInputToOutput(MaterialX::NodePtr UpstreamNode, UInterchangeShaderNode* ParentShaderNode, const FString& InputChannelName, TMap<FString, UInterchangeShaderNode*>& NamesToShaderNodes, UInterchangeBaseNodeContainer& NodeContainer) const;
+	void ConnectIfGreaterEqInputToOutput(MaterialX::NodePtr UpstreamNode, UInterchangeShaderNode* ParentShaderNode, const FString& InputChannelName, TMap<FString, UInterchangeShaderNode*>& NamesToShaderNodes, UInterchangeBaseNodeContainer& NodeContainer) const;
+	void ConnectIfEqualInputToOutput(MaterialX::NodePtr UpstreamNode, UInterchangeShaderNode* ParentShaderNode, const FString& InputChannelName, TMap<FString, UInterchangeShaderNode*>& NamesToShaderNodes, UInterchangeBaseNodeContainer& NodeContainer) const;
 
 	UInterchangeShaderNode* CreateMaskShaderNode(uint8 RGBA, const FString & NodeName, TMap<FString, UInterchangeShaderNode*>& NamesToShaderNodes, UInterchangeBaseNodeContainer& NodeContainer) const;
 
diff --git a/Engine/Plugins/Experimental/Interchange/Runtime/Source/Import/Public/MaterialX/MaterialExpressionAppend3Vector.h b/Engine/Plugins/Experimental/Interchange/Runtime/Source/Import/Public/MaterialX/MaterialExpressionAppend3Vector.h
deleted file mode 100644
diff --git a/Engine/Plugins/Experimental/Interchange/Runtime/Source/Import/Public/MaterialX/MaterialExpressionAppend4Vector.h b/Engine/Plugins/Experimental/Interchange/Runtime/Source/Import/Public/MaterialX/MaterialExpressionAppend4Vector.h
deleted file mode 100644
diff --git a/Engine/Plugins/Experimental/Interchange/Runtime/Source/Import/Public/MaterialX/MaterialExpressionExponential.h b/Engine/Plugins/Experimental/Interchange/Runtime/Source/Import/Public/MaterialX/MaterialExpressionExponential.h
deleted file mode 100644
diff --git a/Engine/Plugins/Experimental/Interchange/Runtime/Source/Import/Public/MaterialX/MaterialExpressionHsvToRgb.h b/Engine/Plugins/Experimental/Interchange/Runtime/Source/Import/Public/MaterialX/MaterialExpressionHsvToRgb.h
deleted file mode 100644
diff --git a/Engine/Plugins/Experimental/Interchange/Runtime/Source/Import/Public/MaterialX/MaterialExpressionLength.h b/Engine/Plugins/Experimental/Interchange/Runtime/Source/Import/Public/MaterialX/MaterialExpressionLength.h
deleted file mode 100644
diff --git a/Engine/Plugins/Experimental/Interchange/Runtime/Source/Import/Public/MaterialX/MaterialExpressionRgbToHsv.h b/Engine/Plugins/Experimental/Interchange/Runtime/Source/Import/Public/MaterialX/MaterialExpressionRgbToHsv.h
deleted file mode 100644
