Interchange Texture Factory
- Set the same compression and mipgen settings as the legacy factory for non-power of two textures.
- Set the IES brightness and texture multipler based on the values in the payload.

diff --git a/Engine/Plugins/Experimental/Interchange/Runtime/Source/Import/Private/Gltf/InterchangeGltfTranslator.cpp b/Engine/Plugins/Experimental/Interchange/Runtime/Source/Import/Private/Gltf/InterchangeGltfTranslator.cpp
--- a/Engine/Plugins/Experimental/Interchange/Runtime/Source/Import/Private/Gltf/InterchangeGltfTranslator.cpp
+++ b/Engine/Plugins/Experimental/Interchange/Runtime/Source/Import/Private/Gltf/InterchangeGltfTranslator.cpp
@@ -1018,7 +1018,8 @@ TOptional< UE::Interchange::FImportImage > UInterchangeGltfTranslator::GetTextur
 	{
 		// Embedded texture -- try using ImageWrapper to decode it
 		TArray64<uint8> ImageData(GltfTexture.Source.Data, GltfTexture.Source.DataByteLength);
-		return UInterchangeImageWrapperTranslator::GetTexturePayloadDataFromBuffer(ImageData);
+		UInterchangeImageWrapperTranslator* ImageWrapperTranslator = NewObject<UInterchangeImageWrapperTranslator>(GetTransientPackage(), NAME_None);
+		return ImageWrapperTranslator->GetTexturePayloadDataFromBuffer(ImageData);
 	}
 	else
 	{
diff --git a/Engine/Plugins/Experimental/Interchange/Runtime/Source/Import/Private/Texture/InterchangeImageWrapperTranslator.cpp b/Engine/Plugins/Experimental/Interchange/Runtime/Source/Import/Private/Texture/InterchangeImageWrapperTranslator.cpp
--- a/Engine/Plugins/Experimental/Interchange/Runtime/Source/Import/Private/Texture/InterchangeImageWrapperTranslator.cpp
+++ b/Engine/Plugins/Experimental/Interchange/Runtime/Source/Import/Private/Texture/InterchangeImageWrapperTranslator.cpp
@@ -10,6 +10,7 @@
 #include "InterchangeImportLog.h"
 #include "InterchangeTextureNode.h"
 #include "Memory/SharedBuffer.h"
+#include "Misc/ConfigCacheIni.h"
 #include "Misc/FileHelper.h"
 #include "Misc/Paths.h"
 #include "Misc/ScopedSlowTask.h"
@@ -63,6 +64,16 @@ static FAutoConsoleVariableRef CCvarInterchangeEnableTGAImport(
 	TEXT("Whether TGA support is enabled."),
 	ECVF_Default);
 
+UInterchangeImageWrapperTranslator::UInterchangeImageWrapperTranslator()
+	: bFillPNGZeroAlpha(true)
+{
+	if (GConfig)
+	{
+		ensure(IsInGameThread()); // Reading config values isn't threadsafe
+		GConfig->GetBool(TEXT("TextureImporter"), TEXT("FillPNGZeroAlpha"), bFillPNGZeroAlpha, GEditorIni);
+	}
+}
+
 TArray<FString> UInterchangeImageWrapperTranslator::GetSupportedFormats() const
 {
 	TArray<FString> Formats;
@@ -144,7 +155,7 @@ TOptional<UE::Interchange::FImportImage> UInterchangeImageWrapperTranslator::Get
 	return GetTexturePayloadDataFromBuffer(SourceDataBuffer);
 }
 
-TOptional<UE::Interchange::FImportImage> UInterchangeImageWrapperTranslator::GetTexturePayloadDataFromBuffer(const TArray64<uint8>& SourceDataBuffer)
+TOptional<UE::Interchange::FImportImage> UInterchangeImageWrapperTranslator::GetTexturePayloadDataFromBuffer(const TArray64<uint8>& SourceDataBuffer) const
 {
 	const uint8* Buffer = SourceDataBuffer.GetData();
 	const uint8* BufferEnd = Buffer + SourceDataBuffer.Num();
@@ -188,7 +199,16 @@ TOptional<UE::Interchange::FImportImage> UInterchangeImageWrapperTranslator::Get
 			}
 
 			// do per-format processing to match legacy behavior :
-			if (ImageFormat == EImageFormat::TGA)
+
+			if (ImageFormat == EImageFormat::PNG)
+			{
+				if (bFillPNGZeroAlpha)
+				{
+					// Replace the pixels with 0.0 alpha with a color value from the nearest neighboring color which has a non-zero alpha
+					UE::TextureUtilitiesCommon::FillZeroAlphaPNGData(PayloadData.SizeX, PayloadData.SizeY, PayloadData.Format, reinterpret_cast<uint8*>(PayloadData.RawData.GetData()));
+				}
+			}
+			else if (ImageFormat == EImageFormat::TGA)
 			{
 				const FTGAFileHeader* TGA = (FTGAFileHeader*)Buffer;
 
diff --git a/Engine/Plugins/Experimental/Interchange/Runtime/Source/Import/Private/Texture/InterchangeTextureFactory.cpp b/Engine/Plugins/Experimental/Interchange/Runtime/Source/Import/Private/Texture/InterchangeTextureFactory.cpp
--- a/Engine/Plugins/Experimental/Interchange/Runtime/Source/Import/Private/Texture/InterchangeTextureFactory.cpp
+++ b/Engine/Plugins/Experimental/Interchange/Runtime/Source/Import/Private/Texture/InterchangeTextureFactory.cpp
@@ -524,10 +524,10 @@ namespace UE::Interchange::Private::InterchangeTextureFactory
 		return {};
 	}
 
-	void SetupTextureSourceData(UTexture* Texture, const FImportImage& Image, UE::Serialization::FEditorBulkData::FSharedBufferWithID&& BufferAndId);
+	void SetupTextureSourceDataFromBulkData(UTexture* Texture, const FImportImage& Image, UE::Serialization::FEditorBulkData::FSharedBufferWithID&& BufferAndId);
 
 #if WITH_EDITOR
-	void SetupTextureSourceDataEditor(UTexture* Texture, const FImportImage& Image, UE::Serialization::FEditorBulkData::FSharedBufferWithID&& BufferAndId)
+	void SetupTextureSourceDataFromBulkData_Editor(UTexture* Texture, const FImportImage& Image, UE::Serialization::FEditorBulkData::FSharedBufferWithID&& BufferAndId)
 	{
 		Texture->Source.InitWithCompressedSourceData(
 			Image.SizeX,
@@ -549,7 +549,7 @@ namespace UE::Interchange::Private::InterchangeTextureFactory
 		}
 	}
 
-	void SetupTexture2DSourceDataEditor(UTexture2D* Texture2D, const FImportBlockedImage& BlockedImage, UE::Serialization::FEditorBulkData::FSharedBufferWithID&& BufferAndId)
+	void SetupTexture2DSourceDataFromBulkData_Editor(UTexture2D* Texture2D, const FImportBlockedImage& BlockedImage, UE::Serialization::FEditorBulkData::FSharedBufferWithID&& BufferAndId)
 	{
 		if (BlockedImage.BlocksData.Num() > 1)
 		{
@@ -585,11 +585,11 @@ namespace UE::Interchange::Private::InterchangeTextureFactory
 			Image.SizeY = Block.SizeY;
 			Image.NumMips = Block.NumMips;
 
-			SetupTextureSourceDataEditor(Texture2D, Image, MoveTemp(BufferAndId));
+			SetupTextureSourceDataFromBulkData_Editor(Texture2D, Image, MoveTemp(BufferAndId));
 		}
 	}
 
-	void SetupTextureSourceDataEditor(UTexture* Texture, const FImportSlicedImage& SlicedImage, UE::Serialization::FEditorBulkData::FSharedBufferWithID&& BufferAndId)
+	void SetupTextureSourceDataFromBulkData_Editor(UTexture* Texture, const FImportSlicedImage& SlicedImage, UE::Serialization::FEditorBulkData::FSharedBufferWithID&& BufferAndId)
 	{
 		Texture->Source.InitLayered(
 			SlicedImage.SizeX,
@@ -640,7 +640,7 @@ namespace UE::Interchange::Private::InterchangeTextureFactory
 		return false;
 	}
 
-	void SetupTextureCubeSourceData(UTextureCube* TextureCube, FProcessedPayload& ProcessedPayload)
+	void SetupTextureCubeSourceDataFromPayload(UTextureCube* TextureCube, FProcessedPayload& ProcessedPayload)
 	{
 		if (TOptional<FImportSlicedImage>* OptionalSlicedImage = ProcessedPayload.SettingsFromPayload.TryGet<TOptional<FImportSlicedImage>>())
 		{
@@ -651,7 +651,7 @@ namespace UE::Interchange::Private::InterchangeTextureFactory
 				// Cube texture always have six slice
 				if (SlicedImage.NumSlice == 6)
 				{
-					SetupTextureSourceDataEditor(TextureCube, SlicedImage, MoveTemp(ProcessedPayload.PayloadAndId));
+					SetupTextureSourceDataFromBulkData_Editor(TextureCube, SlicedImage, MoveTemp(ProcessedPayload.PayloadAndId));
 				}
 			}
 		}
@@ -660,7 +660,7 @@ namespace UE::Interchange::Private::InterchangeTextureFactory
 			if (OptionalImage->IsSet())
 			{
 				FImportImage& Image = OptionalImage->GetValue();
-				SetupTextureSourceData(TextureCube, Image, MoveTemp(ProcessedPayload.PayloadAndId));
+				SetupTextureSourceDataFromBulkData(TextureCube, Image, MoveTemp(ProcessedPayload.PayloadAndId));
 			}
 		}
 		else if (TOptional<FImportLightProfile>* OptionalLightProfile = ProcessedPayload.SettingsFromPayload.TryGet<TOptional<FImportLightProfile>>())
@@ -668,7 +668,7 @@ namespace UE::Interchange::Private::InterchangeTextureFactory
 			if (OptionalLightProfile->IsSet())
 			{
 				FImportLightProfile& LightProfile = OptionalLightProfile->GetValue();
-				SetupTextureSourceData(TextureCube, LightProfile, MoveTemp(ProcessedPayload.PayloadAndId));
+				SetupTextureSourceDataFromBulkData(TextureCube, LightProfile, MoveTemp(ProcessedPayload.PayloadAndId));
 			}
 		}
 		else
@@ -707,7 +707,7 @@ namespace UE::Interchange::Private::InterchangeTextureFactory
 		return false;
 	}
 
-	void SetupTexture2DArraySourceData(UTexture2DArray* Texture2DArray, FProcessedPayload& ProcessedPayload)
+	void SetupTexture2DArraySourceDataFromPayload(UTexture2DArray* Texture2DArray, FProcessedPayload& ProcessedPayload)
 	{
 		if (TOptional<FImportSlicedImage>* OptionalSlicedImage = ProcessedPayload.SettingsFromPayload.TryGet<TOptional<FImportSlicedImage>>())
 		{
@@ -715,7 +715,7 @@ namespace UE::Interchange::Private::InterchangeTextureFactory
 			{
 				FImportSlicedImage& SlicedImage = OptionalSlicedImage->GetValue();
 
-				SetupTextureSourceData(Texture2DArray, SlicedImage, MoveTemp(ProcessedPayload.PayloadAndId));
+				SetupTextureSourceDataFromBulkData(Texture2DArray, SlicedImage, MoveTemp(ProcessedPayload.PayloadAndId));
 			}
 		}
 		else if (TOptional<FImportImage>* OptionalImage = ProcessedPayload.SettingsFromPayload.TryGet<TOptional<FImportImage>>())
@@ -723,7 +723,7 @@ namespace UE::Interchange::Private::InterchangeTextureFactory
 			if (OptionalImage->IsSet())
 			{
 				FImportImage& Image = OptionalImage->GetValue();
-				SetupTextureSourceData(Texture2DArray, Image, MoveTemp(ProcessedPayload.PayloadAndId));
+				SetupTextureSourceDataFromBulkData(Texture2DArray, Image, MoveTemp(ProcessedPayload.PayloadAndId));
 			}
 		}
 		else if (TOptional<FImportLightProfile>* OptionalLightProfile = ProcessedPayload.SettingsFromPayload.TryGet<TOptional<FImportLightProfile>>())
@@ -731,7 +731,7 @@ namespace UE::Interchange::Private::InterchangeTextureFactory
 			if (OptionalLightProfile->IsSet())
 			{
 				FImportLightProfile& LightProfile = OptionalLightProfile->GetValue();
-				SetupTextureSourceData(Texture2DArray, LightProfile, MoveTemp(ProcessedPayload.PayloadAndId));
+				SetupTextureSourceDataFromBulkData(Texture2DArray, LightProfile, MoveTemp(ProcessedPayload.PayloadAndId));
 			}
 		}
 		else
@@ -846,7 +846,7 @@ namespace UE::Interchange::Private::InterchangeTextureFactory
 		return TasksToDo;
 	}
 #else // WITH_EDITOR
-	void SetupTexture2DSourceDataRuntime(UTexture2D* Texture2D, const FImportImage& ImportImage, UE::Serialization::FEditorBulkData::FSharedBufferWithID&& BufferAndId)
+	void SetupTexture2DSourceDataFromBulkData_Runtime(UTexture2D* Texture2D, const FImportImage& ImportImage, UE::Serialization::FEditorBulkData::FSharedBufferWithID&& BufferAndId)
 	{
 		Texture2D->SRGB = ImportImage.bSRGB;
 		ERawImageFormat::Type PixelFormatRawFormat;
@@ -1017,23 +1017,23 @@ namespace UE::Interchange::Private::InterchangeTextureFactory
 		return false;
 	}
 
-	void SetupTextureSourceData(UTexture* Texture, const FImportImage& Image, UE::Serialization::FEditorBulkData::FSharedBufferWithID&& BufferAndId)
+	void SetupTextureSourceDataFromBulkData(UTexture* Texture, const FImportImage& Image, UE::Serialization::FEditorBulkData::FSharedBufferWithID&& BufferAndId)
 	{
 #if WITH_EDITOR
-		SetupTextureSourceDataEditor(Texture, Image, MoveTemp(BufferAndId));
+		SetupTextureSourceDataFromBulkData_Editor(Texture, Image, MoveTemp(BufferAndId));
 #else
-		SetupTexture2DSourceDataRuntime(Cast<UTexture2D>(Texture), Image, MoveTemp(BufferAndId));
+		SetupTexture2DSourceDataFromBulkData_Runtime(Cast<UTexture2D>(Texture), Image, MoveTemp(BufferAndId));
 #endif
 	}
 
-	void SetupTexture2DSourceData(UTexture2D* Texture2D, FProcessedPayload& ProcessedPayload)
+	void SetupTexture2DSourceDataFromPayload(UTexture2D* Texture2D, FProcessedPayload& ProcessedPayload)
 	{
 #if WITH_EDITOR
 		if (TOptional<FImportBlockedImage>* BlockedImage = ProcessedPayload.SettingsFromPayload.TryGet<TOptional<FImportBlockedImage>>())
 		{
 			if (BlockedImage->IsSet())
 			{
-				SetupTexture2DSourceDataEditor(Texture2D, BlockedImage->GetValue(), MoveTemp(ProcessedPayload.PayloadAndId));
+				SetupTexture2DSourceDataFromBulkData_Editor(Texture2D, BlockedImage->GetValue(), MoveTemp(ProcessedPayload.PayloadAndId));
 			}
 		}
 		else
@@ -1042,7 +1042,7 @@ namespace UE::Interchange::Private::InterchangeTextureFactory
 		{
 			if (Image->IsSet())
 			{
-				SetupTextureSourceData(Texture2D, Image->GetValue(), MoveTemp(ProcessedPayload.PayloadAndId));
+				SetupTextureSourceDataFromBulkData(Texture2D, Image->GetValue(), MoveTemp(ProcessedPayload.PayloadAndId));
 			}
 		}
 		else if (TOptional<FImportLightProfile>* OptionalLightProfile = ProcessedPayload.SettingsFromPayload.TryGet<TOptional<FImportLightProfile>>())
@@ -1053,11 +1053,22 @@ namespace UE::Interchange::Private::InterchangeTextureFactory
 
 				if (UTextureLightProfile* TextureLightProfile = Cast<UTextureLightProfile>(Texture2D))
 				{
-					SetupTextureSourceData(TextureLightProfile, LightProfile, MoveTemp(ProcessedPayload.PayloadAndId));
+					SetupTextureSourceDataFromBulkData(TextureLightProfile, LightProfile, MoveTemp(ProcessedPayload.PayloadAndId));
+
+					if (const TOptional<FImportLightProfile>* ImportLightProfilePtr = ProcessedPayload.SettingsFromPayload.TryGet<TOptional<FImportLightProfile>>())
+					{
+						if (ImportLightProfilePtr->IsSet())
+						{
+							const UE::Interchange::FImportLightProfile& ImportLightProfile = ImportLightProfilePtr->GetValue();
+
+							TextureLightProfile->Brightness = ImportLightProfile.Brightness;
+							TextureLightProfile->TextureMultiplier = ImportLightProfile.TextureMultiplier;
+						}
+					}
 				}
 				else
 				{
-					SetupTextureSourceData(Texture2D, LightProfile, MoveTemp(ProcessedPayload.PayloadAndId));
+					SetupTextureSourceDataFromBulkData(Texture2D, LightProfile, MoveTemp(ProcessedPayload.PayloadAndId));
 				}
 			}
 		}
@@ -1094,8 +1105,59 @@ namespace UE::Interchange::Private::InterchangeTextureFactory
 		}
 #endif // WITH_EDITORONLY_DATA
 	}
- }
 
+	void SetupTextureSourceDataFromPayload(UTexture* Texture, FProcessedPayload& ProcessedPayload)
+	{
+		if (UTexture2D* Texture2D = Cast<UTexture2D>(Texture))
+		{
+			SetupTexture2DSourceDataFromPayload(Texture2D, ProcessedPayload);
+		}
+#if WITH_EDITOR
+		else if (UTextureCube* TextureCube = Cast<UTextureCube>(Texture))
+		{
+			SetupTextureCubeSourceDataFromPayload(TextureCube, ProcessedPayload);
+		}
+		else if (UTexture2DArray* Texture2DArray = Cast<UTexture2DArray>(Texture))
+		{
+			SetupTexture2DArraySourceDataFromPayload(Texture2DArray, ProcessedPayload);
+		}
+#endif // WITH_EDITOR
+		else
+		{
+			// This should never happen.
+			ensure(false);
+		}
+
+#if WITH_EDITORONLY_DATA
+		// At some point, these decisions should be left to the pipeline
+		if (!Texture->Source.IsPowerOfTwo())
+		{
+			// try to set some better default options for non-pow2 textures
+
+			// if Texture is not pow2 , change to TMGS_NoMipMaps (if it was default)
+			//   this used to be done by Texture2d.cpp ; it is now optional
+			//	 you can set it back to having mips if you want
+			if (Texture->MipGenSettings == TMGS_FromTextureGroup)
+			{
+				Texture->MipGenSettings = TMGS_NoMipmaps;
+			}
+
+			// if Texture is not multiple of 4, change TC to EditorIcon ("UserInterface2D")
+			//	if you do not do this, you might see "Texture forced to uncompressed because size is not a multiple of 4"
+			//  this needs to match the logic in Texture.cpp : GetDefaultTextureFormatName
+			const int32 SizeX = Texture->Source.GetSizeX();
+			const int32 SizeY = Texture->Source.GetSizeY();
+			if ((SizeX & 3) != 0 || (SizeY & 3) != 0)
+			{
+				if (Texture->CompressionSettings == TC_Default) // AutoDXT/BC1
+				{
+					Texture->CompressionSettings = TC_EditorIcon; // "UserInterface2D"
+				}
+			}
+		}
+#endif // WITH_EDITORONLY_DATA
+	}
+ }
 
 UClass* UInterchangeTextureFactory::GetFactoryClass() const
 {
@@ -1183,7 +1245,7 @@ UObject* UInterchangeTextureFactory::CreateAsset(const FCreateAssetParams& Argum
 {
 	TRACE_CPUPROFILER_EVENT_SCOPE(UInterchangeTextureFactory::CreateAsset);
 
-	using namespace  UE::Interchange::Private::InterchangeTextureFactory;
+	using namespace UE::Interchange::Private::InterchangeTextureFactory;
 
 	if (!Arguments.AssetNode)
 	{
@@ -1318,26 +1380,9 @@ void UInterchangeTextureFactory::PreImportPreCompletedCallback(const FImportPreC
 		using namespace UE::Interchange::Private::InterchangeTextureFactory;
 
 		// Setup source data
-		if (UTexture2D* Texture2D = Cast<UTexture2D>(Texture))
-		{
-			SetupTexture2DSourceData(Texture2D, ProcessedPayload);
-		}
-#if WITH_EDITOR
-		else if (UTextureCube* TextureCube = Cast<UTextureCube>(Texture))
-		{
-			SetupTextureCubeSourceData(TextureCube, ProcessedPayload);
-		}
-		else if (UTexture2DArray* Texture2DArray = Cast<UTexture2DArray>(Arguments.ImportedObject))
-		{
-			SetupTexture2DArraySourceData(Texture2DArray, ProcessedPayload);
-		}
-		else
-		{
-			// This should never happen.
-			ensure(false);
-		}
-
+		SetupTextureSourceDataFromPayload(Texture, ProcessedPayload);
 
+#if WITH_EDITOR
 		UInterchangeFactoryBaseNode* TextureFactoryNode = Arguments.FactoryNode;
 		if (!Arguments.bIsReimport)
 		{
@@ -1445,7 +1490,7 @@ void UInterchangeTextureFactory::CheckForInvalidResolutions(UE::Interchange::Pri
 
 			if (BlockedImage.BlocksData.IsEmpty())
 			{
-				AddErrorMessage(FString(), LOCTEXT("UDIMHasNoValidBlock", "All blocks where invalid. The texture won't be imported."));
+				AddErrorMessage(FString(), LOCTEXT("UDIMHasNoValidBlock", "All blocks were invalid. The texture won't be imported."));
 
 				// Remove the payload
 				InPayloadVariant.Set<FEmptyVariantState>(FEmptyVariantState());
diff --git a/Engine/Plugins/Experimental/Interchange/Runtime/Source/Import/Public/Texture/InterchangeImageWrapperTranslator.h b/Engine/Plugins/Experimental/Interchange/Runtime/Source/Import/Public/Texture/InterchangeImageWrapperTranslator.h
--- a/Engine/Plugins/Experimental/Interchange/Runtime/Source/Import/Public/Texture/InterchangeImageWrapperTranslator.h
+++ b/Engine/Plugins/Experimental/Interchange/Runtime/Source/Import/Public/Texture/InterchangeImageWrapperTranslator.h
@@ -20,6 +20,7 @@ class INTERCHANGEIMPORT_API UInterchangeImageWrapperTranslator : public UInterch
 {
 	GENERATED_BODY()
 public:
+	UInterchangeImageWrapperTranslator();
 
 	virtual TArray<FString> GetSupportedFormats() const override;
 
@@ -47,5 +48,8 @@ public:
 
 	/* IInterchangeTexturePayloadInterface End */
 
-	static TOptional<UE::Interchange::FImportImage> GetTexturePayloadDataFromBuffer(const TArray64<uint8>& SourceDataBuffer);
+	TOptional<UE::Interchange::FImportImage> GetTexturePayloadDataFromBuffer(const TArray64<uint8>& SourceDataBuffer) const;
+
+private:
+	bool bFillPNGZeroAlpha;
 };
\ No newline at end of file
diff --git a/Engine/Source/Editor/UnrealEd/Private/Factories/EditorFactories.cpp b/Engine/Source/Editor/UnrealEd/Private/Factories/EditorFactories.cpp
--- a/Engine/Source/Editor/UnrealEd/Private/Factories/EditorFactories.cpp
+++ b/Engine/Source/Editor/UnrealEd/Private/Factories/EditorFactories.cpp
@@ -2784,180 +2784,6 @@ void UTextureFactory::SuppressImportOverwriteDialog(bool bOverwriteExistingSetti
     bForceOverwriteExistingSettings = bOverwriteExistingSettings;
 }
 
-/**
- * This fills any pixels of a texture with have an alpha value of zero,
- * with an RGB from the nearest neighboring pixel which has non-zero alpha.
-*/
-template<typename PixelDataType, typename ColorDataType, int32 RIdx, int32 GIdx, int32 BIdx, int32 AIdx> class PNGDataFill
-{
-public:
-
-	PNGDataFill( int32 SizeX, int32 SizeY, uint8* SourceTextureData )
-		: SourceData( reinterpret_cast<PixelDataType*>(SourceTextureData) )
-		, TextureWidth(SizeX)
-		, TextureHeight(SizeY)
-	{
-	}
-
-	void ProcessData()
-	{
-		int64 NumZeroedTopRowsToProcess = 0;
-		int64 FillColorRow = -1;
-		for (int64 Y = 0; Y < TextureHeight; ++Y)
-		{
-			if (!ProcessHorizontalRow(Y))
-			{
-				if (FillColorRow != -1)
-				{
-					FillRowColorPixels(FillColorRow, Y);
-				}
-				else
-				{
-					NumZeroedTopRowsToProcess = Y;
-				}
-			}
-			else
-			{
-				FillColorRow = Y;
-			}
-		}
-
-		// Can only fill upwards if image not fully zeroed
-		if (NumZeroedTopRowsToProcess > 0 && NumZeroedTopRowsToProcess + 1 < TextureHeight)
-		{
-			for (int64 Y = 0; Y <= NumZeroedTopRowsToProcess; ++Y)
-			{
-				FillRowColorPixels(NumZeroedTopRowsToProcess + 1, Y);
-			}
-		}
-	}
-
-	/* returns False if requires further processing because entire row is filled with zeroed alpha values */
-	bool ProcessHorizontalRow(int64 Y)
-	{
-		// only wipe out colors that are affected by png turning valid colors white if alpha = 0
-		ColorDataType WhiteWithZeroAlpha;
-		if ( sizeof(ColorDataType) == 4 )
-		{
-			WhiteWithZeroAlpha = FColor(255, 255, 255, 0).DWColor();
-		}
-		else
-		{
-			check( sizeof(ColorDataType) == 8 );
-			uint16 RGBA[4] = { 0xFFFF,0xFFFF,0xFFFF, 0 };
-			check( sizeof(RGBA) == 8 );
-			memcpy(&WhiteWithZeroAlpha,RGBA, sizeof(ColorDataType));
-		}
-
-		// Left -> Right
-		int64 NumLeftmostZerosToProcess = 0;
-		const PixelDataType* FillColor = nullptr;
-		for (int64 X = 0; X < TextureWidth; ++X)
-		{
-			PixelDataType* PixelData = SourceData + (Y * TextureWidth + X) * 4;
-			ColorDataType* ColorData = reinterpret_cast<ColorDataType*>(PixelData);
-
-			if (*ColorData == WhiteWithZeroAlpha)
-			{
-				if (FillColor)
-				{
-					PixelData[RIdx] = FillColor[RIdx];
-					PixelData[GIdx] = FillColor[GIdx];
-					PixelData[BIdx] = FillColor[BIdx];
-				}
-				else
-				{
-					// Mark pixel as needing fill
-					*ColorData = 0;
-
-					// Keep track of how many pixels to fill starting at beginning of row
-					NumLeftmostZerosToProcess = X;
-				}
-			}
-			else
-			{
-				FillColor = PixelData;
-			}
-		}
-
-		if (NumLeftmostZerosToProcess == 0)
-		{
-			// No pixels left that are zero
-			return true;
-		}
-
-		if (NumLeftmostZerosToProcess + 1 >= TextureWidth)
-		{
-			// All pixels in this row are zero and must be filled using rows above or below
-			return false;
-		}
-
-		// Fill using non zero pixel immediately to the right of the beginning series of zeros
-		FillColor = SourceData + (Y * TextureWidth + NumLeftmostZerosToProcess + 1) * 4;
-
-		// Fill zero pixels found at beginning of row that could not be filled during the Left to Right pass
-		for (int64 X = 0; X <= NumLeftmostZerosToProcess; ++X)
-		{
-			PixelDataType* PixelData = SourceData + (Y * TextureWidth + X) * 4;
-			PixelData[RIdx] = FillColor[RIdx];
-			PixelData[GIdx] = FillColor[GIdx];
-			PixelData[BIdx] = FillColor[BIdx];
-		}
-
-		return true;
-	}
-
-	void FillRowColorPixels(int64 FillColorRow, int64 Y)
-	{
-		for (int64 X = 0; X < TextureWidth; ++X)
-		{
-			const PixelDataType* FillColor = SourceData + (FillColorRow * TextureWidth + X) * 4;
-			PixelDataType* PixelData = SourceData + (Y * TextureWidth + X) * 4;
-			PixelData[RIdx] = FillColor[RIdx];
-			PixelData[GIdx] = FillColor[GIdx];
-			PixelData[BIdx] = FillColor[BIdx];
-		}
-	}
-
-	PixelDataType* SourceData;
-	int64 TextureWidth;
-	int64 TextureHeight;
-};
-
-/**
- * For PNG texture importing, this ensures that any pixels with an alpha value of zero have an RGB
- * assigned to them from a neighboring pixel which has non-zero alpha.
- * This is needed as PNG exporters tend to turn pixels that are RGBA = (x,x,x,0) to (1,1,1,0)
- * and this produces artifacts when drawing the texture with bilinear filtering. 
- *
- * @param TextureSource - The source texture
- * @param SourceData - The source texture data
-*/
-void FillZeroAlphaPNGData( int32 SizeX, int32 SizeY, ETextureSourceFormat SourceFormat, uint8* SourceData )
-{
-	switch( SourceFormat )
-	{
-		case TSF_BGRA8:
-		{
-			PNGDataFill<uint8, uint32, 2, 1, 0, 3> PNGFill(SizeX, SizeY, SourceData );
-			PNGFill.ProcessData();
-			break;
-		}
-
-		case TSF_RGBA16:
-		{
-			PNGDataFill<uint16, uint64, 0, 1, 2, 3> PNGFill(SizeX, SizeY, SourceData );
-			PNGFill.ProcessData();
-			break;
-		}
-
-		default:
-		{
-			// G8, G16, no alpha to fill
-			break;
-		}
-	}
-}
 
 extern ENGINE_API bool GUseBilinearLightmaps;
 
@@ -3122,7 +2948,7 @@ bool UTextureFactory::ImportImage(const uint8* Buffer, int64 Length, FFeedbackCo
 				if (bFillPNGZeroAlpha)
 				{
 					// Replace the pixels with 0.0 alpha with a color value from the nearest neighboring color which has a non-zero alpha
-					FillZeroAlphaPNGData(OutImage.SizeX, OutImage.SizeY, OutImage.Format, OutImage.RawData.GetData());
+					UE::TextureUtilitiesCommon::FillZeroAlphaPNGData(OutImage.SizeX, OutImage.SizeY, OutImage.Format, OutImage.RawData.GetData());
 				}
 			}
 			else if ( ImageFormat == EImageFormat::TGA )
diff --git a/Engine/Source/Runtime/TextureUtilitiesCommon/Private/TextureImportUtils.cpp b/Engine/Source/Runtime/TextureUtilitiesCommon/Private/TextureImportUtils.cpp
--- a/Engine/Source/Runtime/TextureUtilitiesCommon/Private/TextureImportUtils.cpp
+++ b/Engine/Source/Runtime/TextureUtilitiesCommon/Private/TextureImportUtils.cpp
@@ -40,5 +40,172 @@ namespace UE
 
 			return true;
 		}
+
+		/**
+		 * This fills any pixels of a texture with have an alpha value of zero,
+		 * with an RGB from the nearest neighboring pixel which has non-zero alpha.
+		 */
+		template<typename PixelDataType, typename ColorDataType, int32 RIdx, int32 GIdx, int32 BIdx, int32 AIdx>
+		class TPNGDataFill
+		{
+		public:
+
+			explicit TPNGDataFill( int32 SizeX, int32 SizeY, uint8* SourceTextureData )
+				: SourceData( reinterpret_cast<PixelDataType*>(SourceTextureData) )
+				, TextureWidth(SizeX)
+				, TextureHeight(SizeY)
+			{
+			}
+
+			void ProcessData()
+			{
+				int64 NumZeroedTopRowsToProcess = 0;
+				int64 FillColorRow = -1;
+				for (int64 Y = 0; Y < TextureHeight; ++Y)
+				{
+					if (!ProcessHorizontalRow(Y))
+					{
+						if (FillColorRow != -1)
+						{
+							FillRowColorPixels(FillColorRow, Y);
+						}
+						else
+						{
+							NumZeroedTopRowsToProcess = Y;
+						}
+					}
+					else
+					{
+						FillColorRow = Y;
+					}
+				}
+
+				// Can only fill upwards if image not fully zeroed
+				if (NumZeroedTopRowsToProcess > 0 && NumZeroedTopRowsToProcess + 1 < TextureHeight)
+				{
+					for (int64 Y = 0; Y <= NumZeroedTopRowsToProcess; ++Y)
+					{
+						FillRowColorPixels(NumZeroedTopRowsToProcess + 1, Y);
+					}
+				}
+			}
+
+			/* returns False if requires further processing because entire row is filled with zeroed alpha values */
+			bool ProcessHorizontalRow(int64 Y)
+			{
+				// only wipe out colors that are affected by png turning valid colors white if alpha = 0
+				ColorDataType WhiteWithZeroAlpha;
+				if ( sizeof(ColorDataType) == 4 )
+				{
+					WhiteWithZeroAlpha = FColor(255, 255, 255, 0).DWColor();
+				}
+				else
+				{
+					check( sizeof(ColorDataType) == 8 );
+					uint16 RGBA[4] = { 0xFFFF,0xFFFF,0xFFFF, 0 };
+					check( sizeof(RGBA) == 8 );
+					memcpy(&WhiteWithZeroAlpha,RGBA, sizeof(ColorDataType));
+				}
+
+				// Left -> Right
+				int64 NumLeftmostZerosToProcess = 0;
+				const PixelDataType* FillColor = nullptr;
+				for (int64 X = 0; X < TextureWidth; ++X)
+				{
+					PixelDataType* PixelData = SourceData + (Y * TextureWidth + X) * 4;
+					ColorDataType* ColorData = reinterpret_cast<ColorDataType*>(PixelData);
+
+					if (*ColorData == WhiteWithZeroAlpha)
+					{
+						if (FillColor)
+						{
+							PixelData[RIdx] = FillColor[RIdx];
+							PixelData[GIdx] = FillColor[GIdx];
+							PixelData[BIdx] = FillColor[BIdx];
+						}
+						else
+						{
+							// Mark pixel as needing fill
+							*ColorData = 0;
+
+							// Keep track of how many pixels to fill starting at beginning of row
+							NumLeftmostZerosToProcess = X;
+						}
+					}
+					else
+					{
+						FillColor = PixelData;
+					}
+				}
+
+				if (NumLeftmostZerosToProcess == 0)
+				{
+					// No pixels left that are zero
+					return true;
+				}
+
+				if (NumLeftmostZerosToProcess + 1 >= TextureWidth)
+				{
+					// All pixels in this row are zero and must be filled using rows above or below
+					return false;
+				}
+
+				// Fill using non zero pixel immediately to the right of the beginning series of zeros
+				FillColor = SourceData + (Y * TextureWidth + NumLeftmostZerosToProcess + 1) * 4;
+
+				// Fill zero pixels found at beginning of row that could not be filled during the Left to Right pass
+				for (int64 X = 0; X <= NumLeftmostZerosToProcess; ++X)
+				{
+					PixelDataType* PixelData = SourceData + (Y * TextureWidth + X) * 4;
+					PixelData[RIdx] = FillColor[RIdx];
+					PixelData[GIdx] = FillColor[GIdx];
+					PixelData[BIdx] = FillColor[BIdx];
+				}
+
+				return true;
+			}
+
+			void FillRowColorPixels(int64 FillColorRow, int64 Y)
+			{
+				for (int64 X = 0; X < TextureWidth; ++X)
+				{
+					const PixelDataType* FillColor = SourceData + (FillColorRow * TextureWidth + X) * 4;
+					PixelDataType* PixelData = SourceData + (Y * TextureWidth + X) * 4;
+					PixelData[RIdx] = FillColor[RIdx];
+					PixelData[GIdx] = FillColor[GIdx];
+					PixelData[BIdx] = FillColor[BIdx];
+				}
+			}
+
+			PixelDataType* SourceData;
+			int64 TextureWidth;
+			int64 TextureHeight;
+		};
+
+		void FillZeroAlphaPNGData(int32 SizeX, int32 SizeY, ETextureSourceFormat SourceFormat, uint8* SourceData)
+		{
+			switch (SourceFormat)
+			{
+				case TSF_BGRA8:
+				{
+					TPNGDataFill<uint8, uint32, 2, 1, 0, 3> PNGFill(SizeX, SizeY, SourceData);
+					PNGFill.ProcessData();
+					break;
+				}
+
+				case TSF_RGBA16:
+				{
+					TPNGDataFill<uint16, uint64, 0, 1, 2, 3> PNGFill(SizeX, SizeY, SourceData);
+					PNGFill.ProcessData();
+					break;
+				}
+
+				default:
+				{
+					// G8, G16, no alpha to fill
+					break;
+				}
+			}
+		}
 	}
 }
\ No newline at end of file
diff --git a/Engine/Source/Runtime/TextureUtilitiesCommon/Public/TextureImportUtils.h b/Engine/Source/Runtime/TextureUtilitiesCommon/Public/TextureImportUtils.h
--- a/Engine/Source/Runtime/TextureUtilitiesCommon/Public/TextureImportUtils.h
+++ b/Engine/Source/Runtime/TextureUtilitiesCommon/Public/TextureImportUtils.h
@@ -3,15 +3,27 @@
 #pragma once
 
 #include "CoreTypes.h"
+#include "Engine/TextureDefines.h"
 
 struct FImage;
 
 namespace UE::TextureUtilitiesCommon
 {
 	/**
-		* Detect the existence of gray scale image in some formats and convert those to a gray scale equivalent image
-		* 
-		* @return true if the image was converted
-		*/
+	 * Detect the existence of gray scale image in some formats and convert those to a gray scale equivalent image
+	 * 
+	 * @return true if the image was converted
+	 */
 	TEXTUREUTILITIESCOMMON_API bool AutoDetectAndChangeGrayScale(FImage& Image);
+
+	/**
+	 * For PNG texture importing, this ensures that any pixels with an alpha value of zero have an RGB
+	 * assigned to them from a neighboring pixel which has non-zero alpha.
+	 * This is needed as PNG exporters tend to turn pixels that are RGBA = (x,x,x,0) to (1,1,1,0)
+	 * and this produces artifacts when drawing the texture with bilinear filtering. 
+	 *
+	 * @param TextureSource - The source texture
+	 * @param SourceData - The source texture data
+	 */
+	TEXTUREUTILITIESCOMMON_API void FillZeroAlphaPNGData(int32 SizeX, int32 SizeY, ETextureSourceFormat SourceFormat, uint8* SourceData);
 }
\ No newline at end of file
