Interchange:
Add some CPU trace for interchange task, fbx translator, skeletalmesh factory.
Fix issues with UVs count crashing when there is no UVs layer.
Add generic pipeline asset to always cook in the plugin game.ini config file.
Make sure we use only half of the worker thread to avoid deadlock because there is no worker thread available.

#rnx

diff --git a/Engine/Plugins/Experimental/Interchange/Runtime/Config/Game.ini b/Engine/Plugins/Experimental/Interchange/Runtime/Config/Game.ini
--- a/Engine/Plugins/Experimental/Interchange/Runtime/Config/Game.ini
+++ b/Engine/Plugins/Experimental/Interchange/Runtime/Config/Game.ini
@@ -2,4 +2,5 @@
 +DirectoriesToAlwaysCook=(Path="/Interchange/Functions")
 +DirectoriesToAlwaysCook=(Path="/Interchange/gltf")
 +DirectoriesToAlwaysCook=(Path="/Interchange/Materials")
++DirectoriesToAlwaysCook=(Path="/Interchange/Pipelines")
 +DirectoriesToAlwaysCook=(Path="/Interchange/Utilities")
diff --git a/Engine/Plugins/Experimental/Interchange/Runtime/Source/Dispatcher/Public/InterchangeCommands.h b/Engine/Plugins/Experimental/Interchange/Runtime/Source/Dispatcher/Public/InterchangeCommands.h
--- a/Engine/Plugins/Experimental/Interchange/Runtime/Source/Dispatcher/Public/InterchangeCommands.h
+++ b/Engine/Plugins/Experimental/Interchange/Runtime/Source/Dispatcher/Public/InterchangeCommands.h
@@ -17,7 +17,7 @@ namespace UE
 		{
 		public:
 			//Major version should be updated when there is an existing API that has been change.
-			static int32 GetMajor() { return 10; }
+			static int32 GetMajor() { return 11; }
 			//Minor version must be updated when there is an addition to the API.
 			static int32 GetMinor() { return 0; }
 			//Patch version should be update if there is some bug fixes in the private code.
diff --git a/Engine/Plugins/Experimental/Interchange/Runtime/Source/Import/Private/Fbx/InterchangeFbxTranslator.cpp b/Engine/Plugins/Experimental/Interchange/Runtime/Source/Import/Private/Fbx/InterchangeFbxTranslator.cpp
--- a/Engine/Plugins/Experimental/Interchange/Runtime/Source/Import/Private/Fbx/InterchangeFbxTranslator.cpp
+++ b/Engine/Plugins/Experimental/Interchange/Runtime/Source/Import/Private/Fbx/InterchangeFbxTranslator.cpp
@@ -195,6 +195,7 @@ TFuture<TOptional<UE::Interchange::FStaticMeshPayloadData>> UInterchangeFbxTrans
 	FString JsonCommand = CreateFetchPayloadFbxCommand(PayLoadKey);
 	const int32 CreatedTaskIndex = Dispatcher->AddTask(JsonCommand, FInterchangeDispatcherTaskCompleted::CreateLambda([this, Promise](const int32 TaskIndex)
 	{
+		TRACE_CPUPROFILER_EVENT_SCOPE("UInterchangeFbxTranslator::GetStaticMeshPayloadData::Dispatcher->AddTaskDone")
 		UE::Interchange::ETaskState TaskState;
 		FString JsonResult;
 		TArray<FString> JsonMessages;
@@ -277,6 +278,7 @@ TFuture<TOptional<UE::Interchange::FSkeletalMeshLodPayloadData>> UInterchangeFbx
 	const FString JsonCommand = CreateFetchPayloadFbxCommand(PayLoadKey);
 	const int32 CreatedTaskIndex = Dispatcher->AddTask(JsonCommand, FInterchangeDispatcherTaskCompleted::CreateLambda([this, Promise](const int32 TaskIndex)
 	{
+		TRACE_CPUPROFILER_EVENT_SCOPE("UInterchangeFbxTranslator::GetSkeletalMeshLodPayloadData::Dispatcher->AddTaskDone")
 		UE::Interchange::ETaskState TaskState;
 		FString JsonResult;
 		TArray<FString> JsonMessages;
@@ -355,6 +357,7 @@ TFuture<TOptional<UE::Interchange::FSkeletalMeshMorphTargetPayloadData>> UInterc
 	const FString JsonCommand = CreateFetchPayloadFbxCommand(PayLoadKey);
 	const int32 CreatedTaskIndex = Dispatcher->AddTask(JsonCommand, FInterchangeDispatcherTaskCompleted::CreateLambda([this, Promise](const int32 TaskIndex)
 	{
+		TRACE_CPUPROFILER_EVENT_SCOPE("UInterchangeFbxTranslator::GetSkeletalMeshMorphTargetPayloadData::Dispatcher->AddTaskDone")
 		UE::Interchange::ETaskState TaskState;
 		FString JsonResult;
 		TArray<FString> JsonMessages;
@@ -424,6 +427,7 @@ TFuture<TOptional<UE::Interchange::FAnimationCurvePayloadData>> UInterchangeFbxT
 	FString JsonCommand = CreateFetchPayloadFbxCommand(PayLoadKey);
 	const int32 CreatedTaskIndex = Dispatcher->AddTask(JsonCommand, FInterchangeDispatcherTaskCompleted::CreateLambda([this, Promise](const int32 TaskIndex)
 		{
+			TRACE_CPUPROFILER_EVENT_SCOPE("UInterchangeFbxTranslator::GetAnimationCurvePayloadData::Dispatcher->AddTaskDone")
 			UE::Interchange::ETaskState TaskState;
 			FString JsonResult;
 			TArray<FString> JsonMessages;
@@ -503,6 +507,7 @@ TFuture<TOptional<UE::Interchange::FAnimationStepCurvePayloadData>> UInterchange
 	FString JsonCommand = CreateFetchPayloadFbxCommand(PayLoadKey);
 	const int32 CreatedTaskIndex = Dispatcher->AddTask(JsonCommand, FInterchangeDispatcherTaskCompleted::CreateLambda([this, Promise](const int32 TaskIndex)
 		{
+			TRACE_CPUPROFILER_EVENT_SCOPE("UInterchangeFbxTranslator::GetAnimationStepCurvePayloadData::Dispatcher->AddTaskDone")
 			UE::Interchange::ETaskState TaskState;
 			FString JsonResult;
 			TArray<FString> JsonMessages;
@@ -576,6 +581,7 @@ TFuture<TOptional<UE::Interchange::FAnimationBakeTransformPayloadData>> UInterch
 	FString JsonCommand = CreateFetchAnimationBakeTransformPayloadFbxCommand(PayLoadKey, BakeFrequency, RangeStartSecond, RangeStopSecond);
 	const int32 CreatedTaskIndex = Dispatcher->AddTask(JsonCommand, FInterchangeDispatcherTaskCompleted::CreateLambda([this, Promise](const int32 TaskIndex)
 		{
+			TRACE_CPUPROFILER_EVENT_SCOPE("UInterchangeFbxTranslator::GetAnimationBakeTransformPayloadData::Dispatcher->AddTaskDone")
 			UE::Interchange::ETaskState TaskState;
 			FString JsonResult;
 			TArray<FString> JsonMessages;
diff --git a/Engine/Plugins/Experimental/Interchange/Runtime/Source/Import/Private/Mesh/InterchangeSkeletalMeshFactory.cpp b/Engine/Plugins/Experimental/Interchange/Runtime/Source/Import/Private/Mesh/InterchangeSkeletalMeshFactory.cpp
--- a/Engine/Plugins/Experimental/Interchange/Runtime/Source/Import/Private/Mesh/InterchangeSkeletalMeshFactory.cpp
+++ b/Engine/Plugins/Experimental/Interchange/Runtime/Source/Import/Private/Mesh/InterchangeSkeletalMeshFactory.cpp
@@ -65,6 +65,7 @@ namespace UE
 																 , const UInterchangeBaseNodeContainer* NodeContainer
 																 , FString AssetName)
 			{
+				TRACE_CPUPROFILER_EVENT_SCOPE("FillMorphTargetMeshDescriptionsPerMorphTargetName")
 				TArray<FString> MorphTargetUids;
 				MeshNodeContext.MeshNode->GetMorphTargetDependencies(MorphTargetUids);
 				TMap<FString, TFuture<TOptional<UE::Interchange::FSkeletalMeshMorphTargetPayloadData>>> TempMorphTargetMeshDescriptionsPerMorphTargetName;
@@ -130,6 +131,7 @@ namespace UE
 
 			void CopyMorphTargetsMeshDescriptionToSkeletalMeshImportData(const TMap<FString, TOptional<UE::Interchange::FSkeletalMeshMorphTargetPayloadData>>& LodMorphTargetMeshDescriptions, FSkeletalMeshImportData& DestinationSkeletalMeshImportData)
 			{
+				TRACE_CPUPROFILER_EVENT_SCOPE("CopyMorphTargetsMeshDescriptionToSkeletalMeshImportData")
 				const int32 OriginalMorphTargetCount = LodMorphTargetMeshDescriptions.Num();
 				TArray<FString> Keys;
 				int32 MorphTargetCount = 0;
@@ -247,6 +249,7 @@ namespace UE
 
 			const UInterchangeSceneNode* RecursiveFindJointByName(const UInterchangeBaseNodeContainer* NodeContainer, const FString& ParentJointNodeId, const FString& JointName)
 			{
+				TRACE_CPUPROFILER_EVENT_SCOPE("RecursiveFindJointByName")
 				if (const UInterchangeSceneNode* JointNode = Cast<UInterchangeSceneNode>(NodeContainer->GetNode(ParentJointNodeId)))
 				{
 					if (JointNode->GetDisplayLabel().Equals(JointName))
@@ -271,6 +274,7 @@ namespace UE
 				, const FString& RootJointNodeId
 				, const FTransform& MeshGlobalTransform)
 			{
+				TRACE_CPUPROFILER_EVENT_SCOPE("SkinVertexPositionToTimeZero")
 				FMeshDescription& MeshDescription = LodMeshPayload.LodMeshDescription;
 				const int32 VertexCount = MeshDescription.Vertices().Num();
 				const TArray<FString>& JointNames = LodMeshPayload.JointNames;
@@ -402,6 +406,7 @@ namespace UE
 																  , const UInterchangeBaseNodeContainer* NodeContainer
 																  , const FString& RootJointNodeId)
 			{
+				TRACE_CPUPROFILER_EVENT_SCOPE("RetrieveAllSkeletalMeshPayloadsAndFillImportData")
 				if (!SkeletalMeshTranslatorPayloadInterface)
 				{
 					return;
@@ -436,6 +441,7 @@ namespace UE
 				//Fill the lod mesh description using all combined mesh part
 				for (const FMeshNodeContext& MeshNodeContext : MeshReferences)
 				{
+					TRACE_CPUPROFILER_EVENT_SCOPE("RetrieveAllSkeletalMeshPayloadsAndFillImportData::GetPayload")
 					TOptional<UE::Interchange::FSkeletalMeshLodPayloadData> LodMeshPayload = LodMeshPayloadPerTranslatorPayloadKey.FindChecked(MeshNodeContext.TranslatorPayloadKey).Get();
 					if (!LodMeshPayload.IsSet())
 					{
@@ -446,8 +452,11 @@ namespace UE
 
 					FSkeletalMeshOperations::FSkeletalMeshAppendSettings SkeletalMeshAppendSettings;
 					SkeletalMeshAppendSettings.SourceVertexIDOffset = VertexOffset;
-					FElementIDRemappings ElementIDRemappings;
-					LodMeshPayload->LodMeshDescription.Compact(ElementIDRemappings);
+					{
+						TRACE_CPUPROFILER_EVENT_SCOPE("RetrieveAllSkeletalMeshPayloadsAndFillImportData::CompactPayload")
+						FElementIDRemappings ElementIDRemappings;
+						LodMeshPayload->LodMeshDescription.Compact(ElementIDRemappings);
+					}
 
 					const bool bIsRigidMesh = LodMeshPayload->JointNames.Num() <= 0 && MeshNodeContext.SceneNode;
 					if (bSkinControlPointToTimeZero && !bIsRigidMesh)
@@ -550,7 +559,7 @@ namespace UE
 
 			void ProcessImportMeshInfluences(const int32 WedgeCount, TArray<SkeletalMeshImportData::FRawBoneInfluence>& Influences)
 			{
-
+				TRACE_CPUPROFILER_EVENT_SCOPE("ProcessImportMeshInfluences")
 				// Sort influences by vertex index.
 				struct FCompareVertexIndex
 				{
@@ -747,6 +756,7 @@ namespace UE
 
 			void RemapSkeletalMeshVertexColorToImportData(const USkeletalMesh* SkeletalMesh, const int32 LODIndex, FSkeletalMeshImportData* SkelMeshImportData)
 			{
+				TRACE_CPUPROFILER_EVENT_SCOPE("RemapSkeletalMeshVertexColorToImportData")
 				//Make sure we have all the source data we need to do the remap
 				if (!SkeletalMesh->GetImportedModel() || !SkeletalMesh->GetImportedModel()->LODModels.IsValidIndex(LODIndex) || !SkeletalMesh->GetHasVertexColors())
 				{
@@ -859,6 +869,7 @@ UClass* UInterchangeSkeletalMeshFactory::GetFactoryClass() const
 
 UObject* UInterchangeSkeletalMeshFactory::CreateEmptyAsset(const FCreateAssetParams& Arguments)
 {
+	TRACE_CPUPROFILER_EVENT_SCOPE("UInterchangeSkeletalMeshFactory::CreateEmptyAsset")
 #if !WITH_EDITOR || !WITH_EDITORONLY_DATA
 
 	UE_LOG(LogInterchangeImport, Error, TEXT("Cannot import skeletalMesh asset in runtime, this is an editor only feature."));
@@ -907,13 +918,14 @@ UObject* UInterchangeSkeletalMeshFactory::CreateEmptyAsset(const FCreateAssetPar
 
 UObject* UInterchangeSkeletalMeshFactory::CreateAsset(const FCreateAssetParams& Arguments)
 {
+	TRACE_CPUPROFILER_EVENT_SCOPE("UInterchangeSkeletalMeshFactory::CreateAsset")
+
 #if !WITH_EDITOR || !WITH_EDITORONLY_DATA
 
 	UE_LOG(LogInterchangeImport, Error, TEXT("Cannot import skeletalMesh asset in runtime, this is an editor only feature."));
 	return nullptr;
 
 #else
-
 	if (!Arguments.AssetNode || !Arguments.AssetNode->GetObjectClass()->IsChildOf(GetFactoryClass()))
 	{
 		return nullptr;
@@ -1064,6 +1076,7 @@ UObject* UInterchangeSkeletalMeshFactory::CreateAsset(const FCreateAssetParams&
 
 	for (int32 LodIndex = 0; LodIndex < LodCount; ++LodIndex)
 	{
+		TRACE_CPUPROFILER_EVENT_SCOPE("UInterchangeSkeletalMeshFactory::CreateAsset_LOD")
 		ESkeletalMeshGeoImportVersions GeoImportVersion = ESkeletalMeshGeoImportVersions::LatestVersion;
 		ESkeletalMeshSkinningImportVersions SkinningImportVersion = ESkeletalMeshSkinningImportVersions::LatestVersion;
 		if (bIsReImport && SkeletalMesh->GetImportedModel() && SkeletalMesh->GetImportedModel()->LODModels.IsValidIndex(CurrentLodIndex))
@@ -1527,6 +1540,7 @@ UObject* UInterchangeSkeletalMeshFactory::CreateAsset(const FCreateAssetParams&
 /* This function is call in the completion task on the main thread, use it to call main thread post creation step for your assets*/
 void UInterchangeSkeletalMeshFactory::PreImportPreCompletedCallback(const FImportPreCompletedCallbackParams& Arguments)
 {
+	TRACE_CPUPROFILER_EVENT_SCOPE("UInterchangeSkeletalMeshFactory::PreImportPreCompletedCallback")
 	check(IsInGameThread());
 	Super::PreImportPreCompletedCallback(Arguments);
 
@@ -1604,6 +1618,7 @@ void UInterchangeSkeletalMeshFactory::PreImportPreCompletedCallback(const FImpor
 
 bool UInterchangeSkeletalMeshFactory::GetSourceFilenames(const UObject* Object, TArray<FString>& OutSourceFilenames) const
 {
+	TRACE_CPUPROFILER_EVENT_SCOPE("UInterchangeSkeletalMeshFactory::GetSourceFilenames")
 #if WITH_EDITORONLY_DATA
 	if (const USkeletalMesh* SkeletalMesh = Cast<USkeletalMesh>(Object))
 	{
@@ -1616,6 +1631,7 @@ bool UInterchangeSkeletalMeshFactory::GetSourceFilenames(const UObject* Object,
 
 bool UInterchangeSkeletalMeshFactory::SetSourceFilename(const UObject* Object, const FString& SourceFilename, int32 SourceIndex) const
 {
+	TRACE_CPUPROFILER_EVENT_SCOPE("UInterchangeSkeletalMeshFactory::SetSourceFilename")
 #if WITH_EDITORONLY_DATA
 	if (const USkeletalMesh* SkeletalMesh = Cast<USkeletalMesh>(Object))
 	{
@@ -1629,6 +1645,7 @@ bool UInterchangeSkeletalMeshFactory::SetSourceFilename(const UObject* Object, c
 
 bool UInterchangeSkeletalMeshFactory::SetReimportSourceIndex(const UObject* Object, int32 SourceIndex) const
 {
+	TRACE_CPUPROFILER_EVENT_SCOPE("UInterchangeSkeletalMeshFactory::SetReimportSourceIndex")
 #if WITH_EDITORONLY_DATA
 	if (const USkeletalMesh* SkeletalMesh = Cast<USkeletalMesh>(Object))
 	{
diff --git a/Engine/Plugins/Experimental/Interchange/Runtime/Source/Import/Private/Mesh/InterchangeStaticMeshFactory.cpp b/Engine/Plugins/Experimental/Interchange/Runtime/Source/Import/Private/Mesh/InterchangeStaticMeshFactory.cpp
--- a/Engine/Plugins/Experimental/Interchange/Runtime/Source/Import/Private/Mesh/InterchangeStaticMeshFactory.cpp
+++ b/Engine/Plugins/Experimental/Interchange/Runtime/Source/Import/Private/Mesh/InterchangeStaticMeshFactory.cpp
@@ -466,7 +466,7 @@ void UInterchangeStaticMeshFactory::SetupSourceModelsSettings(UStaticMesh& Stati
 		if (!bIsAReimport && LodMeshDescriptions.IsValidIndex(LodIndex))
 		{
 			FStaticMeshConstAttributes StaticMeshAttributes(LodMeshDescriptions[LodIndex]);
-			const int32 NumUVChannels = StaticMeshAttributes.GetVertexInstanceUVs().GetNumChannels();
+			const int32 NumUVChannels = StaticMeshAttributes.GetVertexInstanceUVs().IsValid() ? StaticMeshAttributes.GetVertexInstanceUVs().GetNumChannels() : 1;
 			const int32 FirstOpenUVChannel = NumUVChannels >= MAX_MESH_TEXTURE_COORDS_MD ? 1 : NumUVChannels;
 
 			SrcModel.BuildSettings.DstLightmapIndex = FirstOpenUVChannel;
diff --git a/Engine/Source/Developer/MeshBuilder/Private/StaticMeshBuilder.cpp b/Engine/Source/Developer/MeshBuilder/Private/StaticMeshBuilder.cpp
--- a/Engine/Source/Developer/MeshBuilder/Private/StaticMeshBuilder.cpp
+++ b/Engine/Source/Developer/MeshBuilder/Private/StaticMeshBuilder.cpp
@@ -154,7 +154,8 @@ static bool BuildNanite(
 		RemapVerts
 	);
 
-	const uint32 NumTextureCoord = MeshDescription.VertexInstanceAttributes().GetAttributesRef<FVector2f>( MeshAttribute::VertexInstance::TextureCoordinate ).GetNumChannels();
+	TVertexInstanceAttributesRef<FVector2f> VertexInstanceUVs = MeshDescription.VertexInstanceAttributes().GetAttributesRef<FVector2f>(MeshAttribute::VertexInstance::TextureCoordinate);
+	const uint32 NumTextureCoord = VertexInstanceUVs.IsValid() ? VertexInstanceUVs.GetNumChannels() : 0;
 
 	// Only the render data and vertex buffers will be used from now on unless we have more than one source models
 	// This will help with memory usage for Nanite Mesh by releasing memory before doing the build
@@ -566,7 +567,8 @@ bool FStaticMeshBuilder::Build(FStaticMeshRenderData& StaticMeshRenderData, USta
 			RemapVerts
 		);
 
-		const uint32 NumTextureCoord = MeshDescriptions[LodIndex].VertexInstanceAttributes().GetAttributesRef<FVector2f>( MeshAttribute::VertexInstance::TextureCoordinate ).GetNumChannels();
+		TVertexInstanceAttributesRef<FVector2f> VertexInstanceUVs = MeshDescriptions[LodIndex].VertexInstanceAttributes().GetAttributesRef<FVector2f>(MeshAttribute::VertexInstance::TextureCoordinate);
+		const uint32 NumTextureCoord = VertexInstanceUVs.IsValid() ? VertexInstanceUVs.GetNumChannels() : 0;
 
 		// Only the render data and vertex buffers will be used from now on unless we have more than one source models
 		// This will help with memory usage for Nanite Mesh by releasing memory before doing the build
@@ -742,7 +744,7 @@ void BuildVertexBuffer(
 	const bool bHasColors = VertexInstanceColors.IsValid();
 	const bool bIgnoreTangents = StaticMesh->NaniteSettings.bEnabled;
 
-	const uint32 NumTextureCoord = VertexInstanceUVs.GetNumChannels();
+	const uint32 NumTextureCoord = VertexInstanceUVs.IsValid() ? VertexInstanceUVs.GetNumChannels() : 0;
 	const FMatrix ScaleMatrix = FScaleMatrix(BuildSettings.BuildScale3D).Inverse().GetTransposed();
 
 	TMap<FPolygonGroupID, int32> PolygonGroupToSectionIndex;
diff --git a/Engine/Source/Runtime/Experimental/Interchange/Engine/Private/InterchangeManager.cpp b/Engine/Source/Runtime/Experimental/Interchange/Engine/Private/InterchangeManager.cpp
--- a/Engine/Source/Runtime/Experimental/Interchange/Engine/Private/InterchangeManager.cpp
+++ b/Engine/Source/Runtime/Experimental/Interchange/Engine/Private/InterchangeManager.cpp
@@ -578,6 +578,7 @@ bool UInterchangeManager::CanReimport(const UObject* Object, TArray<FString>& Ou
 
 void UInterchangeManager::StartQueuedTasks(bool bCancelAllTasks /*= false*/)
 {
+	TRACE_CPUPROFILER_EVENT_SCOPE("UInterchangeManager::StartQueuedTasks")
 	if (!ensure(IsInGameThread()))
 	{
 		//Do not crash but we will not start any queued tasks if we are not in the game thread
@@ -608,7 +609,11 @@ void UInterchangeManager::StartQueuedTasks(bool bCancelAllTasks /*= false*/)
 		}
 	};
 
-	while (!QueuedTasks.IsEmpty() && (ImportTasks.Num() < FTaskGraphInterface::Get().GetNumWorkerThreads() || bCancelAllTasks))
+	//We need to leave some free task in the pool to avoid deadlock.
+	//Each import can use 2 tasks in same time if the build of the asset ddc use the same task pool (i.e. staticmesh, skeletalmesh, texture...)
+	const int32 PoolWorkerThreadCount = FTaskGraphInterface::Get().GetNumWorkerThreads()/2;
+	const int32 MaxNumWorker = FMath::Max(PoolWorkerThreadCount, 1);
+	while (!QueuedTasks.IsEmpty() && (ImportTasks.Num() < MaxNumWorker || bCancelAllTasks))
 	{
 		FQueuedTaskData QueuedTaskData;
 		if (QueuedTasks.Dequeue(QueuedTaskData))
@@ -700,6 +705,7 @@ UInterchangeManager::ImportSceneAsync(const FString& ContentPath, const UInterch
 TTuple<UE::Interchange::FAssetImportResultRef, UE::Interchange::FSceneImportResultRef>
 UInterchangeManager::ImportInternal(const FString& ContentPath, const UInterchangeSourceData* SourceData, const FImportAssetParameters& ImportAssetParameters, const UE::Interchange::EImportType ImportType)
 {
+	TRACE_CPUPROFILER_EVENT_SCOPE("UInterchangeManager::ImportInternal")
 	if (!ensure(IsInGameThread()))
 	{
 		//Import process can be started only in the game thread
diff --git a/Engine/Source/Runtime/Experimental/Interchange/Engine/Private/Tasks/InterchangeTaskCompletion.cpp b/Engine/Source/Runtime/Experimental/Interchange/Engine/Private/Tasks/InterchangeTaskCompletion.cpp
--- a/Engine/Source/Runtime/Experimental/Interchange/Engine/Private/Tasks/InterchangeTaskCompletion.cpp
+++ b/Engine/Source/Runtime/Experimental/Interchange/Engine/Private/Tasks/InterchangeTaskCompletion.cpp
@@ -21,6 +21,7 @@
 
 void UE::Interchange::FTaskPreAsyncCompletion::DoTask(ENamedThreads::Type CurrentThread, const FGraphEventRef& MyCompletionGraphEvent)
 {
+	TRACE_CPUPROFILER_EVENT_SCOPE("UE::Interchange::FTaskPreAsyncCompletion::DoTask")
 #if INTERCHANGE_TRACE_ASYNCHRONOUS_TASK_ENABLED
 	INTERCHANGE_TRACE_ASYNCHRONOUS_TASK(PreAsyncCompletion)
 #endif
@@ -35,6 +36,7 @@ void UE::Interchange::FTaskPreAsyncCompletion::DoTask(ENamedThreads::Type Curren
 
 void UE::Interchange::FTaskPreCompletion::DoTask(ENamedThreads::Type CurrentThread, const FGraphEventRef& MyCompletionGraphEvent)
 {
+	TRACE_CPUPROFILER_EVENT_SCOPE("UE::Interchange::FTaskPreCompletion::DoTask")
 #if INTERCHANGE_TRACE_ASYNCHRONOUS_TASK_ENABLED
 	INTERCHANGE_TRACE_ASYNCHRONOUS_TASK(PreCompletion)
 #endif
@@ -160,6 +162,7 @@ void UE::Interchange::FTaskPreCompletion::DoTask(ENamedThreads::Type CurrentThre
 
 void UE::Interchange::FTaskCompletion::DoTask(ENamedThreads::Type CurrentThread, const FGraphEventRef& MyCompletionGraphEvent)
 {
+	TRACE_CPUPROFILER_EVENT_SCOPE("UE::Interchange::FTaskCompletion::DoTask")
 #if INTERCHANGE_TRACE_ASYNCHRONOUS_TASK_ENABLED
 	INTERCHANGE_TRACE_ASYNCHRONOUS_TASK(Completion)
 #endif
diff --git a/Engine/Source/Runtime/Experimental/Interchange/Engine/Private/Tasks/InterchangeTaskCreateAsset.cpp b/Engine/Source/Runtime/Experimental/Interchange/Engine/Private/Tasks/InterchangeTaskCreateAsset.cpp
--- a/Engine/Source/Runtime/Experimental/Interchange/Engine/Private/Tasks/InterchangeTaskCreateAsset.cpp
+++ b/Engine/Source/Runtime/Experimental/Interchange/Engine/Private/Tasks/InterchangeTaskCreateAsset.cpp
@@ -27,6 +27,7 @@ namespace UE
 		{
 			void InternalGetPackageName(const UE::Interchange::FImportAsyncHelper& AsyncHelper, const int32 SourceIndex, const FString& PackageBasePath, const UInterchangeFactoryBaseNode* FactoryNode, FString& OutPackageName, FString& OutAssetName)
 			{
+				TRACE_CPUPROFILER_EVENT_SCOPE("UE::Interchange::Private::InternalGetPackageName")
 				const UInterchangeSourceData* SourceData = AsyncHelper.SourceDatas[SourceIndex];
 				check(SourceData);
 				FString NodeDisplayName = FactoryNode->GetAssetName();
@@ -49,6 +50,7 @@ namespace UE
 
 			UObject* GetExistingObjectFromAssetImportData(TSharedPtr<UE::Interchange::FImportAsyncHelper, ESPMode::ThreadSafe> AsyncHelper, UInterchangeFactoryBaseNode* FactoryNode)
 			{
+				TRACE_CPUPROFILER_EVENT_SCOPE("UE::Interchange::Private::GetExistingObjectFromAssetImportData")
 				UInterchangeAssetImportData* OriginalAssetImportData = nullptr;
 				TArray<UObject*> SubObjects;
 				GetObjectsWithOuter(AsyncHelper->TaskData.ReimportObject, SubObjects);
@@ -94,6 +96,7 @@ namespace UE
 
 void UE::Interchange::FTaskCreatePackage::DoTask(ENamedThreads::Type CurrentThread, const FGraphEventRef& MyCompletionGraphEvent)
 {
+	TRACE_CPUPROFILER_EVENT_SCOPE("UE::Interchange::FTaskCreatePackage::DoTask")
 #if INTERCHANGE_TRACE_ASYNCHRONOUS_TASK_ENABLED
 	INTERCHANGE_TRACE_ASYNCHRONOUS_TASK(CreatePackage)
 #endif
diff --git a/Engine/Source/Runtime/Experimental/Interchange/Engine/Private/Tasks/InterchangeTaskCreateSceneObjects.cpp b/Engine/Source/Runtime/Experimental/Interchange/Engine/Private/Tasks/InterchangeTaskCreateSceneObjects.cpp
--- a/Engine/Source/Runtime/Experimental/Interchange/Engine/Private/Tasks/InterchangeTaskCreateSceneObjects.cpp
+++ b/Engine/Source/Runtime/Experimental/Interchange/Engine/Private/Tasks/InterchangeTaskCreateSceneObjects.cpp
@@ -33,6 +33,7 @@ UE::Interchange::FTaskCreateSceneObjects::FTaskCreateSceneObjects(const FString&
 
 void UE::Interchange::FTaskCreateSceneObjects::DoTask(ENamedThreads::Type CurrentThread, const FGraphEventRef& MyCompletionGraphEvent)
 {
+	TRACE_CPUPROFILER_EVENT_SCOPE("UE::Interchange::FTaskCreateSceneObjects::DoTask")
 #if INTERCHANGE_TRACE_ASYNCHRONOUS_TASK_ENABLED
 	INTERCHANGE_TRACE_ASYNCHRONOUS_TASK(SpawnActor)
 #endif
diff --git a/Engine/Source/Runtime/Experimental/Interchange/Engine/Private/Tasks/InterchangeTaskParsing.cpp b/Engine/Source/Runtime/Experimental/Interchange/Engine/Private/Tasks/InterchangeTaskParsing.cpp
--- a/Engine/Source/Runtime/Experimental/Interchange/Engine/Private/Tasks/InterchangeTaskParsing.cpp
+++ b/Engine/Source/Runtime/Experimental/Interchange/Engine/Private/Tasks/InterchangeTaskParsing.cpp
@@ -35,6 +35,7 @@ struct FNodeDependencyCache
 {
 	const TSet<FString>& GetAccumulatedDependencies(const UInterchangeBaseNodeContainer* NodeContainer, const FString& NodeID)
 	{
+		TRACE_CPUPROFILER_EVENT_SCOPE("FNodeDependencyCache::GetAccumulatedDependencies")
 		TSet<FString> NodeStack;
 		return GetAccumulatedDependencies(NodeContainer, NodeID, NodeStack);
 	}
@@ -88,6 +89,7 @@ private:
 
 void UE::Interchange::FTaskParsing::DoTask(ENamedThreads::Type CurrentThread, const FGraphEventRef& MyCompletionGraphEvent)
 {
+	TRACE_CPUPROFILER_EVENT_SCOPE("UE::Interchange::FTaskParsing::DoTask")
 #if INTERCHANGE_TRACE_ASYNCHRONOUS_TASK_ENABLED
 	INTERCHANGE_TRACE_ASYNCHRONOUS_TASK(ParsingGraph)
 #endif
diff --git a/Engine/Source/Runtime/Experimental/Interchange/Engine/Private/Tasks/InterchangeTaskPipeline.cpp b/Engine/Source/Runtime/Experimental/Interchange/Engine/Private/Tasks/InterchangeTaskPipeline.cpp
--- a/Engine/Source/Runtime/Experimental/Interchange/Engine/Private/Tasks/InterchangeTaskPipeline.cpp
+++ b/Engine/Source/Runtime/Experimental/Interchange/Engine/Private/Tasks/InterchangeTaskPipeline.cpp
@@ -17,6 +17,7 @@
 
 void UE::Interchange::FTaskPipelinePreImport::DoTask(ENamedThreads::Type CurrentThread, const FGraphEventRef& MyCompletionGraphEvent)
 {
+	TRACE_CPUPROFILER_EVENT_SCOPE("UE::Interchange::FTaskPipelinePreImport::DoTask")
 #if INTERCHANGE_TRACE_ASYNCHRONOUS_TASK_ENABLED
 	INTERCHANGE_TRACE_ASYNCHRONOUS_TASK(PipelinePreImport)
 #endif
@@ -52,6 +53,7 @@ void UE::Interchange::FTaskPipelinePreImport::DoTask(ENamedThreads::Type Current
 
 void UE::Interchange::FTaskPipelinePostImport::DoTask(ENamedThreads::Type CurrentThread, const FGraphEventRef& MyCompletionGraphEvent)
 {
+	TRACE_CPUPROFILER_EVENT_SCOPE("UE::Interchange::FTaskPipelinePostImport::DoTask")
 #if INTERCHANGE_TRACE_ASYNCHRONOUS_TASK_ENABLED
 	INTERCHANGE_TRACE_ASYNCHRONOUS_TASK(PipelinePostImport)
 #endif
diff --git a/Engine/Source/Runtime/Experimental/Interchange/Engine/Private/Tasks/InterchangeTaskTranslator.cpp b/Engine/Source/Runtime/Experimental/Interchange/Engine/Private/Tasks/InterchangeTaskTranslator.cpp
--- a/Engine/Source/Runtime/Experimental/Interchange/Engine/Private/Tasks/InterchangeTaskTranslator.cpp
+++ b/Engine/Source/Runtime/Experimental/Interchange/Engine/Private/Tasks/InterchangeTaskTranslator.cpp
@@ -15,6 +15,7 @@
 
 void UE::Interchange::FTaskTranslator::DoTask(ENamedThreads::Type CurrentThread, const FGraphEventRef& MyCompletionGraphEvent)
 {
+	TRACE_CPUPROFILER_EVENT_SCOPE("UE::Interchange::FTaskTranslator::DoTask")
 #if INTERCHANGE_TRACE_ASYNCHRONOUS_TASK_ENABLED
 	INTERCHANGE_TRACE_ASYNCHRONOUS_TASK(Translator)
 #endif
