Interchange Multi animation import support added (whilst keeping original usage of UInterchangeAnimationAPI),
Introduced UInterchangeSkeletonAnimationTrackNode,
Solutions allows the use of one UInterchangeSkeletonFactoryNode in multiple UInterchangeAnimSequenceFactoryNode,
Updated Interchange GLTF Importer to support multi animation multi skeletons

diff --git a/Engine/Plugins/Editor/GLTFImporter/Source/GLTFCore/Private/GLTF/GLTFReader.cpp b/Engine/Plugins/Editor/GLTFImporter/Source/GLTFCore/Private/GLTF/GLTFReader.cpp
--- a/Engine/Plugins/Editor/GLTFImporter/Source/GLTFCore/Private/GLTF/GLTFReader.cpp
+++ b/Engine/Plugins/Editor/GLTFImporter/Source/GLTFCore/Private/GLTF/GLTFReader.cpp
@@ -277,7 +277,7 @@ namespace GLTF
 				const int32 NodeIndex = Value->AsNumber();
 				Scene.Nodes.Add(NodeIndex);
 
-				BuildParentIndices(-1, NodeIndex);
+				BuildParentIndices(INDEX_NONE, NodeIndex);
 			}
 		}
 
@@ -876,6 +876,7 @@ namespace GLTF
 
 		SetupNodesType();
 		ExtensionsHandler->SetupAssetExtensions(*JsonRoot);
+		BuildRootJoints();
 	}
 
 	template <typename SetupFunc>
@@ -945,4 +946,27 @@ namespace GLTF
 		}
 	}
 
+	int32 FFileReader::FindRootJointIndex(int32 CurrentIndex) const
+	{
+		if (!ensure(Asset->Nodes.IsValidIndex(CurrentIndex)))
+		{
+			return INDEX_NONE;
+		}
+		while (Asset->Nodes.IsValidIndex(Asset->Nodes[CurrentIndex].ParentIndex) && Asset->Nodes[Asset->Nodes[CurrentIndex].ParentIndex].Type == GLTF::FNode::EType::Joint)
+		{
+			CurrentIndex = Asset->Nodes[CurrentIndex].ParentIndex;
+		}
+		return CurrentIndex;
+	}
+	void FFileReader::BuildRootJoints() const
+	{
+		for (size_t Index = 0; Index < Asset->Nodes.Num(); Index++)
+		{
+			if (Asset->Nodes[Index].Type == GLTF::FNode::EType::Joint)
+			{
+				Asset->Nodes[Index].RootJointIndex = FindRootJointIndex(Index);
+			}
+		}
+	}
+
 }  // namespace GLTF
diff --git a/Engine/Plugins/Editor/GLTFImporter/Source/GLTFCore/Public/GLTFNode.h b/Engine/Plugins/Editor/GLTFImporter/Source/GLTFCore/Public/GLTFNode.h
--- a/Engine/Plugins/Editor/GLTFImporter/Source/GLTFCore/Public/GLTFNode.h
+++ b/Engine/Plugins/Editor/GLTFImporter/Source/GLTFCore/Public/GLTFNode.h
@@ -38,6 +38,7 @@ namespace GLTF
 		int32 LightIndex;   // index into FAsset::Lights
 
 		int32 ParentIndex;
+		int32 RootJointIndex; //only valid if node is of Joint Type
 
 		FNode()
 		    : Type(EType::None)
@@ -46,6 +47,7 @@ namespace GLTF
 		    , CameraIndex(INDEX_NONE)
 		    , LightIndex(INDEX_NONE)
 			, ParentIndex(INDEX_NONE)
+			, RootJointIndex(INDEX_NONE)
 		{
 		}
 	};
diff --git a/Engine/Plugins/Editor/GLTFImporter/Source/GLTFCore/Public/GLTFReader.h b/Engine/Plugins/Editor/GLTFImporter/Source/GLTFCore/Public/GLTFReader.h
--- a/Engine/Plugins/Editor/GLTFImporter/Source/GLTFCore/Public/GLTFReader.h
+++ b/Engine/Plugins/Editor/GLTFImporter/Source/GLTFCore/Public/GLTFReader.h
@@ -61,6 +61,8 @@ namespace GLTF
 		void SetupNodesType() const;
 
 		void BuildParentIndices(int32 ParentNodeIndex, int32 CurrentNodeIndex) const;
+		int32 FindRootJointIndex(int32 CurrentNodeIndex) const;
+		void BuildRootJoints() const;
 
 	private:
 		uint32 BufferCount;
diff --git a/Engine/Plugins/Experimental/Interchange/Runtime/Source/FactoryNodes/Private/InterchangeAnimSequenceFactoryNode.cpp b/Engine/Plugins/Experimental/Interchange/Runtime/Source/FactoryNodes/Private/InterchangeAnimSequenceFactoryNode.cpp
--- a/Engine/Plugins/Experimental/Interchange/Runtime/Source/FactoryNodes/Private/InterchangeAnimSequenceFactoryNode.cpp
+++ b/Engine/Plugins/Experimental/Interchange/Runtime/Source/FactoryNodes/Private/InterchangeAnimSequenceFactoryNode.cpp
@@ -43,6 +43,18 @@ namespace UE::Interchange
 		static FAttributeKey AttributeKey(TEXT("__AnimatedMaterialCurveSuffixes__"));
 		return AttributeKey;
 	}
+
+	const FAttributeKey& FAnimSequenceNodeStaticData::GetSceneNodeAnimationPayloadKeyMapKey()
+	{
+		static FAttributeKey AttributeKey(TEXT("__SceneNodeAnimationPayloadKeyMap__"));
+		return AttributeKey;
+	}
+
+	const FAttributeKey& FAnimSequenceNodeStaticData::GetMorphTargetNodePayloadKeyMapKey()
+	{
+		static FAttributeKey AttributeKey(TEXT("__MorphTargetNodePayloadKeyMap__"));
+		return AttributeKey;
+	}
 }//ns UE::Interchange
 
 
@@ -52,6 +64,9 @@ UInterchangeAnimSequenceFactoryNode::UInterchangeAnimSequenceFactoryNode()
 	AnimatedAttributeCurveNames.Initialize(Attributes, UE::Interchange::FAnimSequenceNodeStaticData::GetAnimatedAttributeCurveNamesKey().ToString());
 	AnimatedAttributeStepCurveNames.Initialize(Attributes, UE::Interchange::FAnimSequenceNodeStaticData::GetAnimatedAttributeStepCurveNamesKey().ToString());
 	AnimatedMaterialCurveSuffixes.Initialize(Attributes, UE::Interchange::FAnimSequenceNodeStaticData::GetAnimatedMaterialCurveSuffixesKey().ToString());
+
+	SceneNodeAnimationPayloadKeyMap.Initialize(Attributes.ToSharedRef(), UE::Interchange::FAnimSequenceNodeStaticData::GetSceneNodeAnimationPayloadKeyMapKey().ToString());
+	MorphTargetNodePayloadKeyMap.Initialize(Attributes.ToSharedRef(), UE::Interchange::FAnimSequenceNodeStaticData::GetMorphTargetNodePayloadKeyMapKey().ToString());
 }
 
 void UInterchangeAnimSequenceFactoryNode::InitializeAnimSequenceNode(const FString& UniqueID, const FString& DisplayLabel)
@@ -393,5 +408,45 @@ bool FInterchangeAnimSequenceTest::RunTest(const FString& Parameters)
 	return true;
 }
 
+void UInterchangeAnimSequenceFactoryNode::GetSceneNodeAnimationPayloadKeys(TMap<FString, FString>& OutSceneNodeAnimationPayloads) const
+{
+	OutSceneNodeAnimationPayloads = SceneNodeAnimationPayloadKeyMap.ToMap();
+}
+
+bool UInterchangeAnimSequenceFactoryNode::GetAnimationPayloadKeyFromSceneNodeUid(const FString& SceneNodeUid, FString& OutPayloadKey) const
+{
+	return SceneNodeAnimationPayloadKeyMap.GetValue(SceneNodeUid, OutPayloadKey);
+}
+
+bool UInterchangeAnimSequenceFactoryNode::SetAnimationPayloadKeyForSceneNodeUid(const FString& SceneNodeUid, const FString& PayloadKey)
+{
+	return SceneNodeAnimationPayloadKeyMap.SetKeyValue(SceneNodeUid, PayloadKey);
+}
+
+bool UInterchangeAnimSequenceFactoryNode::RemoveAnimationPayloadKeyForSceneNodeUid(const FString& SceneNodeUid)
+{
+	return SceneNodeAnimationPayloadKeyMap.RemoveKey(SceneNodeUid);
+}
+
+void UInterchangeAnimSequenceFactoryNode::GetMorphTargetNodeAnimationPayloadKeys(TMap<FString, FString>& OutMorphTargetAnimationPayloads) const
+{
+	OutMorphTargetAnimationPayloads = MorphTargetNodePayloadKeyMap.ToMap();
+}
+
+bool UInterchangeAnimSequenceFactoryNode::GetAnimationPayloadKeyFromMorphTargetNodeUid(const FString& MorphTargetNodeUid, FString& OutPayloadKey) const
+{
+	return MorphTargetNodePayloadKeyMap.GetValue(MorphTargetNodeUid, OutPayloadKey);
+}
+
+bool UInterchangeAnimSequenceFactoryNode::SetAnimationPayloadKeyForMorphTargetNodeUid(const FString& MorphTargetNodeUid, const FString& PayloadKey)
+{
+	return MorphTargetNodePayloadKeyMap.SetKeyValue(MorphTargetNodeUid, PayloadKey);
+}
+
+bool UInterchangeAnimSequenceFactoryNode::RemoveAnimationPayloadKeyForMorphTargetNodeUid(const FString& MorphTargetNodeUid)
+{
+	return MorphTargetNodePayloadKeyMap.RemoveKey(MorphTargetNodeUid);
+}
+
 #endif //WITH_DEV_AUTOMATION_TESTS
 #endif //WITH_EDITOR
diff --git a/Engine/Plugins/Experimental/Interchange/Runtime/Source/FactoryNodes/Public/InterchangeAnimSequenceFactoryNode.h b/Engine/Plugins/Experimental/Interchange/Runtime/Source/FactoryNodes/Public/InterchangeAnimSequenceFactoryNode.h
--- a/Engine/Plugins/Experimental/Interchange/Runtime/Source/FactoryNodes/Public/InterchangeAnimSequenceFactoryNode.h
+++ b/Engine/Plugins/Experimental/Interchange/Runtime/Source/FactoryNodes/Public/InterchangeAnimSequenceFactoryNode.h
@@ -21,6 +21,8 @@ namespace UE::Interchange
 		static const FAttributeKey& GetAnimatedAttributeCurveNamesKey();
 		static const FAttributeKey& GetAnimatedAttributeStepCurveNamesKey();
 		static const FAttributeKey& GetAnimatedMaterialCurveSuffixesKey();
+		static const FAttributeKey& GetSceneNodeAnimationPayloadKeyMapKey();
+		static const FAttributeKey& GetMorphTargetNodePayloadKeyMapKey();
 	};
 }//ns UE::Interchange
 
@@ -40,6 +42,20 @@ public:
 	UFUNCTION(BlueprintCallable, Category = "Interchange | Node | AnimSequence")
 	void InitializeAnimSequenceNode(const FString& UniqueID, const FString& DisplayLabel);
 
+	/**
+	 * Override serialize to restore SlotMaterialDependencies on load.
+	 */
+	virtual void Serialize(FArchive& Ar) override
+	{
+		Super::Serialize(Ar);
+
+		if (Ar.IsLoading() && bIsInitialized)
+		{
+			SceneNodeAnimationPayloadKeyMap.RebuildCache();
+			MorphTargetNodePayloadKeyMap.RebuildCache();
+		}
+	}
+
 	/**
 	 * Return the node type name of the class, we use this when reporting error
 	 */
@@ -357,7 +373,7 @@ public:
 	bool SetCustomDeleteExistingNonCurveCustomAttributes(const bool& AttributeValue);
 
 	/**
-	 * Query the optional existing USkeleton this anim must use. The animsequence factory will use this skeleton instead of the imported one
+	 * Query the optional existing USkeleton this anim must use. The anim sequence factory will use this skeleton instead of the imported one
 	 * (from GetCustomSkeletonFactoryNodeUid) if this attribute is set and the skeleton pointer is valid.
 	 * Pipeline set this attribute in case the user want to specify an existing skeleton.
 	 * Return false if the attribute was not set.
@@ -373,6 +389,30 @@ public:
 	UFUNCTION(BlueprintCallable, Category = "Interchange | Node | AnimSequence")
 	bool SetCustomSkeletonSoftObjectPath(const FSoftObjectPath& AttributeValue);
 
+	UFUNCTION(BlueprintCallable, Category = "Interchange | Node | SkeletalAnimationTrack")
+	void GetSceneNodeAnimationPayloadKeys(TMap<FString, FString>& OutSceneNodeAnimationPayloads) const;
+
+	UFUNCTION(BlueprintCallable, Category = "Interchange | Node | SkeletalAnimationTrack")
+	bool GetAnimationPayloadKeyFromSceneNodeUid(const FString& SceneNodeUid, FString& OutPayloadKey) const;
+
+	UFUNCTION(BlueprintCallable, Category = "Interchange | Node | SkeletalAnimationTrack")
+	bool SetAnimationPayloadKeyForSceneNodeUid(const FString& SceneNodeUid, const FString& PayloadKey);
+
+	UFUNCTION(BlueprintCallable, Category = "Interchange | Node | SkeletalAnimationTrack")
+	bool RemoveAnimationPayloadKeyForSceneNodeUid(const FString& SceneNodeUid);
+
+	UFUNCTION(BlueprintCallable, Category = "Interchange | Node | SkeletalAnimationTrack")
+	void GetMorphTargetNodeAnimationPayloadKeys(TMap<FString, FString>& OutMorphTargetNodeAnimationPayloads) const;
+
+	UFUNCTION(BlueprintCallable, Category = "Interchange | Node | SkeletalAnimationTrack")
+	bool GetAnimationPayloadKeyFromMorphTargetNodeUid(const FString& MorphTargetNodeUid, FString& OutPayloadKey) const;
+
+	UFUNCTION(BlueprintCallable, Category = "Interchange | Node | SkeletalAnimationTrack")
+	bool SetAnimationPayloadKeyForMorphTargetNodeUid(const FString& MorphTargetNodeUid, const FString& PayloadKey);
+
+	UFUNCTION(BlueprintCallable, Category = "Interchange | Node | SkeletalAnimationTrack")
+	bool RemoveAnimationPayloadKeyForMorphTargetNodeUid(const FString& MorphTargetNodeUid);
+
 private:
 	const UE::Interchange::FAttributeKey Macro_CustomSkeletonFactoryNodeUidKey = UE::Interchange::FAttributeKey(TEXT("SkeletonFactoryNodeUid"));
 
@@ -396,4 +436,7 @@ private:
 	UE::Interchange::TArrayAttributeHelper<FString> AnimatedAttributeCurveNames;
 	UE::Interchange::TArrayAttributeHelper<FString> AnimatedAttributeStepCurveNames;
 	UE::Interchange::TArrayAttributeHelper<FString> AnimatedMaterialCurveSuffixes;
+
+	UE::Interchange::TMapAttributeHelper<FString, FString> SceneNodeAnimationPayloadKeyMap;
+	UE::Interchange::TMapAttributeHelper<FString, FString> MorphTargetNodePayloadKeyMap;
 };
diff --git a/Engine/Plugins/Experimental/Interchange/Runtime/Source/Import/Private/Animation/InterchangeAnimSequenceFactory.cpp b/Engine/Plugins/Experimental/Interchange/Runtime/Source/Import/Private/Animation/InterchangeAnimSequenceFactory.cpp
--- a/Engine/Plugins/Experimental/Interchange/Runtime/Source/Import/Private/Animation/InterchangeAnimSequenceFactory.cpp
+++ b/Engine/Plugins/Experimental/Interchange/Runtime/Source/Import/Private/Animation/InterchangeAnimSequenceFactory.cpp
@@ -264,7 +264,7 @@ namespace UE::Interchange::Private
 		, const FString& AssetName
 		, TArray<FString> OutCurvesNotFound)
 	{
-		TMap<FString, TFuture<TOptional<UE::Interchange::FAnimationBakeTransformPayloadData>>> AnimationPayloads;
+		TMap<const UInterchangeSceneNode*, TFuture<TOptional<UE::Interchange::FAnimationBakeTransformPayloadData>>> AnimationPayloads;
 
 		FString SkeletonRootUid;
 		if (!SkeletonFactoryNode->GetCustomRootJointUid(SkeletonRootUid))
@@ -298,14 +298,30 @@ namespace UE::Interchange::Private
 			const double SequenceLength = FMath::Max<double>(RangeEnd - RangeStart, MINIMUM_ANIMATION_LENGTH);
 			int32 BakeKeyCount = (SequenceLength / BakeInterval) + 1;
 
-			for (const FString& NodeUid : SkeletonNodes)
+			TMap<FString, FString> PayloadKeys;
+			AnimSequenceFactoryNode->GetSceneNodeAnimationPayloadKeys(PayloadKeys);
+
+			if (PayloadKeys.Num() == 0)
+			{
+				for (const FString& NodeUid : SkeletonNodes)
+				{
+					if (const UInterchangeSceneNode* SkeletonSceneNode = Cast<UInterchangeSceneNode>(NodeContainer->GetNode(NodeUid)))
+					{
+						FString PayloadKey;
+						if (UInterchangeAnimationAPI::GetCustomNodeTransformPayloadKey(SkeletonSceneNode, PayloadKey))
+						{
+							AnimationPayloads.Add(SkeletonSceneNode, AnimSequenceTranslatorPayloadInterface->GetAnimationBakeTransformPayloadData(PayloadKey, SampleRate, RangeStart, RangeEnd));
+						}
+					}
+				}
+			}
+			else
 			{
-				if (const UInterchangeSceneNode* SkeletonSceneNode = Cast<UInterchangeSceneNode>(NodeContainer->GetNode(NodeUid)))
+				for (const TTuple<FString, FString>& SceneNodeUidAndPayloadKey : PayloadKeys)
 				{
-					FString PayloadKey;
-					if (UInterchangeAnimationAPI::GetCustomNodeTransformPayloadKey(SkeletonSceneNode, PayloadKey))
+					if (const UInterchangeSceneNode* SkeletonSceneNode = Cast<UInterchangeSceneNode>(NodeContainer->GetNode(SceneNodeUidAndPayloadKey.Key)))
 					{
-						AnimationPayloads.Add(PayloadKey, AnimSequenceTranslatorPayloadInterface->GetAnimationBakeTransformPayloadData(PayloadKey, SampleRate, RangeStart, RangeEnd));
+						AnimationPayloads.Add(SkeletonSceneNode, AnimSequenceTranslatorPayloadInterface->GetAnimationBakeTransformPayloadData(SceneNodeUidAndPayloadKey.Value, SampleRate, RangeStart, RangeEnd));
 					}
 				}
 			}
@@ -324,84 +340,78 @@ namespace UE::Interchange::Private
 				GlobalOffsetTransform = FTransform3f(TempTransform);
 			}
 
-			for (const FString& NodeUid : SkeletonNodes)
+			for (const TTuple< const UInterchangeSceneNode*, TFuture<TOptional<UE::Interchange::FAnimationBakeTransformPayloadData>>>& AnimationPayload : AnimationPayloads)
 			{
-				if (const UInterchangeSceneNode* SkeletonSceneNode = Cast<UInterchangeSceneNode>(NodeContainer->GetNode(NodeUid)))
+				const FName BoneName = FName(*(AnimationPayload.Key->GetDisplayLabel()));
+				const int32 BoneIndex = Skeleton->GetReferenceSkeleton().FindBoneIndex(BoneName);
+				if (BoneIndex == INDEX_NONE)
 				{
-					const FName BoneName = FName(*(SkeletonSceneNode->GetDisplayLabel()));
-					const int32 BoneIndex = Skeleton->GetReferenceSkeleton().FindBoneIndex(BoneName);
-					if (BoneIndex == INDEX_NONE)
-					{
-						//Skip this bone, we did not found it in the skeleton
-						continue;
-					}
-					//If we are getting the root 
-					bool bApplyGlobalOffset = NodeUid.Equals(SkeletonRootUid);
-
-					FString PayloadKey;
-					if (UInterchangeAnimationAPI::GetCustomNodeTransformPayloadKey(SkeletonSceneNode, PayloadKey))
-					{
-						TOptional<UE::Interchange::FAnimationBakeTransformPayloadData> AnimationTransformPayload = AnimationPayloads.FindChecked(PayloadKey).Get();
-						if (!AnimationTransformPayload.IsSet())
-						{
-							UE_LOG(LogInterchangeImport, Warning, TEXT("Invalid animation transform payload key [%s] AnimSequence asset %s"), *PayloadKey, *AssetName);
-							continue;
-						}
-
-						if (AnimationTransformPayload->Transforms.Num() == 0)
-						{
-							//We need at least one transform
-							AnimationTransformPayload->Transforms.Add(FTransform::Identity);
-						}
-
-						FRawAnimSequenceTrack RawTrack;
-						RawTrack.PosKeys.Reserve(BakeKeyCount);
-						RawTrack.RotKeys.Reserve(BakeKeyCount);
-						RawTrack.ScaleKeys.Reserve(BakeKeyCount);
-						TArray<float> TimeKeys;
-						TimeKeys.Reserve(BakeKeyCount);
-						
-						//Everything should match Key count, sample rate and range
-						check(AnimationTransformPayload->Transforms.Num() == BakeKeyCount);
-						check(FMath::IsNearlyEqual(AnimationTransformPayload->BakeFrequency, SampleRate, UE_DOUBLE_KINDA_SMALL_NUMBER));
-						check(FMath::IsNearlyEqual(AnimationTransformPayload->RangeStartTime, RangeStart, UE_DOUBLE_KINDA_SMALL_NUMBER));
-						check(FMath::IsNearlyEqual(AnimationTransformPayload->RangeEndTime, RangeEnd, UE_DOUBLE_KINDA_SMALL_NUMBER));
-
-						int32 TransformIndex = 0;
-						for (double CurrentTime = RangeStart; CurrentTime <= RangeEnd + UE_DOUBLE_SMALL_NUMBER; CurrentTime += BakeInterval, ++TransformIndex)
-						{
-							check(AnimationTransformPayload->Transforms.IsValidIndex(TransformIndex));
-							FTransform3f AnimKeyTransform = FTransform3f(AnimationTransformPayload->Transforms[TransformIndex]);
-							if (bApplyGlobalOffset)
-							{
-								AnimKeyTransform = AnimKeyTransform * GlobalOffsetTransform;
-							}
-							//Default value to identity
-							FVector3f Position(0.0f);
-							FQuat4f Quaternion(0.0f);
-							FVector3f Scale(1.0f);
-							
-							Position = AnimKeyTransform.GetLocation();
-							Quaternion = AnimKeyTransform.GetRotation();
-							Scale = AnimKeyTransform.GetScale3D();
-							RawTrack.ScaleKeys.Add(Scale);
-							RawTrack.PosKeys.Add(Position);
-							RawTrack.RotKeys.Add(Quaternion);
-							//Animation are always translated to zero
-							TimeKeys.Add(CurrentTime - RangeStart);
-						}
+					//Skip this bone, we did not found it in the skeleton
+					continue;
+				}
+				//If we are getting the root 
+				bool bApplyGlobalOffset = AnimationPayload.Key->GetUniqueID().Equals(SkeletonRootUid);
+				
+				TOptional<UE::Interchange::FAnimationBakeTransformPayloadData> AnimationTransformPayload = AnimationPayload.Value.Get();
+				if (!AnimationTransformPayload.IsSet())
+				{
+					FString PayloadKey = PayloadKeys[AnimationPayload.Key->GetUniqueID()];
+					UE_LOG(LogInterchangeImport, Warning, TEXT("Invalid animation transform payload key [%s] AnimSequence asset %s"), *PayloadKey, *AssetName);
+					continue;
+				}
 
-						//Make sure we create the correct amount of keys
-						check(RawTrack.ScaleKeys.Num() == BakeKeyCount
-							&& RawTrack.PosKeys.Num() == BakeKeyCount
-							&& RawTrack.RotKeys.Num() == BakeKeyCount
-							&& TimeKeys.Num() == BakeKeyCount);
+				if (AnimationTransformPayload->Transforms.Num() == 0)
+				{
+					//We need at least one transform
+					AnimationTransformPayload->Transforms.Add(FTransform::Identity);
+				}
 
-						//add new track
-						Controller.AddBoneTrack(BoneName);
-						Controller.SetBoneTrackKeys(BoneName, RawTrack.PosKeys, RawTrack.RotKeys, RawTrack.ScaleKeys);
+				FRawAnimSequenceTrack RawTrack;
+				RawTrack.PosKeys.Reserve(BakeKeyCount);
+				RawTrack.RotKeys.Reserve(BakeKeyCount);
+				RawTrack.ScaleKeys.Reserve(BakeKeyCount);
+				TArray<float> TimeKeys;
+				TimeKeys.Reserve(BakeKeyCount);
+
+				//Everything should match Key count, sample rate and range
+				check(AnimationTransformPayload->Transforms.Num() == BakeKeyCount);
+				check(FMath::IsNearlyEqual(AnimationTransformPayload->BakeFrequency, SampleRate, UE_DOUBLE_KINDA_SMALL_NUMBER));
+				check(FMath::IsNearlyEqual(AnimationTransformPayload->RangeStartTime, RangeStart, UE_DOUBLE_KINDA_SMALL_NUMBER));
+				check(FMath::IsNearlyEqual(AnimationTransformPayload->RangeEndTime, RangeEnd, UE_DOUBLE_KINDA_SMALL_NUMBER));
+
+				int32 TransformIndex = 0;
+				for (double CurrentTime = RangeStart; CurrentTime <= RangeEnd + UE_DOUBLE_SMALL_NUMBER; CurrentTime += BakeInterval, ++TransformIndex)
+				{
+					check(AnimationTransformPayload->Transforms.IsValidIndex(TransformIndex));
+					FTransform3f AnimKeyTransform = FTransform3f(AnimationTransformPayload->Transforms[TransformIndex]);
+					if (bApplyGlobalOffset)
+					{
+						AnimKeyTransform = AnimKeyTransform * GlobalOffsetTransform;
 					}
+					//Default value to identity
+					FVector3f Position(0.0f);
+					FQuat4f Quaternion(0.0f);
+					FVector3f Scale(1.0f);
+
+					Position = AnimKeyTransform.GetLocation();
+					Quaternion = AnimKeyTransform.GetRotation();
+					Scale = AnimKeyTransform.GetScale3D();
+					RawTrack.ScaleKeys.Add(Scale);
+					RawTrack.PosKeys.Add(Position);
+					RawTrack.RotKeys.Add(Quaternion);
+					//Animation are always translated to zero
+					TimeKeys.Add(CurrentTime - RangeStart);
 				}
+
+				//Make sure we create the correct amount of keys
+				check(RawTrack.ScaleKeys.Num() == BakeKeyCount
+					&& RawTrack.PosKeys.Num() == BakeKeyCount
+					&& RawTrack.RotKeys.Num() == BakeKeyCount
+					&& TimeKeys.Num() == BakeKeyCount);
+
+				//add new track
+				Controller.AddBoneTrack(BoneName);
+				Controller.SetBoneTrackKeys(BoneName, RawTrack.PosKeys, RawTrack.RotKeys, RawTrack.ScaleKeys);
 			}
 		}
 
diff --git a/Engine/Plugins/Experimental/Interchange/Runtime/Source/Import/Private/Gltf/InterchangeGltfTranslator.cpp b/Engine/Plugins/Experimental/Interchange/Runtime/Source/Import/Private/Gltf/InterchangeGltfTranslator.cpp
--- a/Engine/Plugins/Experimental/Interchange/Runtime/Source/Import/Private/Gltf/InterchangeGltfTranslator.cpp
+++ b/Engine/Plugins/Experimental/Interchange/Runtime/Source/Import/Private/Gltf/InterchangeGltfTranslator.cpp
@@ -19,8 +19,6 @@
 #include "InterchangeVariantSetNode.h"
 #include "Nodes/InterchangeSourceNode.h"
 
-#include "Nodes/InterchangeAnimationAPI.h"
-
 #include "Sections/MovieScene3DTransformSection.h"
 #include "Texture/InterchangeImageWrapperTranslator.h"
 
@@ -972,10 +970,9 @@ bool UInterchangeGltfTranslator::Translate( UInterchangeBaseNodeContainer& NodeC
 
 	// Animations
 	{
-		TArray<FString> JointsAlreadyUsedInAnAnimation;
 		for (int32 AnimationIndex = 0; AnimationIndex < GltfAsset.Animations.Num(); ++AnimationIndex)
 		{
-			HandleGltfAnimation(NodeContainer, AnimationIndex, JointsAlreadyUsedInAnAnimation);
+			HandleGltfAnimation(NodeContainer, AnimationIndex);
 		}
 	}
 
@@ -1093,25 +1090,47 @@ TFuture<TOptional<UE::Interchange::FAnimationBakeTransformPayloadData>> UInterch
 		});
 }
 
-void UInterchangeGltfTranslator::HandleGltfAnimation(UInterchangeBaseNodeContainer& NodeContainer, int32 AnimationIndex, TArray<FString>& JointsAlreadyUsedInAnAnimation /*Currently the system only supports 1 animation per skeleton*/) const
+void UInterchangeGltfTranslator::HandleGltfAnimation(UInterchangeBaseNodeContainer& NodeContainer, int32 AnimationIndex) const
 {
 	const GLTF::FAnimation& GltfAnimation = GltfAsset.Animations[AnimationIndex];
 
 	TMap< const GLTF::FNode*, TArray<int32> > NodeChannelsMap;
 
-	TMap<FString, TArray<int32>> JointNodeUidsToChannelsUsed;
+	TMap<int32, UInterchangeSkeletalAnimationTrackNode*> RootJointIndexToTrackNodeMap;
+	TMap<UInterchangeSkeletalAnimationTrackNode*, TMap<FString, TArray<int32>>> TrackNodeToJointUidWithChannelsUsedMap;
 
 	for (int32 ChannelIndex = 0; ChannelIndex < GltfAnimation.Channels.Num(); ++ChannelIndex)
 	{
 		const GLTF::FAnimation::FChannel& Channel = GltfAnimation.Channels[ChannelIndex];
 		const GLTF::FNode* AnimatedNode = &Channel.Target.Node;
 
-		if (AnimatedNode->Type == GLTF::FNode::EType::Joint)
+		if (AnimatedNode->Type == GLTF::FNode::EType::Joint && GltfAsset.Nodes.IsValidIndex(AnimatedNode->RootJointIndex))
 		{
 			const FString* JointNodeUid = NodeUidMap.Find(AnimatedNode);
-			if (ensure(JointNodeUid))
+			if (ensure(JointNodeUid) && GltfAsset.Nodes.IsValidIndex(AnimatedNode->RootJointIndex))
 			{
-				TArray<int32>& ChannelsUsed = JointNodeUidsToChannelsUsed.FindOrAdd(*JointNodeUid);
+				UInterchangeSkeletalAnimationTrackNode* TrackNode = nullptr;
+				if (RootJointIndexToTrackNodeMap.Contains(AnimatedNode->RootJointIndex))
+				{
+					TrackNode = RootJointIndexToTrackNodeMap[AnimatedNode->RootJointIndex];
+				}
+				else
+				{
+					TrackNode = NewObject< UInterchangeSkeletalAnimationTrackNode >(&NodeContainer);
+					FString TrackNodeUid = "\\SkeletalAnimation\\" + LexToString(AnimatedNode->RootJointIndex) + "_" + LexToString(AnimationIndex);
+					// AnimationIndex is already part of the GltfAnimation.Name (due to GenerateNames)
+					FString DisplayString = GltfAnimation.Name + "_" + LexToString(AnimatedNode->RootJointIndex);
+					TrackNode->InitializeNode(TrackNodeUid, DisplayString, EInterchangeNodeContainerType::TranslatedAsset);
+					const GLTF::FNode& RootJointNode = GltfAsset.Nodes[AnimatedNode->RootJointIndex];
+					const FString* SkeletonNodeUid = NodeUidMap.Find(&RootJointNode);
+					TrackNode->SetCustomSkeletonNodeUid(*SkeletonNodeUid);
+
+					NodeContainer.AddNode(TrackNode);
+
+					RootJointIndexToTrackNodeMap.Add(AnimatedNode->RootJointIndex, TrackNode);
+				}
+				TMap<FString, TArray<int32>>& JointUidWithChannelsUsedMap = TrackNodeToJointUidWithChannelsUsedMap.FindOrAdd(TrackNode);
+				TArray<int32>& ChannelsUsed = JointUidWithChannelsUsedMap.FindOrAdd(*JointNodeUid);
 				ChannelsUsed.Add(ChannelIndex);
 			}
 			
@@ -1124,40 +1143,37 @@ void UInterchangeGltfTranslator::HandleGltfAnimation(UInterchangeBaseNodeContain
 
 	//setup rigged animations:
 	{
-		for (TTuple<FString,TArray<int32>>& JointNodeUidToChannelsUsed : JointNodeUidsToChannelsUsed)
+		for (const TTuple<UInterchangeSkeletalAnimationTrackNode*, TMap<FString, TArray<int32>>>& TrackNodeAndJointNodeChannels : TrackNodeToJointUidWithChannelsUsedMap)
 		{
-			if (JointsAlreadyUsedInAnAnimation.Contains(JointNodeUidToChannelsUsed.Key))
-			{
-				//this joint (part of skeleton) is already part of an animation
-				//due to the Interchange system only supporting a single animation per skeleton at the moment we skip the setup of this joint and animatino pair
-				continue;
-			}
-			
-			//@ AllAnimationStart = 0;
+			//@ StartTime = 0;
 			//from gltf documentation: https://registry.khronos.org/glTF/specs/2.0/glTF-2.0.html#animations
 			//Implementation Note
 			//	For example, if the earliest sampler input for an animation is t = 10, a client implementation must begin playback of that animation channel at t = 0 with output clamped to the first available output value.
-			double AnimationStart = 0.0;
-			double AnimationStop = 0.0;
+			double StartTime = 0.0;
+			double StopTime = 0.0;
 			double FrameRate = 30.0;
 			double SingleFrameDuration = 1.0 / FrameRate;
-			int32 KeyCount = 0;
-			bool CorruptSamplerFound = false;
+			int32 FrameNumber = 0;
+			bool bHasAnimationPayloadSet = false;
 
-			//check channel length:
+			for (const TTuple<FString, TArray<int32>>& JointNodeUidAndChannelsUsedPair : TrackNodeAndJointNodeChannels.Value)
 			{
-				for (int32 ChannelIndex : JointNodeUidToChannelsUsed.Value)
+				double PreviousStopTime = StopTime;
+				//check channel length and build payload
+				FString Payload = LexToString(AnimationIndex);
+				for (int32 ChannelIndex : JointNodeUidAndChannelsUsedPair.Value)
 				{
 					const GLTF::FAnimation::FChannel& Channel = GltfAnimation.Channels[ChannelIndex];
 					const GLTF::FAnimation::FSampler& Sampler = GltfAnimation.Samplers[Channel.Sampler];
 					TArray<float> Seconds;
-					Sampler.Input.GetFloatArray(Seconds);					
+					Sampler.Input.GetFloatArray(Seconds);
 
 					if (Sampler.Interpolation == GLTF::FAnimation::EInterpolation::CubicSpline)
 					{
 						if (Sampler.Input.Count != 3 * Sampler.Output.Count)
 						{
-							CorruptSamplerFound = true;
+							// if any of the channels are corrupt the joint will not receive any of the  animation data
+							Payload = "";
 							UE_LOG(LogInterchangeImport, Warning, TEXT("GLTF Sampler Corrupt. Input and Output not meeting expectations."));
 							break;
 						}
@@ -1166,57 +1182,56 @@ void UInterchangeGltfTranslator::HandleGltfAnimation(UInterchangeBaseNodeContain
 					{
 						if (Sampler.Input.Count != Sampler.Output.Count)
 						{
-							CorruptSamplerFound = true;
+							// if any of the channels are corrupt the joint will not receive any of the  animation data
+							Payload = "";
 							UE_LOG(LogInterchangeImport, Warning, TEXT("GLTF Sampler Corrupt. Input and Output not meeting expectations."));
 							break;
 						}
 					}
 
-					float CurrentAnimationStop = 0;
-					int CurrentKeyCount = 0;
+					float CurrentStopTime = 0;
+					int32 CurrentFrameNumber = 0;
 
 					if (Seconds.Num() > 0)
 					{
-						//calculate keycount and currentAnimationStop:
-						CurrentAnimationStop = Seconds[Seconds.Num()-1];
+						//calculate FrameNumber and currentStopTime:
+						CurrentStopTime = Seconds[Seconds.Num() - 1];
 
-						float CurrentKeyCountCandidate = CurrentAnimationStop / SingleFrameDuration;
-						CurrentKeyCount = int(CurrentKeyCountCandidate);
-						if (int(CurrentKeyCountCandidate) < CurrentKeyCountCandidate)
+						float CurrentFrameNumberCandidate = CurrentStopTime / SingleFrameDuration;
+						CurrentFrameNumber = int(CurrentFrameNumberCandidate);
+						if (int(CurrentFrameNumberCandidate) < CurrentFrameNumberCandidate)
 						{
-							CurrentKeyCount++;
+							CurrentFrameNumber++;
 						}
-						CurrentAnimationStop = CurrentKeyCount * SingleFrameDuration;
+						CurrentStopTime = CurrentFrameNumber * SingleFrameDuration;
 					}
-					
-					if (KeyCount < CurrentKeyCount)
+
+					if (FrameNumber < CurrentFrameNumber)
 					{
-						KeyCount = CurrentKeyCount;
-						AnimationStop = KeyCount * SingleFrameDuration;
+						FrameNumber = CurrentFrameNumber;
+						StopTime = FrameNumber * SingleFrameDuration;
 					}
+					Payload += ":" + LexToString(ChannelIndex);
 				}
-			}
-
-			if (CorruptSamplerFound)
-			{
-				break;
-			}
 
-			if (UInterchangeSceneNode* AnimatedInterchangeSceneNode = const_cast<UInterchangeSceneNode*>(Cast< UInterchangeSceneNode >(NodeContainer.GetNode(*JointNodeUidToChannelsUsed.Key))))
-			{
-				//add the joint to the of already set up joints to track it for handling single animation only.
-				JointsAlreadyUsedInAnAnimation.Add(JointNodeUidToChannelsUsed.Key);
-
-				FString Payload = LexToString(AnimationIndex);
-				for (int32 ChannelIndex : JointNodeUidToChannelsUsed.Value)
+				//set payload:
+				if (Payload.Len() > 0)
 				{
-					Payload += ":" + LexToString(ChannelIndex);
+					TrackNodeAndJointNodeChannels.Key->SetAnimationPayloadKeyForSceneNodeUid(JointNodeUidAndChannelsUsedPair.Key, Payload);
+					bHasAnimationPayloadSet = true;
 				}
+				else
+				{
+					StopTime = PreviousStopTime;
+				}
+			}
 
-				UInterchangeAnimationAPI::SetCustomNodeTransformPayloadKey(AnimatedInterchangeSceneNode, Payload);
-				UInterchangeAnimationAPI::SetCustomIsNodeTransformAnimated(AnimatedInterchangeSceneNode, true);
-				UInterchangeAnimationAPI::SetCustomNodeTransformAnimationStartTime(AnimatedInterchangeSceneNode, AnimationStart);
-				UInterchangeAnimationAPI::SetCustomNodeTransformAnimationEndTime(AnimatedInterchangeSceneNode, AnimationStop);
+			//set animation length:
+			if (bHasAnimationPayloadSet)
+			{
+				TrackNodeAndJointNodeChannels.Key->SetCustomAnimationSampleRate(FrameRate);
+				TrackNodeAndJointNodeChannels.Key->SetCustomAnimationStartTime(StartTime);
+				TrackNodeAndJointNodeChannels.Key->SetCustomAnimationStopTime(StopTime);
 			}
 		}
 	}
@@ -1568,23 +1583,9 @@ TFuture<TOptional<UE::Interchange::FSkeletalMeshMorphTargetPayloadData>> UInterc
 	return Promise->GetFuture();
 }
 
-int32 UInterchangeGltfTranslator::FindRootJointNodeIndex(int32 CurrentIndex) const
-{
-	if (GltfAsset.Nodes[CurrentIndex].Type != GLTF::FNode::EType::Joint)
-	{
-		return CurrentIndex;
-	}
-
-	while (GltfAsset.Nodes.IsValidIndex(GltfAsset.Nodes[CurrentIndex].ParentIndex) && GltfAsset.Nodes[GltfAsset.Nodes[CurrentIndex].ParentIndex].Type == GLTF::FNode::EType::Joint)
-	{
-		CurrentIndex = GltfAsset.Nodes[CurrentIndex].ParentIndex;
-	}
-	return CurrentIndex;
-}
-
 void UInterchangeGltfTranslator::HandleGltfSkeletons(UInterchangeBaseNodeContainer& NodeContainer, const FString& SceneNodeUid, const TArray<int32>& SkinnedMeshNodes) const
 {
-	//if skeletons present set them as a root joint for the join hierarchy as per documentation:
+	//If there are skeletons, set them as a root joint for the joint hierarchy as per documentation:
 	//The skeleton property (if present) points to the node that is the common root of a joints hierarchy or to a direct or indirect parent node of the common root.
 	for (const GLTF::FSkinInfo& Skin : GltfAsset.Skins)
 	{
@@ -1610,7 +1611,18 @@ void UInterchangeGltfTranslator::HandleGltfSkeletons(UInterchangeBaseNodeContain
 		TMap<int32, TArray<int32>>& RootJointGroupedSkinnedMeshNodes = MeshIndexToRootJointGroupedSkinnedMeshNodesMap.FindOrAdd(SkinnedMeshNode.MeshIndex);
 
 		//get the SkinnedMeshNode's skin's first joint as the starting ground and find the top most root joint for it:
-		int32 RootJointIndex = FindRootJointNodeIndex(GltfAsset.Skins[SkinnedMeshNode.Skindex].Joints[0]);
+		if (!GltfAsset.Skins.IsValidIndex(SkinnedMeshNode.Skindex)
+			|| !GltfAsset.Skins[SkinnedMeshNode.Skindex].Joints.IsValidIndex(0)
+			|| !GltfAsset.Nodes.IsValidIndex(GltfAsset.Skins[SkinnedMeshNode.Skindex].Joints[0]))
+		{
+			continue;
+		}
+		
+		int32 RootJointIndex = GltfAsset.Nodes[GltfAsset.Skins[SkinnedMeshNode.Skindex].Joints[0]].RootJointIndex;
+		if (!GltfAsset.Nodes.IsValidIndex(RootJointIndex))
+		{
+			continue;
+		}
 
 		//basedon that root joint group the SkinnedMeshNodes:
 		TArray<int32>& GroupedSkinnedMeshNodes = RootJointGroupedSkinnedMeshNodes.FindOrAdd(RootJointIndex);
diff --git a/Engine/Plugins/Experimental/Interchange/Runtime/Source/Import/Public/Gltf/InterchangeGltfTranslator.h b/Engine/Plugins/Experimental/Interchange/Runtime/Source/Import/Public/Gltf/InterchangeGltfTranslator.h
--- a/Engine/Plugins/Experimental/Interchange/Runtime/Source/Import/Public/Gltf/InterchangeGltfTranslator.h
+++ b/Engine/Plugins/Experimental/Interchange/Runtime/Source/Import/Public/Gltf/InterchangeGltfTranslator.h
@@ -76,7 +76,7 @@ protected:
 	void HandleGltfMaterial( UInterchangeBaseNodeContainer& NodeContainer, const GLTF::FMaterial& GltfMaterial, UInterchangeShaderGraphNode& ShaderGraphNode ) const;
 	void HandleGltfMaterialParameter( UInterchangeBaseNodeContainer& NodeContainer, const GLTF::FTextureMap& TextureMap, UInterchangeShaderNode& ShaderNode,
 		const FString& MapName, const TVariant< FLinearColor, float >& MapFactor, const FString& OutputChannel, const bool bInverse = false, const bool bIsNormal = false ) const;
-	void HandleGltfAnimation(UInterchangeBaseNodeContainer& NodeContainer, int32 AnimationIndex, TArray<FString>& JointsAlreadyUsedInAnAnimation) const;
+	void HandleGltfAnimation(UInterchangeBaseNodeContainer& NodeContainer, int32 AnimationIndex) const;
 	void HandleGltfVariants(UInterchangeBaseNodeContainer& NodeContainer, const FString& FileName) const;
 
 	/** Support for KHR_materials_clearcoat */
@@ -95,8 +95,6 @@ private:
 	void SetTextureSRGB(UInterchangeBaseNodeContainer& NodeContainer, const GLTF::FTextureMap& TextureMap) const;
 	void SetTextureFlipGreenChannel(UInterchangeBaseNodeContainer& NodeContainer, const GLTF::FTextureMap& TextureMap) const;
 
-	int32 FindRootJointNodeIndex(int32 CurrentIndex) const;
-
 	GLTF::FAsset GltfAsset;
 	mutable FNodeUidMap NodeUidMap;
 };
diff --git a/Engine/Plugins/Experimental/Interchange/Runtime/Source/Nodes/Private/InterchangeAnimationTrackSetNode.cpp b/Engine/Plugins/Experimental/Interchange/Runtime/Source/Nodes/Private/InterchangeAnimationTrackSetNode.cpp
--- a/Engine/Plugins/Experimental/Interchange/Runtime/Source/Nodes/Private/InterchangeAnimationTrackSetNode.cpp
+++ b/Engine/Plugins/Experimental/Interchange/Runtime/Source/Nodes/Private/InterchangeAnimationTrackSetNode.cpp
@@ -152,3 +152,101 @@ bool UInterchangeTransformAnimationTrackNode::GetCustomUsedChannels(int32& Attri
 {
 	IMPLEMENT_NODE_ATTRIBUTE_GETTER(UsedChannels, int32);
 }
+
+// UInterchangeSkeletalAnimationTrackNode
+
+UInterchangeSkeletalAnimationTrackNode::UInterchangeSkeletalAnimationTrackNode()
+{
+	SceneNodeAnimationPayloadKeyMap.Initialize(Attributes.ToSharedRef(), TEXT("__SceneNodeAnimationPayloadKeyMap__"));
+	MorphTargetPayloadKeyMap.Initialize(Attributes.ToSharedRef(), TEXT("__MorphTargetPayloadKeyMap__"));
+}
+
+bool UInterchangeSkeletalAnimationTrackNode::GetCustomSkeletonNodeUid(FString& AttributeValue) const
+{
+	IMPLEMENT_NODE_ATTRIBUTE_GETTER(SkeletonNodeUid, FString);
+}
+
+bool UInterchangeSkeletalAnimationTrackNode::SetCustomSkeletonNodeUid(const FString& AttributeValue)
+{
+	IMPLEMENT_NODE_ATTRIBUTE_SETTER_NODELEGATE(SkeletonNodeUid, FString);
+}
+
+bool UInterchangeSkeletalAnimationTrackNode::GetCustomSkeletalMeshNodeUid(FString& AttributeValue) const
+{
+	IMPLEMENT_NODE_ATTRIBUTE_GETTER(SkeletalMeshNodeUid, FString);
+}
+
+bool UInterchangeSkeletalAnimationTrackNode::SetCustomSkeletalMeshNodeUid(const FString& AttributeValue)
+{
+	IMPLEMENT_NODE_ATTRIBUTE_SETTER_NODELEGATE(SkeletalMeshNodeUid, FString);
+}
+
+bool UInterchangeSkeletalAnimationTrackNode::GetCustomAnimationSampleRate(double& AttributeValue) const
+{
+	IMPLEMENT_NODE_ATTRIBUTE_GETTER(AnimationSampleRate, double);
+}
+
+bool UInterchangeSkeletalAnimationTrackNode::SetCustomAnimationSampleRate(const double& AttributeValue)
+{
+	IMPLEMENT_NODE_ATTRIBUTE_SETTER_NODELEGATE(AnimationSampleRate, double);
+}
+
+bool UInterchangeSkeletalAnimationTrackNode::GetCustomAnimationStartTime(double& AttributeValue) const
+{
+	IMPLEMENT_NODE_ATTRIBUTE_GETTER(AnimationStartTime, double);
+}
+
+bool UInterchangeSkeletalAnimationTrackNode::SetCustomAnimationStartTime(const double& AttributeValue)
+{
+	IMPLEMENT_NODE_ATTRIBUTE_SETTER_NODELEGATE(AnimationStartTime, double);
+}
+
+bool UInterchangeSkeletalAnimationTrackNode::GetCustomAnimationStopTime(double& AttributeValue) const
+{
+	IMPLEMENT_NODE_ATTRIBUTE_GETTER(AnimationStopTime, double);
+}
+
+bool UInterchangeSkeletalAnimationTrackNode::SetCustomAnimationStopTime(const double& AttributeValue)
+{
+	IMPLEMENT_NODE_ATTRIBUTE_SETTER_NODELEGATE(AnimationStopTime, double);
+}
+
+void UInterchangeSkeletalAnimationTrackNode::GetSceneNodeAnimationPayloadKeys(TMap<FString, FString>& OutSceneNodeAnimationPayloads) const
+{
+	OutSceneNodeAnimationPayloads = SceneNodeAnimationPayloadKeyMap.ToMap();
+}
+
+bool UInterchangeSkeletalAnimationTrackNode::GetAnimationPayloadKeyFromSceneNodeUid(const FString& SceneNodeUid, FString& OutPayloadKey) const
+{
+	return SceneNodeAnimationPayloadKeyMap.GetValue(SceneNodeUid, OutPayloadKey);
+}
+
+bool UInterchangeSkeletalAnimationTrackNode::SetAnimationPayloadKeyForSceneNodeUid(const FString& SceneNodeUid, const FString& PayloadKey)
+{
+	return SceneNodeAnimationPayloadKeyMap.SetKeyValue(SceneNodeUid, PayloadKey);
+}
+
+bool UInterchangeSkeletalAnimationTrackNode::RemoveAnimationPayloadKeyForSceneNodeUid(const FString& SceneNodeUid)
+{
+	return SceneNodeAnimationPayloadKeyMap.RemoveKey(SceneNodeUid);
+}
+
+void UInterchangeSkeletalAnimationTrackNode::GetMorphTargetNodeAnimationPayloadKeys(TMap<FString, FString>& OutMorphTargetNodeAnimationPayloads) const
+{
+	OutMorphTargetNodeAnimationPayloads = MorphTargetPayloadKeyMap.ToMap();
+}
+
+bool UInterchangeSkeletalAnimationTrackNode::GetAnimationPayloadKeyFromMorphTargetNodeUid(const FString& MorphTargetNodeUid, FString& OutPayloadKey) const
+{
+	return MorphTargetPayloadKeyMap.GetValue(MorphTargetNodeUid, OutPayloadKey);
+}
+
+bool UInterchangeSkeletalAnimationTrackNode::SetAnimationPayloadKeyForMorphTargetNodeUid(const FString& MorphTargetNodeUid, const FString& PayloadKey)
+{
+	return MorphTargetPayloadKeyMap.SetKeyValue(MorphTargetNodeUid, PayloadKey);
+}
+
+bool UInterchangeSkeletalAnimationTrackNode::RemoveAnimationPayloadKeyForMorphTargetNodeUid(const FString& MorphTargetNodeUid)
+{
+	return MorphTargetPayloadKeyMap.RemoveKey(MorphTargetNodeUid);
+}
diff --git a/Engine/Plugins/Experimental/Interchange/Runtime/Source/Nodes/Public/InterchangeAnimationTrackSetNode.h b/Engine/Plugins/Experimental/Interchange/Runtime/Source/Nodes/Public/InterchangeAnimationTrackSetNode.h
--- a/Engine/Plugins/Experimental/Interchange/Runtime/Source/Nodes/Public/InterchangeAnimationTrackSetNode.h
+++ b/Engine/Plugins/Experimental/Interchange/Runtime/Source/Nodes/Public/InterchangeAnimationTrackSetNode.h
@@ -352,4 +352,132 @@ public:
 
 private:
 	const UE::Interchange::FAttributeKey Macro_CustomUsedChannelsKey = UE::Interchange::FAttributeKey(TEXT("UsedChannels"));
-};
\ No newline at end of file
+};
+
+/*
+* Class to hold onto the relationships between a set of animation tracks and the bones, morph targets of a skeleton.
+*/
+UCLASS(BlueprintType, Experimental)
+class INTERCHANGENODES_API UInterchangeSkeletalAnimationTrackNode : public UInterchangeAnimationTrackBaseNode
+{
+	GENERATED_BODY()
+
+public:
+	UInterchangeSkeletalAnimationTrackNode();
+
+	/**
+	 * Override serialize to restore SlotMaterialDependencies on load.
+	 */
+	virtual void Serialize(FArchive& Ar) override
+	{
+		Super::Serialize(Ar);
+
+		if (Ar.IsLoading() && bIsInitialized)
+		{
+			SceneNodeAnimationPayloadKeyMap.RebuildCache();
+			MorphTargetPayloadKeyMap.RebuildCache();
+		}
+	}
+
+	static FStringView StaticAssetTypeName()
+	{
+		return TEXT("SkeletalAnimationTrack");
+	}
+
+	/**
+	* Return the node type name of the class, we use this when reporting errors
+	*/
+	virtual FString GetTypeName() const override
+	{
+		const FString TypeName = TEXT("SkeletalAnimationTrack");
+		return TypeName;
+	}
+
+
+public:
+	/** Get the skeleton factory node unique id. Return false if the attribute is not set. */
+	UFUNCTION(BlueprintCallable, Category = "Interchange | Node | SkeletalAnimationTrack")
+	bool GetCustomSkeletonNodeUid(FString& AttributeValue) const;
+
+	/** Set the skeleton factory node unique id. Return false if the attribute cannot be set. */
+	UFUNCTION(BlueprintCallable, Category = "Interchange | Node | SkeletalAnimationTrack")
+	bool SetCustomSkeletonNodeUid(const FString& AttributeValue);
+
+	/** Get the skeletal mesh node unique id. Return false if the attribute is not set. */
+	UFUNCTION(BlueprintCallable, Category = "Interchange | Node | SkeletalAnimationTrack")
+	bool GetCustomSkeletalMeshNodeUid(FString& AttributeValue) const;
+
+	/** Set the skeletal mesh node unique id. Return false if the attribute cannot be set. */
+	UFUNCTION(BlueprintCallable, Category = "Interchange | Node | SkeletalAnimationTrack")
+	bool SetCustomSkeletalMeshNodeUid(const FString& AttributeValue);
+
+	/**
+	 * Set the animation sample rate. Return false if the attribute cannot be set.
+	 */
+	UFUNCTION(BlueprintCallable, Category = "Interchange | Node | SkeletalAnimationTrack")
+	bool SetCustomAnimationSampleRate(const double& SampleRate);
+
+	/**
+	 * Get the animation sample rate. Return false if the attribute is not set.
+	 */
+	UFUNCTION(BlueprintCallable, Category = "Interchange | Node | SkeletalAnimationTrack")
+	bool GetCustomAnimationSampleRate(double& SampleRate) const;
+
+	/**
+	 * Set the animation start time. Return false if the attribute cannot be set.
+	 */
+	UFUNCTION(BlueprintCallable, Category = "Interchange | Node | SkeletalAnimationTrack")
+	bool SetCustomAnimationStartTime(const double& StartTime);
+
+	/**
+	 * Get the animation start time. Return false if the attribute is not set.
+	 */
+	UFUNCTION(BlueprintCallable, Category = "Interchange | Node | SkeletalAnimationTrack")
+	bool GetCustomAnimationStartTime(double& StartTime) const;
+	
+	/**
+	 * Set the animation stop time. Return false if the attribute cannot be set.
+	 */
+	UFUNCTION(BlueprintCallable, Category = "Interchange | Node | SkeletalAnimationTrack")
+	bool SetCustomAnimationStopTime(const double& StopTime);
+
+	/**
+	 * Get the animation stop time. Return false if the attribute is not set.
+	 */
+	UFUNCTION(BlueprintCallable, Category = "Interchange | Node | SkeletalAnimationTrack")
+	bool GetCustomAnimationStopTime(double& StopTime) const;
+
+	UFUNCTION(BlueprintCallable, Category = "Interchange | Node | SkeletalAnimationTrack")
+	void GetSceneNodeAnimationPayloadKeys(TMap<FString, FString>& OutSceneNodeAnimationPayloadKeys) const;
+
+	UFUNCTION(BlueprintCallable, Category = "Interchange | Node | SkeletalAnimationTrack")
+	bool GetAnimationPayloadKeyFromSceneNodeUid(const FString& SceneNodeUid, FString& OutPayloadKey) const;
+
+	UFUNCTION(BlueprintCallable, Category = "Interchange | Node | SkeletalAnimationTrack")
+	bool SetAnimationPayloadKeyForSceneNodeUid(const FString& SceneNodeUid, const FString& PayloadKey);
+
+	UFUNCTION(BlueprintCallable, Category = "Interchange | Node | SkeletalAnimationTrack")
+	bool RemoveAnimationPayloadKeyForSceneNodeUid(const FString& SceneNodeUid);
+
+	UFUNCTION(BlueprintCallable, Category = "Interchange | Node | SkeletalAnimationTrack")
+	void GetMorphTargetNodeAnimationPayloadKeys(TMap<FString, FString>& OutMorphTargetNodeAnimationPayloads) const;
+
+	UFUNCTION(BlueprintCallable, Category = "Interchange | Node | SkeletalAnimationTrack")
+	bool GetAnimationPayloadKeyFromMorphTargetNodeUid(const FString& MorphTargetNodeUid, FString& OutPayloadKey) const;
+
+	UFUNCTION(BlueprintCallable, Category = "Interchange | Node | SkeletalAnimationTrack")
+	bool SetAnimationPayloadKeyForMorphTargetNodeUid(const FString& MorphTargetNodeUid, const FString& PayloadKey);
+
+	UFUNCTION(BlueprintCallable, Category = "Interchange | Node | SkeletalAnimationTrack")
+	bool RemoveAnimationPayloadKeyForMorphTargetNodeUid(const FString& MorphTargetNodeUid);
+
+private:
+	const UE::Interchange::FAttributeKey Macro_CustomSkeletonNodeUidKey = UE::Interchange::FAttributeKey(TEXT("SkeletonNodeUid"));
+	const UE::Interchange::FAttributeKey Macro_CustomSkeletalMeshNodeUidKey = UE::Interchange::FAttributeKey(TEXT("SkeletalMeshNodeUid"));
+	const UE::Interchange::FAttributeKey Macro_CustomAnimationSampleRateKey = UE::Interchange::FAttributeKey(TEXT("AnimationSampleRate"));
+	const UE::Interchange::FAttributeKey Macro_CustomAnimationStartTimeKey = UE::Interchange::FAttributeKey(TEXT("AnimationStartTime"));
+	const UE::Interchange::FAttributeKey Macro_CustomAnimationStopTimeKey = UE::Interchange::FAttributeKey(TEXT("AnimationStopTime"));
+	
+	UE::Interchange::TMapAttributeHelper<FString, FString> SceneNodeAnimationPayloadKeyMap;
+	UE::Interchange::TMapAttributeHelper<FString, FString> MorphTargetPayloadKeyMap;
+};
diff --git a/Engine/Plugins/Experimental/Interchange/Runtime/Source/Pipelines/Private/InterchangeGenericAnimationPipeline.cpp b/Engine/Plugins/Experimental/Interchange/Runtime/Source/Pipelines/Private/InterchangeGenericAnimationPipeline.cpp
--- a/Engine/Plugins/Experimental/Interchange/Runtime/Source/Pipelines/Private/InterchangeGenericAnimationPipeline.cpp
+++ b/Engine/Plugins/Experimental/Interchange/Runtime/Source/Pipelines/Private/InterchangeGenericAnimationPipeline.cpp
@@ -132,6 +132,26 @@ void UInterchangeGenericAnimationPipeline::ExecutePreImportPipeline(UInterchange
 		SourceDatas.Add(SourceData);
 	}
 
+	TArray<UInterchangeSkeletalAnimationTrackNode*> TrackNodes;
+	BaseNodeContainer->IterateNodesOfType<UInterchangeSkeletalAnimationTrackNode>([&](const FString& NodeUid, UInterchangeSkeletalAnimationTrackNode* Node)
+		{
+			TrackNodes.Add(Node);
+		});
+
+	if (TrackNodes.Num() > 0)
+	{
+		//UInterchangeSkeletalAnimationTrackNode is exclusively and only used by the GLTF translator
+		//The class will be re-visited to unify management of skeletal animations in Interchange.
+		for (UInterchangeSkeletalAnimationTrackNode* TrackNode : TrackNodes)
+		{
+			if (TrackNode)
+			{
+				CreateAnimSequenceFactoryNode(*TrackNode);
+			}
+		}
+		return;
+	}
+
 	double SampleRate = 30.0;
 	double RangeStart = 0;
 	double RangeStop = 0;
@@ -564,3 +584,326 @@ void UInterchangeGenericAnimationPipeline::CreateAnimationTrackSetFactoryNode(UI
 	BaseNodeContainer->AddNode(FactoryNode);
 }
 
+void UInterchangeGenericAnimationPipeline::CreateAnimSequenceFactoryNode(UInterchangeSkeletalAnimationTrackNode& TrackNode)
+{
+	FString SkeletonNodeUid;
+	if (!ensure(TrackNode.GetCustomSkeletonNodeUid(SkeletonNodeUid)))
+	{
+		// TODO: Warn something wrong happened
+		return;
+	}
+
+	const FString SkeletonFactoryNodeUid = UInterchangeFactoryBaseNode::BuildFactoryNodeUid(SkeletonNodeUid);
+	const UInterchangeSkeletonFactoryNode* SkeletonFactoryNode = Cast<const UInterchangeSkeletonFactoryNode>(BaseNodeContainer->GetFactoryNode(SkeletonFactoryNodeUid));
+	if (!ensure(SkeletonFactoryNode))
+	{
+		// TODO: Warn something wrong happened
+		return;
+	}
+
+	FString SkeletalMeshNodeUid;
+	const UInterchangeSkeletalMeshFactoryNode* SkeletalMeshFactoryNode = nullptr;
+	if (TrackNode.GetCustomSkeletalMeshNodeUid(SkeletalMeshNodeUid))
+	{
+		const FString SkeletalMeshFactoryNodeUid = UInterchangeFactoryBaseNode::BuildFactoryNodeUid(SkeletalMeshNodeUid);
+		SkeletalMeshFactoryNode = Cast<const UInterchangeSkeletalMeshFactoryNode>(BaseNodeContainer->GetFactoryNode(SkeletalMeshFactoryNodeUid));
+	}
+
+	// NOTE: See whether it would not make sense to add a method to UInterchangeSkeletalAnimationTrackNode
+	// to collect all the mesh unique ids of the morph targets.
+	// In that case, GetCustomSkeletalMeshNodeUid might not be necessary.
+	TArray<const UInterchangeMeshNode*> MorphTargetNodes;
+	if (SkeletalMeshFactoryNode)
+	{
+		//Get the LOD data factory node
+		TArray<FString> LodDataChildren;
+		SkeletalMeshFactoryNode->GetLodDataUniqueIds(LodDataChildren);
+
+		for (const FString& ChildUid : LodDataChildren)
+		{
+			if (const UInterchangeSkeletalMeshLodDataNode* LodData = Cast<UInterchangeSkeletalMeshLodDataNode>(BaseNodeContainer->GetNode(ChildUid)))
+			{
+				TArray<FString> MeshUids;
+				LodData->GetMeshUids(MeshUids);
+				for (const FString& MeshUid : MeshUids)
+				{
+					if (const UInterchangeBaseNode* BaseNode = BaseNodeContainer->GetNode(MeshUid))
+					{
+						FString RealMeshUid;
+						if (const UInterchangeSceneNode* SceneNode = Cast<UInterchangeSceneNode>(BaseNode))
+						{
+							SceneNode->GetCustomAssetInstanceUid(RealMeshUid);
+						}
+						else
+						{
+							RealMeshUid = MeshUid;
+						}
+
+						if (const UInterchangeMeshNode* MeshNode = Cast<UInterchangeMeshNode>(BaseNodeContainer->GetNode(RealMeshUid)))
+						{
+							TArray<FString> MorphTargetUids;
+							MeshNode->GetMorphTargetDependencies(MorphTargetUids);
+
+							TArray<const UInterchangeMeshNode*> LocalMorphTargetNodes;
+
+							for (const FString& MorphTargetUid : MorphTargetUids)
+							{
+								if (const UInterchangeMeshNode* MorphTargetNode = Cast<UInterchangeMeshNode>(BaseNodeContainer->GetNode(MorphTargetUid)))
+								{
+									TOptional<FString> MorphTargetAnimationPayloadKey = MorphTargetNode->GetAnimationCurvePayLoadKey();
+									if (MorphTargetAnimationPayloadKey.IsSet())
+									{
+										LocalMorphTargetNodes.AddUnique(MorphTargetNode);
+									}
+								}
+							}
+							if (LocalMorphTargetNodes.Num() > 0)
+							{
+								// This is weird logic. Why not break?
+								MorphTargetNodes.Reset(LocalMorphTargetNodes.Num());
+								MorphTargetNodes.Append(LocalMorphTargetNodes);
+							}
+						}
+					}
+				}
+			}
+		}
+	}
+
+	double SampleRate = 30.;
+	double StartTime = 0.;
+	double StopTime = 0.;
+	bool bTimeRangeIsValid = false;
+
+	if (bImportBoneTracks)
+	{
+		if (const UInterchangeSourceNode* SourceNode = UInterchangeSourceNode::GetUniqueInstance(BaseNodeContainer))
+		{
+			int32 Numerator, Denominator;
+			if (!bUse30HzToBakeBoneAnimation && CustomBoneAnimationSampleRate == 0 && SourceNode->GetCustomSourceFrameRateNumerator(Numerator))
+			{
+				if (SourceNode->GetCustomSourceFrameRateDenominator(Denominator) && Denominator > 0 && Numerator > 0)
+				{
+					SampleRate = static_cast<double>(Numerator) / static_cast<double>(Denominator);
+				}
+			}
+			else if ((!bUse30HzToBakeBoneAnimation && CustomBoneAnimationSampleRate > 0))
+			{
+				SampleRate = static_cast<double>(CustomBoneAnimationSampleRate);
+			}
+
+			if (AnimationRange == EInterchangeAnimationRange::Timeline)
+			{
+				if (SourceNode->GetCustomSourceTimelineStart(StartTime))
+				{
+					if (SourceNode->GetCustomSourceTimelineEnd(StopTime))
+					{
+						bTimeRangeIsValid = true;
+					}
+				}
+			}
+			else if (AnimationRange == EInterchangeAnimationRange::Animated)
+			{
+				if (SourceNode->GetCustomAnimatedTimeStart(StartTime))
+				{
+					if (SourceNode->GetCustomAnimatedTimeEnd(StopTime))
+					{
+						bTimeRangeIsValid = true;
+					}
+				}
+			}
+			else if (AnimationRange == EInterchangeAnimationRange::SetRange)
+			{
+				StartTime = static_cast<double>(FrameImportRange.Min) / SampleRate;
+				StopTime = static_cast<double>(FrameImportRange.Max) / SampleRate;
+				bTimeRangeIsValid = true;
+			}
+		}
+		else
+		{
+			bTimeRangeIsValid = TrackNode.GetCustomAnimationSampleRate(SampleRate);
+			bTimeRangeIsValid &= TrackNode.GetCustomAnimationStartTime(StartTime);
+			bTimeRangeIsValid &= TrackNode.GetCustomAnimationStopTime(StopTime);
+		}
+
+		FFrameRate FrameRate = UE::Interchange::Animation::ConvertSampleRatetoFrameRate(SampleRate);
+
+		const double SequenceLength = FMath::Max<double>(StopTime - StartTime, MINIMUM_ANIMATION_LENGTH);
+
+		const float SubFrame = FrameRate.AsFrameTime(SequenceLength).GetSubFrame();
+
+		if (!FMath::IsNearlyZero(SubFrame, KINDA_SMALL_NUMBER) && !FMath::IsNearlyEqual(SubFrame, 1.0f, KINDA_SMALL_NUMBER))
+		{
+			UInterchangeResultError_Generic* Message = AddMessage<UInterchangeResultError_Generic>();
+			Message->SourceAssetName = SourceDatas[0]->GetFilename();
+			Message->DestinationAssetName = TrackNode.GetDisplayLabel();
+			Message->AssetType = UAnimSequence::StaticClass();
+			Message->Text = FText::Format(NSLOCTEXT("UInterchangeGenericAnimationPipeline", "WrongSequenceLength", "Animation length {0} is not compatible with import frame-rate {1} (sub frame {2}), animation has to be frame-border aligned."),
+				FText::AsNumber(SequenceLength), FrameRate.ToPrettyText(), FText::AsNumber(SubFrame));
+			//Skip this anim sequence factory node
+			return;
+		}
+	}
+
+	const FString AnimSequenceUid = TEXT("\\AnimSequence") + TrackNode.GetUniqueID();
+
+	UInterchangeAnimSequenceFactoryNode* AnimSequenceFactoryNode = NewObject<UInterchangeAnimSequenceFactoryNode>(BaseNodeContainer, NAME_None);
+	AnimSequenceFactoryNode->InitializeAnimSequenceNode(AnimSequenceUid, TrackNode.GetDisplayLabel());
+
+	AnimSequenceFactoryNode->SetCustomSkeletonFactoryNodeUid(SkeletonFactoryNode->GetUniqueID());
+	if (SkeletalMeshFactoryNode)
+	{
+		AnimSequenceFactoryNode->AddFactoryDependencyUid(SkeletalMeshFactoryNode->GetUniqueID());
+	}
+
+	AnimSequenceFactoryNode->SetCustomImportBoneTracks(bImportBoneTracks);
+	AnimSequenceFactoryNode->SetCustomImportBoneTracksSampleRate(SampleRate);
+	if (bTimeRangeIsValid)
+	{
+		AnimSequenceFactoryNode->SetCustomImportBoneTracksRangeStart(StartTime);
+		AnimSequenceFactoryNode->SetCustomImportBoneTracksRangeStop(StopTime);
+	}
+
+	AnimSequenceFactoryNode->SetCustomImportAttributeCurves(bImportCustomAttribute);
+	AnimSequenceFactoryNode->SetCustomDoNotImportCurveWithZero(bDoNotImportCurveWithZero);
+	AnimSequenceFactoryNode->SetCustomRemoveCurveRedundantKeys(bRemoveCurveRedundantKeys);
+	AnimSequenceFactoryNode->SetCustomDeleteExistingMorphTargetCurves(bDeleteExistingMorphTargetCurves);
+	AnimSequenceFactoryNode->SetCustomDeleteExistingCustomAttributeCurves(bDeleteExistingCustomAttributeCurves);
+	AnimSequenceFactoryNode->SetCustomDeleteExistingNonCurveCustomAttributes(bDeleteExistingNonCurveCustomAttributes);
+
+	AnimSequenceFactoryNode->SetCustomMaterialDriveParameterOnCustomAttribute(bSetMaterialDriveParameterOnCustomAttribute);
+	for (const FString& MaterialSuffixe : MaterialCurveSuffixes)
+	{
+		AnimSequenceFactoryNode->SetAnimatedMaterialCurveSuffixe(MaterialSuffixe);
+	}
+
+	//Add the animated morph targets uid so the factory can import them
+	if (MorphTargetNodes.Num() > 0)
+	{
+		for (const UInterchangeMeshNode* MorphTargetNode : MorphTargetNodes)
+		{
+			AnimSequenceFactoryNode->SetAnimatedMorphTargetDependencyUid(MorphTargetNode->GetUniqueID());
+		}
+	}
+
+	//USkeleton cannot be created without a valid skeletal mesh
+	const FString SkeletonUid = SkeletonFactoryNode->GetUniqueID();
+	AnimSequenceFactoryNode->AddFactoryDependencyUid(SkeletonUid);
+
+	FString RootJointUid;
+	if (SkeletonFactoryNode->GetCustomRootJointUid(RootJointUid))
+	{
+		// NOTE: Could this be added as an array of FString attributes on the UInterchangeSkeletalAnimationTrackNode
+#if WITH_EDITOR
+		//Iterate all joints to set the meta data value in the anim sequence factory node
+		UE::Interchange::Private::FSkeletonHelper::RecursiveAddSkeletonMetaDataValues(BaseNodeContainer, AnimSequenceFactoryNode, RootJointUid);
+#endif //WITH_EDITOR
+
+		const TArray<FString> CustomAttributeNamesToImport = UAnimationSettings::Get()->GetBoneCustomAttributeNamesToImport();
+
+		BaseNodeContainer->IterateNodeChildren(RootJointUid, [&AnimSequenceFactoryNode, &CustomAttributeNamesToImport](const UInterchangeBaseNode* Node)
+			{
+				if (const UInterchangeSceneNode* SceneNode = Cast<UInterchangeSceneNode>(Node))
+				{
+					using namespace UE::Interchange;
+
+					FString BoneName = SceneNode->GetDisplayLabel();
+					bool bImportAllAttributesOnBone = UAnimationSettings::Get()->BoneNamesWithCustomAttributes.Contains(BoneName);
+
+					TArray<FInterchangeUserDefinedAttributeInfo> AttributeInfos;
+					UInterchangeUserDefinedAttributesAPI::GetUserDefinedAttributeInfos(SceneNode, AttributeInfos);
+					for (const FInterchangeUserDefinedAttributeInfo& AttributeInfo : AttributeInfos)
+					{
+						if (AttributeInfo.PayloadKey.IsSet())
+						{
+							bool bDecimalType = AttributeInfo.Type == EAttributeTypes::Float ||
+								AttributeInfo.Type == EAttributeTypes::Float16 ||
+								AttributeInfo.Type == EAttributeTypes::Double;
+
+							const bool bForceImportBoneCustomAttribute = CustomAttributeNamesToImport.Contains(AttributeInfo.Name);
+							//Material attribute curve
+							if (!bImportAllAttributesOnBone && bDecimalType && !bForceImportBoneCustomAttribute)
+							{
+								AnimSequenceFactoryNode->SetAnimatedAttributeCurveName(AttributeInfo.Name);
+							}
+							else if (bForceImportBoneCustomAttribute || bImportAllAttributesOnBone)
+							{
+								AnimSequenceFactoryNode->SetAnimatedAttributeStepCurveName(AttributeInfo.Name);
+							}
+						}
+					}
+				}
+			});
+	}
+
+	const bool bImportOnlyAnimation = CommonSkeletalMeshesAndAnimationsProperties->bImportOnlyAnimations;
+
+	//Iterate dependencies
+	{
+		TArray<FString> SkeletalMeshNodeUids;
+		BaseNodeContainer->GetNodes(UInterchangeSkeletalMeshFactoryNode::StaticClass(), SkeletalMeshNodeUids);
+		for (const FString& SkelMeshFactoryNodeUid : SkeletalMeshNodeUids)
+		{
+			if (const UInterchangeSkeletalMeshFactoryNode* ConstSkeletalMeshFactoryNode = Cast<const UInterchangeSkeletalMeshFactoryNode>(BaseNodeContainer->GetFactoryNode(SkelMeshFactoryNodeUid)))
+			{
+				TArray<FString> SkeletalMeshDependencies;
+				ConstSkeletalMeshFactoryNode->GetFactoryDependencies(SkeletalMeshDependencies);
+				for (const FString& SkeletalMeshDependencyUid : SkeletalMeshDependencies)
+				{
+					if (SkeletonUid.Equals(SkeletalMeshDependencyUid))
+					{
+						AnimSequenceFactoryNode->AddFactoryDependencyUid(SkelMeshFactoryNodeUid);
+						break;
+					}
+				}
+			}
+		}
+	}
+
+	if (!CommonSkeletalMeshesAndAnimationsProperties->Skeleton.IsNull())
+	{
+		bool bSkeletonCompatible = true;
+
+		//TODO: support skeleton helper in runtime
+#if WITH_EDITOR
+		bSkeletonCompatible = UE::Interchange::Private::FSkeletonHelper::IsCompatibleSkeleton(CommonSkeletalMeshesAndAnimationsProperties->Skeleton, RootJointUid, BaseNodeContainer);
+#endif
+		if (bSkeletonCompatible)
+		{
+			FSoftObjectPath SkeletonSoftObjectPath(CommonSkeletalMeshesAndAnimationsProperties->Skeleton.Get());
+			AnimSequenceFactoryNode->SetCustomSkeletonSoftObjectPath(SkeletonSoftObjectPath);
+		}
+		else
+		{
+			UInterchangeResultError_Generic* Message = AddMessage<UInterchangeResultError_Generic>();
+			Message->Text = FText::Format(NSLOCTEXT("UInterchangeGenericAnimationPipeline", "IncompatibleSkeleton", "Incompatible skeleton {0} when importing AnimSequence {1}."),
+				FText::FromString(CommonSkeletalMeshesAndAnimationsProperties->Skeleton->GetName()),
+				FText::FromString(TrackNode.GetDisplayLabel()));
+		}
+	}
+
+	{
+		TMap<FString, FString> SceneNodeAnimationPayloadKeys;
+		TrackNode.GetSceneNodeAnimationPayloadKeys(SceneNodeAnimationPayloadKeys);
+		for (const TPair<FString, FString>& Entry : SceneNodeAnimationPayloadKeys)
+		{
+			AnimSequenceFactoryNode->SetAnimationPayloadKeyForSceneNodeUid(Entry.Key, Entry.Value);
+		}
+	}
+
+	{
+		TMap<FString, FString> MorphTargetNodeAnimationPayloads;
+		TrackNode.GetMorphTargetNodeAnimationPayloadKeys(MorphTargetNodeAnimationPayloads);
+		for (const TPair<FString, FString>& Entry : MorphTargetNodeAnimationPayloads)
+		{
+			AnimSequenceFactoryNode->SetAnimationPayloadKeyForMorphTargetNodeUid(Entry.Key, Entry.Value);
+		}
+	}
+
+	UInterchangeUserDefinedAttributesAPI::DuplicateAllUserDefinedAttribute(&TrackNode, AnimSequenceFactoryNode, false);
+
+	AnimSequenceFactoryNode->AddTargetNodeUid(TrackNode.GetUniqueID());
+	TrackNode.AddTargetNodeUid(AnimSequenceFactoryNode->GetUniqueID());
+
+	BaseNodeContainer->AddNode(AnimSequenceFactoryNode);
+}
diff --git a/Engine/Plugins/Experimental/Interchange/Runtime/Source/Pipelines/Public/InterchangeGenericAnimationPipeline.h b/Engine/Plugins/Experimental/Interchange/Runtime/Source/Pipelines/Public/InterchangeGenericAnimationPipeline.h
--- a/Engine/Plugins/Experimental/Interchange/Runtime/Source/Pipelines/Public/InterchangeGenericAnimationPipeline.h
+++ b/Engine/Plugins/Experimental/Interchange/Runtime/Source/Pipelines/Public/InterchangeGenericAnimationPipeline.h
@@ -14,6 +14,7 @@
 
 class UInterchangeAnimationTrackSetNode;
 class UInterchangeMeshNode;
+class UInterchangeSkeletalAnimationTrackNode;
 
 /** Animation length type when importing */
 UENUM(BlueprintType)
@@ -116,6 +117,8 @@ private:
 
 	void CreateAnimationTrackSetFactoryNode(UInterchangeAnimationTrackSetNode& Node);
 
+	void CreateAnimSequenceFactoryNode(UInterchangeSkeletalAnimationTrackNode& Node);
+
 	UInterchangeBaseNodeContainer* BaseNodeContainer = nullptr;
 	TArray<const UInterchangeSourceData*> SourceDatas;
 
