Fixed an issue when importing a texture cube or a texture array via interchange.
The wrong override of SetupTextureSourceDataFromBulkData_Editor was called by SetupTextureSourceDataFromBulkData for those because the type information that is used to select the right override was lost when calling the function.

diff --git a/Engine/Plugins/Experimental/Interchange/Runtime/Source/Import/Private/Texture/InterchangeTextureFactory.cpp b/Engine/Plugins/Experimental/Interchange/Runtime/Source/Import/Private/Texture/InterchangeTextureFactory.cpp
--- a/Engine/Plugins/Experimental/Interchange/Runtime/Source/Import/Private/Texture/InterchangeTextureFactory.cpp
+++ b/Engine/Plugins/Experimental/Interchange/Runtime/Source/Import/Private/Texture/InterchangeTextureFactory.cpp
@@ -660,7 +660,7 @@ namespace UE::Interchange::Private::InterchangeTextureFactory
 			if (OptionalImage->IsSet())
 			{
 				FImportImage& Image = OptionalImage->GetValue();
-				SetupTextureSourceDataFromBulkData(TextureCube, Image, MoveTemp(ProcessedPayload.PayloadAndId));
+				SetupTextureSourceDataFromBulkData_Editor(TextureCube, Image, MoveTemp(ProcessedPayload.PayloadAndId));
 			}
 		}
 		else if (TOptional<FImportLightProfile>* OptionalLightProfile = ProcessedPayload.SettingsFromPayload.TryGet<TOptional<FImportLightProfile>>())
@@ -668,7 +668,7 @@ namespace UE::Interchange::Private::InterchangeTextureFactory
 			if (OptionalLightProfile->IsSet())
 			{
 				FImportLightProfile& LightProfile = OptionalLightProfile->GetValue();
-				SetupTextureSourceDataFromBulkData(TextureCube, LightProfile, MoveTemp(ProcessedPayload.PayloadAndId));
+				SetupTextureSourceDataFromBulkData_Editor(TextureCube, LightProfile, MoveTemp(ProcessedPayload.PayloadAndId));
 			}
 		}
 		else
@@ -715,7 +715,7 @@ namespace UE::Interchange::Private::InterchangeTextureFactory
 			{
 				FImportSlicedImage& SlicedImage = OptionalSlicedImage->GetValue();
 
-				SetupTextureSourceDataFromBulkData(Texture2DArray, SlicedImage, MoveTemp(ProcessedPayload.PayloadAndId));
+				SetupTextureSourceDataFromBulkData_Editor(Texture2DArray, SlicedImage, MoveTemp(ProcessedPayload.PayloadAndId));
 			}
 		}
 		else if (TOptional<FImportImage>* OptionalImage = ProcessedPayload.SettingsFromPayload.TryGet<TOptional<FImportImage>>())
@@ -723,7 +723,7 @@ namespace UE::Interchange::Private::InterchangeTextureFactory
 			if (OptionalImage->IsSet())
 			{
 				FImportImage& Image = OptionalImage->GetValue();
-				SetupTextureSourceDataFromBulkData(Texture2DArray, Image, MoveTemp(ProcessedPayload.PayloadAndId));
+				SetupTextureSourceDataFromBulkData_Editor(Texture2DArray, Image, MoveTemp(ProcessedPayload.PayloadAndId));
 			}
 		}
 		else if (TOptional<FImportLightProfile>* OptionalLightProfile = ProcessedPayload.SettingsFromPayload.TryGet<TOptional<FImportLightProfile>>())
@@ -731,7 +731,7 @@ namespace UE::Interchange::Private::InterchangeTextureFactory
 			if (OptionalLightProfile->IsSet())
 			{
 				FImportLightProfile& LightProfile = OptionalLightProfile->GetValue();
-				SetupTextureSourceDataFromBulkData(Texture2DArray, LightProfile, MoveTemp(ProcessedPayload.PayloadAndId));
+				SetupTextureSourceDataFromBulkData_Editor(Texture2DArray, LightProfile, MoveTemp(ProcessedPayload.PayloadAndId));
 			}
 		}
 		else
