Interchange Translators: Make sure they are always enabled when running automated tests.

diff --git a/Engine/Plugins/Experimental/Interchange/Runtime/Source/Import/Private/Fbx/InterchangeFbxTranslator.cpp b/Engine/Plugins/Experimental/Interchange/Runtime/Source/Import/Private/Fbx/InterchangeFbxTranslator.cpp
--- a/Engine/Plugins/Experimental/Interchange/Runtime/Source/Import/Private/Fbx/InterchangeFbxTranslator.cpp
+++ b/Engine/Plugins/Experimental/Interchange/Runtime/Source/Import/Private/Fbx/InterchangeFbxTranslator.cpp
@@ -48,7 +48,7 @@ EInterchangeTranslatorAssetType UInterchangeFbxTranslator::GetSupportedAssetType
 
 TArray<FString> UInterchangeFbxTranslator::GetSupportedFormats() const
 {
-	if (GInterchangeEnableFBXImport)
+	if (GInterchangeEnableFBXImport || GIsAutomationTesting)
 	{
 		TArray<FString> Formats{ TEXT("fbx;Filmbox") };
 		return Formats;
diff --git a/Engine/Plugins/Experimental/Interchange/Runtime/Source/Import/Private/Gltf/InterchangeGltfTranslator.cpp b/Engine/Plugins/Experimental/Interchange/Runtime/Source/Import/Private/Gltf/InterchangeGltfTranslator.cpp
--- a/Engine/Plugins/Experimental/Interchange/Runtime/Source/Import/Private/Gltf/InterchangeGltfTranslator.cpp
+++ b/Engine/Plugins/Experimental/Interchange/Runtime/Source/Import/Private/Gltf/InterchangeGltfTranslator.cpp
@@ -799,7 +799,7 @@ TArray<FString> UInterchangeGltfTranslator::GetSupportedFormats() const
 {
 	TArray<FString> GltfExtensions;
 
-	if ( GInterchangeEnableGLTFImport )
+	if ( GInterchangeEnableGLTFImport || GIsAutomationTesting )
 	{
 		GltfExtensions.Reserve(2);
 		GltfExtensions.Add(TEXT("gltf;GL Transmission Format"));
diff --git a/Engine/Plugins/Experimental/Interchange/Runtime/Source/Import/Private/MaterialX/InterchangeMaterialXTranslator.cpp b/Engine/Plugins/Experimental/Interchange/Runtime/Source/Import/Private/MaterialX/InterchangeMaterialXTranslator.cpp
--- a/Engine/Plugins/Experimental/Interchange/Runtime/Source/Import/Private/MaterialX/InterchangeMaterialXTranslator.cpp
+++ b/Engine/Plugins/Experimental/Interchange/Runtime/Source/Import/Private/MaterialX/InterchangeMaterialXTranslator.cpp
@@ -202,7 +202,7 @@ TArray<FString> UInterchangeMaterialXTranslator::GetSupportedFormats() const
 {
 	// Call to UInterchangeMaterialXTranslator::GetSupportedFormats is not supported out of game thread
 	// A more global solution must be found for translators which require some initialization
-	if(!IsInGameThread() || !GInterchangeEnableMaterialXImport)
+	if(!IsInGameThread() || (!GInterchangeEnableMaterialXImport && !GIsAutomationTesting))
 	{
 		return TArray<FString>{};
 	}
diff --git a/Engine/Plugins/Experimental/Interchange/Runtime/Source/Import/Private/Mesh/InterchangeOBJTranslator.cpp b/Engine/Plugins/Experimental/Interchange/Runtime/Source/Import/Private/Mesh/InterchangeOBJTranslator.cpp
--- a/Engine/Plugins/Experimental/Interchange/Runtime/Source/Import/Private/Mesh/InterchangeOBJTranslator.cpp
+++ b/Engine/Plugins/Experimental/Interchange/Runtime/Source/Import/Private/Mesh/InterchangeOBJTranslator.cpp
@@ -860,7 +860,7 @@ UInterchangeOBJTranslator::~UInterchangeOBJTranslator()
 
 TArray<FString> UInterchangeOBJTranslator::GetSupportedFormats() const
 {
-	if (GInterchangeEnableOBJImport)
+	if (GInterchangeEnableOBJImport || GIsAutomationTesting)
 	{
 		TArray<FString> Formats{ TEXT("obj;OBJ File Format") };
 		return Formats;
diff --git a/Engine/Plugins/Experimental/Interchange/Runtime/Source/Import/Private/Texture/InterchangeDDSTranslator.cpp b/Engine/Plugins/Experimental/Interchange/Runtime/Source/Import/Private/Texture/InterchangeDDSTranslator.cpp
--- a/Engine/Plugins/Experimental/Interchange/Runtime/Source/Import/Private/Texture/InterchangeDDSTranslator.cpp
+++ b/Engine/Plugins/Experimental/Interchange/Runtime/Source/Import/Private/Texture/InterchangeDDSTranslator.cpp
@@ -58,7 +58,7 @@ namespace UE::Interchange::Private::InterchangeDDSTranslator
 
 TArray<FString> UInterchangeDDSTranslator::GetSupportedFormats() const
 {
-	if (GInterchangeEnableDDSImport)
+	if (GInterchangeEnableDDSImport || GIsAutomationTesting)
 	{
 		TArray<FString> Formats{ TEXT("dds;DirectDraw Surface") };
 		return Formats;
diff --git a/Engine/Plugins/Experimental/Interchange/Runtime/Source/Import/Private/Texture/InterchangeIESTranslator.cpp b/Engine/Plugins/Experimental/Interchange/Runtime/Source/Import/Private/Texture/InterchangeIESTranslator.cpp
--- a/Engine/Plugins/Experimental/Interchange/Runtime/Source/Import/Private/Texture/InterchangeIESTranslator.cpp
+++ b/Engine/Plugins/Experimental/Interchange/Runtime/Source/Import/Private/Texture/InterchangeIESTranslator.cpp
@@ -20,7 +20,7 @@ static FAutoConsoleVariableRef CCvarInterchangeEnableIESImport(
 
 TArray<FString> UInterchangeIESTranslator::GetSupportedFormats() const
 {
-	if (GInterchangeEnableIESImport)
+	if (GInterchangeEnableIESImport || GIsAutomationTesting)
 	{
 		TArray<FString> Formats{ TEXT("ies;IES light profile") };
 		return Formats;
diff --git a/Engine/Plugins/Experimental/Interchange/Runtime/Source/Import/Private/Texture/InterchangeImageWrapperTranslator.cpp b/Engine/Plugins/Experimental/Interchange/Runtime/Source/Import/Private/Texture/InterchangeImageWrapperTranslator.cpp
--- a/Engine/Plugins/Experimental/Interchange/Runtime/Source/Import/Private/Texture/InterchangeImageWrapperTranslator.cpp
+++ b/Engine/Plugins/Experimental/Interchange/Runtime/Source/Import/Private/Texture/InterchangeImageWrapperTranslator.cpp
@@ -79,35 +79,35 @@ TArray<FString> UInterchangeImageWrapperTranslator::GetSupportedFormats() const
 	TArray<FString> Formats;
 	Formats.Reserve(7);
 
-	if (GInterchangeEnablePNGImport)
+	if (GInterchangeEnablePNGImport || GIsAutomationTesting)
 	{
 		Formats.Emplace(TEXT("png;Portable Network Graphic"));
 	}
 
-	if (GInterchangeEnableBMPImport)
+	if (GInterchangeEnableBMPImport || GIsAutomationTesting)
 	{
 		Formats.Emplace(TEXT("bmp;Bitmap image"));
 	}
 
-	if (GInterchangeEnableEXRImport)
+	if (GInterchangeEnableEXRImport || GIsAutomationTesting)
 	{
 		Formats.Emplace(TEXT("exr;OpenEXR image"));
 	}
 
-	if (GInterchangeEnableHDRImport)
+	if (GInterchangeEnableHDRImport || GIsAutomationTesting)
 	{
 		Formats.Emplace(TEXT("hdr;High Dynamic Range image"));
 	}
 
 #if WITH_LIBTIFF
-	if (GInterchangeEnableTIFFImport)
+	if (GInterchangeEnableTIFFImport || GIsAutomationTesting)
 	{
 		Formats.Emplace(TEXT("tif;Tag Image File Format"));
 		Formats.Emplace(TEXT("tiff;Tag Image File Format"));
 	}
 #endif
 
-	if (GInterchangeEnableTGAImport)
+	if (GInterchangeEnableTGAImport || GIsAutomationTesting)
 	{
 		Formats.Emplace(TEXT("tga;Targa image"));
 	}
diff --git a/Engine/Plugins/Experimental/Interchange/Runtime/Source/Import/Private/Texture/InterchangeJPGTranslator.cpp b/Engine/Plugins/Experimental/Interchange/Runtime/Source/Import/Private/Texture/InterchangeJPGTranslator.cpp
--- a/Engine/Plugins/Experimental/Interchange/Runtime/Source/Import/Private/Texture/InterchangeJPGTranslator.cpp
+++ b/Engine/Plugins/Experimental/Interchange/Runtime/Source/Import/Private/Texture/InterchangeJPGTranslator.cpp
@@ -28,7 +28,7 @@ TArray<FString> UInterchangeJPGTranslator::GetSupportedFormats() const
 {
 	TArray<FString> Formats;
 
-	if (GInterchangeEnableJPGImport)
+	if (GInterchangeEnableJPGImport || GIsAutomationTesting)
 	{
 		Formats.Reserve(2);
 
diff --git a/Engine/Plugins/Experimental/Interchange/Runtime/Source/Import/Private/Texture/InterchangePCXTranslator.cpp b/Engine/Plugins/Experimental/Interchange/Runtime/Source/Import/Private/Texture/InterchangePCXTranslator.cpp
--- a/Engine/Plugins/Experimental/Interchange/Runtime/Source/Import/Private/Texture/InterchangePCXTranslator.cpp
+++ b/Engine/Plugins/Experimental/Interchange/Runtime/Source/Import/Private/Texture/InterchangePCXTranslator.cpp
@@ -75,7 +75,7 @@ namespace UE
 
 TArray<FString> UInterchangePCXTranslator::GetSupportedFormats() const
 {
-	if (GInterchangeEnablePCXImport)
+	if (GInterchangeEnablePCXImport || GIsAutomationTesting)
 	{
 		TArray<FString> Formats{ TEXT("pcx;Picture Exchange") };
 		return Formats;
diff --git a/Engine/Plugins/Experimental/Interchange/Runtime/Source/Import/Private/Texture/InterchangePSDTranslator.cpp b/Engine/Plugins/Experimental/Interchange/Runtime/Source/Import/Private/Texture/InterchangePSDTranslator.cpp
--- a/Engine/Plugins/Experimental/Interchange/Runtime/Source/Import/Private/Texture/InterchangePSDTranslator.cpp
+++ b/Engine/Plugins/Experimental/Interchange/Runtime/Source/Import/Private/Texture/InterchangePSDTranslator.cpp
@@ -376,7 +376,7 @@ namespace UE
 
 TArray<FString> UInterchangePSDTranslator::GetSupportedFormats() const
 {
-	if (GInterchangeEnablePSDImport)
+	if (GInterchangeEnablePSDImport || GIsAutomationTesting)
 	{
 		TArray<FString> Formats{ TEXT("psd;Photoshop Document") };
 		return Formats;
