Interchange GLTF importer updated to be in parity with Datasmith regarding Camera properties.
Added usage of EnableDepthOfField/FocusMode variable getting/setting to Interchange Framework

diff --git a/Engine/Plugins/Experimental/Interchange/Runtime/Source/FactoryNodes/Public/InterchangeCineCameraFactoryNode.h b/Engine/Plugins/Experimental/Interchange/Runtime/Source/FactoryNodes/Public/InterchangeCineCameraFactoryNode.h
--- a/Engine/Plugins/Experimental/Interchange/Runtime/Source/FactoryNodes/Public/InterchangeCineCameraFactoryNode.h
+++ b/Engine/Plugins/Experimental/Interchange/Runtime/Source/FactoryNodes/Public/InterchangeCineCameraFactoryNode.h
@@ -53,13 +53,27 @@ public:
 		IMPLEMENT_NODE_ATTRIBUTE_SETTER_WITH_CUSTOM_DELEGATE_WITH_CUSTOM_CLASS(UInterchangeCineCameraFactoryNode, SensorHeight, float, UCineCameraComponent);
 	}
 
+	UFUNCTION(BlueprintCallable, Category = "Interchange | Node | CameraFactory")
+	bool GetCustomFocusMethod(ECameraFocusMethod& AttributeValue) const
+	{
+		IMPLEMENT_NODE_ATTRIBUTE_GETTER(FocusMethod, ECameraFocusMethod);
+	}
+
+	UFUNCTION(BlueprintCallable, Category = "Interchange | Node | CameraFactory")
+	bool SetCustomFocusMethod(const ECameraFocusMethod& AttributeValue, bool bAddApplyDelegate = true)
+	{
+		IMPLEMENT_NODE_ATTRIBUTE_SETTER_WITH_CUSTOM_DELEGATE_WITH_CUSTOM_CLASS(UInterchangeCineCameraFactoryNode, FocusMethod, ECameraFocusMethod, UCineCameraComponent);
+	}	
+
 private:
 	const UE::Interchange::FAttributeKey Macro_CustomFocalLengthKey = UE::Interchange::FAttributeKey(TEXT("FocalLength"));
 	const UE::Interchange::FAttributeKey Macro_CustomSensorWidthKey = UE::Interchange::FAttributeKey(TEXT("SensorWidth"));
 	const UE::Interchange::FAttributeKey Macro_CustomSensorHeightKey = UE::Interchange::FAttributeKey(TEXT("SensorHeight"));
+	const UE::Interchange::FAttributeKey Macro_CustomFocusMethodKey = UE::Interchange::FAttributeKey(TEXT("FocusMethod"));
 
 private:
 	IMPLEMENT_NODE_ATTRIBUTE_DELEGATE_BY_PROPERTYNAME(FocalLength, float, UCineCameraComponent, TEXT("CurrentFocalLength"));
 	IMPLEMENT_NODE_ATTRIBUTE_DELEGATE_BY_PROPERTYNAME(SensorWidth, float, UCineCameraComponent, TEXT("Filmback.SensorWidth"));
 	IMPLEMENT_NODE_ATTRIBUTE_DELEGATE_BY_PROPERTYNAME(SensorHeight, float, UCineCameraComponent, TEXT("Filmback.SensorHeight"));
+	IMPLEMENT_NODE_ATTRIBUTE_DELEGATE_BY_PROPERTYNAME(FocusMethod, ECameraFocusMethod, UCineCameraComponent, TEXT("FocusSettings.FocusMethod"));
 };
diff --git a/Engine/Plugins/Experimental/Interchange/Runtime/Source/Import/Private/Gltf/InterchangeGltfTranslator.cpp b/Engine/Plugins/Experimental/Interchange/Runtime/Source/Import/Private/Gltf/InterchangeGltfTranslator.cpp
--- a/Engine/Plugins/Experimental/Interchange/Runtime/Source/Import/Private/Gltf/InterchangeGltfTranslator.cpp
+++ b/Engine/Plugins/Experimental/Interchange/Runtime/Source/Import/Private/Gltf/InterchangeGltfTranslator.cpp
@@ -915,6 +915,27 @@ bool UInterchangeGltfTranslator::Translate( UInterchangeBaseNodeContainer& NodeC
 			UInterchangeCameraNode* CameraNode = NewObject< UInterchangeCameraNode >( &NodeContainer );
 			FString CameraNodeUid = TEXT("\\Camera\\") + GltfCamera.Name;
 			CameraNode->InitializeNode( CameraNodeUid, GltfCamera.Name, EInterchangeNodeContainerType::TranslatedAsset );
+
+			float       AspectRatio;
+			float       FocalLength;
+			const float SensorWidth = 36.f;  // mm
+			CameraNode->SetCustomSensorWidth(SensorWidth);
+			if (GltfCamera.bIsPerspective)
+			{
+				AspectRatio = GltfCamera.Perspective.AspectRatio;
+				FocalLength = (SensorWidth / GltfCamera.Perspective.AspectRatio) / (2.0 * tan(GltfCamera.Perspective.Fov / 2.0));
+			}
+			else
+			{
+				AspectRatio = GltfCamera.Orthographic.XMagnification / GltfCamera.Orthographic.YMagnification;
+				FocalLength = (SensorWidth / AspectRatio) / (AspectRatio * tan(AspectRatio / 4.0));  // can only approximate Fov
+			}
+
+			CameraNode->SetCustomSensorHeight(SensorWidth / AspectRatio);
+			CameraNode->SetCustomFocalLength(FocalLength);
+			CameraNode->SetCustomEnableDepthOfField(false);
+			// ignoring znear and zfar
+
 			NodeContainer.AddNode( CameraNode );
 			++CameraIndex;
 		}
diff --git a/Engine/Plugins/Experimental/Interchange/Runtime/Source/Nodes/Public/InterchangeCameraNode.h b/Engine/Plugins/Experimental/Interchange/Runtime/Source/Nodes/Public/InterchangeCameraNode.h
--- a/Engine/Plugins/Experimental/Interchange/Runtime/Source/Nodes/Public/InterchangeCameraNode.h
+++ b/Engine/Plugins/Experimental/Interchange/Runtime/Source/Nodes/Public/InterchangeCameraNode.h
@@ -64,8 +64,21 @@ public:
 		IMPLEMENT_NODE_ATTRIBUTE_SETTER_NODELEGATE(SensorHeight, float);
 	}
 
+	UFUNCTION(BlueprintCallable, Category = "Interchange | Node | Camera")
+	bool GetCustomEnableDepthOfField(bool& AttributeValue) const
+	{
+		IMPLEMENT_NODE_ATTRIBUTE_GETTER(EnableDepthOfField, bool);
+	}
+
+	UFUNCTION(BlueprintCallable, Category = "Interchange | Node | Camera")
+	bool SetCustomEnableDepthOfField(const bool& AttributeValue)
+	{
+		IMPLEMENT_NODE_ATTRIBUTE_SETTER_NODELEGATE(EnableDepthOfField, bool);
+	}
+
 private:
 	const UE::Interchange::FAttributeKey Macro_CustomFocalLengthKey = UE::Interchange::FAttributeKey(TEXT("FocalLength"));
 	const UE::Interchange::FAttributeKey Macro_CustomSensorWidthKey = UE::Interchange::FAttributeKey(TEXT("SensorWidth"));
 	const UE::Interchange::FAttributeKey Macro_CustomSensorHeightKey = UE::Interchange::FAttributeKey(TEXT("SensorHeight"));
+	const UE::Interchange::FAttributeKey Macro_CustomEnableDepthOfFieldKey = UE::Interchange::FAttributeKey(TEXT("EnableDepthOfField"));
 };
diff --git a/Engine/Plugins/Experimental/Interchange/Runtime/Source/Pipelines/Private/InterchangeGenericScenesPipeline.cpp b/Engine/Plugins/Experimental/Interchange/Runtime/Source/Pipelines/Private/InterchangeGenericScenesPipeline.cpp
--- a/Engine/Plugins/Experimental/Interchange/Runtime/Source/Pipelines/Private/InterchangeGenericScenesPipeline.cpp
+++ b/Engine/Plugins/Experimental/Interchange/Runtime/Source/Pipelines/Private/InterchangeGenericScenesPipeline.cpp
@@ -360,6 +360,13 @@ void UInterchangeGenericLevelPipeline::SetUpFactoryNode(UInterchangeActorFactory
 			{
 				CineCameraFactoryNode->SetCustomSensorWidth(SensorWidth);
 			}
+
+			bool bEnableDepthOfField;
+			if (CameraNode->GetCustomEnableDepthOfField(bEnableDepthOfField))
+			{
+				CineCameraFactoryNode->SetCustomFocusMethod(bEnableDepthOfField ? ECameraFocusMethod::Manual : ECameraFocusMethod::DoNotOverride);
+			}
+			
 		}
 	}
 }
