Interchange GLTF Importer Texcoord index handling fixes

diff --git a/Engine/Plugins/Experimental/Interchange/Runtime/Source/Import/Private/Gltf/InterchangeGltfTranslator.cpp b/Engine/Plugins/Experimental/Interchange/Runtime/Source/Import/Private/Gltf/InterchangeGltfTranslator.cpp
--- a/Engine/Plugins/Experimental/Interchange/Runtime/Source/Import/Private/Gltf/InterchangeGltfTranslator.cpp
+++ b/Engine/Plugins/Experimental/Interchange/Runtime/Source/Import/Private/Gltf/InterchangeGltfTranslator.cpp
@@ -253,9 +253,19 @@ void UInterchangeGltfTranslator::HandleGltfMaterialParameter( UInterchangeBaseNo
 
 		ColorNode->AddStringAttribute( UInterchangeShaderPortsAPI::MakeInputValueKey( Standard::Nodes::TextureSample::Inputs::Texture.ToString() ), TextureUid );
 
-		if ( TextureMap.bHasTextureTransform )
+		if (TextureMap.TexCoord > 0 || TextureMap.bHasTextureTransform)
 		{
-			HandleGltfTextureTransform( NodeContainer, TextureMap.TextureTransform, TextureMap.TexCoord, *ColorNode );
+			UInterchangeShaderNode* TexCoordNode = UInterchangeShaderNode::Create(&NodeContainer, TEXT("TexCoord"), ShaderNode.GetUniqueID());
+			TexCoordNode->SetCustomShaderType(Standard::Nodes::TextureCoordinate::Name.ToString());
+
+			TexCoordNode->AddInt32Attribute(UInterchangeShaderPortsAPI::MakeInputValueKey(Standard::Nodes::TextureCoordinate::Inputs::Index.ToString()), TextureMap.TexCoord);
+
+			UInterchangeShaderPortsAPI::ConnectDefaultOuputToInput(ColorNode, Standard::Nodes::TextureSample::Inputs::Coordinates.ToString(), TexCoordNode->GetUniqueID());
+
+			if (TextureMap.bHasTextureTransform)
+			{
+				HandleGltfTextureTransform(NodeContainer, TextureMap.TextureTransform, TextureMap.TexCoord, *TexCoordNode);
+			}
 		}
 
 		bool bNeedsFactorNode = false;
@@ -726,18 +736,6 @@ void UInterchangeGltfTranslator::HandleGltfTextureTransform( UInterchangeBaseNod
 {
 	using namespace UE::Interchange::Materials;
 
-	UInterchangeShaderNode* CurrentNode = &ShaderNode;
-
-	// Texture Coordinates
-	{
-		UInterchangeShaderNode* TexCoordNode = UInterchangeShaderNode::Create( &NodeContainer, TEXT("TexCoord"), ShaderNode.GetUniqueID() );
-		TexCoordNode->SetCustomShaderType( Standard::Nodes::TextureCoordinate::Name.ToString() );
-
-		TexCoordNode->AddInt32Attribute( UInterchangeShaderPortsAPI::MakeInputValueKey( Standard::Nodes::TextureCoordinate::Inputs::Index.ToString() ), TexCoordIndex );
-
-		CurrentNode = TexCoordNode;
-	}
-
 	// Scale
 	if ( !FMath::IsNearlyEqual( TextureTransform.Scale[0], 1.f ) || 
 		 !FMath::IsNearlyEqual( TextureTransform.Scale[1], 1.f ) )
@@ -746,7 +744,7 @@ void UInterchangeGltfTranslator::HandleGltfTextureTransform( UInterchangeBaseNod
 		TextureScale.X = TextureTransform.Scale[0];
 		TextureScale.Y = TextureTransform.Scale[1];
 
-		CurrentNode->SetAttribute< FVector2f >( UInterchangeShaderPortsAPI::MakeInputValueKey( Standard::Nodes::TextureCoordinate::Inputs::Scale.ToString() ), TextureScale );
+		ShaderNode.SetAttribute< FVector2f >( UInterchangeShaderPortsAPI::MakeInputValueKey( Standard::Nodes::TextureCoordinate::Inputs::Scale.ToString() ), TextureScale );
 	}
 
 	// Offset
@@ -757,7 +755,7 @@ void UInterchangeGltfTranslator::HandleGltfTextureTransform( UInterchangeBaseNod
 		TextureOffset.X = TextureTransform.Offset[0];
 		TextureOffset.Y = TextureTransform.Offset[1];
 
-		CurrentNode->SetAttribute< FVector2f >( UInterchangeShaderPortsAPI::MakeInputValueKey( Standard::Nodes::TextureCoordinate::Inputs::Offset.ToString() ), TextureOffset );
+		ShaderNode.SetAttribute< FVector2f >( UInterchangeShaderPortsAPI::MakeInputValueKey( Standard::Nodes::TextureCoordinate::Inputs::Offset.ToString() ), TextureOffset );
 	}
 
 	// Rotate
@@ -772,16 +770,12 @@ void UInterchangeGltfTranslator::HandleGltfTextureTransform( UInterchangeBaseNod
 
 		AngleRadians = 1.0f - ( AngleRadians / TWO_PI );
 
-		CurrentNode->AddFloatAttribute( UInterchangeShaderPortsAPI::MakeInputValueKey( Standard::Nodes::TextureCoordinate::Inputs::Rotate.ToString() ), AngleRadians );
+		ShaderNode.AddFloatAttribute( UInterchangeShaderPortsAPI::MakeInputValueKey( Standard::Nodes::TextureCoordinate::Inputs::Rotate.ToString() ), AngleRadians );
 
 		FVector2f RotationCenter = FVector2f::Zero();
-		CurrentNode->SetAttribute< FVector2f >( UInterchangeShaderPortsAPI::MakeInputValueKey( Standard::Nodes::TextureCoordinate::Inputs::RotationCenter.ToString() ), RotationCenter );
+		ShaderNode.SetAttribute< FVector2f >( UInterchangeShaderPortsAPI::MakeInputValueKey( Standard::Nodes::TextureCoordinate::Inputs::RotationCenter.ToString() ), RotationCenter );
 	}
 
-	if ( CurrentNode != &ShaderNode )
-	{
-		UInterchangeShaderPortsAPI::ConnectDefaultOuputToInput( &ShaderNode, Standard::Nodes::TextureSample::Inputs::Coordinates.ToString(), CurrentNode->GetUniqueID() );
-	}
 }
 
 EInterchangeTranslatorType UInterchangeGltfTranslator::GetTranslatorType() const
