Interchange GLTF importer updated to be in parity with Datasmith regarding Texture minfilter usage.
Added usage of Texture filter variable getting/setting to Interchange Framework

diff --git a/Engine/Plugins/Experimental/Interchange/Runtime/Source/Import/Private/Gltf/InterchangeGltfTranslator.cpp b/Engine/Plugins/Experimental/Interchange/Runtime/Source/Import/Private/Gltf/InterchangeGltfTranslator.cpp
--- a/Engine/Plugins/Experimental/Interchange/Runtime/Source/Import/Private/Gltf/InterchangeGltfTranslator.cpp
+++ b/Engine/Plugins/Experimental/Interchange/Runtime/Source/Import/Private/Gltf/InterchangeGltfTranslator.cpp
@@ -78,6 +78,22 @@ namespace UE::Interchange::Gltf::Private
 		}
 	}
 
+	EInterchangeTextureFilterMode ConvertFilter(GLTF::FSampler::EFilter Filter)
+	{
+		switch (Filter)
+		{
+			case GLTF::FSampler::EFilter::Nearest:
+				return EInterchangeTextureFilterMode::Nearest;
+			case GLTF::FSampler::EFilter::LinearMipmapNearest:
+				return EInterchangeTextureFilterMode::Bilinear;
+			case GLTF::FSampler::EFilter::LinearMipmapLinear:
+				return EInterchangeTextureFilterMode::Trilinear;
+				// Other glTF filter values have no direct correlation to Unreal
+			default:
+				return EInterchangeTextureFilterMode::Default;
+		}
+	}
+
 	bool CheckForVariants(const GLTF::FMesh& Mesh, int32 VariantCount, int32 MaterialCount)
 	{
 		for (const GLTF::FPrimitive& Primitive : Mesh.Primitives)
@@ -853,6 +869,7 @@ bool UInterchangeGltfTranslator::Translate( UInterchangeBaseNodeContainer& NodeC
 			const FString TextureName = GltfTexture.Name;
 
 			UInterchangeTexture2DNode* TextureNode = UInterchangeTexture2DNode::Create( &NodeContainer, TextureName );
+			TextureNode->SetCustomFilter(ConvertFilter(GltfTexture.Sampler.MinFilter));
 			TextureNode->SetPayLoadKey( LexToString( TextureIndex++ ) );
 
 			TextureNode->SetCustomWrapU( UE::Interchange::Gltf::Private::ConvertWrap( GltfTexture.Sampler.WrapS ) );
diff --git a/Engine/Plugins/Experimental/Interchange/Runtime/Source/Nodes/Public/InterchangeTextureNode.h b/Engine/Plugins/Experimental/Interchange/Runtime/Source/Nodes/Public/InterchangeTextureNode.h
--- a/Engine/Plugins/Experimental/Interchange/Runtime/Source/Nodes/Public/InterchangeTextureNode.h
+++ b/Engine/Plugins/Experimental/Interchange/Runtime/Source/Nodes/Public/InterchangeTextureNode.h
@@ -31,6 +31,16 @@ enum class EInterchangeTextureWrapMode : uint8
 	Mirror
 };
 
+UENUM(BlueprintType)
+enum class EInterchangeTextureFilterMode : uint8
+{
+	Nearest,
+	Bilinear,
+	Trilinear,
+	/** Use setting from the Texture Group. */
+	Default
+};
+
 UCLASS(BlueprintType, Abstract, Experimental)
 class INTERCHANGENODES_API UInterchangeTextureNode : public UInterchangeBaseNode
 {
@@ -125,8 +135,21 @@ public:
 		IMPLEMENT_NODE_ATTRIBUTE_SETTER_NODELEGATE(bFlipGreenChannel, bool)
 	}
 
+	UFUNCTION(BlueprintCallable, Category = "Interchange | Node | Texture")
+	bool GetCustomFilter(EInterchangeTextureFilterMode& AttributeValue) const
+	{
+		IMPLEMENT_NODE_ATTRIBUTE_GETTER(Filter, EInterchangeTextureFilterMode);
+	}
+
+	UFUNCTION(BlueprintCallable, Category = "Interchange | Node | Texture")
+	bool SetCustomFilter(const EInterchangeTextureFilterMode& AttributeValue)
+	{
+		IMPLEMENT_NODE_ATTRIBUTE_SETTER_NODELEGATE(Filter, EInterchangeTextureFilterMode)
+	}
+
 private:
 
 	IMPLEMENT_NODE_ATTRIBUTE_KEY(SRGB)
 	IMPLEMENT_NODE_ATTRIBUTE_KEY(bFlipGreenChannel)
+	IMPLEMENT_NODE_ATTRIBUTE_KEY(Filter)
 };
diff --git a/Engine/Plugins/Experimental/Interchange/Runtime/Source/Pipelines/Private/InterchangeGenericTexturePipeline.cpp b/Engine/Plugins/Experimental/Interchange/Runtime/Source/Pipelines/Private/InterchangeGenericTexturePipeline.cpp
--- a/Engine/Plugins/Experimental/Interchange/Runtime/Source/Pipelines/Private/InterchangeGenericTexturePipeline.cpp
+++ b/Engine/Plugins/Experimental/Interchange/Runtime/Source/Pipelines/Private/InterchangeGenericTexturePipeline.cpp
@@ -24,6 +24,7 @@
 #include "InterchangeTextureLightProfileNode.h"
 #include "InterchangeTextureNode.h"
 #include "Misc/Paths.h"
+#include "Engine/Texture.h"
 
 #if WITH_EDITOR
 #include "NormalMapIdentification.h"
@@ -292,6 +293,21 @@ UInterchangeTextureFactoryNode* UInterchangeGenericTexturePipeline::CreateTextur
 		TextureFactoryNode->SetCustombFlipGreenChannel(bFlipGreenChannel);
 	}
 
+	using FInterchangeTextureFilterMode = std::underlying_type_t<EInterchangeTextureFilterMode>;
+	using FTextureFilter = std::underlying_type_t<TextureFilter>;
+	using FCommonTextureFilterModes = std::common_type_t<FInterchangeTextureFilterMode, FTextureFilter>;
+
+	static_assert(FCommonTextureFilterModes(EInterchangeTextureFilterMode::Nearest) == FCommonTextureFilterModes(TextureFilter::TF_Nearest), "EInterchangeTextureFilterMode::Nearest differs from TextureFilter::TF_Nearest");
+	static_assert(FCommonTextureFilterModes(EInterchangeTextureFilterMode::Bilinear) == FCommonTextureFilterModes(TextureFilter::TF_Bilinear), "EInterchangeTextureFilterMode::Bilinear differs from TextureFilter::TF_Bilinear");
+	static_assert(FCommonTextureFilterModes(EInterchangeTextureFilterMode::Trilinear) == FCommonTextureFilterModes(TextureFilter::TF_Trilinear), "EInterchangeTextureFilterMode::Trilinear differs from TextureFilter::TF_Trilinear");
+	static_assert(FCommonTextureFilterModes(EInterchangeTextureFilterMode::Default) == FCommonTextureFilterModes(TextureFilter::TF_Default), "EInterchangeTextureFilterMode::Default differs from TextureFilter::TF_Default");
+
+	if (EInterchangeTextureFilterMode TextureFilter; TextureNode->GetCustomFilter(TextureFilter))
+	{
+
+		TextureFactoryNode->SetCustomFilter(uint8(TextureFilter));
+	}
+
 #if WITH_EDITORONLY_DATA
 	if (bPreferCompressedSourceData)
 	{
