Fixed issue with lightmap UV destination index when the imported mesh has more  UV channels than the default lightmap UV destination index set on the pipeline.

diff --git a/Engine/Plugins/Experimental/Interchange/Runtime/Source/Import/Private/Mesh/InterchangeStaticMeshFactory.cpp b/Engine/Plugins/Experimental/Interchange/Runtime/Source/Import/Private/Mesh/InterchangeStaticMeshFactory.cpp
--- a/Engine/Plugins/Experimental/Interchange/Runtime/Source/Import/Private/Mesh/InterchangeStaticMeshFactory.cpp
+++ b/Engine/Plugins/Experimental/Interchange/Runtime/Source/Import/Private/Mesh/InterchangeStaticMeshFactory.cpp
@@ -374,6 +374,15 @@ UObject* UInterchangeStaticMeshFactory::CreateAsset(const FCreateAssetParams& Ar
 	{
 		const bool bIsAReimport = Arguments.ReimportObject != nullptr;
 		SetupSourceModelsSettings(*StaticMesh, LodMeshDescriptions, PrevLodCount, FinalLodCount, bIsAReimport);
+
+		// SetupSourceModelsSettings can change the destination lightmap UV index
+		// Make sure the destination lightmap UV index on the factory node takes
+		// in account the potential change
+		int32 FactoryDstLightmapIndex;
+		if (StaticMeshFactoryNode->GetCustomDstLightmapIndex(FactoryDstLightmapIndex) && StaticMesh->GetLightMapCoordinateIndex() > FactoryDstLightmapIndex)
+		{
+			StaticMeshFactoryNode->SetCustomDstLightmapIndex(StaticMesh->GetLightMapCoordinateIndex());
+		}
 	}
 #endif // WITH_EDITOR
 
