Interchange GLTF importer updated to be in parity with Datasmith regarding Lights.

diff --git a/Engine/Plugins/Experimental/Interchange/Runtime/Source/Import/Private/Gltf/InterchangeGltfTranslator.cpp b/Engine/Plugins/Experimental/Interchange/Runtime/Source/Import/Private/Gltf/InterchangeGltfTranslator.cpp
--- a/Engine/Plugins/Experimental/Interchange/Runtime/Source/Import/Private/Gltf/InterchangeGltfTranslator.cpp
+++ b/Engine/Plugins/Experimental/Interchange/Runtime/Source/Import/Private/Gltf/InterchangeGltfTranslator.cpp
@@ -957,11 +957,56 @@ bool UInterchangeGltfTranslator::Translate( UInterchangeBaseNodeContainer& NodeC
 		int32 LightIndex = 0;
 		for ( const GLTF::FLight& GltfLight : GltfAsset.Lights )
 		{
-			UInterchangeLightNode* LightNode = NewObject< UInterchangeLightNode >( &NodeContainer );
 			FString LightNodeUid = TEXT("\\Light\\") + GltfLight.Name;
-			LightNode->InitializeNode( LightNodeUid, GltfLight.Name, EInterchangeNodeContainerType::TranslatedAsset );
-			NodeContainer.AddNode( LightNode );
-			++LightIndex;
+
+			switch (GltfLight.Type)
+			{
+			case GLTF::FLight::EType::Directional:
+			{
+				UInterchangeDirectionalLightNode* LightNode = NewObject< UInterchangeDirectionalLightNode >(&NodeContainer);
+				LightNode->InitializeNode(LightNodeUid, GltfLight.Name, EInterchangeNodeContainerType::TranslatedAsset);
+
+				LightNode->SetCustomLightColor(FLinearColor(GltfLight.Color));
+				LightNode->SetCustomIntensity(GltfLight.Intensity);
+
+				NodeContainer.AddNode(LightNode);
+				++LightIndex;
+			}
+				break;
+			case GLTF::FLight::EType::Point:
+			{
+				UInterchangePointLightNode* LightNode = NewObject< UInterchangePointLightNode >(&NodeContainer);
+				LightNode->InitializeNode(LightNodeUid, GltfLight.Name, EInterchangeNodeContainerType::TranslatedAsset);
+				
+				LightNode->SetCustomIntensityUnits(EInterchangeLightUnits::Candelas);
+				LightNode->SetCustomLightColor(FLinearColor(GltfLight.Color));
+				LightNode->SetCustomIntensity(GltfLight.Intensity);
+				
+				LightNode->SetCustomAttenuationRadius(GltfLight.Range * GltfUnitConversionMultiplier);
+
+				NodeContainer.AddNode(LightNode);
+				++LightIndex;
+			}
+				break;
+			case GLTF::FLight::EType::Spot:
+			{
+				UInterchangeSpotLightNode* LightNode = NewObject< UInterchangeSpotLightNode >(&NodeContainer);
+				LightNode->InitializeNode(LightNodeUid, GltfLight.Name, EInterchangeNodeContainerType::TranslatedAsset);
+
+				LightNode->SetCustomIntensityUnits(EInterchangeLightUnits::Candelas);
+				LightNode->SetCustomLightColor(FLinearColor(GltfLight.Color));
+				LightNode->SetCustomIntensity(GltfLight.Intensity);
+
+				LightNode->SetCustomInnerConeAngle(FMath::RadiansToDegrees(GltfLight.Spot.InnerConeAngle));
+				LightNode->SetCustomOuterConeAngle(FMath::RadiansToDegrees(GltfLight.Spot.OuterConeAngle));
+				
+				NodeContainer.AddNode(LightNode);
+				++LightIndex;
+			}
+				break;
+			default:
+				break;
+			}
 		}
 	}
 
