Interchange
- Fix ensure when a texture translator is created on a background thread by an asset creation process

diff --git a/Engine/Plugins/Experimental/Interchange/Runtime/Source/Import/Private/Texture/InterchangeImageWrapperTranslator.cpp b/Engine/Plugins/Experimental/Interchange/Runtime/Source/Import/Private/Texture/InterchangeImageWrapperTranslator.cpp
--- a/Engine/Plugins/Experimental/Interchange/Runtime/Source/Import/Private/Texture/InterchangeImageWrapperTranslator.cpp
+++ b/Engine/Plugins/Experimental/Interchange/Runtime/Source/Import/Private/Texture/InterchangeImageWrapperTranslator.cpp
@@ -65,12 +65,12 @@ static FAutoConsoleVariableRef CCvarInterchangeEnableTGAImport(
 	ECVF_Default);
 
 UInterchangeImageWrapperTranslator::UInterchangeImageWrapperTranslator()
-	: bFillPNGZeroAlpha(true)
 {
-	if (GConfig)
+	if (IsTemplate() && GConfig)
 	{
 		ensure(IsInGameThread()); // Reading config values isn't threadsafe
-		GConfig->GetBool(TEXT("TextureImporter"), TEXT("FillPNGZeroAlpha"), bFillPNGZeroAlpha, GEditorIni);
+		bFillPNGZeroAlpha = true;
+		GConfig->GetBool(TEXT("TextureImporter"), TEXT("FillPNGZeroAlpha"), bFillPNGZeroAlpha.GetValue(), GEditorIni);
 	}
 }
 
@@ -202,7 +202,7 @@ TOptional<UE::Interchange::FImportImage> UInterchangeImageWrapperTranslator::Get
 
 			if (ImageFormat == EImageFormat::PNG)
 			{
-				if (bFillPNGZeroAlpha)
+				if (GetDefault<UInterchangeImageWrapperTranslator>()->bFillPNGZeroAlpha.GetValue())
 				{
 					// Replace the pixels with 0.0 alpha with a color value from the nearest neighboring color which has a non-zero alpha
 					UE::TextureUtilitiesCommon::FillZeroAlphaPNGData(PayloadData.SizeX, PayloadData.SizeY, PayloadData.Format, reinterpret_cast<uint8*>(PayloadData.RawData.GetData()));
diff --git a/Engine/Plugins/Experimental/Interchange/Runtime/Source/Import/Public/Texture/InterchangeImageWrapperTranslator.h b/Engine/Plugins/Experimental/Interchange/Runtime/Source/Import/Public/Texture/InterchangeImageWrapperTranslator.h
--- a/Engine/Plugins/Experimental/Interchange/Runtime/Source/Import/Public/Texture/InterchangeImageWrapperTranslator.h
+++ b/Engine/Plugins/Experimental/Interchange/Runtime/Source/Import/Public/Texture/InterchangeImageWrapperTranslator.h
@@ -51,5 +51,5 @@ public:
 	TOptional<UE::Interchange::FImportImage> GetTexturePayloadDataFromBuffer(const TArray64<uint8>& SourceDataBuffer) const;
 
 private:
-	bool bFillPNGZeroAlpha;
+	TOptional<bool> bFillPNGZeroAlpha; // Set on the CDO only
 };
\ No newline at end of file
