Interchange: Do not run the test if the interchange worker is not available.

#rnx

diff --git a/Engine/Plugins/Experimental/Interchange/Runtime/Source/Dispatcher/Private/InterchangeDispatcher.cpp b/Engine/Plugins/Experimental/Interchange/Runtime/Source/Dispatcher/Private/InterchangeDispatcher.cpp
--- a/Engine/Plugins/Experimental/Interchange/Runtime/Source/Dispatcher/Private/InterchangeDispatcher.cpp
+++ b/Engine/Plugins/Experimental/Interchange/Runtime/Source/Dispatcher/Private/InterchangeDispatcher.cpp
@@ -207,6 +207,39 @@ namespace UE
 			return CompletedTaskCount == TaskPool.Num();
 		}
 
+		//this is a static function
+		bool FInterchangeDispatcher::IsInterchangeWorkerAvailable()
+		{
+			const FString ResultFolder = FPaths::ConvertRelativePathToFull(FPaths::ProjectSavedDir());
+			TUniquePtr<UE::Interchange::FInterchangeDispatcher> Dispatcher = MakeUnique<UE::Interchange::FInterchangeDispatcher>(ResultFolder);
+			bool bWorkerValid = false;
+			if (ensure(Dispatcher.IsValid()))
+			{
+				Dispatcher->StartProcess();
+				//Wait the conection handshake is done
+				//The interchange worker will validate the version and issue a error message if the connection is incorrect
+				//If version is good the interchange worker will issue a ping command
+				while (Dispatcher->IsHandlerAlive() && Dispatcher->GetInterchangeWorkerFatalError().IsEmpty())
+				{
+					if (Dispatcher->WorkerHandler->IsPingCommandReceived())
+					{
+						break;
+					}
+					else
+					{
+						FPlatformProcess::Sleep(0.001f);
+					}
+				}
+				if (Dispatcher->IsHandlerAlive())
+				{
+					Dispatcher->WaitAllTaskToCompleteExecution();
+					bWorkerValid = Dispatcher->GetInterchangeWorkerFatalError().IsEmpty();
+				}
+				Dispatcher->CloseHandler();
+			}
+			return bWorkerValid;
+		}
+
 		void FInterchangeDispatcher::SpawnHandler()
 		{
 			WorkerHandler = MakeUnique<FInterchangeWorkerHandler>(*this, ResultFolder);
diff --git a/Engine/Plugins/Experimental/Interchange/Runtime/Source/Dispatcher/Private/InterchangeDispatcherNetworking.cpp b/Engine/Plugins/Experimental/Interchange/Runtime/Source/Dispatcher/Private/InterchangeDispatcherNetworking.cpp
--- a/Engine/Plugins/Experimental/Interchange/Runtime/Source/Dispatcher/Private/InterchangeDispatcherNetworking.cpp
+++ b/Engine/Plugins/Experimental/Interchange/Runtime/Source/Dispatcher/Private/InterchangeDispatcherNetworking.cpp
@@ -30,27 +30,27 @@ namespace UE
 			FScopeLock Lock(&SendReceiveCriticalSection);
 			if (bWriteError)
 			{
-				UE_LOG(LogInterchangeDispatcher, Error, TEXT("bWriteError flag raised, can't write"));
+				UE_LOG(LogInterchangeDispatcher, Display, TEXT("bWriteError flag raised, can't write"));
 				return false;
 			}
 
 			if (ConnectedSocket == nullptr)
 			{
-				UE_LOG(LogInterchangeDispatcher, Error, TEXT("node not connected, can't write"));
+				UE_LOG(LogInterchangeDispatcher, Display, TEXT("node not connected, can't write"));
 				return false;
 			}
 
 			ESocketConnectionState State = ConnectedSocket->GetConnectionState();
 			if (State == ESocketConnectionState::SCS_ConnectionError)
 			{
-				UE_LOG(LogInterchangeDispatcher, Error, TEXT("connection state error"));
+				UE_LOG(LogInterchangeDispatcher, Display, TEXT("connection state error"));
 				return false;
 			}
 
 			bool bCanWrite = ConnectedSocket->Wait(ESocketWaitConditions::WaitForWrite, TimespanFromSeconds(Timeout_s));
 			if (!bCanWrite)
 			{
-				UE_LOG(LogInterchangeDispatcher, Warning, TEXT("can't write on socket"));
+				UE_LOG(LogInterchangeDispatcher, Display, TEXT("can't write on socket"));
 				return false;
 			}
 
@@ -82,7 +82,7 @@ namespace UE
 
 			if (!SendBuffer(HeaderBuffer))
 			{
-				UE_LOG(LogInterchangeDispatcher, Error, TEXT("can't write header on socket"));
+				UE_LOG(LogInterchangeDispatcher, Display, TEXT("can't write header on socket"));
 				bWriteError = true;
 				return false;
 			}
@@ -90,7 +90,7 @@ namespace UE
 			// Send the content
 			if (!SendBuffer(Message))
 			{
-				UE_LOG(LogInterchangeDispatcher, Error, TEXT("can't write content on socket"));
+				UE_LOG(LogInterchangeDispatcher, Display, TEXT("can't write content on socket"));
 				bWriteError = true;
 				return false;
 			}
@@ -103,13 +103,13 @@ namespace UE
 			FScopeLock Lock(&SendReceiveCriticalSection);
 			if (bReadError)
 			{
-				UE_LOG(LogInterchangeDispatcher, Error, TEXT("ReadError flag raised, can't read"));
+				UE_LOG(LogInterchangeDispatcher, Display, TEXT("ReadError flag raised, can't read"));
 				return false;
 			}
 
 			if (ConnectedSocket == nullptr)
 			{
-				UE_LOG(LogInterchangeDispatcher, Error, TEXT("node not connected, can't read"));
+				UE_LOG(LogInterchangeDispatcher, Display, TEXT("node not connected, can't read"));
 				return false;
 			}
 
@@ -145,7 +145,7 @@ namespace UE
 				bool bRecvSucceed = ConnectedSocket->Recv(HeaderBuffer.GetData(), HeaderBuffer.Num(), BytesRead);
 				if (!bRecvSucceed || BytesRead != HeaderByteSize)
 				{
-					UE_LOG(LogInterchangeDispatcher, Error, TEXT("Parsed header failed"));
+					UE_LOG(LogInterchangeDispatcher, Display, TEXT("Parsed header failed"));
 					bReadError = true;
 					return false;
 				}
@@ -155,7 +155,7 @@ namespace UE
 
 				if (IncommingMessage.Header.ByteSize < 0 || IncommingMessage.Header.ByteSize > 1 << 20)
 				{
-					UE_LOG(LogInterchangeDispatcher, Error, TEXT("Parsed header failed: bad message size %d"), IncommingMessage.Header.ByteSize);
+					UE_LOG(LogInterchangeDispatcher, Display, TEXT("Parsed header failed: bad message size %d"), IncommingMessage.Header.ByteSize);
 					bReadError = true;
 					return false;
 				}
@@ -177,7 +177,7 @@ namespace UE
 				if (BytesRead != ReadTarget)
 				{
 					bReadError = true;
-					UE_LOG(LogInterchangeDispatcher, Error, TEXT("Recv issue"));
+					UE_LOG(LogInterchangeDispatcher, Display, TEXT("Recv issue"));
 					return false;
 				}
 
@@ -207,7 +207,7 @@ namespace UE
 				Protocol = InternetAddress->GetProtocolType();
 			}
 			FSocket* NewSocket = SocketSubsystem->CreateSocket(NAME_Stream, *Description, Protocol);
-			UE_CLOG(!NewSocket, LogInterchangeDispatcher, Error, TEXT("Socket creation failure"));
+			UE_CLOG(!NewSocket, LogInterchangeDispatcher, Display, TEXT("Socket creation failure"));
 			return NewSocket;
 		}
 
@@ -244,7 +244,7 @@ namespace UE
 			if (BoundPort == 0)
 			{
 				ConnectedSocketError = SocketErrorCode::Error_Bind;
-				UE_LOG(LogInterchangeDispatcher, Error, TEXT("Socket binding failure"));
+				UE_LOG(LogInterchangeDispatcher, Display, TEXT("Socket binding failure"));
 				return;
 			}
 			ensure(BoundPort == ListeningSocket->GetPortNo());
@@ -252,7 +252,7 @@ namespace UE
 			if (!ListeningSocket->Listen(0))
 			{
 				ConnectedSocketError = SocketErrorCode::Error_Listen;
-				UE_LOG(LogInterchangeDispatcher, Error, TEXT("Socket listen failure"));
+				UE_LOG(LogInterchangeDispatcher, Display, TEXT("Socket listen failure"));
 				return;
 			}
 		}
@@ -314,7 +314,7 @@ namespace UE
 				}
 			} while (FPlatformTime::Seconds() < ConnectTimeout_s);
 
-			UE_LOG(LogInterchangeDispatcher, Error, TEXT("Client socket failed to connect"));
+			UE_LOG(LogInterchangeDispatcher, Display, TEXT("Client socket failed to connect"));
 			CloseSocket(ConnectedSocket);
 			return false;
 		}
diff --git a/Engine/Plugins/Experimental/Interchange/Runtime/Source/Dispatcher/Private/InterchangeWorkerHandler.cpp b/Engine/Plugins/Experimental/Interchange/Runtime/Source/Dispatcher/Private/InterchangeWorkerHandler.cpp
--- a/Engine/Plugins/Experimental/Interchange/Runtime/Source/Dispatcher/Private/InterchangeWorkerHandler.cpp
+++ b/Engine/Plugins/Experimental/Interchange/Runtime/Source/Dispatcher/Private/InterchangeWorkerHandler.cpp
@@ -88,7 +88,7 @@ namespace UE
 #endif //#elif PLATFORM_WINDOWS
 				if (!FPaths::FileExists(Path))
 				{
-					UE_LOG(LogInterchangeDispatcher, Warning, TEXT("InterchangeWorker executable not found. Expected location: %s"), *FPaths::ConvertRelativePathToFull(Path));
+					UE_LOG(LogInterchangeDispatcher, Display, TEXT("InterchangeWorker executable not found. Expected location: %s"), *FPaths::ConvertRelativePathToFull(Path));
 				}
 				return Path;
 			}();
@@ -156,13 +156,13 @@ namespace UE
 		{
 			if (!NetworkInterface.IsValid())
 			{
-				UE_LOG(LogInterchangeDispatcher, Error, TEXT("NetworkInterface lost"));
+				UE_LOG(LogInterchangeDispatcher, Display, TEXT("NetworkInterface lost"));
 				WorkerState = EWorkerState::Closing;
 				ErrorState = EWorkerErrorState::ConnectionLost;
 			}
 			else if (WorkerHandle.IsValid() && !FPlatformProcess::IsProcRunning(WorkerHandle))
 			{
-				UE_LOG(LogInterchangeDispatcher, Error, TEXT("Worker lost"));
+				UE_LOG(LogInterchangeDispatcher, Display, TEXT("Worker lost"));
 				WorkerState = EWorkerState::Closing;
 				ErrorState = EWorkerErrorState::WorkerProcess_Lost;
 			}
@@ -213,7 +213,7 @@ namespace UE
 			WorkerState = EWorkerState::Uninitialized;
 			RunInternal();
 			WorkerState = EWorkerState::Terminated;
-			UE_CLOG(ErrorState != EWorkerErrorState::Ok, LogInterchangeDispatcher, Warning, TEXT("Handler ended with error: %s"), EWorkerErrorStateAsString(ErrorState));
+			UE_CLOG(ErrorState != EWorkerErrorState::Ok, LogInterchangeDispatcher, Display, TEXT("Handler ended with error: %s"), EWorkerErrorStateAsString(ErrorState));
 		}
 
 		void FInterchangeWorkerHandler::KillAllCurrentTasks()
@@ -429,13 +429,14 @@ namespace UE
 
 		void FInterchangeWorkerHandler::ProcessCommand(FPingCommand& PingCommand)
 		{
+			bPingCommandReceived = true;
 			UE::Interchange::FBackPingCommand BackPing;
 			CommandIO.SendCommand(BackPing, 0);
 		}
 
 		void FInterchangeWorkerHandler::ProcessCommand(FErrorCommand& ErrorCommand)
 		{
-			UE_LOG(LogInterchangeDispatcher, Error, TEXT("%s"), *ErrorCommand.ErrorMessage);
+			UE_LOG(LogInterchangeDispatcher, Display, TEXT("%s"), *ErrorCommand.ErrorMessage);
 			this->Dispatcher.SetInterchangeWorkerFatalError(ErrorCommand.ErrorMessage);
 			FTerminateCommand Terminate;
 			CommandIO.SendCommand(Terminate, 0);
diff --git a/Engine/Plugins/Experimental/Interchange/Runtime/Source/Dispatcher/Public/InterchangeDispatcher.h b/Engine/Plugins/Experimental/Interchange/Runtime/Source/Dispatcher/Public/InterchangeDispatcher.h
--- a/Engine/Plugins/Experimental/Interchange/Runtime/Source/Dispatcher/Public/InterchangeDispatcher.h
+++ b/Engine/Plugins/Experimental/Interchange/Runtime/Source/Dispatcher/Public/InterchangeDispatcher.h
@@ -41,6 +41,8 @@ namespace UE
 			{
 				return InterchangeWorkerFatalError;
 			}
+
+			static bool IsInterchangeWorkerAvailable();
 		private:
 			void SpawnHandler();
 			bool IsHandlerAlive();
diff --git a/Engine/Plugins/Experimental/Interchange/Runtime/Source/Dispatcher/Public/InterchangeWorkerHandler.h b/Engine/Plugins/Experimental/Interchange/Runtime/Source/Dispatcher/Public/InterchangeWorkerHandler.h
--- a/Engine/Plugins/Experimental/Interchange/Runtime/Source/Dispatcher/Public/InterchangeWorkerHandler.h
+++ b/Engine/Plugins/Experimental/Interchange/Runtime/Source/Dispatcher/Public/InterchangeWorkerHandler.h
@@ -52,6 +52,8 @@ namespace UE
 			void Stop();
 			void StopBlocking();
 
+			bool IsPingCommandReceived() { return bPingCommandReceived; }
+
 			FSimpleMulticastDelegate OnWorkerHandlerExitLoop;
 		protected:
 			void ProcessCommand(ICommand& Command);
@@ -89,6 +91,9 @@ namespace UE
 			bool bShouldTerminate;
 			double LastProgressMessageTime;
 
+			//When the worker start, it send a ping command. This flag is turn on when we receive the ping command
+			bool bPingCommandReceived = false;
+
 			friend Dispatcher::FTaskProcessCommand;
 		};
 	} //ns Interchange
diff --git a/Engine/Plugins/Tests/InterchangeTests/Source/InterchangeTests/InterchangeTests.Build.cs b/Engine/Plugins/Tests/InterchangeTests/Source/InterchangeTests/InterchangeTests.Build.cs
--- a/Engine/Plugins/Tests/InterchangeTests/Source/InterchangeTests/InterchangeTests.Build.cs
+++ b/Engine/Plugins/Tests/InterchangeTests/Source/InterchangeTests/InterchangeTests.Build.cs
@@ -25,6 +25,7 @@ namespace UnrealBuildTool.Rules
 					"AssetRegistry",
 					"InterchangePipelines",
 					"InterchangeCore",
+					"InterchangeDispatcher",
 					"InterchangeEngine",
 					"MeshDescription",
 					"StaticMeshDescription",
diff --git a/Engine/Plugins/Tests/InterchangeTests/Source/InterchangeTests/Private/InterchangeImportTest.cpp b/Engine/Plugins/Tests/InterchangeTests/Source/InterchangeTests/Private/InterchangeImportTest.cpp
--- a/Engine/Plugins/Tests/InterchangeTests/Source/InterchangeTests/Private/InterchangeImportTest.cpp
+++ b/Engine/Plugins/Tests/InterchangeTests/Source/InterchangeTests/Private/InterchangeImportTest.cpp
@@ -8,17 +8,19 @@
 #include "Misc/AutomationTest.h"
 #include "InterchangeTestsLog.h"
 #include "ImportTestFunctions/StaticMeshImportTestFunctions.h"
+#include "InterchangeDispatcher.h"
 #include "InterchangeImportTestData.h"
 #include "InterchangeImportTestSettings.h"
 #include "InterchangeImportTestPlan.h"
 #include "InterchangeImportTestStepBase.h"
+#include "InterchangeManager.h"
 #include "InterchangeTestsModule.h"
 #include "UObject/StrongObjectPtr.h"
 #include "AssetRegistry/AssetRegistryModule.h"
 #include "AssetRegistry/AssetData.h"
 #include "Editor.h"
 #include "Editor/Transactor.h"
-#include "InterchangeManager.h"
+
 #include "UObject/SavePackage.h"
 #include "ObjectTools.h"
 
@@ -37,6 +39,12 @@ void FInterchangeImportTest::GetTests(TArray<FString>& OutBeautifiedNames, TArra
 	return;
 #endif
 
+	//Make sure interchange worker is available, do not run the test if unavailable
+	if (!UE::Interchange::FInterchangeDispatcher::IsInterchangeWorkerAvailable())
+	{
+		return;
+	}
+
 	const FAssetRegistryModule& AssetRegistryModule = FModuleManager::LoadModuleChecked<FAssetRegistryModule>(TEXT("AssetRegistry"));
 	TArray<FAssetData> AllTestPlans;
 	AssetRegistryModule.Get().GetAssetsByClass(UInterchangeImportTestPlan::StaticClass()->GetClassPathName(), AllTestPlans, true);
