- Fixing the crash occuring in standard_surface_brick_procedural. It was caused by a wrong append of scalars that later made the rgbtohsv custom node to crash.
- Preinit some standard math nodes, by transforming them into the dot node, if they have a neutral element, mostly useful for the MaterialX divide and sub nodes, because UE inits those material expression to respectively 2 and 1 when it should be 1 and 0 like Multiply and Add
- Preinit the tiledimage nodes to have a texcoord node in case it's missing, because after the substitution of the nodegraph we lose that info if the original node doesn't refer to a texcoord input
- Adding a function to add an attribute from a value or an interfacename, occuring in an input (heavily used in the test file standard_surface_brick_procedural
- Fixing a wrong connection of inputs by replacing the FindOrAdd method by Add on the Map data structure

diff --git a/Engine/Plugins/Experimental/Interchange/Runtime/Source/Import/Private/MaterialX/InterchangeMaterialXDefinitions.h b/Engine/Plugins/Experimental/Interchange/Runtime/Source/Import/Private/MaterialX/InterchangeMaterialXDefinitions.h
--- a/Engine/Plugins/Experimental/Interchange/Runtime/Source/Import/Private/MaterialX/InterchangeMaterialXDefinitions.h
+++ b/Engine/Plugins/Experimental/Interchange/Runtime/Source/Import/Private/MaterialX/InterchangeMaterialXDefinitions.h
@@ -162,7 +162,6 @@ MATERIALX_NAMESPACE_BEGIN
 	namespace Attributes
 	{
 		static constexpr const char* IsVisited = "UE:IsVisited";
-		static constexpr const char* OriginalName = "UE:OriginalName";
 		static constexpr const char* NewName = "UE:NewName";
 	}
 
@@ -312,6 +311,15 @@ MATERIALX_NAMESPACE_BEGIN
 				static constexpr const char* FrameEndAction = "frameendaction";
 			}
 		}
+
+		namespace Math
+		{
+			template<typename T>
+			static const T NeutralZero = T{ 0 };
+
+			template<typename T>
+			static const T NeutralOne = T{ 1 };
+		}
 	}
 
 	namespace Image
diff --git a/Engine/Plugins/Experimental/Interchange/Runtime/Source/Import/Private/MaterialX/InterchangeMaterialXTranslator.cpp b/Engine/Plugins/Experimental/Interchange/Runtime/Source/Import/Private/MaterialX/InterchangeMaterialXTranslator.cpp
--- a/Engine/Plugins/Experimental/Interchange/Runtime/Source/Import/Private/MaterialX/InterchangeMaterialXTranslator.cpp
+++ b/Engine/Plugins/Experimental/Interchange/Runtime/Source/Import/Private/MaterialX/InterchangeMaterialXTranslator.cpp
@@ -880,13 +880,12 @@ bool UInterchangeMaterialXTranslator::ConnectNodeGraphOutputToInput(MaterialX::I
 										  FText::FromString(GetInputName(InputToNodeGraph)));
 			return false;
 		}
-
-		//Let's substitue all the sub graphs, so we don't have to handle them ourselves, but they will be handled by smaller operations		
+	
 		//extract and separate nodes will be handled by ourselves, since they are defined with a swizzle node, which is not present in UE
 		auto FilterNodes = [](mx::NodePtr Node) ->bool
 		{
 			const std::string& Category= Node->getCategory();
-			return 
+			return
 				Category != mx::Category::Extract &&
 				Category != mx::Category::Saturate &&
 				Category != mx::Category::Separate2 &&
@@ -894,7 +893,16 @@ bool UInterchangeMaterialXTranslator::ConnectNodeGraphOutputToInput(MaterialX::I
 				Category != mx::Category::Separate4;
 		};
 
-		Output->getParent()->asA<mx::NodeGraph>()->flattenSubgraphs(mx::EMPTY_STRING, FilterNodes);
+		mx::NodeGraphPtr NodeGraph = Output->getParent()->asA<mx::NodeGraph>();
+
+		AddTexCoordToTiledImageNodes(NodeGraph);
+
+		//This function will replace all the nodes that are defined with a nodegraph with their equivalent node network
+		//This allows us to handle nodes defined with a nodegraph, with nodes that we already support while traversing the current node graph
+		//For example the node tiledimage is just a succession of the nodes image, multiply, add, etc.
+		NodeGraph->flattenSubgraphs(mx::EMPTY_STRING, FilterNodes);
+
+		ConvertNeutralNodesToDot(NodeGraph);
 
 		for(mx::Edge Edge : Output->traverseGraph())
 		{
@@ -997,24 +1005,9 @@ bool UInterchangeMaterialXTranslator::ConnectNodeOutputToInput(MaterialX::NodePt
 
 		for(mx::InputPtr Input : Node->getInputs())
 		{
-			if(Input->hasValue())
+			if(const FString* InputNameFound = UEInputs.Find(GetInputName(Input)))
 			{
-				if(const FString* InputNameFound = UEInputs.Find(GetInputName(Input)))
-				{
-					AddAttribute(Input, *InputNameFound, OperatorNode);
-				}
-			}
-			else if(Input->hasInterfaceName())
-			{
-				mx::InputPtr InputInterface = Input->getInterfaceInput();
-				if(InputInterface && InputInterface->hasValue())
-				{
-					// We need to take the input name from the original Input not the Interface
-					if(const FString* InputNameFound = UEInputs.Find(GetInputName(Input)))
-					{
-						AddAttribute(InputInterface, *InputNameFound, OperatorNode);
-					}
-				}
+				AddAttributeFromValueOrInterface(Input, *InputNameFound, OperatorNode);
 			}
 		}
 
@@ -1170,6 +1163,26 @@ MaterialX::InputPtr UInterchangeMaterialXTranslator::GetSpotLightInput(MaterialX
 	return Input;
 }
 
+bool UInterchangeMaterialXTranslator::AddAttributeFromValueOrInterface(MaterialX::InputPtr Input, const FString& InputChannelName, UInterchangeShaderNode* ShaderNode) const
+{
+	if(Input)
+	{
+		if(Input->hasValue())
+		{
+			return AddAttribute(Input, InputChannelName, ShaderNode);
+		}
+		else if(Input->hasInterfaceName())
+		{
+			if(mx::InputPtr InputInterface = Input->getInterfaceInput(); InputInterface->hasValue())
+			{
+				return AddAttribute(InputInterface, InputChannelName, ShaderNode);
+			}
+		}
+	}
+
+	return false;
+}
+
 bool UInterchangeMaterialXTranslator::AddAttribute(MaterialX::InputPtr Input, const FString& InputChannelName, UInterchangeShaderNode* ShaderNode) const
 {
 	if(Input)
@@ -1203,7 +1216,7 @@ bool UInterchangeMaterialXTranslator::AddFloatAttribute(MaterialX::InputPtr Inpu
 {
 	if(Input)
 	{
-		if(Input->hasValueString())
+		if(Input->hasValue())
 		{
 			float Value = mx::fromValueString<float>(Input->getValueString());
 
@@ -1221,7 +1234,7 @@ bool UInterchangeMaterialXTranslator::AddLinearColorAttribute(MaterialX::InputPt
 {
 	if(Input)
 	{
-		if(Input->hasValueString())
+		if(Input->hasValue())
 		{
 			const FLinearColor Value = GetLinearColor(Input);
 
@@ -1317,13 +1330,7 @@ FLinearColor UInterchangeMaterialXTranslator::GetLinearColor(MaterialX::InputPtr
 
 void UInterchangeMaterialXTranslator::ConnectConstantInputToOutput(MaterialX::NodePtr UpstreamNode, UInterchangeShaderNode* ParentShaderNode, const FString& InputChannelName, TMap<FString, UInterchangeShaderNode*>& NamesToShaderNodes, UInterchangeBaseNodeContainer& NodeContainer) const
 {
-	if(mx::InputPtr InputConstant = UpstreamNode->getInput("value"); !AddAttribute(InputConstant, InputChannelName, ParentShaderNode))
-	{
-		UInterchangeResultWarning_Generic* Message = AddMessage<UInterchangeResultWarning_Generic>();
-		Message->Text = FText::Format(LOCTEXT("InputTypeNotSupported", "<{0}>: \"{1}\" is not supported yet"),
-									  FText::FromString(GetInputName(InputConstant)),
-									  FText::FromString(InputConstant->getType().c_str()));
-	}
+	AddAttributeFromValueOrInterface(UpstreamNode->getInput("value"), InputChannelName, ParentShaderNode);
 }
 
 void UInterchangeMaterialXTranslator::ConnectExtractInputToOutput(MaterialX::NodePtr UpstreamNode, UInterchangeShaderNode* ParentShaderNode, const FString& InputChannelName, TMap<FString, UInterchangeShaderNode*>& NamesToShaderNodes, UInterchangeBaseNodeContainer& NodeContainer) const
@@ -1356,29 +1363,21 @@ void UInterchangeMaterialXTranslator::ConnectDotInputToOutput(MaterialX::NodePtr
 	if(mx::InputPtr Input = UpstreamNode->getInput("in"))
 	{
 		RenameInput(Input, TCHAR_TO_UTF8(*InputChannelName)); //let's take the parent node's input name
-		NamesToShaderNodes.FindOrAdd(UpstreamNode->getName().c_str(), ParentShaderNode);
+		NamesToShaderNodes.Add(UpstreamNode->getName().c_str(), ParentShaderNode);
 	}
 }
 
 void UInterchangeMaterialXTranslator::ConnectTransformPositionInputToOutput(MaterialX::NodePtr UpstreamNode, UInterchangeShaderNode* ParentShaderNode, const FString& InputChannelName, TMap<FString, UInterchangeShaderNode*>& NamesToShaderNodes, UInterchangeBaseNodeContainer& NodeContainer) const
 {
 	UInterchangeShaderNode* TransformNode = CreateShaderNode<UInterchangeShaderNode>(UpstreamNode->getName().c_str(), TEXT("TransformPosition"), NamesToShaderNodes, NodeContainer);
-	if(mx::InputPtr Input = UpstreamNode->getInput("in"); Input->hasValue())
-	{
-		FLinearColor Vector = GetVector(Input);
-		TransformNode->AddLinearColorAttribute(UInterchangeShaderPortsAPI::MakeInputValueKey(TEXT("Input")), Vector);
-	}
+	AddAttributeFromValueOrInterface(UpstreamNode->getInput("in"), TEXT("Input"), TransformNode);
 	UInterchangeShaderPortsAPI::ConnectDefaultOuputToInput(ParentShaderNode, InputChannelName, TransformNode->GetUniqueID());
 }
 
 void UInterchangeMaterialXTranslator::ConnectTransformVectorInputToOutput(MaterialX::NodePtr UpstreamNode, UInterchangeShaderNode* ParentShaderNode, const FString& InputChannelName, TMap<FString, UInterchangeShaderNode*>& NamesToShaderNodes, UInterchangeBaseNodeContainer& NodeContainer) const
 {
 	UInterchangeShaderNode* TransformNode = CreateShaderNode<UInterchangeShaderNode>(UpstreamNode->getName().c_str(), TEXT("Transform"), NamesToShaderNodes, NodeContainer);
-	if(mx::InputPtr Input = UpstreamNode->getInput("in"); Input->hasValue())
-	{
-		FLinearColor Vector = GetVector(Input);
-		TransformNode->AddLinearColorAttribute(UInterchangeShaderPortsAPI::MakeInputValueKey(TEXT("Input")), Vector);
-	}
+	AddAttributeFromValueOrInterface(UpstreamNode->getInput("in"), TEXT("Input"), TransformNode);
 	UInterchangeShaderPortsAPI::ConnectDefaultOuputToInput(ParentShaderNode, InputChannelName, TransformNode->GetUniqueID());
 }
 
@@ -1388,19 +1387,10 @@ void UInterchangeMaterialXTranslator::ConnectRotate2DInputToOutput(mx::NodePtr U
 	Rotate2DNode->AddLinearColorAttribute(UInterchangeShaderPortsAPI::MakeInputValueKey(TEXT("NormalizedRotationAxis")), FLinearColor(0, 0, 1));
 	Rotate2DNode->AddLinearColorAttribute(UInterchangeShaderPortsAPI::MakeInputValueKey(TEXT("PivotPoint")), FLinearColor(0.5, 0.5, 0));
 
-	if(mx::InputPtr Input = UpstreamNode->getInput("in"); Input->hasValue())
-	{
-		FLinearColor Vector = GetVector(Input);
-		Rotate2DNode->AddLinearColorAttribute(UInterchangeShaderPortsAPI::MakeInputValueKey(TEXT("Position")), Vector);
-	}
+	AddAttributeFromValueOrInterface(UpstreamNode->getInput("in"), TEXT("Input"), Rotate2DNode);
 
-	if(mx::InputPtr Input = UpstreamNode->getInput("amount"); Input->hasValue())
-	{
-		//MaterialX angle is in degrees, in UE it's in the range [0,1]
-		float Amount = mx::fromValueString<float>(Input->getValueString()) / 360.;
-		Rotate2DNode->AddFloatAttribute(UInterchangeShaderPortsAPI::MakeInputValueKey(TEXT("RotationAngle")), Amount);
-	}
-	else // we need to create a Divide node
+	// we create a Divide node to convert MaterialX angle in degrees to UE's angle which is a value between [0,1]
+	if(mx::InputPtr Input = UpstreamNode->getInput("amount"); !AddAttributeFromValueOrInterface(Input, TEXT("RotationAngle"), Rotate2DNode))
 	{
 		//copy "in" input into "in1" and create "amount" input as "in2" with the value 360 because UE's angle is a value between [0,1]
 		mx::NodePtr NewDivideNode = CreateNode(UpstreamNode->getParent()->asA<mx::NodeGraph>(),
@@ -1421,25 +1411,11 @@ void UInterchangeMaterialXTranslator::ConnectRotate3DInputToOutput(mx::NodePtr U
 	UInterchangeShaderNode* Rotate3DNode = CreateShaderNode<UInterchangeShaderNode>(UpstreamNode->getName().c_str(), TEXT("RotateAboutAxis"), NamesToShaderNodes, NodeContainer);
 	Rotate3DNode->AddLinearColorAttribute(UInterchangeShaderPortsAPI::MakeInputValueKey(TEXT("PivotPoint")), FLinearColor(0.5, 0.5, 0));
 
-	if(mx::InputPtr Input = UpstreamNode->getInput("in"); Input->hasValue())
-	{
-		FLinearColor Vector = GetVector(Input);
-		Rotate3DNode->AddLinearColorAttribute(UInterchangeShaderPortsAPI::MakeInputValueKey(TEXT("Position")), Vector);
-	}
+	AddAttributeFromValueOrInterface(UpstreamNode->getInput("in"), TEXT("Position"), Rotate3DNode);
+	AddAttributeFromValueOrInterface(UpstreamNode->getInput("axis"), TEXT("NormalizedRotationAxis"), Rotate3DNode);
 
-	if(mx::InputPtr Input = UpstreamNode->getInput("axis"); Input->hasValue())
-	{
-		FLinearColor Vector = GetVector(Input);
-		Rotate3DNode->AddLinearColorAttribute(UInterchangeShaderPortsAPI::MakeInputValueKey(TEXT("NormalizedRotationAxis")), Vector);
-	}
-
-	if(mx::InputPtr Input = UpstreamNode->getInput("amount"); Input->hasValue())
-	{
-		//MaterialX angle is in degrees, in UE it's in the range [0,1]
-		float Amount = mx::fromValueString<float>(Input->getValueString()) / 360.;
-		Rotate3DNode->AddFloatAttribute(UInterchangeShaderPortsAPI::MakeInputValueKey(TEXT("RotationAngle")), Amount);
-	}
-	else // we need to create a Divide node
+	// we create a Divide node to convert MaterialX angle in degrees to UE's angle which is a value between [0,1]
+	if(mx::InputPtr Input = UpstreamNode->getInput("amount"); !AddAttributeFromValueOrInterface(Input, TEXT("RotationAngle"), Rotate3DNode))
 	{
 		//copy "in" input into "in1" and create "amount" input as "in2" with the value 360 because UE's angle is a value between [0,1]
 		mx::NodePtr NewDivideNode = CreateNode(UpstreamNode->getParent()->asA<mx::NodeGraph>(),
@@ -1461,13 +1437,17 @@ void UInterchangeMaterialXTranslator::ConnectImageInputToOutput(MaterialX::NodeP
 
 	if(UInterchangeTextureNode* TextureNode = CreateTextureNode(UpstreamNode, NodeContainer))
 	{
-		//By default set the output of a texture to RGB, if the type is float then it's either up to an extract node or the Material Input to handle it
+		//By default set the output of a texture to RGB
 		FString OutputChannel{ TEXT("RGB") };
 
 		if(UpstreamNode->getType() == mx::Type::Vector4 || UpstreamNode->getType() == mx::Type::Color4)
 		{
 			OutputChannel = TEXT("RGBA");
 		}
+		else if(UpstreamNode->getType() == mx::Type::Float)
+		{
+			OutputChannel = TEXT("R");
+		}
 
 		UInterchangeShaderNode* TextureShaderNode = CreateShaderNode<UInterchangeShaderNode>(UpstreamNode->getName().c_str(), TextureSample::Name.ToString(), NamesToShaderNodes, NodeContainer);
 		TextureShaderNode->AddStringAttribute(UInterchangeShaderPortsAPI::MakeInputValueKey(TextureSample::Inputs::Texture.ToString()), TextureNode->GetUniqueID());
@@ -1475,7 +1455,7 @@ void UInterchangeMaterialXTranslator::ConnectImageInputToOutput(MaterialX::NodeP
 	}
 	else
 	{
-		AddAttribute(UpstreamNode->getInput(mx::NodeGroup::Texture2D::Inputs::Default), InputChannelName, ParentShaderNode);
+		AddAttributeFromValueOrInterface(UpstreamNode->getInput(mx::NodeGroup::Texture2D::Inputs::Default), InputChannelName, ParentShaderNode);
 	}
 }
 
@@ -1550,10 +1530,7 @@ void UInterchangeMaterialXTranslator::ConnectConvertInputToOutput(MaterialX::Nod
 			MaskShaderNode = CreateMaskShaderNode(0b1000, UpstreamNode->getName().c_str(), NamesToShaderNodes, NodeContainer);
 		}
 
-		if(Input->hasValue())
-		{
-			AddAttribute(Input, TEXT("Input"), MaskShaderNode);
-		}
+		AddAttributeFromValueOrInterface(Input, TEXT("Input"), MaskShaderNode);
 
 		UInterchangeShaderPortsAPI::ConnectDefaultOuputToInput(ParentShaderNode, InputChannelName, MaskShaderNode->GetUniqueID());
 	}
@@ -1561,15 +1538,12 @@ void UInterchangeMaterialXTranslator::ConnectConvertInputToOutput(MaterialX::Nod
 	{
 		//same as dot, just connect the next output to this parent input
 		RenameInput(Input, TCHAR_TO_UTF8(*InputChannelName)); //let's take the parent node's input name
-		NamesToShaderNodes.FindOrAdd(UpstreamNode->getName().c_str(), ParentShaderNode);
+		NamesToShaderNodes.Add(UpstreamNode->getName().c_str(), ParentShaderNode);
 
 		//No need to create a node, since the input and the node have the same channels, we just check if there's a value
 		if(NodeN == InputN || (NodeN < '2' && InputN < '2'))
 		{
-			if(Input->hasValue())
-			{
-				AddAttribute(Input, InputChannelName, ParentShaderNode);
-			}
+			AddAttributeFromValueOrInterface(Input, InputChannelName, ParentShaderNode);
 			return;
 		}
 		std::string Category;
@@ -1625,31 +1599,29 @@ void UInterchangeMaterialXTranslator::ConnectIfGreaterInputToOutput(MaterialX::N
 	UInterchangeShaderNode* NodeIf = CreateShaderNode<UInterchangeShaderNode>(UpstreamNode->getName().c_str(), TEXT("If"), NamesToShaderNodes, NodeContainer);
 	UInterchangeShaderPortsAPI::ConnectDefaultOuputToInput(ParentShaderNode, InputChannelName, NodeIf->GetUniqueID());
 
-	if(mx::InputPtr Input = UpstreamNode->getInput("value1"); Input && Input->hasValueString())
+	if(mx::InputPtr Input = UpstreamNode->getInput("value1"))
 	{
-		AddAttribute(Input, GetInputName(Input), NodeIf);
+		AddAttributeFromValueOrInterface(Input, GetInputName(Input), NodeIf);
 	}
 	
-	if(mx::InputPtr Input = UpstreamNode->getInput("value2"); Input && Input->hasValueString())
+	if(mx::InputPtr Input = UpstreamNode->getInput("value2"))
 	{
-		AddAttribute(Input, GetInputName(Input), NodeIf);
+		AddAttributeFromValueOrInterface(Input, GetInputName(Input), NodeIf);
 	}
 	
-	if(mx::InputPtr Input = UpstreamNode->getInput("in1"); Input && Input->hasValueString())
+	if(mx::InputPtr Input = UpstreamNode->getInput("in1"))
 	{
-		AddAttribute(Input, GetInputName(Input), NodeIf);
+		AddAttributeFromValueOrInterface(Input, GetInputName(Input), NodeIf);
 	}
 
 	//In that case we also need to add an attribute to AEqualsB
 	mx::InputPtr Input = UpstreamNode->getInput("in2");
 	if(Input)
 	{
-		if(Input->hasValueString())
-		{
-			AddAttribute(Input, GetInputName(Input), NodeIf);
-			AddAttribute(Input, TEXT("AEqualsB"), NodeIf);
-		}
-		else
+		bool bHasValue = AddAttributeFromValueOrInterface(Input, GetInputName(Input), NodeIf);
+		AddAttributeFromValueOrInterface(Input, TEXT("AEqualsB"), NodeIf);
+
+		if(bHasValue)
 		{
 			//Let's add a new input that is a copy of in2 to connect it to the equal input
 			mx::InputPtr Input3 = UpstreamNode->addInput("in3");
@@ -1664,31 +1636,29 @@ void UInterchangeMaterialXTranslator::ConnectIfGreaterEqInputToOutput(MaterialX:
 	UInterchangeShaderNode* NodeIf = CreateShaderNode<UInterchangeShaderNode>(UpstreamNode->getName().c_str(), TEXT("If"), NamesToShaderNodes, NodeContainer);
 	UInterchangeShaderPortsAPI::ConnectDefaultOuputToInput(ParentShaderNode, InputChannelName, NodeIf->GetUniqueID());
 
-	if(mx::InputPtr Input = UpstreamNode->getInput("value1"); Input && Input->hasValueString())
+	if(mx::InputPtr Input = UpstreamNode->getInput("value1"))
 	{
-		AddAttribute(Input, GetInputName(Input), NodeIf);
+		AddAttributeFromValueOrInterface(Input, GetInputName(Input), NodeIf);
 	}
 
-	if(mx::InputPtr Input = UpstreamNode->getInput("value2"); Input && Input->hasValueString())
+	if(mx::InputPtr Input = UpstreamNode->getInput("value2"))
 	{
-		AddAttribute(Input, GetInputName(Input), NodeIf);
+		AddAttributeFromValueOrInterface(Input, GetInputName(Input), NodeIf);
 	}
 
-	if(mx::InputPtr Input = UpstreamNode->getInput("in2"); Input && Input->hasValueString())
+	if(mx::InputPtr Input = UpstreamNode->getInput("in2"))
 	{
-		AddAttribute(Input, GetInputName(Input), NodeIf);
+		AddAttributeFromValueOrInterface(Input, GetInputName(Input), NodeIf);
 	}
 
 	//In that case we also need to add an attribute to AEqualsB
 	mx::InputPtr Input = UpstreamNode->getInput("in1");
 	if(Input)
 	{
-		if(Input->hasValueString())
-		{
-			AddAttribute(Input, GetInputName(Input), NodeIf);
-			AddAttribute(Input, TEXT("AEqualsB"), NodeIf);
-		}
-		else
+		bool bHasValue = AddAttributeFromValueOrInterface(Input, GetInputName(Input), NodeIf);
+		AddAttributeFromValueOrInterface(Input, TEXT("AEqualsB"), NodeIf);
+
+		if(bHasValue)
 		{
 			//Let's add a new input that is a copy of in2 to connect it to the equal input
 			mx::InputPtr Input3 = UpstreamNode->addInput("in3");
@@ -1703,31 +1673,29 @@ void UInterchangeMaterialXTranslator::ConnectIfEqualInputToOutput(MaterialX::Nod
 	UInterchangeShaderNode* NodeIf = CreateShaderNode<UInterchangeShaderNode>(UpstreamNode->getName().c_str(), TEXT("If"), NamesToShaderNodes, NodeContainer);
 	UInterchangeShaderPortsAPI::ConnectDefaultOuputToInput(ParentShaderNode, InputChannelName, NodeIf->GetUniqueID());
 
-	if(mx::InputPtr Input = UpstreamNode->getInput("value1"); Input && Input->hasValueString())
+	if(mx::InputPtr Input = UpstreamNode->getInput("value1"); Input && Input->hasValue())
 	{
-		AddAttribute(Input, GetInputName(Input), NodeIf);
+		AddAttributeFromValueOrInterface(Input, GetInputName(Input), NodeIf);
 	}
 
-	if(mx::InputPtr Input = UpstreamNode->getInput("value2"); Input && Input->hasValueString())
+	if(mx::InputPtr Input = UpstreamNode->getInput("value2"); Input && Input->hasValue())
 	{
-		AddAttribute(Input, GetInputName(Input), NodeIf);
+		AddAttributeFromValueOrInterface(Input, GetInputName(Input), NodeIf);
 	}
 
-	if(mx::InputPtr Input = UpstreamNode->getInput("in1"); Input && Input->hasValueString())
+	if(mx::InputPtr Input = UpstreamNode->getInput("in1"); Input && Input->hasValue())
 	{
-		AddAttribute(Input, GetInputName(Input), NodeIf);
+		AddAttributeFromValueOrInterface(Input, GetInputName(Input), NodeIf);
 	}
 
 	//In that case we also need to add an attribute to AGreaterThanB
 	mx::InputPtr Input = UpstreamNode->getInput("in2");
 	if(Input)
 	{
-		if(Input->hasValueString())
-		{
-			AddAttribute(Input, GetInputName(Input), NodeIf);
-			AddAttribute(Input, TEXT("AGreaterThanB"), NodeIf);
-		}
-		else
+		bool bHasValue = AddAttributeFromValueOrInterface(Input, GetInputName(Input), NodeIf);
+		AddAttributeFromValueOrInterface(Input, TEXT("AGreaterThanB"), NodeIf);
+
+		if(bHasValue)
 		{
 			//Let's add a new input that is a copy of in2 to connect it to the equal input
 			mx::InputPtr Input3 = UpstreamNode->addInput("in3");
@@ -1797,8 +1765,8 @@ MaterialX::NodePtr UInterchangeMaterialXTranslator::CreateNode(MaterialX::NodeGr
 
 			mx::InputPtr InputCopy = Node->addInput();
 			InputCopy->copyContentFrom(Input);
-			//remove the attribute OriginalName since when the input will be copied, the name will be the one from the spec (before a renaming)
-			InputCopy->removeAttribute(mx::Attributes::OriginalName);
+			//remove the attribute NewName since when the input will be copied, the name will be the one from the spec (before a renaming)
+			InputCopy->removeAttribute(mx::Attributes::NewName);
 
 			if(NewInputName)
 			{
@@ -1825,6 +1793,102 @@ MaterialX::NodePtr UInterchangeMaterialXTranslator::CreateNode(MaterialX::NodeGr
 	return Node;
 }
 
+void UInterchangeMaterialXTranslator::ConvertNeutralNodesToDot(MaterialX::NodeGraphPtr NodeGraph) const
+{
+	const std::vector<mx::NodePtr> Nodes = NodeGraph->getNodes();
+	for(mx::NodePtr Node : Nodes)
+	{
+		const std::string& Category = Node->getCategory();
+		if(Category == mx::Category::Multiply || Category == mx::Category::Divide || Category == mx::Category::Add || Category == mx::Category::Sub)
+		{
+			mx::InputPtr Input2 = Node->getInput("in2");
+			bool bIsNeutral = false;
+
+			if(!Input2)
+			{
+				bIsNeutral = true;
+				if(mx::InputPtr Input1 = Node->getInput("in1"))
+				{
+					Input1->setName("in");
+				}
+			}
+			else
+			{
+				if(!Input2->hasValue() && !Input2->hasNodeName() && !Input2->hasInterfaceName())
+				{
+					bIsNeutral = true;
+					Node->removeInput("in2");
+					if(mx::InputPtr Input1 = Node->getInput("in1"))
+					{
+						Input1->setName("in");
+					}
+				}
+				else if(Input2->hasValue())
+				{
+					using namespace mx::NodeGroup::Math;
+
+					auto IsValueStringNeutral = [&bIsNeutral, &Category, Input2](auto T)
+					{
+						bIsNeutral = (Category == mx::Category::Multiply || Category == mx::Category::Divide) ?
+							mx::fromValueString<decltype(T)>(Input2->getValueString()) == NeutralOne<decltype(T)> :
+							mx::fromValueString<decltype(T)>(Input2->getValueString()) == NeutralZero<decltype(T)>;
+					};
+
+					if(Input2->getType() == mx::Type::Float)
+					{
+						IsValueStringNeutral(float());
+					}
+					else if(Input2->getType() == mx::Type::Vector2)
+					{
+						IsValueStringNeutral(mx::Vector2());
+					}
+					else if(Input2->getType() == mx::Type::Vector3)
+					{
+						IsValueStringNeutral(mx::Vector3());
+					}
+					else if(Input2->getType() == mx::Type::Vector4)
+					{
+						IsValueStringNeutral(mx::Vector4());
+					}
+					else if(Input2->getType() == mx::Type::Color3)
+					{
+						IsValueStringNeutral(mx::Color3());
+					}
+					else if(Input2->getType() == mx::Type::Color4)
+					{
+						IsValueStringNeutral(mx::Color4());
+					}
+				}
+			}
+
+			if(bIsNeutral)
+			{
+				Node->setCategory(mx::Category::Dot);
+			}
+		}
+	}
+}
+
+void UInterchangeMaterialXTranslator::AddTexCoordToTiledImageNodes(MaterialX::NodeGraphPtr NodeGraph) const
+{
+	for(mx::NodePtr Node : NodeGraph->getNodes(mx::Category::TiledImage))
+	{
+		mx::InputPtr InputTexCoord = Node->getInput("texcoord");
+
+		if(!InputTexCoord)
+		{
+			mx::NodePtr NodeTexCoord = CreateNode(Node->getParent()->asA<mx::NodeGraph>(),
+												  "node_ue_texcoord",
+												  mx::Category::TexCoord,
+												  {},
+												  { {"index", FAttributeValueArray{{"type", "integer"}, {"value", "0"}}} });
+
+			InputTexCoord = Node->addInput(mx::NodeGroup::Texture2D::Inputs::TexCoord, mx::Type::Vector2);
+			InputTexCoord->setNodeName(NodeTexCoord->getName());
+		}
+	}
+}
+
 #endif //WITH_EDITOR
 
 #undef LOCTEXT_NAMESPACE
\ No newline at end of file
diff --git a/Engine/Plugins/Experimental/Interchange/Runtime/Source/Import/Public/MaterialX/InterchangeMaterialXTranslator.h b/Engine/Plugins/Experimental/Interchange/Runtime/Source/Import/Public/MaterialX/InterchangeMaterialXTranslator.h
--- a/Engine/Plugins/Experimental/Interchange/Runtime/Source/Import/Public/MaterialX/InterchangeMaterialXTranslator.h
+++ b/Engine/Plugins/Experimental/Interchange/Runtime/Source/Import/Public/MaterialX/InterchangeMaterialXTranslator.h
@@ -260,7 +260,18 @@ protected:
 	MaterialX::InputPtr GetSpotLightInput(MaterialX::NodePtr SpotLight, const char* InputName, MaterialX::DocumentPtr Document) const;
 
 	/**
-	 * Add an attribute to a shader node, only floats and linear colors are supported for the moment
+	 * Add an attribute to a shader node from the given MaterialX input if that input has either a value or an interface name
+	 *
+	 * @param Input - The MaterialX input to retrieve and add the value from, must be of type float/color/vector
+	 * @param InputChannelName - The name of the shader node's input to add the attribute
+	 * @param ShaderNode - The shader node to which we want to add the attribute
+	 *
+	 * @return true if the attribute was successfully added
+	 */
+	bool AddAttributeFromValueOrInterface(MaterialX::InputPtr Input, const FString& InputChannelName, UInterchangeShaderNode* ShaderNode) const;
+
+	/**
+	 * Add an attribute to a shader node from the given MaterialX input, only floats and linear colors are supported for the moment
 	 * 
 	 * @param Input - The MaterialX input to retrieve and add the value from, must be of type float/color/vector
 	 * @param InputChannelName - The name of the shader node's input to add the attribute
@@ -383,6 +394,24 @@ protected:
 	 */
 	MaterialX::NodePtr CreateNode(MaterialX::NodeGraphPtr NodeGraph, const char * NodeName, const char * Category, TArray<FInputToCopy> InputsToCopy, TArray<FInputToCreate> InputsToCreate) const;
 
+	/**
+	 * Convert math nodes (add, sub, div, mul) that have a neutral input, to the MaterialX <dot> node which is simply the identity or a noop
+	 * This allows us to optimize the final shader by avoiding unnecessary operation on a material (for example a*1 = a)
+	 * This function should be called after the subgraphs substitution in order to search for all the nodes even in the subgraphs
+	 * 
+	 * @param NodeGraph - The node graph to retrieve the different math nodes
+	 */
+	void ConvertNeutralNodesToDot(MaterialX::NodeGraphPtr NodeGraph) const;
+
+	/**
+	 * Add a texcoord input to the tiledimage nodes if the input is missing, and connect it to a texcoord node in the nodegraph
+	 * This function should be called before the subgraphs substitution to avoid missing UVs input on a texture
+	 * See for example standard_surface_brick_procedural where the tiledimage nodes have no texcoord input
+	 * 
+	 * @param Node - The nodegraph to retrieve the different tiledimage nodes
+	 */
+	void AddTexCoordToTiledImageNodes(MaterialX::NodeGraphPtr NodeGraph) const;
+
 private:
 
 	TMap<TPair<FString, FString>, FString> InputNamesMaterialX2UE; //given a MaterialX node (category - input), return the UE/Interchange input name.
