Prevent import of variants in game mode. This crashes the game.
Note: Variants will have to be reworked in order to be usable in game.

#rnx

diff --git a/Engine/Plugins/Experimental/Interchange/Runtime/Source/Import/Private/Gltf/InterchangeGltfTranslator.cpp b/Engine/Plugins/Experimental/Interchange/Runtime/Source/Import/Private/Gltf/InterchangeGltfTranslator.cpp
--- a/Engine/Plugins/Experimental/Interchange/Runtime/Source/Import/Private/Gltf/InterchangeGltfTranslator.cpp
+++ b/Engine/Plugins/Experimental/Interchange/Runtime/Source/Import/Private/Gltf/InterchangeGltfTranslator.cpp
@@ -24,6 +24,7 @@
 
 #include "Algo/Find.h"
 #include "Async/Async.h"
+#include "Misc/App.h"
 #include "StaticMeshAttributes.h"
 #include "SkeletalMeshAttributes.h"
 #include "UObject/GCObjectScopeGuard.h"
@@ -1053,8 +1054,9 @@ bool UInterchangeGltfTranslator::Translate( UInterchangeBaseNodeContainer& NodeC
 		}
 	}
 
-	// Create SceneVariantSetsNode if model contains variants.
-	if (bHasVariants)
+	// Variants
+	// Note: Variants are not supported yet in game play mode
+	if (!FApp::IsGame() && bHasVariants)
 	{
 		HandleGltfVariants(NodeContainer, FileName);
 	}
