[interchange] gltf import fix for rotations on (Rigid) animations

diff --git a/Engine/Plugins/Experimental/Interchange/Runtime/Source/Import/Private/Gltf/InterchangeGltfAnimation.cpp b/Engine/Plugins/Experimental/Interchange/Runtime/Source/Import/Private/Gltf/InterchangeGltfAnimation.cpp
--- a/Engine/Plugins/Experimental/Interchange/Runtime/Source/Import/Private/Gltf/InterchangeGltfAnimation.cpp
+++ b/Engine/Plugins/Experimental/Interchange/Runtime/Source/Import/Private/Gltf/InterchangeGltfAnimation.cpp
@@ -304,10 +304,42 @@ namespace UE::Interchange::Gltf::Private
 			{
 				FrameDataBuffer.SetNumUninitialized(Sampler.Output.Count * 4);
 				Sampler.Output.GetQuatArray(reinterpret_cast<FVector4f*>(FrameDataBuffer.GetData()));
-				ConvertCurve([](const float* QuatData)
+
+				struct HeadingHelper
+				{
+					FRotator3f Heading;
+					bool bHeadingSet = false;
+				};
+				HeadingHelper Helper;
+				ConvertCurve([&Helper](const float* QuatData)
 					{
-						const FQuat4f Quat(QuatData[0], QuatData[1], QuatData[2], QuatData[3]);
-						return Quat.Euler();
+						//GLTF uses quaternions while Uneal uses Eulers for level animations
+						// Quaternions can represent multiple Orientations in Eulers.
+						// We cannot tell which precise Orientation version was the author of gltf trying to achieve at the moment of creation.
+						// For that reason we compare each consecutive orientation and pick the closest (by ManhattanDistance) from the 2 EquivalentRotators.
+						// This means that animations of Rotations can be correctly represented with this system, if
+						//    every consecutive frame's intended (absolute) rotations are less than 3x180 degrees (aka less than 180 degrees per axis)
+
+						FQuat4f Quat(QuatData[0], QuatData[1], QuatData[2], QuatData[3]);
+						FRotator3f Rotator = Quat.Rotator();
+						if (Helper.bHeadingSet)
+						{
+							FRotator3f OtherChoice = Rotator.GetEquivalentRotator().GetNormalized();
+							float FirstDiff = Helper.Heading.GetManhattanDistance(Rotator);
+							float SecondDiff = Helper.Heading.GetManhattanDistance(OtherChoice);
+							if (SecondDiff < FirstDiff)
+							{
+								Rotator = OtherChoice;
+							}
+						}
+						else
+						{
+							Helper.bHeadingSet = true;
+						}
+
+						Helper.Heading = Rotator;
+						
+						return Rotator.Euler();
 					}, 3, 4, true);
 				break;
 			}
