Forcing use of MaterialInstances if in Standalone game

#rnx

diff --git a/Engine/Plugins/Experimental/Interchange/Runtime/Source/Pipelines/Private/InterchangeGenericMaterialPipeline.cpp b/Engine/Plugins/Experimental/Interchange/Runtime/Source/Pipelines/Private/InterchangeGenericMaterialPipeline.cpp
--- a/Engine/Plugins/Experimental/Interchange/Runtime/Source/Pipelines/Private/InterchangeGenericMaterialPipeline.cpp
+++ b/Engine/Plugins/Experimental/Interchange/Runtime/Source/Pipelines/Private/InterchangeGenericMaterialPipeline.cpp
@@ -34,6 +34,7 @@
 #include "Materials/MaterialExpressionVectorParameter.h"
 #include "Materials/MaterialInstanceConstant.h"
 #include "Materials/MaterialInstanceDynamic.h"
+#include "Misc/CoreMisc.h"
 #include "Misc/Paths.h"
 #include "Nodes/InterchangeBaseNode.h"
 #include "Nodes/InterchangeBaseNodeContainer.h"
@@ -1741,7 +1742,8 @@ UInterchangeMaterialInstanceFactoryNode* UInterchangeGenericMaterialPipeline::Cr
 	}
 
 #if WITH_EDITOR
-	MaterialInstanceFactoryNode->SetCustomInstanceClassName(UMaterialInstanceConstant::StaticClass()->GetPathName());
+	const UClass* MaterialClass = IsRunningGame() ? UMaterialInstanceDynamic::StaticClass() : UMaterialInstanceConstant::StaticClass();
+	MaterialInstanceFactoryNode->SetCustomInstanceClassName(MaterialClass->GetPathName());
 #else
 	MaterialInstanceFactoryNode->SetCustomInstanceClassName(UMaterialInstanceDynamic::StaticClass()->GetPathName());
 #endif
