Interchange:
The option bImportMaterials was not use by the material pipeline to avoid importing materials
Make sure weak references type use TWeakObjectPtr instead of TObjectPtr to allow deletion of asset loaded in memory

#rnx

diff --git a/Engine/Plugins/Experimental/Interchange/Runtime/Source/Pipelines/Private/InterchangeGenericAnimationPipeline.cpp b/Engine/Plugins/Experimental/Interchange/Runtime/Source/Pipelines/Private/InterchangeGenericAnimationPipeline.cpp
--- a/Engine/Plugins/Experimental/Interchange/Runtime/Source/Pipelines/Private/InterchangeGenericAnimationPipeline.cpp
+++ b/Engine/Plugins/Experimental/Interchange/Runtime/Source/Pipelines/Private/InterchangeGenericAnimationPipeline.cpp
@@ -51,7 +51,7 @@ void UInterchangeGenericAnimationPipeline::AdjustSettingsForContext(EInterchange
 {
 	Super::AdjustSettingsForContext(ImportType, ReimportAsset);
 
-	check(!CommonSkeletalMeshesAndAnimationsProperties.IsNull());
+	check(CommonSkeletalMeshesAndAnimationsProperties.IsValid());
 	
 	if (ImportType == EInterchangePipelineContext::AssetCustomLODImport
 		|| ImportType == EInterchangePipelineContext::AssetCustomLODReimport
@@ -116,12 +116,12 @@ void UInterchangeGenericAnimationPipeline::ExecutePreImportPipeline(UInterchange
 		}
 	}
 
-	if (CommonSkeletalMeshesAndAnimationsProperties.IsNull())
+	if (!CommonSkeletalMeshesAndAnimationsProperties.IsValid())
 	{
 		return;
 	}
 
-	if (CommonSkeletalMeshesAndAnimationsProperties->bImportOnlyAnimations && CommonSkeletalMeshesAndAnimationsProperties->Skeleton.IsNull())
+	if (CommonSkeletalMeshesAndAnimationsProperties->bImportOnlyAnimations && !CommonSkeletalMeshesAndAnimationsProperties->Skeleton.IsValid())
 	{
 		UE_LOG(LogInterchangePipeline, Warning, TEXT("UInterchangeGenericAnimationPipeline: Cannot execute pre-import pipeline because we cannot import animation only but not specify any valid skeleton"));
 		return;
@@ -504,13 +504,13 @@ void UInterchangeGenericAnimationPipeline::ExecutePreImportPipeline(UInterchange
 			}
 		}
 
-		if (!CommonSkeletalMeshesAndAnimationsProperties->Skeleton.IsNull())
+		if (CommonSkeletalMeshesAndAnimationsProperties->Skeleton.IsValid())
 		{
 			bool bSkeletonCompatible = true;
 
 			//TODO: support skeleton helper in runtime
 #if WITH_EDITOR
-			bSkeletonCompatible = UE::Interchange::Private::FSkeletonHelper::IsCompatibleSkeleton(CommonSkeletalMeshesAndAnimationsProperties->Skeleton, RootJointUid, BaseNodeContainer);
+			bSkeletonCompatible = UE::Interchange::Private::FSkeletonHelper::IsCompatibleSkeleton(CommonSkeletalMeshesAndAnimationsProperties->Skeleton.Get(), RootJointUid, BaseNodeContainer);
 #endif
 			if(bSkeletonCompatible)
 			{
@@ -860,13 +860,13 @@ void UInterchangeGenericAnimationPipeline::CreateAnimSequenceFactoryNode(UInterc
 		}
 	}
 
-	if (!CommonSkeletalMeshesAndAnimationsProperties->Skeleton.IsNull())
+	if (CommonSkeletalMeshesAndAnimationsProperties->Skeleton.IsValid())
 	{
 		bool bSkeletonCompatible = true;
 
 		//TODO: support skeleton helper in runtime
 #if WITH_EDITOR
-		bSkeletonCompatible = UE::Interchange::Private::FSkeletonHelper::IsCompatibleSkeleton(CommonSkeletalMeshesAndAnimationsProperties->Skeleton, RootJointUid, BaseNodeContainer);
+		bSkeletonCompatible = UE::Interchange::Private::FSkeletonHelper::IsCompatibleSkeleton(CommonSkeletalMeshesAndAnimationsProperties->Skeleton.Get(), RootJointUid, BaseNodeContainer);
 #endif
 		if (bSkeletonCompatible)
 		{
diff --git a/Engine/Plugins/Experimental/Interchange/Runtime/Source/Pipelines/Private/InterchangeGenericMaterialPipeline.cpp b/Engine/Plugins/Experimental/Interchange/Runtime/Source/Pipelines/Private/InterchangeGenericMaterialPipeline.cpp
--- a/Engine/Plugins/Experimental/Interchange/Runtime/Source/Pipelines/Private/InterchangeGenericMaterialPipeline.cpp
+++ b/Engine/Plugins/Experimental/Interchange/Runtime/Source/Pipelines/Private/InterchangeGenericMaterialPipeline.cpp
@@ -149,6 +149,12 @@ void UInterchangeGenericMaterialPipeline::ExecutePreImportPipeline(UInterchangeB
 		TexturePipeline->ScriptedExecutePreImportPipeline(InBaseNodeContainer, InSourceDatas);
 	}
 
+	//Skip Material import if the toggle is off
+	if (!bImportMaterials)
+	{
+		return;
+	}
+
 	//Find all translated node we need for this pipeline
 	BaseNodeContainer->IterateNodes([this](const FString& NodeUid, UInterchangeBaseNode* Node)
 	{
diff --git a/Engine/Plugins/Experimental/Interchange/Runtime/Source/Pipelines/Private/InterchangeGenericMeshPipeline.cpp b/Engine/Plugins/Experimental/Interchange/Runtime/Source/Pipelines/Private/InterchangeGenericMeshPipeline.cpp
--- a/Engine/Plugins/Experimental/Interchange/Runtime/Source/Pipelines/Private/InterchangeGenericMeshPipeline.cpp
+++ b/Engine/Plugins/Experimental/Interchange/Runtime/Source/Pipelines/Private/InterchangeGenericMeshPipeline.cpp
@@ -19,7 +19,7 @@ void UInterchangeGenericMeshPipeline::AdjustSettingsForContext(EInterchangePipel
 {
 	Super::AdjustSettingsForContext(ImportType, ReimportAsset);
 
-	check(!CommonSkeletalMeshesAndAnimationsProperties.IsNull());
+	check(CommonSkeletalMeshesAndAnimationsProperties.IsValid());
 	if (ImportType == EInterchangePipelineContext::None)
 	{
 		//We do not change the setting if we are in editing context
diff --git a/Engine/Plugins/Experimental/Interchange/Runtime/Source/Pipelines/Private/InterchangeGenericSkeletalMeshPipeline.cpp b/Engine/Plugins/Experimental/Interchange/Runtime/Source/Pipelines/Private/InterchangeGenericSkeletalMeshPipeline.cpp
--- a/Engine/Plugins/Experimental/Interchange/Runtime/Source/Pipelines/Private/InterchangeGenericSkeletalMeshPipeline.cpp
+++ b/Engine/Plugins/Experimental/Interchange/Runtime/Source/Pipelines/Private/InterchangeGenericSkeletalMeshPipeline.cpp
@@ -72,7 +72,7 @@ namespace UE::Interchange::SkeletalMeshGenericPipeline
 
 void UInterchangeGenericMeshPipeline::ExecutePreImportPipelineSkeletalMesh()
 {
-	check(!CommonMeshesProperties.IsNull());
+	check(CommonMeshesProperties.IsValid());
 	if (!bImportSkeletalMeshes)
 	{
 		//Nothing to import
@@ -264,7 +264,7 @@ void UInterchangeGenericMeshPipeline::ExecutePreImportPipelineSkeletalMesh()
 
 UInterchangeSkeletonFactoryNode* UInterchangeGenericMeshPipeline::CreateSkeletonFactoryNode(const FString& RootJointUid)
 {
-	check(!CommonSkeletalMeshesAndAnimationsProperties.IsNull());
+	check(CommonSkeletalMeshesAndAnimationsProperties.IsValid());
 	const UInterchangeBaseNode* RootJointNode = BaseNodeContainer->GetNode(RootJointUid);
 	if (!RootJointNode)
 	{
@@ -306,10 +306,10 @@ UInterchangeSkeletonFactoryNode* UInterchangeGenericMeshPipeline::CreateSkeleton
 	}
 
 	//If we have a specified skeleton
-	if (CommonSkeletalMeshesAndAnimationsProperties->Skeleton)
+	if (CommonSkeletalMeshesAndAnimationsProperties->Skeleton.IsValid())
 	{
 		SkeletonFactoryNode->SetEnabled(false);
-		SkeletonFactoryNode->SetCustomReferenceObject(FSoftObjectPath(CommonSkeletalMeshesAndAnimationsProperties->Skeleton));
+		SkeletonFactoryNode->SetCustomReferenceObject(FSoftObjectPath(CommonSkeletalMeshesAndAnimationsProperties->Skeleton.Get()));
 	}
 #if WITH_EDITOR
 	//Iterate all joints to set the meta data value in the skeleton node
@@ -321,8 +321,8 @@ UInterchangeSkeletonFactoryNode* UInterchangeGenericMeshPipeline::CreateSkeleton
 
 UInterchangeSkeletalMeshFactoryNode* UInterchangeGenericMeshPipeline::CreateSkeletalMeshFactoryNode(const FString& RootJointUid, const TMap<int32, TArray<FString>>& MeshUidsPerLodIndex)
 {
-	check(!CommonMeshesProperties.IsNull());
-	check(!CommonSkeletalMeshesAndAnimationsProperties.IsNull());
+	check(CommonMeshesProperties.IsValid());
+	check(CommonSkeletalMeshesAndAnimationsProperties.IsValid());
 	//Get the skeleton factory node
 	const UInterchangeBaseNode* RootJointNode = BaseNodeContainer->GetNode(RootJointUid);
 	if (!RootJointNode)
@@ -424,13 +424,13 @@ UInterchangeSkeletalMeshFactoryNode* UInterchangeGenericMeshPipeline::CreateSkel
 	SkeletalMeshFactoryNode->SetCustomImportContentType(SkeletalMeshImportContentType);
 
 	//If we have a specified skeleton
-	if (CommonSkeletalMeshesAndAnimationsProperties->Skeleton)
+	if (CommonSkeletalMeshesAndAnimationsProperties->Skeleton.IsValid())
 	{
 		bool bSkeletonCompatible = false;
 
 		//TODO: support skeleton helper in runtime
 #if WITH_EDITOR
-		bSkeletonCompatible = UE::Interchange::Private::FSkeletonHelper::IsCompatibleSkeleton(CommonSkeletalMeshesAndAnimationsProperties->Skeleton, RootJointNode->GetUniqueID(), BaseNodeContainer);
+		bSkeletonCompatible = UE::Interchange::Private::FSkeletonHelper::IsCompatibleSkeleton(CommonSkeletalMeshesAndAnimationsProperties->Skeleton.Get(), RootJointNode->GetUniqueID(), BaseNodeContainer);
 #endif
 		if (bSkeletonCompatible)
 		{
@@ -461,7 +461,7 @@ UInterchangeSkeletalMeshFactoryNode* UInterchangeGenericMeshPipeline::CreateSkel
 		}
 	}
 	SkeletalMeshFactoryNode->SetCustomCreatePhysicsAsset(bCreatePhysicsAsset);
-	if (!bCreatePhysicsAsset && !PhysicsAsset.IsNull())
+	if (!bCreatePhysicsAsset && PhysicsAsset.IsValid())
 	{
 		FSoftObjectPath PhysicSoftObjectPath(PhysicsAsset.Get());
 		SkeletalMeshFactoryNode->SetCustomPhysicAssetSoftObjectPath(PhysicSoftObjectPath);
@@ -530,8 +530,8 @@ UInterchangeSkeletalMeshLodDataNode* UInterchangeGenericMeshPipeline::CreateSkel
 
 void UInterchangeGenericMeshPipeline::AddLodDataToSkeletalMesh(const UInterchangeSkeletonFactoryNode* SkeletonFactoryNode, UInterchangeSkeletalMeshFactoryNode* SkeletalMeshFactoryNode, const TMap<int32, TArray<FString>>& NodeUidsPerLodIndex)
 {
-	check(!CommonMeshesProperties.IsNull());
-	check(!CommonSkeletalMeshesAndAnimationsProperties.IsNull());
+	check(CommonMeshesProperties.IsValid());
+	check(CommonSkeletalMeshesAndAnimationsProperties.IsValid());
 
 	const FString SkeletalMeshUid = SkeletalMeshFactoryNode->GetUniqueID();
 	const FString SkeletonUid = SkeletonFactoryNode->GetUniqueID();
@@ -606,7 +606,7 @@ void UInterchangeGenericMeshPipeline::AddLodDataToSkeletalMesh(const UInterchang
 
 void UInterchangeGenericMeshPipeline::PostImportSkeletalMesh(UObject* CreatedAsset, const UInterchangeFactoryBaseNode* FactoryNode)
 {
-	check(!CommonSkeletalMeshesAndAnimationsProperties.IsNull());
+	check(CommonSkeletalMeshesAndAnimationsProperties.IsValid());
 
 	if (!BaseNodeContainer)
 	{
@@ -621,7 +621,7 @@ void UInterchangeGenericMeshPipeline::PostImportSkeletalMesh(UObject* CreatedAss
 
 	//If we import only the geometry we do not want to update the skeleton reference pose.
 	const bool bImportGeometryOnlyContent = SkeletalMeshImportContentType == EInterchangeSkeletalMeshContentType::Geometry;
-	if (!bImportGeometryOnlyContent && bUpdateSkeletonReferencePose && !CommonSkeletalMeshesAndAnimationsProperties->Skeleton.IsNull() && SkeletalMesh->GetSkeleton() == CommonSkeletalMeshesAndAnimationsProperties->Skeleton.Get())
+	if (!bImportGeometryOnlyContent && bUpdateSkeletonReferencePose && CommonSkeletalMeshesAndAnimationsProperties->Skeleton.IsValid() && SkeletalMesh->GetSkeleton() == CommonSkeletalMeshesAndAnimationsProperties->Skeleton.Get())
 	{
 		SkeletalMesh->GetSkeleton()->UpdateReferencePoseFromMesh(SkeletalMesh);
 		//TODO: notify editor the skeleton has change
@@ -686,7 +686,7 @@ void UInterchangeGenericMeshPipeline::PostImportPhysicsAssetImport(UObject* Crea
 
 void UInterchangeGenericMeshPipeline::ImplementUseSourceNameForAssetOptionSkeletalMesh(const int32 MeshesImportedNodeCount, const bool bUseSourceNameForAsset)
 {
-	check(!CommonSkeletalMeshesAndAnimationsProperties.IsNull());
+	check(CommonSkeletalMeshesAndAnimationsProperties.IsValid());
 
 	const UClass* SkeletalMeshFactoryNodeClass = UInterchangeSkeletalMeshFactoryNode::StaticClass();
 	TArray<FString> SkeletalMeshNodeUids;
@@ -721,7 +721,7 @@ void UInterchangeGenericMeshPipeline::ImplementUseSourceNameForAssetOptionSkelet
 		if (const UInterchangeSkeletalMeshLodDataNode* SkeletalMeshLodDataNode = Cast<const UInterchangeSkeletalMeshLodDataNode>(BaseNodeContainer->GetFactoryNode(LodDataUids[0])))
 		{
 			//If the user did not specify any skeleton
-			if (CommonSkeletalMeshesAndAnimationsProperties->Skeleton.IsNull())
+			if (!CommonSkeletalMeshesAndAnimationsProperties->Skeleton.IsValid())
 			{
 				FString SkeletalMeshSkeletonUid;
 				SkeletalMeshLodDataNode->GetCustomSkeletonUid(SkeletalMeshSkeletonUid);
diff --git a/Engine/Plugins/Experimental/Interchange/Runtime/Source/Pipelines/Private/InterchangeGenericStaticMeshPipeline.cpp b/Engine/Plugins/Experimental/Interchange/Runtime/Source/Pipelines/Private/InterchangeGenericStaticMeshPipeline.cpp
--- a/Engine/Plugins/Experimental/Interchange/Runtime/Source/Pipelines/Private/InterchangeGenericStaticMeshPipeline.cpp
+++ b/Engine/Plugins/Experimental/Interchange/Runtime/Source/Pipelines/Private/InterchangeGenericStaticMeshPipeline.cpp
@@ -136,7 +136,7 @@ static void BuildMeshToCollisionMeshMap(const TArray<FString>& MeshUids, TMap<FS
 
 void UInterchangeGenericMeshPipeline::ExecutePreImportPipelineStaticMesh()
 {
-	check(!CommonMeshesProperties.IsNull());
+	check(CommonMeshesProperties.IsValid());
 
 #if WITH_EDITOR
 	//Make sure the generic pipeline will cover all staticmesh build settings when we import
@@ -382,7 +382,7 @@ bool UInterchangeGenericMeshPipeline::MakeMeshFactoryNodeUidAndDisplayLabel(cons
 
 UInterchangeStaticMeshFactoryNode* UInterchangeGenericMeshPipeline::CreateStaticMeshFactoryNode(const TMap<int32, TArray<FString>>& MeshUidsPerLodIndex)
 {
-	check(!CommonMeshesProperties.IsNull());
+	check(CommonMeshesProperties.IsValid());
 	if (MeshUidsPerLodIndex.Num() == 0)
 	{
 		return nullptr;
@@ -448,7 +448,7 @@ UInterchangeStaticMeshFactoryNode* UInterchangeGenericMeshPipeline::CreateStatic
 	StaticMeshFactoryNode->SetCustomDstLightmapIndex(DstLightmapIndex);
 	StaticMeshFactoryNode->SetCustomBuildScale3D(BuildScale3D);
 	StaticMeshFactoryNode->SetCustomDistanceFieldResolutionScale(DistanceFieldResolutionScale);
-	StaticMeshFactoryNode->SetCustomDistanceFieldReplacementMesh(DistanceFieldReplacementMesh);
+	StaticMeshFactoryNode->SetCustomDistanceFieldReplacementMesh(DistanceFieldReplacementMesh.Get());
 	StaticMeshFactoryNode->SetCustomMaxLumenMeshCards(MaxLumenMeshCards);
 	StaticMeshFactoryNode->SetCustomBuildNanite(bBuildNanite);
 
@@ -476,7 +476,7 @@ UInterchangeStaticMeshLodDataNode* UInterchangeGenericMeshPipeline::CreateStatic
 
 void UInterchangeGenericMeshPipeline::AddLodDataToStaticMesh(UInterchangeStaticMeshFactoryNode* StaticMeshFactoryNode, const TMap<int32, TArray<FString>>& NodeUidsPerLodIndex)
 {
-	check(!CommonMeshesProperties.IsNull());
+	check(CommonMeshesProperties.IsValid());
 	const FString StaticMeshFactoryUid = StaticMeshFactoryNode->GetUniqueID();
 
 	for (const TPair<int32, TArray<FString>>& LodIndexAndNodeUids : NodeUidsPerLodIndex)
diff --git a/Engine/Plugins/Experimental/Interchange/Runtime/Source/Pipelines/Public/InterchangeGenericAnimationPipeline.h b/Engine/Plugins/Experimental/Interchange/Runtime/Source/Pipelines/Public/InterchangeGenericAnimationPipeline.h
--- a/Engine/Plugins/Experimental/Interchange/Runtime/Source/Pipelines/Public/InterchangeGenericAnimationPipeline.h
+++ b/Engine/Plugins/Experimental/Interchange/Runtime/Source/Pipelines/Public/InterchangeGenericAnimationPipeline.h
@@ -39,7 +39,7 @@ public:
 
 	//Common SkeletalMeshes And Animations Properties Settings Pointer
 	UPROPERTY(Transient)
-	TObjectPtr<UInterchangeGenericCommonSkeletalMeshesAndAnimationsProperties> CommonSkeletalMeshesAndAnimationsProperties;
+	TWeakObjectPtr<UInterchangeGenericCommonSkeletalMeshesAndAnimationsProperties> CommonSkeletalMeshesAndAnimationsProperties;
 
 	//////	ANIMATION_CATEGORY Properties //////
 	/** If enable, import all animation assets find in the sources. */
diff --git a/Engine/Plugins/Experimental/Interchange/Runtime/Source/Pipelines/Public/InterchangeGenericAssetsPipelineSharedSettings.h b/Engine/Plugins/Experimental/Interchange/Runtime/Source/Pipelines/Public/InterchangeGenericAssetsPipelineSharedSettings.h
--- a/Engine/Plugins/Experimental/Interchange/Runtime/Source/Pipelines/Public/InterchangeGenericAssetsPipelineSharedSettings.h
+++ b/Engine/Plugins/Experimental/Interchange/Runtime/Source/Pipelines/Public/InterchangeGenericAssetsPipelineSharedSettings.h
@@ -115,7 +115,7 @@ public:
 
 	/** Skeleton to use for imported asset. When importing a skeletal mesh, leaving this as "None" will create a new skeleton. When importing an animation this MUST be specified to import the asset. */
 	UPROPERTY(EditAnywhere, BlueprintReadWrite, Category = "Common Skeletal Meshes and Animations")
-	TObjectPtr<USkeleton> Skeleton;
+	TWeakObjectPtr<USkeleton> Skeleton;
 
 	/** If checked, meshes nested in bone hierarchies will be imported instead of being converted to bones. */
 	UPROPERTY(EditAnywhere, BlueprintReadWrite, Category = "Common Skeletal Meshes and Animations")
@@ -127,7 +127,7 @@ public:
 
 	virtual bool IsSettingsAreValid(TOptional<FText>& OutInvalidReason) const override
 	{
-		if (bImportOnlyAnimations && Skeleton.IsNull())
+		if (bImportOnlyAnimations && !Skeleton.IsValid())
 		{
 			OutInvalidReason = NSLOCTEXT("UInterchangeGenericCommonSkeletalMeshesAndAnimationsProperties", "SkeletonMustBeSpecified", "When importing only animations, a valid skeleton must be set.");
 			return false;
diff --git a/Engine/Plugins/Experimental/Interchange/Runtime/Source/Pipelines/Public/InterchangeGenericMaterialPipeline.h b/Engine/Plugins/Experimental/Interchange/Runtime/Source/Pipelines/Public/InterchangeGenericMaterialPipeline.h
--- a/Engine/Plugins/Experimental/Interchange/Runtime/Source/Pipelines/Public/InterchangeGenericMaterialPipeline.h
+++ b/Engine/Plugins/Experimental/Interchange/Runtime/Source/Pipelines/Public/InterchangeGenericMaterialPipeline.h
@@ -37,7 +37,7 @@ class INTERCHANGEPIPELINES_API UInterchangeGenericMaterialPipeline : public UInt
 public:
 	UInterchangeGenericMaterialPipeline();
 
-	/** If enabled, imports the texture assets found in the sources. */
+	/** If enabled, imports the material assets found in the sources. */
 	UPROPERTY(EditAnywhere, BlueprintReadWrite, Category = "Materials")
 	bool bImportMaterials = true;
 
diff --git a/Engine/Plugins/Experimental/Interchange/Runtime/Source/Pipelines/Public/InterchangeGenericMeshPipeline.h b/Engine/Plugins/Experimental/Interchange/Runtime/Source/Pipelines/Public/InterchangeGenericMeshPipeline.h
--- a/Engine/Plugins/Experimental/Interchange/Runtime/Source/Pipelines/Public/InterchangeGenericMeshPipeline.h
+++ b/Engine/Plugins/Experimental/Interchange/Runtime/Source/Pipelines/Public/InterchangeGenericMeshPipeline.h
@@ -32,11 +32,11 @@ public:
 
 	//Common Meshes Properties Settings Pointer
 	UPROPERTY(Transient)
-	TObjectPtr<UInterchangeGenericCommonMeshesProperties> CommonMeshesProperties;
+	TWeakObjectPtr<UInterchangeGenericCommonMeshesProperties> CommonMeshesProperties;
 
 	//Common SkeletalMeshes And Animations Properties Settings Pointer
 	UPROPERTY(Transient)
-	TObjectPtr<UInterchangeGenericCommonSkeletalMeshesAndAnimationsProperties> CommonSkeletalMeshesAndAnimationsProperties;
+	TWeakObjectPtr<UInterchangeGenericCommonSkeletalMeshesAndAnimationsProperties> CommonSkeletalMeshesAndAnimationsProperties;
 	
 	//////	STATIC_MESHES_CATEGORY Properties //////
 
@@ -110,7 +110,7 @@ public:
 	float DistanceFieldResolutionScale = 1.0f;
 	
 	UPROPERTY(EditAnywhere, BlueprintReadWrite, Category = "Static Meshes", meta = (SubCategory = "Build"))
-	TObjectPtr<class UStaticMesh> DistanceFieldReplacementMesh = nullptr;
+	TWeakObjectPtr<class UStaticMesh> DistanceFieldReplacementMesh = nullptr;
 	
 	/** 
 	 * Max Lumen mesh cards to generate for this mesh.
@@ -153,7 +153,7 @@ public:
 
 	/** If this is set, use this specified PhysicsAsset. If its not set and bCreatePhysicsAsset is false, the importer will not generate or set any physic asset. */
 	UPROPERTY(EditAnywhere, BlueprintReadWrite, Category = "Skeletal Meshes", meta = (editcondition = "!bCreatePhysicsAsset"))
-	TObjectPtr<UPhysicsAsset> PhysicsAsset;
+	TWeakObjectPtr<UPhysicsAsset> PhysicsAsset;
 
 	/** Threshold use to decide if two vertex position are equal. */
 	UPROPERTY(EditAnywhere, BlueprintReadWrite, Category = "Skeletal Meshes", meta = (SubCategory = "Build"))
diff --git a/Engine/Source/Developer/AnimationDataController/Private/AnimDataController.cpp b/Engine/Source/Developer/AnimationDataController/Private/AnimDataController.cpp
--- a/Engine/Source/Developer/AnimationDataController/Private/AnimDataController.cpp
+++ b/Engine/Source/Developer/AnimationDataController/Private/AnimDataController.cpp
@@ -193,7 +193,7 @@ void UAnimDataController::SetFrameRate(FFrameRate FrameRate, bool bShouldTransac
 	{
 		CONDITIONAL_TRANSACTION(LOCTEXT("SetFrameRate", "Setting Frame Rate"));
 
-		CONDITIONAL_ACTION(UE::Anim::FSetFrameRateAction, Model);
+		CONDITIONAL_ACTION(UE::Anim::FSetFrameRateAction, Model.Get());
 
 		FFrameRateChangedPayload Payload;
 		Payload.PreviousFrameRate = Model->FrameRate;
@@ -1207,7 +1207,7 @@ void UAnimDataController::ResizePlayLength(float Length, float T0, float T1, boo
 					Payload.T1 = T1;
 					Payload.PreviousLength = Model->PlayLength;
 
-					CONDITIONAL_ACTION(UE::Anim::FResizePlayLengthAction, Model, T0, T1);
+					CONDITIONAL_ACTION(UE::Anim::FResizePlayLengthAction, Model.Get(), T0, T1);
 
 					Model->PlayLength = Length;
 
diff --git a/Engine/Source/Developer/AnimationDataController/Public/AnimDataController.h b/Engine/Source/Developer/AnimationDataController/Public/AnimDataController.h
--- a/Engine/Source/Developer/AnimationDataController/Public/AnimDataController.h
+++ b/Engine/Source/Developer/AnimationDataController/Public/AnimDataController.h
@@ -45,8 +45,8 @@ public:
 #if WITH_EDITOR
 	/** Begin IAnimationDataController overrides */
 	virtual void SetModel(UAnimDataModel* InModel)  override;
-    virtual UAnimDataModel* GetModel() override { return Model; }
-	virtual const UAnimDataModel* const GetModel() const override { return Model; }
+    virtual UAnimDataModel* GetModel() override { return Model.Get(); }
+	virtual const UAnimDataModel* const GetModel() const override { return Model.Get(); }
 	virtual void OpenBracket(const FText& InTitle, bool bShouldTransact = true) override;
 	virtual void CloseBracket(bool bShouldTransact = true) override;
 	virtual void SetPlayLength(float Length, bool bShouldTransact = true) override;
@@ -152,7 +152,7 @@ private:
 
 #if WITH_EDITORONLY_DATA
 	UPROPERTY(transient)
-	TObjectPtr<UAnimDataModel> Model;
+	TWeakObjectPtr<UAnimDataModel> Model;
 #endif // WITH_EDITORONLY_DATA
 
 	friend class FAnimDataControllerTestBase;
