Interchange error handling
- Added a generic error message when we fail to retrieve a texture payload.
- Added error messages in the GLTF, OBJ and IES translators.
- Skip over invalid imported objects in the completion task so that we don't try to posteditchange them and we don't log a success.
- Display the correct icon for the "import done" notification based on if we have a success message or not.

diff --git a/Engine/Plugins/Interchange/Runtime/Source/Import/Private/Gltf/InterchangeGltfTranslator.cpp b/Engine/Plugins/Interchange/Runtime/Source/Import/Private/Gltf/InterchangeGltfTranslator.cpp
--- a/Engine/Plugins/Interchange/Runtime/Source/Import/Private/Gltf/InterchangeGltfTranslator.cpp
+++ b/Engine/Plugins/Interchange/Runtime/Source/Import/Private/Gltf/InterchangeGltfTranslator.cpp
@@ -34,6 +34,7 @@
 
 #include "Gltf/InterchangeGltfPrivate.h"
 
+#define LOCTEXT_NAMESPACE "InterchangeGltfTranslator"
 
 static bool GInterchangeEnableGLTFImport = true;
 static FAutoConsoleVariableRef CCvarInterchangeEnableGLTFImport(
@@ -838,28 +839,34 @@ bool UInterchangeGltfTranslator::Translate( UInterchangeBaseNodeContainer& NodeC
 
 	ScaleNodeTranslations(const_cast<UInterchangeGltfTranslator*>(this)->GltfAsset.Nodes, GltfUnitConversionMultiplier);
 
+	const FString FileName = FPaths::GetBaseFilename(FilePath);
+
 	if (GltfAsset.RequiredExtensions.Num() != 0)
 	{
 		FString NotSupportedExtensions;
-		for (const FString& RequiredExtension: GltfAsset.RequiredExtensions)
+		for ( const FString& RequiredExtension: GltfAsset.RequiredExtensions )
 		{
-			if (SupportedExtensions.Find(RequiredExtension) == INDEX_NONE)
+			if ( SupportedExtensions.Find( RequiredExtension ) == INDEX_NONE )
 			{
-				if (NotSupportedExtensions.Len() > 0)
+				if ( NotSupportedExtensions.Len() > 0 )
 				{
 					NotSupportedExtensions += ", ";
 				}
 				NotSupportedExtensions += RequiredExtension;
 			}
 		}
-		if (NotSupportedExtensions.Len() > 0)
+		if ( NotSupportedExtensions.Len() > 0 )
 		{
-			UE_LOG(LogInterchangeImport, Error, TEXT("Not All Required Extensions are supported. (Unsupported extensions: %s)"), *NotSupportedExtensions);
+			UInterchangeResultError_Generic* ErrorResult = AddMessage< UInterchangeResultError_Generic >();
+			ErrorResult->SourceAssetName = FileName;
+			ErrorResult->Text = FText::Format(
+				LOCTEXT( "UnsupportedRequiredExtensions", "Not All Required Extensions are supported. (Unsupported extensions: {0})"),
+				FText::FromString( NotSupportedExtensions ) );
+
 			return false;
 		}
 	}
 
-	const FString FileName = FPaths::GetBaseFilename(FilePath);
 	const_cast< UInterchangeGltfTranslator* >( this )->GltfAsset.GenerateNames(FileName);
 
 	// Textures
@@ -904,7 +911,7 @@ bool UInterchangeGltfTranslator::Translate( UInterchangeBaseNodeContainer& NodeC
 			NodeContainer.AddNode( MeshNode );
 
 			
-			for (size_t PrimitiveCounter = 0; PrimitiveCounter < GltfMesh.Primitives.Num(); PrimitiveCounter++)
+			for (int32 PrimitiveCounter = 0; PrimitiveCounter < GltfMesh.Primitives.Num(); PrimitiveCounter++)
 			{
 				const GLTF::FPrimitive& Primitive = GltfMesh.Primitives[PrimitiveCounter];
 
@@ -1042,23 +1049,54 @@ bool UInterchangeGltfTranslator::Translate( UInterchangeBaseNodeContainer& NodeC
 			}
 
 			// Skeletons:
-			HandleGltfSkeletons(NodeContainer, SceneNodeUid, SkinnedMeshNodes);
+			HandleGltfSkeletons( NodeContainer, SceneNodeUid, SkinnedMeshNodes );
 		}
 	}
 
 	// Animations
 	{
-		for (int32 AnimationIndex = 0; AnimationIndex < GltfAsset.Animations.Num(); ++AnimationIndex)
+		for ( int32 AnimationIndex = 0; AnimationIndex < GltfAsset.Animations.Num(); ++AnimationIndex )
 		{
-			HandleGltfAnimation(NodeContainer, AnimationIndex);
+			HandleGltfAnimation( NodeContainer, AnimationIndex );
 		}
 	}
 
 	// Variants
 	// Note: Variants are not supported yet in game play mode
-	if (!FApp::IsGame() && bHasVariants)
+	if ( !FApp::IsGame() && bHasVariants )
 	{
-		HandleGltfVariants(NodeContainer, FileName);
+		HandleGltfVariants( NodeContainer, FileName );
+	}
+
+	// Add glTF errors and warnings to the Interchange results
+	for ( const GLTF::FLogMessage& LogMessage : GltfFileReader.GetLogMessages() )
+	{
+		UInterchangeResult* Result = nullptr;
+
+		switch ( LogMessage.Get<0>() )
+		{
+		case GLTF::EMessageSeverity::Error :
+			{
+				UInterchangeResultError_Generic* ErrorResult = AddMessage< UInterchangeResultError_Generic >();
+				ErrorResult->Text = FText::FromString(LogMessage.Get<1>());
+				Result = ErrorResult;
+			}
+			break;
+		case GLTF::EMessageSeverity::Warning:
+		default:
+			{
+				UInterchangeResultWarning_Generic* WarningResult = AddMessage< UInterchangeResultWarning_Generic >();
+				WarningResult->Text = FText::FromString(LogMessage.Get<1>());
+				Result = WarningResult;
+			}
+
+			break;
+		}
+
+		if ( Result )
+		{
+			Result->SourceAssetName = FileName;
+		}
 	}
 
 	return true;
@@ -1784,4 +1822,6 @@ void UInterchangeGltfTranslator::HandleGltfSkeletons(UInterchangeBaseNodeContain
 			}
 		}
 	}
-}
\ No newline at end of file
+}
+
+#undef LOCTEXT_NAMESPACE
\ No newline at end of file
diff --git a/Engine/Plugins/Interchange/Runtime/Source/Import/Private/Mesh/InterchangeOBJTranslator.cpp b/Engine/Plugins/Interchange/Runtime/Source/Import/Private/Mesh/InterchangeOBJTranslator.cpp
--- a/Engine/Plugins/Interchange/Runtime/Source/Import/Private/Mesh/InterchangeOBJTranslator.cpp
+++ b/Engine/Plugins/Interchange/Runtime/Source/Import/Private/Mesh/InterchangeOBJTranslator.cpp
@@ -921,6 +921,15 @@ bool UInterchangeOBJTranslator::Translate(UInterchangeBaseNodeContainer& BaseNod
 
 	if (bSuccess)
 	{
+		if (ObjDataPtr->Groups.IsEmpty() && ObjDataPtr->Materials.IsEmpty())
+		{
+			UInterchangeResultError_Generic* ErrorResult = AddMessage<UInterchangeResultError_Generic>();
+			ErrorResult->SourceAssetName = FPaths::GetBaseFilename(Filename);
+			ErrorResult->Text = NSLOCTEXT("InterchangeOBJTranslator", "EmptyFileError", "The OBJ file appears to be empty.");
+
+			return false;
+		}
+
 		// Add mesh nodes to the container for each mesh group
 
 		for (const TPair<FString, FObjData::FGroupData>& Group : ObjDataPtr->Groups)
diff --git a/Engine/Plugins/Interchange/Runtime/Source/Import/Private/Texture/InterchangeIESTranslator.cpp b/Engine/Plugins/Interchange/Runtime/Source/Import/Private/Texture/InterchangeIESTranslator.cpp
--- a/Engine/Plugins/Interchange/Runtime/Source/Import/Private/Texture/InterchangeIESTranslator.cpp
+++ b/Engine/Plugins/Interchange/Runtime/Source/Import/Private/Texture/InterchangeIESTranslator.cpp
@@ -11,6 +11,8 @@
 #include "Misc/Paths.h"
 #include "Texture/TextureTranslatorUtilities.h"
 
+#define LOCTEXT_NAMESPACE "InterchangeIESTranslator"
+
 static bool GInterchangeEnableIESImport = true;
 static FAutoConsoleVariableRef CCvarInterchangeEnableIESImport(
 	TEXT("Interchange.FeatureFlags.Import.IES"),
@@ -18,6 +20,20 @@ static FAutoConsoleVariableRef CCvarInterchangeEnableIESImport(
 	TEXT("Whether IES support is enabled."),
 	ECVF_Default);
 
+namespace UE::Interchange::IESTranslator::Private
+{
+	void LogError(const UInterchangeIESTranslator& IESTranslator, const FString& UniqueID, FText&& ErrorText)
+	{
+		check(IESTranslator.GetSourceData() != nullptr);
+
+		UInterchangeResultError_Generic* ErrorMessage = IESTranslator.AddMessage<UInterchangeResultError_Generic>();
+		ErrorMessage->SourceAssetName = IESTranslator.GetSourceData()->GetFilename();
+		ErrorMessage->AssetType = UTextureLightProfile::StaticClass();
+		ErrorMessage->InterchangeKey = UniqueID;
+		ErrorMessage->Text = MoveTemp(ErrorText);
+	}
+}
+
 TArray<FString> UInterchangeIESTranslator::GetSupportedFormats() const
 {
 	if (GInterchangeEnableIESImport || GIsAutomationTesting)
@@ -38,11 +54,15 @@ bool UInterchangeIESTranslator::Translate(UInterchangeBaseNodeContainer& BaseNod
 
 TOptional<UE::Interchange::FImportLightProfile> UInterchangeIESTranslator::GetLightProfilePayloadData(const UInterchangeSourceData* PayloadSourceData, const FString& PayLoadKey) const
 {
+	using namespace UE::Interchange;
+
 	check(GetSourceData() == PayloadSourceData);
 
 	if (!GetSourceData())
 	{
-		UE_LOG(LogInterchangeImport, Error, TEXT("Failed to import IES, bad source data."));
+		IESTranslator::Private::LogError(*this, PayLoadKey,
+			LOCTEXT("IESImportFailure_BadData", "Failed to import IES, bad source data."));
+
 		return {};
 	}
 
@@ -52,19 +72,25 @@ TOptional<UE::Interchange::FImportLightProfile> UInterchangeIESTranslator::GetLi
 	//Make sure the key fit the filename, The key should always be valid
 	if (!Filename.Equals(PayLoadKey))
 	{
-		UE_LOG(LogInterchangeImport, Error, TEXT("Failed to import IES, wrong payload key. [%s]"), *Filename);
+		IESTranslator::Private::LogError(*this, PayLoadKey,
+			LOCTEXT("IESImportFailure_BadPayloadKey", "Failed to import IES, wrong payload key."));
+
 		return {};
 	}
 
 	if (!FPaths::FileExists(Filename))
 	{
-		UE_LOG(LogInterchangeImport, Error, TEXT("Failed to import IES, cannot open file. [%s]"), *Filename);
+		IESTranslator::Private::LogError(*this, PayLoadKey,
+			LOCTEXT("IESImportFailure_OpenFile", "Failed to import IES, cannot open file."));
+
 		return {};
 	}
 
 	if (!FFileHelper::LoadFileToArray(SourceDataBuffer, *Filename))
 	{
-		UE_LOG(LogInterchangeImport, Error, TEXT("Failed to import IES, cannot load file content into an array. [%s]"), *Filename);
+		IESTranslator::Private::LogError(*this, PayLoadKey,
+			LOCTEXT("IESImportFailure_LoadFile", "Failed to import IES, cannot load file content into an array."));
+
 		return {};
 	}
 
@@ -94,6 +120,13 @@ TOptional<UE::Interchange::FImportLightProfile> UInterchangeIESTranslator::GetLi
 
 		return Payload;
 	}
+	else
+	{
+		IESTranslator::Private::LogError(*this, PayLoadKey,
+			FText::Format(LOCTEXT("IESImportFailure_Generic", "IES import failed: \"{0}\""), FText::FromString(IESConverter.GetError())));
+	}
 
 	return {};
 }
+
+#undef LOCTEXT_NAMESPACE
\ No newline at end of file
diff --git a/Engine/Plugins/Interchange/Runtime/Source/Import/Private/Texture/InterchangeImageWrapperTranslator.cpp b/Engine/Plugins/Interchange/Runtime/Source/Import/Private/Texture/InterchangeImageWrapperTranslator.cpp
--- a/Engine/Plugins/Interchange/Runtime/Source/Import/Private/Texture/InterchangeImageWrapperTranslator.cpp
+++ b/Engine/Plugins/Interchange/Runtime/Source/Import/Private/Texture/InterchangeImageWrapperTranslator.cpp
@@ -177,7 +177,10 @@ TOptional<UE::Interchange::FImportImage> UInterchangeImageWrapperTranslator::Get
 		if (ImageWrapperModule.DecompressImage(Buffer, Length, LoadedImage))
 		{
 			// Todo interchange: should these payload modification be part of the pipeline, factory or stay there?
-			UE::TextureUtilitiesCommon::AutoDetectAndChangeGrayScale(LoadedImage);
+			if (UE::TextureUtilitiesCommon::AutoDetectAndChangeGrayScale(LoadedImage))
+			{
+				UE_LOG(LogInterchangeImport, Display, TEXT("Auto-detected grayscale, image changed to G8"));
+			}
 
 			ETextureSourceFormat TextureFormat = FImageCoreUtils::ConvertToTextureSourceFormat(LoadedImage.Format);
 			bool bSRGB = LoadedImage.GammaSpace != EGammaSpace::Linear;
diff --git a/Engine/Plugins/Interchange/Runtime/Source/Import/Private/Texture/InterchangeTextureFactory.cpp b/Engine/Plugins/Interchange/Runtime/Source/Import/Private/Texture/InterchangeTextureFactory.cpp
--- a/Engine/Plugins/Interchange/Runtime/Source/Import/Private/Texture/InterchangeTextureFactory.cpp
+++ b/Engine/Plugins/Interchange/Runtime/Source/Import/Private/Texture/InterchangeTextureFactory.cpp
@@ -980,9 +980,13 @@ namespace UE::Interchange::Private::InterchangeTextureFactory
 		return FilesToHash;
 	}
 
-	void LogErrorInvalidPayload(const FString& TextureClass, const FString& ObjectName)
+	void LogErrorInvalidPayload(UInterchangeFactoryBase& Factory, const UClass* TextureClass, const FString& SourceAssetName, const FString& DestinationAssetName)
 	{
-		UE_LOG(LogInterchangeImport, Error, TEXT("UInterchangeTextureFactory: The Payload was invalid for a %s. (%s)"), *TextureClass, *ObjectName);
+		UInterchangeResultError_Generic* ErrorMessage = Factory.AddMessage<UInterchangeResultError_Generic>();
+		ErrorMessage->AssetType = TextureClass;
+		ErrorMessage->SourceAssetName = SourceAssetName;
+		ErrorMessage->DestinationAssetName = DestinationAssetName;
+		ErrorMessage->Text = LOCTEXT("InvalidPayload", "Unable to retrieve the payload from the source file.");
 	}
 
 	bool CanSetupTexture2DSourceData(FTexturePayloadVariant& TexturePayload)
@@ -1338,7 +1342,7 @@ UObject* UInterchangeTextureFactory::CreateAsset(const FCreateAssetParams& Argum
 
 	if (!bCanSetup)
 	{
-		LogErrorInvalidPayload(Texture->GetClass()->GetName(), Texture->GetName());
+		LogErrorInvalidPayload(*this, Texture->GetClass(), Arguments.SourceData->GetFilename(), Texture->GetName());
 		return Texture;
 	}
 
diff --git a/Engine/Plugins/Interchange/Runtime/Source/Pipelines/Private/InterchangeGenericTexturePipeline.cpp b/Engine/Plugins/Interchange/Runtime/Source/Pipelines/Private/InterchangeGenericTexturePipeline.cpp
--- a/Engine/Plugins/Interchange/Runtime/Source/Pipelines/Private/InterchangeGenericTexturePipeline.cpp
+++ b/Engine/Plugins/Interchange/Runtime/Source/Pipelines/Private/InterchangeGenericTexturePipeline.cpp
@@ -67,6 +67,8 @@ namespace UE::Interchange::Private
 			Texture->PreEditChange(nullptr);
 			if (UE::NormalMapIdentification::HandleAssetPostImport(Texture))
 			{
+				UE_LOG(LogInterchangePipeline, Display, TEXT("Auto-detected normal map"));
+
 				if (bFlipNormalMapGreenChannel)
 				{
 					Texture->bFlipGreenChannel = true;
diff --git a/Engine/Source/Runtime/Interchange/Engine/Private/InterchangeManager.cpp b/Engine/Source/Runtime/Interchange/Engine/Private/InterchangeManager.cpp
--- a/Engine/Source/Runtime/Interchange/Engine/Private/InterchangeManager.cpp
+++ b/Engine/Source/Runtime/Interchange/Engine/Private/InterchangeManager.cpp
@@ -1316,7 +1316,36 @@ TSharedRef<UE::Interchange::FImportAsyncHelper, ESPMode::ThreadSafe> UInterchang
 
 void UInterchangeManager::ReleaseAsyncHelper(TWeakPtr<UE::Interchange::FImportAsyncHelper, ESPMode::ThreadSafe> AsyncHelper)
 {
+	using namespace UE::Interchange;
+
 	check(AsyncHelper.IsValid());
+	
+	bool bSucceeded = false;
+	{
+		TSharedPtr<FImportAsyncHelper> AsyncHelperPtr = AsyncHelper.Pin();
+		
+		for (UInterchangeResult* Result : AsyncHelperPtr->AssetImportResult->GetResults()->GetResults())
+		{
+			if (Result && Result->GetResultType() == EInterchangeResultType::Success)
+			{
+				bSucceeded = true;
+				break;
+			}
+		}
+
+		if (!bSucceeded)
+		{
+			for (UInterchangeResult* Result : AsyncHelperPtr->SceneImportResult->GetResults()->GetResults())
+			{
+				if (Result && Result->GetResultType() == EInterchangeResultType::Success)
+				{
+					bSucceeded = true;
+					break;
+				}
+			}
+		}
+	}
+
 	ImportTasks.RemoveSingle(AsyncHelper.Pin());
 	//Make sure the async helper is destroy, if not destroy its because we are canceling the import and we still have a shared ptr on it
 	{
@@ -1334,9 +1363,7 @@ void UInterchangeManager::ReleaseAsyncHelper(TWeakPtr<UE::Interchange::FImportAs
 		if (Notification.IsValid())
 		{
 			FText TitleText = NSLOCTEXT("Interchange", "Asynchronous_import_end", "Import Done");
-			//TODO make sure any error are reported so we can control success or not
-			const bool bSuccess = true;
-			Notification->SetComplete(TitleText, FText::GetEmpty(), bSuccess);
+			Notification->SetComplete(TitleText, FText::GetEmpty(), bSucceeded);
 			Notification = nullptr; //This should delete the notification
 		}
 	}
diff --git a/Engine/Source/Runtime/Interchange/Engine/Private/Tasks/InterchangeTaskCompletion.cpp b/Engine/Source/Runtime/Interchange/Engine/Private/Tasks/InterchangeTaskCompletion.cpp
--- a/Engine/Source/Runtime/Interchange/Engine/Private/Tasks/InterchangeTaskCompletion.cpp
+++ b/Engine/Source/Runtime/Interchange/Engine/Private/Tasks/InterchangeTaskCompletion.cpp
@@ -76,7 +76,7 @@ void UE::Interchange::FTaskPreCompletion::DoTask(ENamedThreads::Type CurrentThre
 					ObjectInfo.Factory->PreImportPreCompletedCallback(Arguments);
 				}
 
-				if (ImportedObject == nullptr)
+				if (ImportedObject == nullptr || !IsValid(ImportedObject))
 				{
 					continue;
 				}
