Prevent deleting the saved folder when verifying if interchange worker is build with the correct version.

#rnx

diff --git a/Engine/Plugins/Interchange/Runtime/Source/Dispatcher/Private/InterchangeDispatcher.cpp b/Engine/Plugins/Interchange/Runtime/Source/Dispatcher/Private/InterchangeDispatcher.cpp
--- a/Engine/Plugins/Interchange/Runtime/Source/Dispatcher/Private/InterchangeDispatcher.cpp
+++ b/Engine/Plugins/Interchange/Runtime/Source/Dispatcher/Private/InterchangeDispatcher.cpp
@@ -210,8 +210,9 @@ namespace UE
 		//this is a static function
 		bool FInterchangeDispatcher::IsInterchangeWorkerAvailable()
 		{
-			const FString ResultFolder = FPaths::ConvertRelativePathToFull(FPaths::ProjectSavedDir());
-			TUniquePtr<UE::Interchange::FInterchangeDispatcher> Dispatcher = MakeUnique<UE::Interchange::FInterchangeDispatcher>(ResultFolder);
+			FString RandomGuid = FGuid::NewGuid().ToString(EGuidFormats::Base36Encoded);
+			const FString TmpResultFolder = FPaths::ConvertRelativePathToFull(FPaths::ProjectSavedDir() + TEXT("Interchange/Temp/") + RandomGuid);
+			TUniquePtr<UE::Interchange::FInterchangeDispatcher> Dispatcher = MakeUnique<UE::Interchange::FInterchangeDispatcher>(TmpResultFolder);
 			bool bWorkerValid = false;
 			if (ensure(Dispatcher.IsValid()))
 			{
