Interchange
- Ported the new DDS import code from UTextureFactory to Interchange.
- Added the option to only parse the header in FDDSFile.
- Added support for texture cube arrays and volume textures.
- Dump warnings and errors in the log at the end of the import.

diff --git a/Engine/Plugins/Interchange/Runtime/Source/FactoryNodes/Public/InterchangeTextureCubeArrayFactoryNode.h b/Engine/Plugins/Interchange/Runtime/Source/FactoryNodes/Public/InterchangeTextureCubeArrayFactoryNode.h
new file mode 100644
--- /dev/null
+++ b/Engine/Plugins/Interchange/Runtime/Source/FactoryNodes/Public/InterchangeTextureCubeArrayFactoryNode.h
@@ -0,0 +1,25 @@
+// Copyright Epic Games, Inc. All Rights Reserved.
+
+#pragma once
+
+#include "CoreMinimal.h"
+#include "InterchangeTextureFactoryNode.h"
+
+#if WITH_ENGINE
+#include "Engine/TextureCubeArray.h"
+#endif
+
+
+#include "InterchangeTextureCubeArrayFactoryNode.generated.h"
+
+UCLASS(BlueprintType, Experimental)
+class INTERCHANGEFACTORYNODES_API UInterchangeTextureCubeArrayFactoryNode : public UInterchangeTextureFactoryNode
+{
+	GENERATED_BODY()
+
+private:
+	virtual UClass* GetObjectClass() const override
+	{
+		return UTextureCubeArray::StaticClass();
+	}
+};
\ No newline at end of file
diff --git a/Engine/Plugins/Interchange/Runtime/Source/FactoryNodes/Public/InterchangeVolumeTextureFactoryNode.h b/Engine/Plugins/Interchange/Runtime/Source/FactoryNodes/Public/InterchangeVolumeTextureFactoryNode.h
new file mode 100644
--- /dev/null
+++ b/Engine/Plugins/Interchange/Runtime/Source/FactoryNodes/Public/InterchangeVolumeTextureFactoryNode.h
@@ -0,0 +1,25 @@
+// Copyright Epic Games, Inc. All Rights Reserved.
+
+#pragma once
+
+#include "CoreMinimal.h"
+#include "InterchangeTextureFactoryNode.h"
+
+#if WITH_ENGINE
+#include "Engine/VolumeTexture.h"
+#endif
+
+
+#include "InterchangeVolumeTextureFactoryNode.generated.h"
+
+UCLASS(BlueprintType, Experimental)
+class INTERCHANGEFACTORYNODES_API UInterchangeVolumeTextureFactoryNode : public UInterchangeTextureFactoryNode
+{
+	GENERATED_BODY()
+
+private:
+	virtual UClass* GetObjectClass() const override
+	{
+		return UVolumeTexture::StaticClass();
+	}
+};
\ No newline at end of file
diff --git a/Engine/Plugins/Interchange/Runtime/Source/Import/Private/Texture/InterchangeDDSTranslator.cpp b/Engine/Plugins/Interchange/Runtime/Source/Import/Private/Texture/InterchangeDDSTranslator.cpp
--- a/Engine/Plugins/Interchange/Runtime/Source/Import/Private/Texture/InterchangeDDSTranslator.cpp
+++ b/Engine/Plugins/Interchange/Runtime/Source/Import/Private/Texture/InterchangeDDSTranslator.cpp
@@ -2,12 +2,14 @@
 
 #include "Texture/InterchangeDDSTranslator.h"
 
+#include "DDSFile.h"
 #include "DDSLoader.h"
 #include "Engine/Texture.h"
 #include "Engine/Texture2D.h"
 #include "HAL/FileManager.h"
 #include "IImageWrapper.h"
 #include "IImageWrapperModule.h"
+#include "ImageCoreUtils.h"
 #include "InterchangeImportLog.h"
 #include "InterchangeTextureCubeNode.h"
 #include "InterchangeTextureNode.h"
@@ -27,7 +29,7 @@ static FAutoConsoleVariableRef CCvarInterchangeEnableDDSImport(
 	ECVF_Default);
 
 namespace UE::Interchange::Private::InterchangeDDSTranslator
-{ 
+{
 	bool LoadDDSHeaderFromFile(TArray64<uint8>& OutHeader, const FString& Filename)
 	{
 		// The file is close when the archive is destroyed
@@ -35,14 +37,15 @@ namespace UE::Interchange::Private::InterchangeDDSTranslator
 
 		if (FileReaderArchive)
 		{
+			const int64 SizeOfDDSMarker = 4;
 			const int64 SizeOfFile = FileReaderArchive->TotalSize();
-			const int64 MinimalHeaderSize = FDDSLoadHelper::GetDDSHeaderMinimalSize();
+			const int64 MinimalHeaderSize = FDDSLoadHelper::GetDDSHeaderMinimalSize() + SizeOfDDSMarker;
 
-			// If the file is not bigger then the smallest header possible then clearly the file is not valid as a dds file.
+			// If the file is not bigger than the smallest header possible then clearly the file is not valid as a dds file.
 			if (SizeOfFile > MinimalHeaderSize)
 			{
-				const int64 MaximalHeaderSize = FDDSLoadHelper::GetDDSHeaderMaximalSize();
-				const int64 BytesToRead = SizeOfFile >= MaximalHeaderSize ? MaximalHeaderSize : MinimalHeaderSize;
+				const int64 MaximalHeaderSize = FDDSLoadHelper::GetDDSHeaderMaximalSize() + SizeOfDDSMarker;
+				const int64 BytesToRead = FMath::Min(SizeOfFile, MaximalHeaderSize);
 
 				OutHeader.Reset(BytesToRead);
 				OutHeader.AddUninitialized(BytesToRead);
@@ -54,6 +57,137 @@ namespace UE::Interchange::Private::InterchangeDDSTranslator
 
 		return false;
 	}
+
+	TUniquePtr<UE::DDS::FDDSFile> CreateDDSFileReader(const FString& Filename, bool bHeaderOnly)
+	{
+		TArray64<uint8> DDSSourceData;
+
+		if (!FPaths::FileExists(Filename))
+		{
+			UE_LOG(LogInterchangeImport, Error, TEXT("Failed to import DDS, cannot open file. [%s]"), *Filename);
+			return {};
+		}
+
+		if (bHeaderOnly)
+		{
+			LoadDDSHeaderFromFile(DDSSourceData, Filename);
+		}
+		else if (!FFileHelper::LoadFileToArray(DDSSourceData, *Filename))
+		{
+			UE_LOG(LogInterchangeImport, Error, TEXT("Failed to import DDS, cannot load file content into an array. [%s]"), *Filename);
+			return {};
+		}
+
+		UE::DDS::EDDSError Error;
+		TUniquePtr<UE::DDS::FDDSFile> DDS = TUniquePtr<UE::DDS::FDDSFile>(UE::DDS::FDDSFile::CreateFromDDSInMemory(DDSSourceData.GetData(), DDSSourceData.Num(), &Error, bHeaderOnly));
+		if (!DDS.IsValid())
+		{
+			// NotADds is okay/expected , we try this on all image buffers
+			// IoError means buffer is too small to be a DDS
+			check(Error != UE::DDS::EDDSError::OK);
+			if (Error != UE::DDS::EDDSError::NotADds && Error != UE::DDS::EDDSError::IoError)
+			{
+				UE_LOG(LogInterchangeImport, Warning, TEXT("Failed to load DDS (Error=%d)"), (int)Error);
+			}
+			return {};
+		}
+
+		if (!bHeaderOnly)
+		{
+			// change X8 formats to A8 :	
+			DDS->ConvertRGBXtoRGBA();
+			// change RGBA8 to BGRA8 before DXGIFormatGetClosestRawFormat :
+			DDS->ConvertChannelOrder(UE::DDS::EChannelOrder::BGRA);
+		}
+
+		return DDS;
+	}
+
+	uint32 GetDDSMipCount(const TUniquePtr<UE::DDS::FDDSFile>& DDS)
+	{
+		uint32 MipCount = (uint32)DDS->MipCount;
+
+		if (MipCount > MAX_TEXTURE_MIP_COUNT)
+		{
+			// resolution can be above MAX_TEXTURE_MIP_COUNT for VT
+			UE_LOG(LogInterchangeImport, Warning, TEXT("DDS exceeds MAX_TEXTURE_MIP_COUNT"));
+			MipCount = MAX_TEXTURE_MIP_COUNT;
+		}
+
+		return MipCount;
+	}
+
+	bool PerformPrePayloadValidations(const UInterchangeDDSTranslator& DDSTranslator, const UInterchangeSourceData* PayloadSourceData, const FString& PayLoadKey)
+	{
+		check(PayloadSourceData == DDSTranslator.GetSourceData());
+
+		if (!PayloadSourceData)
+		{
+			UE_LOG(LogInterchangeImport, Error, TEXT("Failed to import DDS, bad source data."));
+			return false;
+		}
+
+		FString Filename = PayloadSourceData->GetFilename();
+
+		//Make sure the key fit the filename, The key should always be valid
+		if (!Filename.Equals(PayLoadKey))
+		{
+			UE_LOG(LogInterchangeImport, Error, TEXT("Failed to import DDS, wrong payload key. [%s]"), *Filename);
+			return false;
+		}
+
+		return true;
+	}
+
+	bool DDSPayloadIsSRGB(const TUniquePtr<UE::DDS::FDDSFile>& DDS, ERawImageFormat::Type RawFormat)
+	{
+		bool bSRGB = false;
+
+		if (ERawImageFormat::GetFormatNeedsGammaSpace(RawFormat))
+		{
+			if (DDS->CreateFlags & UE::DDS::FDDSFile::CREATE_FLAG_WAS_D3D9)
+			{
+				// no SRGB info in Dx9 format
+				// assume SRGB yes
+				bSRGB = true;
+			}
+			else if (UE::DDS::DXGIFormatHasLinearAndSRGBForm(DDS->DXGIFormat))
+			{
+				// Dx10 file with format that has linear/srgb pair
+				//	( _UNORM when _UNORM_SRGB)
+				bSRGB = UE::DDS::DXGIFormatIsSRGB(DDS->DXGIFormat);
+			}
+			else
+			{
+				// Dx10 format that doesn't have linear/srgb pairs
+
+				// R8G8_UNORM and R8_UNORM have no _SRGB pair
+				// so no way to clearly indicate SRGB or Linear for them
+				// assume SRGB yes
+				bSRGB = true;
+			}
+		}
+
+		return bSRGB;
+	}
+
+	template < typename PayloadType >
+	void SetupPayload(const TUniquePtr<UE::DDS::FDDSFile>& DDS, PayloadType& PayloadData, ERawImageFormat::Type RawFormat, const int32 MipCount)
+	{
+		PayloadData.bSRGB = DDSPayloadIsSRGB(DDS, RawFormat);
+
+		if (ERawImageFormat::IsHDR(RawFormat))
+		{
+			PayloadData.CompressionSettings = TC_HDR;
+			ensure(PayloadData.bSRGB == false);
+		}
+
+		if (MipCount > 1)
+		{
+			// if the source has mips we keep the mips by default, unless the user changes that
+			PayloadData.MipGenSettings = TMGS_LeaveExistingMips;
+		}
+	}
 }
 
 TArray<FString> UInterchangeDDSTranslator::GetSupportedFormats() const
@@ -71,15 +205,13 @@ TArray<FString> UInterchangeDDSTranslator::GetSupportedFormats() const
 
 bool UInterchangeDDSTranslator::CanImportSourceData(const UInterchangeSourceData* InSourceData) const
 {
+	using namespace UE::Interchange::Private::InterchangeDDSTranslator;
+
 	if (!UInterchangeTranslatorBase::CanImportSourceData(InSourceData))
 	{
 		return false;
 	}
 
-	/*
-	 * DDS file can also be a texture array so we have to open the file and see if its a valid 2D texture.
-	 */
-	TArray64<uint8> HeaderDataBuffer;
 	FString Filename = InSourceData->GetFilename();
 
 	if (!FPaths::FileExists(Filename))
@@ -87,17 +219,21 @@ bool UInterchangeDDSTranslator::CanImportSourceData(const UInterchangeSourceData
 		return false;
 	}
 
-	if (!UE::Interchange::Private::InterchangeDDSTranslator::LoadDDSHeaderFromFile(HeaderDataBuffer, Filename))
+	constexpr bool bHeaderOnly = true;
+	TUniquePtr<UE::DDS::FDDSFile> DDS = CreateDDSFileReader(Filename, bHeaderOnly);
+
+	if (!DDS)
 	{
 		return false;
 	}
 
-	FDDSLoadHelper  DDSLoadHelper(HeaderDataBuffer.GetData(), HeaderDataBuffer.Num());
-	return DDSLoadHelper.IsValid2DTexture() || DDSLoadHelper.IsValidCubemapTexture() || DDSLoadHelper.IsValidArrayTexture();
+	return DDS->IsValidTexture2D() || DDS->IsValidTextureCube() || DDS->IsValidTextureArray() || DDS->IsValidTextureVolume();
 }
 
 bool UInterchangeDDSTranslator::Translate(UInterchangeBaseNodeContainer& BaseNodeContainer) const
 {
+	using namespace UE::Interchange::Private::InterchangeDDSTranslator;
+
 	/*
 	 * DDS file can also be a cube map so we have to open the file and see if its a valid 2D texture.
 	 */
@@ -107,218 +243,234 @@ bool UInterchangeDDSTranslator::Translate(UInterchangeBaseNodeContainer& BaseNod
 		return false;
 	}
 
-	TArray64<uint8> HeaderDataBuffer;
-	if (!UE::Interchange::Private::InterchangeDDSTranslator::LoadDDSHeaderFromFile(HeaderDataBuffer, Filename))
+	constexpr bool bHeaderOnly = true;
+	TUniquePtr<UE::DDS::FDDSFile> DDS = CreateDDSFileReader(Filename, bHeaderOnly);
+
+	if (!DDS)
 	{
 		return false;
 	}
 
-
-	FDDSLoadHelper DDSLoadHelper(HeaderDataBuffer.GetData(), HeaderDataBuffer.Num());
-	if (DDSLoadHelper.IsValid2DTexture())
+	if (DDS->IsValidTexture2D())
 	{
 		return UE::Interchange::FTextureTranslatorUtilities::Generic2DTextureTranslate(GetSourceData(), BaseNodeContainer);
 	}
 
-	if (DDSLoadHelper.IsValidCubemapTexture())
+	if (DDS->IsValidTextureCube())
 	{
-		return UE::Interchange::FTextureTranslatorUtilities::GenericTextureCubeTranslate(GetSourceData(), BaseNodeContainer);
+		if (DDS->ArraySize > 6)
+		{
+			return UE::Interchange::FTextureTranslatorUtilities::GenericTextureCubeArrayTranslate(GetSourceData(), BaseNodeContainer);
+		}
+		else
+		{
+			return UE::Interchange::FTextureTranslatorUtilities::GenericTextureCubeTranslate(GetSourceData(), BaseNodeContainer);
+		}
 	}
 
-	if (DDSLoadHelper.IsValidArrayTexture())
+	if (DDS->IsValidTextureArray())
 	{
 		return UE::Interchange::FTextureTranslatorUtilities::GenericTexture2DArrayTranslate(GetSourceData(), BaseNodeContainer);
 	}
 
+	if (DDS->IsValidTextureVolume())
+	{
+		return UE::Interchange::FTextureTranslatorUtilities::GenericVolumeTextureTranslate(GetSourceData(), BaseNodeContainer);
+	}
+
 	return false;
 }
 
 TOptional<UE::Interchange::FImportImage> UInterchangeDDSTranslator::GetTexturePayloadData(const UInterchangeSourceData* PayloadSourceData, const FString& PayLoadKey) const
 {
-	check(PayloadSourceData == GetSourceData());
+	using namespace UE::Interchange::Private::InterchangeDDSTranslator;
 
-	if (!GetSourceData())
-	{
-		UE_LOG(LogInterchangeImport, Error, TEXT("Failed to import DDS, bad source data."));
-		return TOptional<UE::Interchange::FImportImage>();
-	}
-
-	FString Filename = GetSourceData()->GetFilename();
-	
-	//Make sure the key fit the filename, The key should always be valid
-	if (!Filename.Equals(PayLoadKey))
+	if (!PerformPrePayloadValidations(*this, PayloadSourceData, PayLoadKey))
 	{
-		UE_LOG(LogInterchangeImport, Error, TEXT("Failed to import DDS, wrong payload key. [%s]"), *Filename);
-		return TOptional<UE::Interchange::FImportImage>();
+		return {};
 	}
 
-	TArray64<uint8> DDSSourceData;
+	const FString Filename = GetSourceData()->GetFilename();
+	constexpr bool bHeaderOnly = false;
+	TUniquePtr<UE::DDS::FDDSFile> DDS = CreateDDSFileReader(Filename, bHeaderOnly);
 
-	if (!FPaths::FileExists(Filename))
+	if (!DDS || !DDS->IsValidTexture2D())
 	{
-		UE_LOG(LogInterchangeImport, Error, TEXT("Failed to import DDS, cannot open file. [%s]"), *Filename);
+		UE_LOG(LogInterchangeImport, Error, TEXT("Failed to import DDS, unsupported format. [%s]"), *Filename);
 		return TOptional<UE::Interchange::FImportImage>();
 	}
 
-	if (!FFileHelper::LoadFileToArray(DDSSourceData, *Filename))
+	// map format to RawFormat and ETextureSourceFormat
+	ERawImageFormat::Type RawFormat = UE::DDS::DXGIFormatGetClosestRawFormat(DDS->DXGIFormat);
+	if (RawFormat == ERawImageFormat::Invalid)
 	{
-		UE_LOG(LogInterchangeImport, Error, TEXT("Failed to import DDS, cannot load file content into an array. [%s]"), *Filename);
-		return TOptional<UE::Interchange::FImportImage>();
+		UInterchangeResultError_Generic* Error = AddMessage<UInterchangeResultError_Generic>();
+		Error->AssetType = UTexture2D::StaticClass();
+		Error->Text = FText::Format(NSLOCTEXT("InterchangeDDSTranslator", "InvalidDXGIFormat", "DDS DXGIFormat not supported : {0} : {1}"), (int)DDS->DXGIFormat, FText::FromString(UE::DDS::DXGIFormatGetName(DDS->DXGIFormat)));
+		return {};
 	}
 
+	ETextureSourceFormat TSFormat = FImageCoreUtils::ConvertToTextureSourceFormat(RawFormat);
 
-	//
-	// DDS Texture
-	//
-	FDDSLoadHelper  DDSLoadHelper(DDSSourceData.GetData(), DDSSourceData.Num());
-	if (!DDSLoadHelper.IsValid2DTexture())
-	{
-		UE_LOG(LogInterchangeImport, Error, TEXT("Failed to import DDS, unsupported format. [%s]"), *Filename);
-		return TOptional<UE::Interchange::FImportImage>();
-	}
+	const uint32 MipCount = GetDDSMipCount(DDS);
+
+	UE::Interchange::FImportImage PayloadData;
+	PayloadData.Init2DWithParams(
+		DDS->Width,
+		DDS->Height,
+		MipCount,
+		TSFormat,
+		DDSPayloadIsSRGB(DDS, RawFormat)
+	);
 
-	ETextureSourceFormat SourceFormat = DDSLoadHelper.ComputeSourceFormat();
+	int64 MipOffset = 0;
 
-	// Invalid DDS format
-	if (SourceFormat == TSF_Invalid)
+	for (uint32 MipIndex = 0; MipIndex < MipCount; ++MipIndex)
 	{
-		UE_LOG(LogInterchangeImport, Error, TEXT("DDS file [%s] contains data in an unsupported format"), *Filename);
-		return TOptional<UE::Interchange::FImportImage>();
-	}
+		const uint32 MipWidth = FMath::Max(DDS->Width >> MipIndex, 1u);
+		const uint32 MipHeight = FMath::Max(DDS->Height >> MipIndex, 1u);
 
-	uint32 MipMapCount = DDSLoadHelper.ComputeMipMapCount();
-	if (MipMapCount <= 0)
-	{
-		UE_LOG(LogInterchangeImport, Error, TEXT("DDS file [%s] do not have any mipmap"), *Filename);
-		return TOptional<UE::Interchange::FImportImage>();
-	}
+		FImageView DestMipData((void*)(reinterpret_cast<const uint8*>(PayloadData.RawData.GetData()) + MipOffset),
+			MipWidth, MipHeight, RawFormat);
 
-	UE::Interchange::FImportImage PayloadData;
-	PayloadData.Init2DWithMips(
-		DDSLoadHelper.DDSHeader->dwWidth,
-		DDSLoadHelper.DDSHeader->dwHeight,
-		MipMapCount,
-		SourceFormat,
-		DDSLoadHelper.GetDDSDataPointer()
-	);
+		if (!DDS->GetMipImage(DestMipData, MipIndex))
+		{
+			UE_LOG(LogInterchangeImport, Error, TEXT("DDS could not convert pixel format"));
+			return {};
+		}
 
-	if (MipMapCount > 1)
-	{
-		PayloadData.MipGenSettings = TMGS_LeaveExistingMips;
-	}
-	if (FTextureSource::IsHDR(SourceFormat))
-	{
-		// the loader can suggest a compression setting
-		PayloadData.CompressionSettings = TC_HDR;
+		MipOffset += PayloadData.GetMipSize(MipIndex);
 	}
 
+	SetupPayload(DDS, PayloadData, RawFormat, MipCount);
+
 	return PayloadData;
 }
 
 TOptional<UE::Interchange::FImportSlicedImage> UInterchangeDDSTranslator::GetSlicedTexturePayloadData(const UInterchangeSourceData* PayloadSourceData, const FString& PayLoadKey) const
 {
-	check(PayloadSourceData == GetSourceData());
+	using namespace UE::Interchange::Private::InterchangeDDSTranslator;
 
-	if (!GetSourceData())
+	if (!PerformPrePayloadValidations(*this, PayloadSourceData, PayLoadKey))
 	{
-		UE_LOG(LogInterchangeImport, Error, TEXT("Failed to import DDS, bad source data."));
-		return TOptional<UE::Interchange::FImportSlicedImage>();
+		return {};
 	}
 
-	FString Filename = GetSourceData()->GetFilename();
-	
-	//Make sure the key fit the filename, The key should always be valid
-	if (!Filename.Equals(PayLoadKey))
+	const FString Filename = GetSourceData()->GetFilename();
+	constexpr bool bHeaderOnly = false;
+	TUniquePtr<UE::DDS::FDDSFile> DDS = CreateDDSFileReader(Filename, bHeaderOnly);
+
+	if (!DDS)
 	{
-		UE_LOG(LogInterchangeImport, Error, TEXT("Failed to import DDS, wrong payload key. [%s]"), *Filename);
-		return TOptional<UE::Interchange::FImportSlicedImage>();
+		UE_LOG(LogInterchangeImport, Error, TEXT("Failed to import DDS, unsupported format. [%s]"), *Filename);
+		return {};
 	}
 
-	TArray64<uint8> DDSSourceData;
-	if (!FPaths::FileExists(Filename))
+	uint32 NumSlices = 1;
+
+	if (DDS->IsValidTextureCube() || DDS->IsValidTextureArray())
 	{
-		UE_LOG(LogInterchangeImport, Error, TEXT("Failed to import DDS, cannot open file. [%s]"), *Filename);
-		return TOptional<UE::Interchange::FImportSlicedImage>();
+		NumSlices = DDS->ArraySize;
 	}
-
-	if (!FFileHelper::LoadFileToArray(DDSSourceData, *Filename))
+	else if (DDS->IsValidTextureVolume())
 	{
-		UE_LOG(LogInterchangeImport, Error, TEXT("Failed to import DDS, cannot load file content into an array. [%s]"), *Filename);
-		return TOptional<UE::Interchange::FImportSlicedImage>();
+		NumSlices = DDS->Depth;
 	}
-
-	//
-	// DDS Texture
-	//
-	FDDSLoadHelper DDSLoadHelper(DDSSourceData.GetData(), DDSSourceData.Num());
-
-	if (!DDSLoadHelper.IsValidCubemapTexture() && !DDSLoadHelper.IsValidArrayTexture() )
+	else
 	{
 		UE_LOG(LogInterchangeImport, Error, TEXT("Failed to import DDS, unsupported format. [%s]"), *Filename);
 		return TOptional<UE::Interchange::FImportSlicedImage>();
 	}
 
-	ETextureSourceFormat SourceFormat = DDSLoadHelper.ComputeSourceFormat();
-
-	// Invalid DDS format
-	if (SourceFormat == TSF_Invalid)
+	// map format to RawFormat and ETextureSourceFormat
+	ERawImageFormat::Type RawFormat = UE::DDS::DXGIFormatGetClosestRawFormat(DDS->DXGIFormat);
+	if (RawFormat == ERawImageFormat::Invalid)
 	{
-		UE_LOG(LogInterchangeImport, Error, TEXT("DDS file [%s] contains data in an unsupported format"), *Filename);
-		return TOptional<UE::Interchange::FImportSlicedImage>();
+		UInterchangeResultError_Generic* Error = AddMessage<UInterchangeResultError_Generic>();
+		Error->AssetType = UTexture2D::StaticClass();
+		Error->Text = FText::Format(NSLOCTEXT("InterchangeDDSTranslator", "InvalidDXGIFormat", "DDS DXGIFormat not supported : {0} : {1}"), (int)DDS->DXGIFormat, FText::FromString(UE::DDS::DXGIFormatGetName(DDS->DXGIFormat)));
+		return {};
 	}
 
-	uint32 MipMapCount = DDSLoadHelper.ComputeMipMapCount();
-	if (MipMapCount <= 0)
-	{
-		UE_LOG(LogInterchangeImport, Error, TEXT("DDS file [%s] do not have any mipmap"), *Filename);
-		return TOptional<UE::Interchange::FImportSlicedImage>();
-	}
+	ETextureSourceFormat TSFormat = FImageCoreUtils::ConvertToTextureSourceFormat(RawFormat);
 
-	UE::Interchange::FImportSlicedImage PayloadData;
-	PayloadData.Init(
-		DDSLoadHelper.GetSizeX(),
-		DDSLoadHelper.GetSizeY(),
-		DDSLoadHelper.GetSliceCount(),
-		MipMapCount,
-		SourceFormat,
-		!FTextureSource::IsHDR(SourceFormat)
-	);
+	const uint32 MipCount = GetDDSMipCount(DDS);
+	const bool bIsVolume = (DDS->Dimension == 3);
 
-	uint8* DestMipData[MAX_TEXTURE_MIP_COUNT] = { 0 };
-	int64 MipsSize[MAX_TEXTURE_MIP_COUNT] = { 0 };
+	UE::Interchange::FImportSlicedImage PayloadData;
 
-	for (uint32 MipIndex = 0; MipIndex < MipMapCount; ++MipIndex)
+	if (bIsVolume)
 	{
-		int64 MipSize = PayloadData.GetMipSize(MipIndex);
-		MipsSize[MipIndex] = MipSize;
+		PayloadData.InitVolume(
+			DDS->Width,
+			DDS->Height,
+			NumSlices,
+			MipCount,
+			TSFormat,
+			DDSPayloadIsSRGB(DDS, RawFormat)
+		);
 	}
-
-	for (uint32 MipIndex = 0; MipIndex < MipMapCount; ++MipIndex)
+	else
 	{
-		DestMipData[MipIndex] = PayloadData.GetMipData(MipIndex);
+		PayloadData.Init(
+			DDS->Width,
+			DDS->Height,
+			NumSlices,
+			MipCount,
+			TSFormat,
+			DDSPayloadIsSRGB(DDS, RawFormat)
+		);
 	}
 
-	for (uint32 SliceIndex = 0; SliceIndex < DDSLoadHelper.GetSliceCount(); ++SliceIndex)
+	for (uint32 MipIndex = 0; MipIndex < MipCount; ++MipIndex)
 	{
-		const uint8* SrcMipData = DDSLoadHelper.GetDDSDataPointer((ECubeFace)SliceIndex);
-		for (uint32 MipIndex = 0; MipIndex < MipMapCount; ++MipIndex)
+		const uint32 MipWidth = FMath::Max(DDS->Width >> MipIndex, 1u);
+		const uint32 MipHeight = FMath::Max(DDS->Height >> MipIndex, 1u);
+		const uint32 MipNumSlices = bIsVolume ? FMath::Max(NumSlices >> MipIndex, 1u) : NumSlices;
+
+		FImageView DestMipData(PayloadData.GetMipData(MipIndex),
+			MipWidth,
+			MipHeight,
+			MipNumSlices,
+			RawFormat,
+			PayloadData.bSRGB ? EGammaSpace::sRGB : EGammaSpace::Linear
+		);
+
+		if (bIsVolume)
 		{
-			FMemory::Memcpy(DestMipData[MipIndex] + MipsSize[MipIndex] * SliceIndex, SrcMipData, MipsSize[MipIndex]);
-			SrcMipData += MipsSize[MipIndex];
+			check(DDS->Mips.Num() == DDS->MipCount);
+			check(DDS->Mips[MipIndex].Depth == MipNumSlices);
+
+			if (!DDS->GetMipImage(DestMipData, MipIndex))
+			{
+				UE_LOG(LogInterchangeImport, Error, TEXT("DDS could not convert pixel format"));
+				return {};
+			}
 		}
-	}
+		else
+		{
+			// DDS->Mips[] contains both mips and arrays
+			check(DDS->Mips.Num() == DDS->MipCount * MipNumSlices);
 
-	if (MipMapCount > 1)
-	{
-		PayloadData.MipGenSettings = TMGS_LeaveExistingMips;
-	}
+			for (uint32 SliceIndex = 0; SliceIndex < MipNumSlices; ++SliceIndex)
+			{
+				FImageView DestSliceData = DestMipData.GetSlice(SliceIndex);
 
-	if (FTextureSource::IsHDR(SourceFormat))
-	{
-		// the loader can suggest a compression setting
-		PayloadData.CompressionSettings = TC_HDR;
+				// DDS Mips[] array has whole mip chain of each slice, then next slice
+				// we have the opposite (all slices of top mip first, then next mip)
+				int DDSMipIndex = SliceIndex * DDS->MipCount + MipIndex;
+
+				if (!DDS->GetMipImage(DestSliceData, DDSMipIndex))
+				{
+					UE_LOG(LogInterchangeImport, Error, TEXT("DDS could not convert pixel format"));
+					return {};
+				}
+			}
+		}
 	}
 
+	SetupPayload(DDS, PayloadData, RawFormat, MipCount);
+
 	return PayloadData;
 }
 
diff --git a/Engine/Plugins/Interchange/Runtime/Source/Import/Private/Texture/InterchangeSlicedTexturePayloadData.cpp b/Engine/Plugins/Interchange/Runtime/Source/Import/Private/Texture/InterchangeSlicedTexturePayloadData.cpp
--- a/Engine/Plugins/Interchange/Runtime/Source/Import/Private/Texture/InterchangeSlicedTexturePayloadData.cpp
+++ b/Engine/Plugins/Interchange/Runtime/Source/Import/Private/Texture/InterchangeSlicedTexturePayloadData.cpp
@@ -4,7 +4,6 @@
 
 #include "CoreMinimal.h"
 
-
 void UE::Interchange::FImportSlicedImage::Init(int32 InSizeX, int32 InSizeY, int32 InNumSlice, int32 InNumMips, ETextureSourceFormat InFormat, bool InSRGB)
 {
 	NumSlice = InNumSlice;
@@ -12,17 +11,26 @@ void UE::Interchange::FImportSlicedImage::Init(int32 InSizeX, int32 InSizeY, int
 	FImportImage::Init2DWithParams(InSizeX, InSizeY, InNumMips, InFormat, InSRGB);
 }
 
+void UE::Interchange::FImportSlicedImage::InitVolume(int32 InSizeX, int32 InSizeY, int32 InNumSlice, int32 InNumMips, ETextureSourceFormat InFormat, bool InSRGB)
+{
+	bIsVolume = true;
+	Init(InSizeX, InSizeY, InNumSlice, InNumMips, InFormat, InSRGB);
+}
+
 const uint8* UE::Interchange::FImportSlicedImage::GetMipData(int32 InMipIndex, int32 InSliceIndex) const
 {
+	const int32 SliceFactor = bIsVolume ? 1 : NumSlice;
+
 	int64 Offset = 0;
 	int32 MipIndex;
 	for (MipIndex = 0; MipIndex < InMipIndex; ++MipIndex)
 	{
-		Offset += GetMipSize(MipIndex) * NumSlice;
+		Offset += GetMipSize(MipIndex) * SliceFactor;
 	}
 
 	if (InSliceIndex != INDEX_NONE)
 	{ 
+		check(!bIsVolume);
 		Offset += GetMipSize(MipIndex) * InSliceIndex;
 	}
 
@@ -37,9 +45,23 @@ uint8* UE::Interchange::FImportSlicedImage::GetMipData(int32 InMipIndex, int32 I
 	return const_cast<uint8*>(ConstBufferPtr);
 }
 
+int64 UE::Interchange::FImportSlicedImage::GetMipSize(int32 InMipIndex) const
+{
+	const int64 BytesPerPixel = FTextureSource::GetBytesPerPixel(Format);
+
+	int32 MipSizeX = FMath::Max<int32>(SizeX >> InMipIndex, 1);
+	int32 MipSizeY = FMath::Max<int32>(SizeY >> InMipIndex, 1);
+	int32 MipSizeZ = bIsVolume ? FMath::Max(NumSlice >> InMipIndex, 1) : 1;
+
+	const int64 MipSize = MipSizeX * MipSizeY * MipSizeZ * BytesPerPixel;
+
+	return MipSize;
+}
+
 int64 UE::Interchange::FImportSlicedImage::ComputeBufferSize() const
 {
-	return FImportImage::ComputeBufferSize() * NumSlice;
+	const int32 SliceFactor = bIsVolume ? 1 : NumSlice;
+	return FImportImage::ComputeBufferSize() * SliceFactor;
 }
 
 bool UE::Interchange::FImportSlicedImage::IsValid() const
diff --git a/Engine/Plugins/Interchange/Runtime/Source/Import/Private/Texture/InterchangeTextureFactory.cpp b/Engine/Plugins/Interchange/Runtime/Source/Import/Private/Texture/InterchangeTextureFactory.cpp
--- a/Engine/Plugins/Interchange/Runtime/Source/Import/Private/Texture/InterchangeTextureFactory.cpp
+++ b/Engine/Plugins/Interchange/Runtime/Source/Import/Private/Texture/InterchangeTextureFactory.cpp
@@ -24,11 +24,15 @@
 #include "InterchangeTexture2DNode.h"
 #include "InterchangeTextureCubeFactoryNode.h"
 #include "InterchangeTextureCubeNode.h"
+#include "InterchangeTextureCubeArrayFactoryNode.h"
+#include "InterchangeTextureCubeArrayNode.h"
 #include "InterchangeTextureFactoryNode.h"
 #include "InterchangeTextureLightProfileFactoryNode.h"
 #include "InterchangeTextureLightProfileNode.h"
 #include "InterchangeTextureNode.h"
 #include "InterchangeTranslatorBase.h"
+#include "InterchangeVolumeTextureNode.h"
+#include "InterchangeVolumeTextureFactoryNode.h"
 #include "Nodes/InterchangeBaseNode.h"
 #include "ProfilingDebugging/CpuProfilerTrace.h"
 #include "Serialization/EditorBulkData.h"
@@ -129,9 +133,11 @@ namespace UE::Interchange::Private::InterchangeTextureFactory
 	UClass* GetSupportedFactoryNodeClass(const UInterchangeFactoryBaseNode* AssetNode)
 	{
 		UClass* TextureCubeFactoryClass = UInterchangeTextureCubeFactoryNode::StaticClass();
+		UClass* TextureCubeArrayFactoryClass = UInterchangeTextureCubeArrayFactoryNode::StaticClass();
 		UClass* Texture2DFactoryClass = UInterchangeTexture2DFactoryNode::StaticClass();
 		UClass* Texture2DArrayFactoryClass = UInterchangeTexture2DArrayFactoryNode::StaticClass();
 		UClass* TextureLightProfileFactoryClass = UInterchangeTextureLightProfileFactoryNode::StaticClass();
+		UClass* VolumeTextureFactoryClass = UInterchangeVolumeTextureFactoryNode::StaticClass();
 		UClass* AssetClass = AssetNode->GetClass();
 #if USTRUCT_FAST_ISCHILDOF_IMPL == USTRUCT_ISCHILDOF_STRUCTARRAY
 		if (AssetClass->IsChildOf(Texture2DArrayFactoryClass))
@@ -142,10 +148,18 @@ namespace UE::Interchange::Private::InterchangeTextureFactory
 		{
 			return TextureCubeFactoryClass;
 		}
+		else if (AssetClass->IsChildOf(TextureCubeArrayFactoryClass))
+		{
+			return TextureCubeArrayFactoryClass;
+		}
 		else if (AssetClass->IsChildOf(TextureLightProfileFactoryClass))
 		{
 			return TextureLightProfileFactoryClass;
 		}
+		else if (AssetClass->IsChildOf(VolumeTextureFactoryClass))
+		{
+			return VolumeTextureFactoryClass;
+		}
 		else if (AssetClass->IsChildOf(Texture2DFactoryClass))
 		{
 			return Texture2DFactoryClass;
@@ -154,9 +168,11 @@ namespace UE::Interchange::Private::InterchangeTextureFactory
 		while (AssetClass)
 		{
 			if (AssetClass == TextureCubeFactoryClass
+				|| AssetClass == TextureCubeArrayFactoryClass
 				|| AssetClass == Texture2DFactoryClass
 				|| AssetClass == Texture2DArrayFactoryClass
-				|| AssetClass == TextureLightProfileFactoryClass)
+				|| AssetClass == TextureLightProfileFactoryClass
+				|| AssetClass == VolumeTextureFactoryClass)
 			{
 				return AssetClass;
 			}
@@ -171,12 +187,14 @@ namespace UE::Interchange::Private::InterchangeTextureFactory
 	using FTextureFactoryNodeVariant = TVariant<FEmptyVariantState
 		, UInterchangeTexture2DFactoryNode*
 		, UInterchangeTextureCubeFactoryNode*
+		, UInterchangeTextureCubeArrayFactoryNode*
 		, UInterchangeTexture2DArrayFactoryNode*
-		, UInterchangeTextureLightProfileFactoryNode*>;
+		, UInterchangeTextureLightProfileFactoryNode*
+		, UInterchangeVolumeTextureFactoryNode*>;
 
 	FTextureFactoryNodeVariant GetAsTextureFactoryNodeVariant(UInterchangeFactoryBaseNode* AssetNode, UClass* SupportedFactoryNodeClass)
 	{
-		static_assert(TVariantSize<FTextureFactoryNodeVariant>::Value == 5, "Please update the code below and this assert to reflect the change to the variant type.");
+		static_assert(TVariantSize<FTextureFactoryNodeVariant>::Value == 7, "Please update the code below and this assert to reflect the change to the variant type.");
 
 		if (AssetNode)
 		{
@@ -195,6 +213,11 @@ namespace UE::Interchange::Private::InterchangeTextureFactory
 				return FTextureFactoryNodeVariant(TInPlaceType<UInterchangeTextureCubeFactoryNode*>(), static_cast<UInterchangeTextureCubeFactoryNode*>(AssetNode));
 			}
 
+			if (SupportedFactoryNodeClass == UInterchangeTextureCubeArrayFactoryNode::StaticClass())
+			{
+				return FTextureFactoryNodeVariant(TInPlaceType<UInterchangeTextureCubeArrayFactoryNode*>(), static_cast<UInterchangeTextureCubeArrayFactoryNode*>(AssetNode));
+			}
+
 			if (SupportedFactoryNodeClass == UInterchangeTexture2DArrayFactoryNode::StaticClass())
 			{
 				return FTextureFactoryNodeVariant(TInPlaceType<UInterchangeTexture2DArrayFactoryNode*>(), static_cast<UInterchangeTexture2DArrayFactoryNode*>(AssetNode));
@@ -204,6 +227,11 @@ namespace UE::Interchange::Private::InterchangeTextureFactory
 			{
 				return FTextureFactoryNodeVariant(TInPlaceType<UInterchangeTextureLightProfileFactoryNode*>(), static_cast<UInterchangeTextureLightProfileFactoryNode*>(AssetNode));
 			}
+
+			if (SupportedFactoryNodeClass == UInterchangeVolumeTextureFactoryNode::StaticClass())
+			{
+				return FTextureFactoryNodeVariant(TInPlaceType<UInterchangeVolumeTextureFactoryNode*>(), static_cast<UInterchangeVolumeTextureFactoryNode*>(AssetNode));
+			}
 		}
 
 		ensureMsgf(false
@@ -230,7 +258,7 @@ namespace UE::Interchange::Private::InterchangeTextureFactory
 
 	UInterchangeTextureFactoryNode* GetTextureFactoryNodeFromVariant(const FTextureFactoryNodeVariant& FactoryNodeVariant)
 	{
-		static_assert(TVariantSize<FTextureFactoryNodeVariant>::Value == 5, "Please update the code below and this assert to reflect the change to the variant type.");
+		static_assert(TVariantSize<FTextureFactoryNodeVariant>::Value == 7, "Please update the code below and this assert to reflect the change to the variant type.");
 
 		if (UInterchangeTexture2DFactoryNode* const* TextureFactoryNode = FactoryNodeVariant.TryGet<UInterchangeTexture2DFactoryNode*>())
 		{
@@ -242,6 +270,11 @@ namespace UE::Interchange::Private::InterchangeTextureFactory
 			return *TextureCubeFactoryNode;
 		}
 
+		if (UInterchangeTextureCubeArrayFactoryNode* const* TextureCubeArrayFactoryNode = FactoryNodeVariant.TryGet<UInterchangeTextureCubeArrayFactoryNode*>())
+		{
+			return *TextureCubeArrayFactoryNode;
+		}
+
 		if (UInterchangeTexture2DArrayFactoryNode* const* Texture2DArrayFactoryNode = FactoryNodeVariant.TryGet<UInterchangeTexture2DArrayFactoryNode*>())
 		{
 			return *Texture2DArrayFactoryNode;
@@ -252,20 +285,27 @@ namespace UE::Interchange::Private::InterchangeTextureFactory
 			return *TextureLightProfileFactoryNode;
 		}
 
+		if (UInterchangeVolumeTextureFactoryNode* const* VolumeTextureFactoryNode = FactoryNodeVariant.TryGet<UInterchangeVolumeTextureFactoryNode*>())
+		{
+			return *VolumeTextureFactoryNode;
+		}
+
 		return nullptr;
 	}
 
 	using FTextureNodeVariant = TVariant<FEmptyVariantState
 		, const UInterchangeTexture2DNode*
 		, const UInterchangeTextureCubeNode*
+		, const UInterchangeTextureCubeArrayNode*
 		, const UInterchangeTexture2DArrayNode*
-		, const UInterchangeTextureLightProfileNode*>;
+		, const UInterchangeTextureLightProfileNode*
+		, const UInterchangeVolumeTextureNode*>;
 
 	FTextureNodeVariant GetTextureNodeVariantFromFactoryVariant(const FTextureFactoryNodeVariant& FactoryVariant, const UInterchangeBaseNodeContainer* NodeContainer)
 	{
-		static_assert(TVariantSize<FTextureFactoryNodeVariant>::Value == 5, "Please update the code below and this assert to reflect the change to the variant type.");
+		static_assert(TVariantSize<FTextureFactoryNodeVariant>::Value == 7, "Please update the code below and this assert to reflect the change to the variant type.");
 
-		static_assert(TVariantSize<FTextureNodeVariant>::Value == 5, "Please update the code below and this assert to reflect the change to the variant type.");
+		static_assert(TVariantSize<FTextureNodeVariant>::Value == 7, "Please update the code below and this assert to reflect the change to the variant type.");
 
 		FString TextureNodeUniqueID;
 
@@ -277,6 +317,10 @@ namespace UE::Interchange::Private::InterchangeTextureFactory
 		{
 			(*TextureCubeFactoryNode)->GetCustomTranslatedTextureNodeUid(TextureNodeUniqueID);
 		}
+		else if (UInterchangeTextureCubeArrayFactoryNode* const* TextureCubeArrayFactoryNode = FactoryVariant.TryGet<UInterchangeTextureCubeArrayFactoryNode*>())
+		{
+			(*TextureCubeArrayFactoryNode)->GetCustomTranslatedTextureNodeUid(TextureNodeUniqueID);
+		}
 		else if (UInterchangeTexture2DArrayFactoryNode* const* Texture2DArrayFactoryNode = FactoryVariant.TryGet<UInterchangeTexture2DArrayFactoryNode*>())
 		{
 			(*Texture2DArrayFactoryNode)->GetCustomTranslatedTextureNodeUid(TextureNodeUniqueID);
@@ -285,6 +329,10 @@ namespace UE::Interchange::Private::InterchangeTextureFactory
 		{
 			(*TextureLightProfileFactoryNode)->GetCustomTranslatedTextureNodeUid(TextureNodeUniqueID);
 		}
+		else if (UInterchangeVolumeTextureFactoryNode* const* VolumeTextureFactoryNode = FactoryVariant.TryGet<UInterchangeVolumeTextureFactoryNode*>())
+		{
+			(*VolumeTextureFactoryNode)->GetCustomTranslatedTextureNodeUid(TextureNodeUniqueID);
+		}
 
 		if (const UInterchangeBaseNode* TranslatedNode = NodeContainer->GetNode(TextureNodeUniqueID))
 		{
@@ -293,6 +341,11 @@ namespace UE::Interchange::Private::InterchangeTextureFactory
 				return FTextureNodeVariant(TInPlaceType<const UInterchangeTextureCubeNode*>(), TextureCubeTranslatedNode);
 			}
 
+			if (const UInterchangeTextureCubeArrayNode* TextureCubeArrayTranslatedNode = Cast<UInterchangeTextureCubeArrayNode>(TranslatedNode))
+			{
+				return FTextureNodeVariant(TInPlaceType<const UInterchangeTextureCubeArrayNode*>(), TextureCubeArrayTranslatedNode);
+			}
+
 			if (const UInterchangeTexture2DArrayNode* Texture2DArrayTranslatedNode = Cast<UInterchangeTexture2DArrayNode>(TranslatedNode))
 			{
 				return FTextureNodeVariant(TInPlaceType<const UInterchangeTexture2DArrayNode*>(), Texture2DArrayTranslatedNode);
@@ -307,6 +360,11 @@ namespace UE::Interchange::Private::InterchangeTextureFactory
 			{
 				return FTextureNodeVariant(TInPlaceType<const UInterchangeTexture2DNode*>(), TextureTranslatedNode);
 			}
+
+			if (const UInterchangeVolumeTextureNode* TextureTranslatedNode = Cast<UInterchangeVolumeTextureNode>(TranslatedNode))
+			{
+				return FTextureNodeVariant(TInPlaceType<const UInterchangeVolumeTextureNode*>(), TextureTranslatedNode);
+			}
 		}
 
 		return {};
@@ -314,7 +372,7 @@ namespace UE::Interchange::Private::InterchangeTextureFactory
 
 	bool HasPayloadKey(const FTextureNodeVariant& TextureNodeVariant)
 	{
-		static_assert(TVariantSize<FTextureNodeVariant>::Value == 5, "Please update the code below and this assert to reflect the change to the variant type.");
+		static_assert(TVariantSize<FTextureNodeVariant>::Value == 7, "Please update the code below and this assert to reflect the change to the variant type.");
 
 		if (const UInterchangeTexture2DNode* const* TextureNode =  TextureNodeVariant.TryGet<const UInterchangeTexture2DNode*>())
 		{
@@ -326,6 +384,11 @@ namespace UE::Interchange::Private::InterchangeTextureFactory
 			return (*TextureCubeNode)->GetPayLoadKey().IsSet();
 		}
 
+		if (const UInterchangeTextureCubeArrayNode* const* TextureCubeArrayNode = TextureNodeVariant.TryGet<const UInterchangeTextureCubeArrayNode*>())
+		{
+			return (*TextureCubeArrayNode)->GetPayLoadKey().IsSet();
+		}
+
 		if (const UInterchangeTexture2DArrayNode* const* Texture2DArrayNode = TextureNodeVariant.TryGet<const UInterchangeTexture2DArrayNode*>())
 		{
 			return (*Texture2DArrayNode)->GetPayLoadKey().IsSet();
@@ -336,12 +399,17 @@ namespace UE::Interchange::Private::InterchangeTextureFactory
 			return (*TextureLightProfileNode)->GetPayLoadKey().IsSet();
 		}
 
+		if (const UInterchangeVolumeTextureNode* const* VolumeTextureNode = TextureNodeVariant.TryGet<const UInterchangeVolumeTextureNode*>())
+		{
+			return (*VolumeTextureNode)->GetPayLoadKey().IsSet();
+		}
+
 		return false;
 	}
 
 	TOptional<FString> GetPayloadKey(const FTextureNodeVariant& TextureNodeVariant)
 	{
-		static_assert(TVariantSize<FTextureNodeVariant>::Value == 5, "Please update the code below and this assert to reflect the change to the variant type.");
+		static_assert(TVariantSize<FTextureNodeVariant>::Value == 7, "Please update the code below and this assert to reflect the change to the variant type.");
 
 		if (const UInterchangeTexture2DNode* const* TextureNode =  TextureNodeVariant.TryGet<const UInterchangeTexture2DNode*>())
 		{
@@ -353,6 +421,11 @@ namespace UE::Interchange::Private::InterchangeTextureFactory
 			return (*TextureCubeNode)->GetPayLoadKey();
 		}
 
+		if (const UInterchangeTextureCubeArrayNode* const* TextureCubeArrayNode = TextureNodeVariant.TryGet<const UInterchangeTextureCubeArrayNode*>())
+		{
+			return (*TextureCubeArrayNode)->GetPayLoadKey();
+		}
+
 		if (const UInterchangeTexture2DArrayNode* const* Texture2DArrayNode = TextureNodeVariant.TryGet<const UInterchangeTexture2DArrayNode*>())
 		{
 			return (*Texture2DArrayNode)->GetPayLoadKey();
@@ -363,6 +436,11 @@ namespace UE::Interchange::Private::InterchangeTextureFactory
 			return (*TextureLightProfileNode)->GetPayLoadKey();
 		}
 
+		if (const UInterchangeVolumeTextureNode* const* VolumeTextureNode = TextureNodeVariant.TryGet<const UInterchangeVolumeTextureNode*>())
+		{
+			return (*VolumeTextureNode)->GetPayLoadKey();
+		}
+
 		return {};
 	}
 
@@ -457,11 +535,11 @@ namespace UE::Interchange::Private::InterchangeTextureFactory
 
 	FTexturePayloadVariant GetTexturePayload(const UInterchangeSourceData* SourceData, const FString& PayloadKey, const FTextureNodeVariant& TextureNodeVariant, const FTextureFactoryNodeVariant& FactoryNodeVariant, const UInterchangeTranslatorBase* Translator)
 	{
-		static_assert(TVariantSize<FTextureFactoryNodeVariant>::Value == 5, "Please update the code below and this assert to reflect the change to the variant type.");
+		static_assert(TVariantSize<FTextureFactoryNodeVariant>::Value == 7, "Please update the code below and this assert to reflect the change to the variant type.");
 
 		static_assert(TVariantSize<FTexturePayloadVariant>::Value == 5, "Please update the code below and this assert to reflect the change to the variant type.");
 
-		static_assert(TVariantSize<FTextureNodeVariant>::Value == 5, "Please update the code below and this assert to reflect the change to the variant type.");
+		static_assert(TVariantSize<FTextureNodeVariant>::Value == 7, "Please update the code below and this assert to reflect the change to the variant type.");
 
 		// Standard texture 2D payload
 		if (const UInterchangeTexture2DNode* const* TextureNode =  TextureNodeVariant.TryGet<const UInterchangeTexture2DNode*>())
@@ -503,8 +581,9 @@ namespace UE::Interchange::Private::InterchangeTextureFactory
 			}
 		}
 
-		// Cube or array texture payload 
-		if (TextureNodeVariant.IsType<const UInterchangeTextureCubeNode*>() || TextureNodeVariant.IsType<const UInterchangeTexture2DArrayNode*>())
+		// Cube, array, cube array or volume texture payload 
+		if (TextureNodeVariant.IsType<const UInterchangeTextureCubeNode*>() || TextureNodeVariant.IsType<const UInterchangeTexture2DArrayNode*>()
+			|| TextureNodeVariant.IsType<const UInterchangeTextureCubeArrayNode*>() || TextureNodeVariant.IsType<const UInterchangeVolumeTextureNode*>())
 		{
 			if (const IInterchangeSlicedTexturePayloadInterface* SlicedTextureTranslator = Cast<IInterchangeSlicedTexturePayloadInterface>(Translator))
 			{
@@ -678,6 +757,72 @@ namespace UE::Interchange::Private::InterchangeTextureFactory
 		}
 	}
 
+	bool CanSetupTextureCubeArraySourceData(FTexturePayloadVariant& TexturePayload)
+	{
+		static_assert(TVariantSize<FTexturePayloadVariant>::Value == 5, "Please update the code below and this assert to reflect the change to the variant type.");
+
+		if (TOptional<FImportSlicedImage>* SlicedImage = TexturePayload.TryGet<TOptional<FImportSlicedImage>>())
+		{
+			if (SlicedImage->IsSet())
+			{
+				return (*SlicedImage)->IsValid() && (*SlicedImage)->NumSlice > 6;
+			}
+		}
+		else if (TOptional<FImportImage>* Image = TexturePayload.TryGet<TOptional<FImportImage>>())
+		{
+			if (Image->IsSet())
+			{
+				return (*Image)->IsValid();
+			}
+		}
+		else if (TOptional<FImportLightProfile>* LightProfile = TexturePayload.TryGet<TOptional<FImportLightProfile>>())
+		{
+			if (LightProfile->IsSet())
+			{
+				return (*LightProfile)->IsValid();
+			}
+		}
+
+		return false;
+	}
+
+	void SetupTextureCubeArraySourceDataFromPayload(UTextureCubeArray* TextureCubeArray, FProcessedPayload& ProcessedPayload)
+	{
+		if (TOptional<FImportSlicedImage>* OptionalSlicedImage = ProcessedPayload.SettingsFromPayload.TryGet<TOptional<FImportSlicedImage>>())
+		{
+			if (OptionalSlicedImage->IsSet())
+			{
+				FImportSlicedImage& SlicedImage = OptionalSlicedImage->GetValue();
+
+				if (SlicedImage.NumSlice > 6)
+				{
+					SetupTextureSourceDataFromBulkData_Editor(TextureCubeArray, SlicedImage, MoveTemp(ProcessedPayload.PayloadAndId));
+				}
+			}
+		}
+		else if (TOptional<FImportImage>* OptionalImage = ProcessedPayload.SettingsFromPayload.TryGet<TOptional<FImportImage>>())
+		{
+			if (OptionalImage->IsSet())
+			{
+				FImportImage& Image = OptionalImage->GetValue();
+				SetupTextureSourceDataFromBulkData_Editor(TextureCubeArray, Image, MoveTemp(ProcessedPayload.PayloadAndId));
+			}
+		}
+		else if (TOptional<FImportLightProfile>* OptionalLightProfile = ProcessedPayload.SettingsFromPayload.TryGet<TOptional<FImportLightProfile>>())
+		{
+			if (OptionalLightProfile->IsSet())
+			{
+				FImportLightProfile& LightProfile = OptionalLightProfile->GetValue();
+				SetupTextureSourceDataFromBulkData_Editor(TextureCubeArray, LightProfile, MoveTemp(ProcessedPayload.PayloadAndId));
+			}
+		}
+		else
+		{
+			// The TexturePayload should be validated before calling this function
+			checkNoEntry();
+		}
+	}
+
 	bool CanSetupTexture2DArraySourceData(FTexturePayloadVariant& TexturePayload)
 	{
 		static_assert(TVariantSize<FTexturePayloadVariant>::Value == 5, "Please update the code below and this assert to reflect the change to the variant type.");
@@ -741,6 +886,69 @@ namespace UE::Interchange::Private::InterchangeTextureFactory
 		}
 	}
 
+	bool CanSetupVolumeTextureSourceData(FTexturePayloadVariant& TexturePayload)
+	{
+		static_assert(TVariantSize<FTexturePayloadVariant>::Value == 5, "Please update the code below and this assert to reflect the change to the variant type.");
+
+		if (TOptional<FImportSlicedImage>* SlicedImage = TexturePayload.TryGet<TOptional<FImportSlicedImage>>())
+		{
+			if (SlicedImage->IsSet())
+			{
+				return (*SlicedImage)->IsValid();
+			}
+		}
+		else if (TOptional<FImportImage>* Image = TexturePayload.TryGet<TOptional<FImportImage>>())
+		{
+			if (Image->IsSet())
+			{
+				return (*Image)->IsValid();
+			}
+		}
+		else if (TOptional<FImportLightProfile>* LightProfile = TexturePayload.TryGet<TOptional<FImportLightProfile>>())
+		{
+			if (LightProfile->IsSet())
+			{
+				return (*LightProfile)->IsValid();
+			}
+		}
+
+		return false;
+	}
+
+	void SetupVolumeTextureSourceDataFromPayload(UVolumeTexture* VolumeTexture, FProcessedPayload& ProcessedPayload)
+	{
+		if (TOptional<FImportSlicedImage>* OptionalSlicedImage = ProcessedPayload.SettingsFromPayload.TryGet<TOptional<FImportSlicedImage>>())
+		{
+			if (OptionalSlicedImage->IsSet())
+			{
+				FImportSlicedImage& SlicedImage = OptionalSlicedImage->GetValue();
+
+				SetupTextureSourceDataFromBulkData_Editor(VolumeTexture, SlicedImage, MoveTemp(ProcessedPayload.PayloadAndId));
+			}
+		}
+		else if (TOptional<FImportImage>* OptionalImage = ProcessedPayload.SettingsFromPayload.TryGet<TOptional<FImportImage>>())
+		{
+			if (OptionalImage->IsSet())
+			{
+				FImportImage& Image = OptionalImage->GetValue();
+				SetupTextureSourceDataFromBulkData_Editor(VolumeTexture, Image, MoveTemp(ProcessedPayload.PayloadAndId));
+			}
+		}
+		else if (TOptional<FImportLightProfile>* OptionalLightProfile = ProcessedPayload.SettingsFromPayload.TryGet<TOptional<FImportLightProfile>>())
+		{
+			if (OptionalLightProfile->IsSet())
+			{
+				FImportLightProfile& LightProfile = OptionalLightProfile->GetValue();
+				SetupTextureSourceDataFromBulkData_Editor(VolumeTexture, LightProfile, MoveTemp(ProcessedPayload.PayloadAndId));
+			}
+		}
+		else
+		{
+			// The TexturePayload should be validated before calling this function
+			checkNoEntry();
+		}
+	}
+
 	FGraphEventArray GenerateHashSourceFilesTasks(const UInterchangeSourceData* SourceData, TArray<FString>&& FilesToHash, TArray<FAssetImportInfo::FSourceFile>& OutSourceFiles)
 	{
 		struct FHashSourceTaskBase
@@ -1121,10 +1329,18 @@ namespace UE::Interchange::Private::InterchangeTextureFactory
 		{
 			SetupTextureCubeSourceDataFromPayload(TextureCube, ProcessedPayload);
 		}
+		else if (UTextureCubeArray* TextureCubeArray = Cast<UTextureCubeArray>(Texture))
+		{
+			SetupTextureCubeArraySourceDataFromPayload(TextureCubeArray, ProcessedPayload);
+		}
 		else if (UTexture2DArray* Texture2DArray = Cast<UTexture2DArray>(Texture))
 		{
 			SetupTexture2DArraySourceDataFromPayload(Texture2DArray, ProcessedPayload);
 		}
+		else if (UVolumeTexture* VolumeTexture = Cast<UVolumeTexture>(Texture))
+		{
+			SetupVolumeTextureSourceDataFromPayload(VolumeTexture, ProcessedPayload);
+		}
 #endif // WITH_EDITOR
 		else
 		{
@@ -1334,10 +1550,18 @@ UObject* UInterchangeTextureFactory::CreateAsset(const FCreateAssetParams& Argum
 	{
 		bCanSetup = CanSetupTextureCubeSourceData(TexturePayload);
 	}
+	else if (UTextureCubeArray* TextureCubeArray = Cast<UTextureCubeArray>(Texture))
+	{
+		bCanSetup = CanSetupTextureCubeArraySourceData(TexturePayload);
+	}
 	else if (UTexture2DArray* Texture2DArray = Cast<UTexture2DArray>(Texture))
 	{
 		bCanSetup = CanSetupTexture2DArraySourceData(TexturePayload);
 	}
+	else if (UVolumeTexture* VolumeTexture = Cast<UVolumeTexture>(Texture))
+	{
+		bCanSetup = CanSetupVolumeTextureSourceData(TexturePayload);
+	}
 #endif // WITH_EDITOR
 
 	if (!bCanSetup)
diff --git a/Engine/Plugins/Interchange/Runtime/Source/Import/Private/Texture/TextureTranslatorUtilities.cpp b/Engine/Plugins/Interchange/Runtime/Source/Import/Private/Texture/TextureTranslatorUtilities.cpp
--- a/Engine/Plugins/Interchange/Runtime/Source/Import/Private/Texture/TextureTranslatorUtilities.cpp
+++ b/Engine/Plugins/Interchange/Runtime/Source/Import/Private/Texture/TextureTranslatorUtilities.cpp
@@ -5,9 +5,11 @@
 #include "InterchangeSourceData.h"
 #include "InterchangeTexture2DArrayNode.h"
 #include "InterchangeTexture2DNode.h"
+#include "InterchangeTextureCubeArrayNode.h"
 #include "InterchangeTextureCubeNode.h"
 #include "InterchangeTextureLightProfileNode.h"
 #include "InterchangeTextureNode.h"
+#include "InterchangeVolumeTextureNode.h"
 #include "Misc/Paths.h"
 #include "Nodes/InterchangeBaseNodeContainer.h"
 #include "Templates/SubclassOf.h"
@@ -58,11 +60,21 @@ bool UE::Interchange::FTextureTranslatorUtilities::GenericTextureCubeTranslate(c
 	return Private::GenericTextureTranslate(SourceData, BaseNodeContainer, UInterchangeTextureCubeNode::StaticClass());
 }
 
+bool UE::Interchange::FTextureTranslatorUtilities::GenericTextureCubeArrayTranslate(const UInterchangeSourceData* SourceData, UInterchangeBaseNodeContainer& BaseNodeContainer)
+{
+	return Private::GenericTextureTranslate(SourceData, BaseNodeContainer, UInterchangeTextureCubeArrayNode::StaticClass());
+}
+
 bool UE::Interchange::FTextureTranslatorUtilities::GenericTexture2DArrayTranslate(const UInterchangeSourceData* SourceData, UInterchangeBaseNodeContainer& BaseNodeContainer)
 {
 	return Private::GenericTextureTranslate(SourceData, BaseNodeContainer, UInterchangeTexture2DArrayNode::StaticClass());
 }
 
+bool UE::Interchange::FTextureTranslatorUtilities::GenericVolumeTextureTranslate(const UInterchangeSourceData* SourceData, UInterchangeBaseNodeContainer& BaseNodeContainer)
+{
+	return Private::GenericTextureTranslate(SourceData, BaseNodeContainer, UInterchangeVolumeTextureNode::StaticClass());
+}
+
 bool UE::Interchange::FTextureTranslatorUtilities::GenericTextureLightProfileTranslate(const UInterchangeSourceData* SourceData, UInterchangeBaseNodeContainer& BaseNodeContainer)
 {
 	return Private::GenericTextureTranslate(SourceData, BaseNodeContainer, UInterchangeTextureLightProfileNode::StaticClass());
diff --git a/Engine/Plugins/Interchange/Runtime/Source/Import/Private/Texture/TextureTranslatorUtilities.h b/Engine/Plugins/Interchange/Runtime/Source/Import/Private/Texture/TextureTranslatorUtilities.h
--- a/Engine/Plugins/Interchange/Runtime/Source/Import/Private/Texture/TextureTranslatorUtilities.h
+++ b/Engine/Plugins/Interchange/Runtime/Source/Import/Private/Texture/TextureTranslatorUtilities.h
@@ -23,8 +23,12 @@ namespace UE
 
 			static bool GenericTextureCubeTranslate(const UInterchangeSourceData* SourceData, UInterchangeBaseNodeContainer& BaseNodeContainer);
 
+			static bool GenericTextureCubeArrayTranslate(const UInterchangeSourceData* SourceData, UInterchangeBaseNodeContainer& BaseNodeContainer);
+
 			static bool GenericTexture2DArrayTranslate(const UInterchangeSourceData* SourceData, UInterchangeBaseNodeContainer& BaseNodeContainer);
 
+			static bool GenericVolumeTextureTranslate(const UInterchangeSourceData* SourceData, UInterchangeBaseNodeContainer& BaseNodeContainer);
+
 			static bool GenericTextureLightProfileTranslate(const UInterchangeSourceData* SourceData, UInterchangeBaseNodeContainer& BaseNodeContainer);
 		};
 	}//ns Interchange
diff --git a/Engine/Plugins/Interchange/Runtime/Source/Import/Public/Texture/InterchangeSlicedTexturePayloadData.h b/Engine/Plugins/Interchange/Runtime/Source/Import/Public/Texture/InterchangeSlicedTexturePayloadData.h
--- a/Engine/Plugins/Interchange/Runtime/Source/Import/Public/Texture/InterchangeSlicedTexturePayloadData.h
+++ b/Engine/Plugins/Interchange/Runtime/Source/Import/Public/Texture/InterchangeSlicedTexturePayloadData.h
@@ -12,12 +12,15 @@ namespace UE
 		struct INTERCHANGEIMPORT_API FImportSlicedImage : public FImportImage
 		{
 			int32 NumSlice = 0;
+			bool bIsVolume = false;
 
-			void Init(int32 InSizeX, int32 InSizeY, int32 InNumSlice, int32 InNumMips, ETextureSourceFormat InFormat , bool InSRGB);
+			void Init(int32 InSizeX, int32 InSizeY, int32 InNumSlice, int32 InNumMips, ETextureSourceFormat InFormat, bool InSRGB);
+			void InitVolume(int32 InSizeX, int32 InSizeY, int32 InNumSlice, int32 InNumMips, ETextureSourceFormat InFormat, bool InSRGB);
 
 			const uint8* GetMipData(int32 InMipIndex, int32 InSliceIndex = INDEX_NONE) const;
 			uint8* GetMipData(int32 InMipIndex, int32 InSliceIndex = INDEX_NONE);
 
+			virtual int64 GetMipSize(int32 InMipIndex) const override;
 			virtual int64 ComputeBufferSize() const override;
 
 			virtual bool IsValid() const override;
diff --git a/Engine/Plugins/Interchange/Runtime/Source/Import/Public/Texture/InterchangeTexturePayloadData.h b/Engine/Plugins/Interchange/Runtime/Source/Import/Public/Texture/InterchangeTexturePayloadData.h
--- a/Engine/Plugins/Interchange/Runtime/Source/Import/Public/Texture/InterchangeTexturePayloadData.h
+++ b/Engine/Plugins/Interchange/Runtime/Source/Import/Public/Texture/InterchangeTexturePayloadData.h
@@ -39,7 +39,7 @@ namespace UE
 			void Init2DWithOneMip(int32 InSizeX, int32 InSizeY, ETextureSourceFormat InFormat, const void* InData = nullptr);
 			void Init2DWithMips(int32 InSizeX, int32 InSizeY, int32 InNumMips, ETextureSourceFormat InFormat, const void* InData = nullptr);
 
-			int64 GetMipSize(int32 InMipIndex) const;
+			virtual int64 GetMipSize(int32 InMipIndex) const;
 			virtual int64 ComputeBufferSize() const;
 
 			TArrayView64<uint8> GetArrayViewOfRawData();
diff --git a/Engine/Plugins/Interchange/Runtime/Source/Nodes/Public/InterchangeTextureCubeArrayNode.h b/Engine/Plugins/Interchange/Runtime/Source/Nodes/Public/InterchangeTextureCubeArrayNode.h
new file mode 100644
--- /dev/null
+++ b/Engine/Plugins/Interchange/Runtime/Source/Nodes/Public/InterchangeTextureCubeArrayNode.h
@@ -0,0 +1,26 @@
+// Copyright Epic Games, Inc. All Rights Reserved.
+
+#pragma once
+
+#include "CoreMinimal.h"
+#include "InterchangeTextureNode.h"
+
+#include "InterchangeTextureCubeArrayNode.generated.h"
+
+
+UCLASS(BlueprintType, Experimental)
+class INTERCHANGENODES_API UInterchangeTextureCubeArrayNode : public UInterchangeTextureNode
+{
+	GENERATED_BODY()
+
+public:
+
+	/**
+	 * Return the node type name of the class, we use this when reporting error
+	 */
+	virtual FString GetTypeName() const override
+	{
+		const FString TypeName = TEXT("TextureCubeArrayNode");
+		return TypeName;
+	}
+};
diff --git a/Engine/Plugins/Interchange/Runtime/Source/Nodes/Public/InterchangeVolumeTextureNode.h b/Engine/Plugins/Interchange/Runtime/Source/Nodes/Public/InterchangeVolumeTextureNode.h
new file mode 100644
--- /dev/null
+++ b/Engine/Plugins/Interchange/Runtime/Source/Nodes/Public/InterchangeVolumeTextureNode.h
@@ -0,0 +1,26 @@
+// Copyright Epic Games, Inc. All Rights Reserved.
+
+#pragma once
+
+#include "CoreMinimal.h"
+#include "InterchangeTextureNode.h"
+
+#include "InterchangeVolumeTextureNode.generated.h"
+
+
+UCLASS(BlueprintType, Experimental)
+class INTERCHANGENODES_API UInterchangeVolumeTextureNode : public UInterchangeTextureNode
+{
+	GENERATED_BODY()
+
+public:
+
+	/**
+	 * Return the node type name of the class, we use this when reporting error
+	 */
+	virtual FString GetTypeName() const override
+	{
+		const FString TypeName = TEXT("VolumeTextureNode");
+		return TypeName;
+	}
+};
diff --git a/Engine/Plugins/Interchange/Runtime/Source/Pipelines/Private/InterchangeGenericTexturePipeline.cpp b/Engine/Plugins/Interchange/Runtime/Source/Pipelines/Private/InterchangeGenericTexturePipeline.cpp
--- a/Engine/Plugins/Interchange/Runtime/Source/Pipelines/Private/InterchangeGenericTexturePipeline.cpp
+++ b/Engine/Plugins/Interchange/Runtime/Source/Pipelines/Private/InterchangeGenericTexturePipeline.cpp
@@ -7,22 +7,19 @@
 #include "Engine/TextureLightProfile.h"
 #include "InterchangePipelineLog.h"
 #include "InterchangeTexture2DArrayFactoryNode.h"
-#include "InterchangeTexture2DArrayFactoryNode.h"
-#include "InterchangeTexture2DArrayNode.h"
 #include "InterchangeTexture2DArrayNode.h"
 #include "InterchangeTexture2DFactoryNode.h"
 #include "InterchangeTexture2DNode.h"
-#include "InterchangeTexture2DNode.h"
-#include "InterchangeTextureCubeFactoryNode.h"
+#include "InterchangeTextureCubeArrayFactoryNode.h"
+#include "InterchangeTextureCubeArrayNode.h"
 #include "InterchangeTextureCubeFactoryNode.h"
 #include "InterchangeTextureCubeNode.h"
-#include "InterchangeTextureCubeNode.h"
 #include "InterchangeTextureFactoryNode.h"
 #include "InterchangeTextureLightProfileFactoryNode.h"
-#include "InterchangeTextureLightProfileFactoryNode.h"
-#include "InterchangeTextureLightProfileNode.h"
 #include "InterchangeTextureLightProfileNode.h"
 #include "InterchangeTextureNode.h"
+#include "InterchangeVolumeTextureNode.h"
+#include "InterchangeVolumeTextureFactoryNode.h"
 #include "Misc/Paths.h"
 #include "Engine/Texture.h"
 
@@ -46,6 +43,11 @@ namespace UE::Interchange::Private
 			return UInterchangeTextureCubeFactoryNode::StaticClass();
 		}
 
+		if (UInterchangeTextureCubeArrayNode::StaticClass() == NodeClass)
+		{
+			return UInterchangeTextureCubeArrayFactoryNode::StaticClass();
+		}
+
 		if (UInterchangeTexture2DArrayNode::StaticClass() == NodeClass)
 		{
 			return UInterchangeTexture2DArrayFactoryNode::StaticClass();
@@ -56,6 +58,11 @@ namespace UE::Interchange::Private
 			return UInterchangeTextureLightProfileFactoryNode::StaticClass();
 		}
 
+		if (UInterchangeVolumeTextureNode::StaticClass() == NodeClass)
+		{
+			return UInterchangeVolumeTextureFactoryNode::StaticClass();
+		}
+
 		return nullptr;
 	}
 
diff --git a/Engine/Source/Runtime/ImageCore/Private/DDSFile.cpp b/Engine/Source/Runtime/ImageCore/Private/DDSFile.cpp
--- a/Engine/Source/Runtime/ImageCore/Private/DDSFile.cpp
+++ b/Engine/Source/Runtime/ImageCore/Private/DDSFile.cpp
@@ -716,7 +716,7 @@ static EDDSError ReadPayload(FDDSFile* InDDS, const uint8* InReadCursor, const u
 }
 
 
-/* static */ FDDSFile* FDDSFile::CreateFromDDSInMemory(const uint8* InDDS, int64 InDDSSize, EDDSError *OutError)
+/* static */ FDDSFile* FDDSFile::CreateFromDDSInMemory(const uint8* InDDS, int64 InDDSSize, EDDSError *OutError, bool bHeaderOnly)
 {
 	// If no OutError passed in, redirect it to dummy storage on stack.
 	EDDSError DummyError;
@@ -771,7 +771,7 @@ static EDDSError ReadPayload(FDDSFile* InDDS, const uint8* InReadCursor, const u
 	FDDSFile* DDS = new FDDSFile();
 
 	*OutError = ParseHeader(DDS, &DDSHeader, &DDS10Header);
-	if (*OutError == EDDSError::OK)
+	if (*OutError == EDDSError::OK && !bHeaderOnly)
 	{
 		*OutError = ReadPayload(DDS, ReadCursor, ReadEnd);
 	}
@@ -1331,4 +1331,4 @@ EDXGIFormat DXGIFormatFromRawFormat(ERawImageFormat::Type RawFormat,EGammaSpace
 	}
 }
 
-} } // end UE::DDS namespace
\ No newline at end of file
+} } // end UE::DDS namespace
diff --git a/Engine/Source/Runtime/ImageCore/Public/DDSFile.h b/Engine/Source/Runtime/ImageCore/Public/DDSFile.h
--- a/Engine/Source/Runtime/ImageCore/Public/DDSFile.h
+++ b/Engine/Source/Runtime/ImageCore/Public/DDSFile.h
@@ -290,7 +290,7 @@ namespace UE { namespace DDS
 		// 
 		// On error, returns nullptr. If a non-null OutError is supplied, error information
 		// is written there.
-		IMAGECORE_API static FDDSFile* CreateFromDDSInMemory(const uint8* InDDS, int64 InDDSSize, EDDSError* OutError=nullptr);
+		IMAGECORE_API static FDDSFile* CreateFromDDSInMemory(const uint8* InDDS, int64 InDDSSize, EDDSError* OutError=nullptr, bool bHeaderOnly = false);
 		
 		// is this buffer a DDS ?
 		IMAGECORE_API static bool IsADDS(const uint8* InDDS, int64 InDDSSize);
diff --git a/Engine/Source/Runtime/Interchange/Engine/Private/InterchangeManager.cpp b/Engine/Source/Runtime/Interchange/Engine/Private/InterchangeManager.cpp
--- a/Engine/Source/Runtime/Interchange/Engine/Private/InterchangeManager.cpp
+++ b/Engine/Source/Runtime/Interchange/Engine/Private/InterchangeManager.cpp
@@ -1320,30 +1320,45 @@ void UInterchangeManager::ReleaseAsyncHelper(TWeakPtr<UE::Interchange::FImportAs
 
 	check(AsyncHelper.IsValid());
 	
+	constexpr bool bLogWarningsAndErrors = true;
+
 	bool bSucceeded = false;
 	{
 		TSharedPtr<FImportAsyncHelper> AsyncHelperPtr = AsyncHelper.Pin();
 		
-		for (UInterchangeResult* Result : AsyncHelperPtr->AssetImportResult->GetResults()->GetResults())
+		auto ForEachResult = [&bSucceeded, bLogWarningsAndErrors](TArray<UInterchangeResult*>&& Results)
 		{
-			if (Result && Result->GetResultType() == EInterchangeResultType::Success)
+			if (!bSucceeded || bLogWarningsAndErrors)
 			{
-				bSucceeded = true;
-				break;
-			}
-		}
-
-		if (!bSucceeded)
-		{
-			for (UInterchangeResult* Result : AsyncHelperPtr->SceneImportResult->GetResults()->GetResults())
-			{
-				if (Result && Result->GetResultType() == EInterchangeResultType::Success)
+				for (UInterchangeResult* Result : Results)
 				{
-					bSucceeded = true;
-					break;
+					if (!Result)
+					{
+						continue;
+					}
+
+					if (Result->GetResultType() == EInterchangeResultType::Success)
+					{
+						bSucceeded = true;
+					}
+					else if (bLogWarningsAndErrors)
+					{
+						switch (Result->GetResultType())
+						{
+						case EInterchangeResultType::Warning:
+							UE_LOG(LogInterchangeEngine, Warning, TEXT("%s"), *Result->ToJson());
+							break;
+						case EInterchangeResultType::Error:
+							UE_LOG(LogInterchangeEngine, Error, TEXT("%s"), *Result->ToJson());
+							break;
+						}
+					}
 				}
 			}
-		}
+		};
+
+		ForEachResult(AsyncHelperPtr->AssetImportResult->GetResults()->GetResults());
+		ForEachResult(AsyncHelperPtr->SceneImportResult->GetResults()->GetResults());
 	}
 
 	ImportTasks.RemoveSingle(AsyncHelper.Pin());
