Prevented update of Nanite settings at runtime: NaniteSettings property is not available at runtime

diff --git a/Engine/Plugins/Interchange/Runtime/Source/FactoryNodes/Private/InterchangeStaticMeshFactoryNode.cpp b/Engine/Plugins/Interchange/Runtime/Source/FactoryNodes/Private/InterchangeStaticMeshFactoryNode.cpp
--- a/Engine/Plugins/Interchange/Runtime/Source/FactoryNodes/Private/InterchangeStaticMeshFactoryNode.cpp
+++ b/Engine/Plugins/Interchange/Runtime/Source/FactoryNodes/Private/InterchangeStaticMeshFactoryNode.cpp
@@ -95,7 +95,11 @@ bool UInterchangeStaticMeshFactoryNode::GetCustomBuildNanite(bool& AttributeValu
 
 bool UInterchangeStaticMeshFactoryNode::SetCustomBuildNanite(const bool& AttributeValue, bool bAddApplyDelegate /*= true*/)
 {
+#if WITH_EDITORONLY_DATA
 	IMPLEMENT_NODE_ATTRIBUTE_SETTER_WITH_CUSTOM_DELEGATE(UInterchangeStaticMeshFactoryNode, BuildNanite, bool);
+#else
+	return false;
+#endif
 }
 
 int32 UInterchangeStaticMeshFactoryNode::GetSocketUidCount() const
diff --git a/Engine/Plugins/Interchange/Runtime/Source/FactoryNodes/Public/InterchangeStaticMeshFactoryNode.h b/Engine/Plugins/Interchange/Runtime/Source/FactoryNodes/Public/InterchangeStaticMeshFactoryNode.h
--- a/Engine/Plugins/Interchange/Runtime/Source/FactoryNodes/Public/InterchangeStaticMeshFactoryNode.h
+++ b/Engine/Plugins/Interchange/Runtime/Source/FactoryNodes/Public/InterchangeStaticMeshFactoryNode.h
@@ -207,7 +207,9 @@ private:
 
 protected:
 	
+#if WITH_EDITORONLY_DATA
 	IMPLEMENT_NODE_ATTRIBUTE_DELEGATE_BY_PROPERTYNAME(BuildNanite, bool, UStaticMesh, TEXT("NaniteSettings.bEnabled"));
+#endif
 
 	bool ApplyCustomBuildReversedIndexBufferToAsset(UObject* Asset) const;
 	bool FillCustomBuildReversedIndexBufferFromAsset(UObject* Asset);
