Interchange Textures
- Support for UDIMs of different formats.
- Support for 32k virtual textures.
- Added a default value to the large texture dialog in UTextureFactory so that it works with automation.

diff --git a/Engine/Plugins/Interchange/Runtime/Source/Import/Private/Texture/InterchangeTextureFactory.cpp b/Engine/Plugins/Interchange/Runtime/Source/Import/Private/Texture/InterchangeTextureFactory.cpp
--- a/Engine/Plugins/Interchange/Runtime/Source/Import/Private/Texture/InterchangeTextureFactory.cpp
+++ b/Engine/Plugins/Interchange/Runtime/Source/Import/Private/Texture/InterchangeTextureFactory.cpp
@@ -518,6 +518,46 @@ namespace UE::Interchange::Private::InterchangeTextureFactory
 		{
 			BlockedImage.BlocksData.Reserve(Images.Num());
 
+			bool bMismatchedFormats = false;
+			for (int32 Index = 0; Index < Images.Num(); ++Index)
+			{
+				if (BlockedImage.Format != Images[Index].Format)
+				{
+					bMismatchedFormats = true;
+					break;
+				}
+			}
+
+			if (bMismatchedFormats)
+			{
+				UE_LOG(LogInterchangeImport, Warning, TEXT("Mismatched UDIM image formats, converting all to BGRA8 or RGBA16F ..."));
+
+				if (FTextureSource::IsHDR(BlockedImage.Format))
+				{
+					BlockedImage.Format = TSF_RGBA16F;
+					BlockedImage.bSRGB = false;
+				}
+				else
+				{
+					BlockedImage.Format = TSF_BGRA8;
+					BlockedImage.bSRGB = true;
+				}
+
+				for (UE::Interchange::FImportImage& Image : Images)
+				{
+					if (Image.Format != BlockedImage.Format)
+					{
+						ERawImageFormat::Type ImageRawFormat = FImageCoreUtils::ConvertToRawImageFormat(Image.Format);
+						FImageView SourceImage(Image.RawData.GetData(), Image.SizeX, Image.SizeY, ImageRawFormat);
+						FImage DestImage(Image.SizeX, Image.SizeY, FImageCoreUtils::ConvertToRawImageFormat(BlockedImage.Format));
+						FImageCore::CopyImage(SourceImage, DestImage);
+
+						Image.RawData = MakeUniqueBufferFromArray(MoveTemp(DestImage.RawData));
+						Image.Format = BlockedImage.Format;
+					}
+				}
+			}
+
 			for (int32 Index = 0; Index < Images.Num(); ++Index)
 			{
 				int32 UDIMIndex = UDIMsAndSourcesFileArray[Index]->Key;
diff --git a/Engine/Plugins/Interchange/Runtime/Source/Import/Private/Texture/InterchangeTexturePayloadData.cpp b/Engine/Plugins/Interchange/Runtime/Source/Import/Private/Texture/InterchangeTexturePayloadData.cpp
--- a/Engine/Plugins/Interchange/Runtime/Source/Import/Private/Texture/InterchangeTexturePayloadData.cpp
+++ b/Engine/Plugins/Interchange/Runtime/Source/Import/Private/Texture/InterchangeTexturePayloadData.cpp
@@ -10,26 +10,78 @@
 
 bool UE::Interchange::FImportImageHelper::IsImportResolutionValid(int32 Width, int32 Height, bool bAllowNonPowerOfTwo, FText* OutErrorMessage)
 {
-	static const auto CVarVirtualTexturesEnabled = IConsoleManager::Get().FindTConsoleVariableDataInt(TEXT("r.VirtualTextures")); check(CVarVirtualTexturesEnabled);
+	// VT res is currently limited by pixel count fitting in int32
+	const int64 MaximumSupportedVirtualTextureResolution = 32768;
 
-	// In theory this value could be much higher, but various UE image code currently uses 32bit size/offset values
-	const int32 MaximumSupportedVirtualTextureResolution = 16 * 1024;
+	// Get the non-VT size limit :
+	int64 MaximumSupportedResolutionNonVT = (int64)UTexture::GetMaximumDimensionOfNonVT();
 
-	// Calculate the maximum supported resolution utilizing the global max texture mip count
-	// (Note, have to subtract 1 because 1x1 is a valid mip-size; this means a GMaxTextureMipCount of 4 means a max resolution of 8x8, not 2^4 = 16x16)
-	const int32 MaximumSupportedResolution = CVarVirtualTexturesEnabled->GetValueOnAnyThread() ? MaximumSupportedVirtualTextureResolution : (1 << (GMaxTextureMipCount - 1));
+	// limit on current rendering RHI : == GetMax2DTextureDimension()
+	const int64 CurrentRHIMaxResolution = int64(1) << (GMaxTextureMipCount - 1);
 
-	// Check if the texture is above the supported resolution and prompt the user if they wish to continue if it is
-	if (Width > MaximumSupportedResolution || Height > MaximumSupportedResolution)
+	// MaximumSupportedResolutionNonVT is only a popup/warning , not a hard limit
+	MaximumSupportedResolutionNonVT = FMath::Min(MaximumSupportedResolutionNonVT, CurrentRHIMaxResolution);
+
+	// No zero-size textures :
+	if (Width == 0 || Height == 0)
+	{
+		if (OutErrorMessage)
+		{
+			*OutErrorMessage = NSLOCTEXT("Interchange", "Warning_TextureSizeZero", "Texture has zero width or height");
+		}
+
+		return false;
+	}
+
+	bool bTextureTooLargeOrInvalidResolution = false;
+
+	// Dimensions must fit in signed int32
+	//  could be negative here if it was over 2G and int32 was used earlier
+	if (Width < 0 || Height < 0 || Width > MAX_int32 || Height > MAX_int32)
+	{
+		bTextureTooLargeOrInvalidResolution = true;
+	}
+
+	// pixel count must fit in int32 :
+	//  mip surface could still be larger than 2 GB, that's allowed
+	//	eg. 16k RGBA float = 4 GB
+	if (Width * Height > MAX_int32)
+	{
+		bTextureTooLargeOrInvalidResolution = true;
+	}
+
+	if ((Width * Height) > FMath::Square(MaximumSupportedVirtualTextureResolution))
+	{
+		bTextureTooLargeOrInvalidResolution = true;
+	}
+
+	if (bTextureTooLargeOrInvalidResolution)
 	{
-		if ((Width * Height) > FMath::Square(MaximumSupportedVirtualTextureResolution))
+		if (OutErrorMessage)
 		{
-			if (OutErrorMessage)
+			*OutErrorMessage = FText::Format(
+				NSLOCTEXT("Interchange", "Warning_TextureSizeTooLargeOrInvalid", "Texture is too large to import or it has an invalid resolution. The current maximum is {0} pixels"),
+				FText::AsNumber(FMath::Square(MaximumSupportedVirtualTextureResolution)));
+		}
+
+		return false;
+	}
+
+	if (Width > MaximumSupportedResolutionNonVT || Height > MaximumSupportedResolutionNonVT)
+	{
+		const TConsoleVariableData<int32>* CVarVirtualTexturesEnabled = IConsoleManager::Get().FindTConsoleVariableDataInt(TEXT("r.VirtualTextures")); check(CVarVirtualTexturesEnabled);
+		check(CVarVirtualTexturesEnabled != nullptr);
+
+		if (!CVarVirtualTexturesEnabled->GetValueOnAnyThread())
+		{
+			const FText VTMessage = NSLOCTEXT("Interchange", "Warning_LargeTextureVTDisabled", "\nWarning: Virtual Textures are disabled in this project.");
+
+			if (EAppReturnType::Yes != FMessageDialog::Open(EAppMsgType::YesNo, EAppReturnType::Yes, FText::Format(
+				NSLOCTEXT("Interchange", "Warning_LargeTextureImport", "Attempting to import {0} x {1} texture, proceed?\nLargest supported non-VT texture size: {2} x {3}{4}"),
+				FText::AsNumber(Width), FText::AsNumber(Height), FText::AsNumber(MaximumSupportedResolutionNonVT), FText::AsNumber(MaximumSupportedResolutionNonVT), VTMessage)))
 			{
-				*OutErrorMessage = NSLOCTEXT("Interchange", "Warning_TextureSizeTooLarge", "Texture is too large to import");
+				return false;
 			}
-			
-			return false;
 		}
 	}
 
diff --git a/Engine/Source/Editor/UnrealEd/Private/Factories/EditorFactories.cpp b/Engine/Source/Editor/UnrealEd/Private/Factories/EditorFactories.cpp
--- a/Engine/Source/Editor/UnrealEd/Private/Factories/EditorFactories.cpp
+++ b/Engine/Source/Editor/UnrealEd/Private/Factories/EditorFactories.cpp
@@ -4508,7 +4508,7 @@ bool UTextureFactory::IsImportResolutionValid(int64 Width, int64 Height, bool bA
 		const FText VTMessage = CVarVirtualTexturesEnabled->GetValueOnAnyThread() ? FText() :
 			NSLOCTEXT("UnrealEd","Warning_LargeTextureVTDisabled", "\nWarning: Virtual Textures are disabled in this project.");
 
-		if ( EAppReturnType::Yes != FMessageDialog::Open( EAppMsgType::YesNo, FText::Format(
+		if ( EAppReturnType::Yes != FMessageDialog::Open( EAppMsgType::YesNo, EAppReturnType::Yes, FText::Format(
 				NSLOCTEXT("UnrealEd", "Warning_LargeTextureImport", "Attempting to import {0} x {1} texture, proceed?\nLargest supported non-VT texture size: {2} x {3}{4}"),
 				FText::AsNumber(Width), FText::AsNumber(Height), FText::AsNumber(MaximumSupportedResolutionNonVT), FText::AsNumber(MaximumSupportedResolutionNonVT), VTMessage) ) )
 		{
