Convert MaterialX rotate2d node input from vector2 to vector3 to fit RotateAboutAxis and convert it to a vector2

diff --git a/Engine/Plugins/Interchange/Runtime/Source/Import/Private/MaterialX/InterchangeMaterialXTranslator.cpp b/Engine/Plugins/Interchange/Runtime/Source/Import/Private/MaterialX/InterchangeMaterialXTranslator.cpp
--- a/Engine/Plugins/Interchange/Runtime/Source/Import/Private/MaterialX/InterchangeMaterialXTranslator.cpp
+++ b/Engine/Plugins/Interchange/Runtime/Source/Import/Private/MaterialX/InterchangeMaterialXTranslator.cpp
@@ -1387,23 +1387,48 @@ void UInterchangeMaterialXTranslator::ConnectRotate2DInputToOutput(mx::NodePtr U
 	Rotate2DNode->AddLinearColorAttribute(UInterchangeShaderPortsAPI::MakeInputValueKey(TEXT("NormalizedRotationAxis")), FLinearColor(0, 0, 1));
 	Rotate2DNode->AddLinearColorAttribute(UInterchangeShaderPortsAPI::MakeInputValueKey(TEXT("PivotPoint")), FLinearColor(0.5, 0.5, 0));
 
-	AddAttributeFromValueOrInterface(UpstreamNode->getInput("in"), TEXT("Input"), Rotate2DNode);
+	{
+		mx::InputPtr Input = UpstreamNode->getInput("in");
+		AddAttributeFromValueOrInterface(UpstreamNode->getInput("in"), TEXT("Position"), Rotate2DNode);
+
+		// Rotate2D "in" input takes a vector2 input, but RotateAboutAxis takes a vector3 so we convert
+		if(Input)
+		{
+			mx::NodePtr ConvertNode = CreateNode(UpstreamNode->getParent()->asA<mx::NodeGraph>(),
+												 UpstreamNode->getName().c_str(),
+												 mx::Category::Convert,
+												 { {"in", Input} },
+												 {});
+			ConvertNode->setType(mx::Type::Vector3);
+			Input->setNodeName(ConvertNode->getName());
+		}
+	}
 
-	// we create a Divide node to convert MaterialX angle in degrees to UE's angle which is a value between [0,1]
-	if(mx::InputPtr Input = UpstreamNode->getInput("amount"); !AddAttributeFromValueOrInterface(Input, TEXT("RotationAngle"), Rotate2DNode))
 	{
-		//copy "in" input into "in1" and create "amount" input as "in2" with the value 360 because UE's angle is a value between [0,1]
-		mx::NodePtr NewDivideNode = CreateNode(UpstreamNode->getParent()->asA<mx::NodeGraph>(),
-											   UpstreamNode->getName().c_str(),
-											   mx::Category::Divide,
-											   { {"in1", Input} }, // rename input to match <divide> input "in1"
-											   { {"in2", FAttributeValueArray{{"type", "float"}, {"value", "360"}}} });
 
-		// Input now points to the new node
-		Input->setNodeName(NewDivideNode->getName());
+		mx::InputPtr Input = UpstreamNode->getInput("amount");
+		AddAttributeFromValueOrInterface(Input, TEXT("RotationAngle"), Rotate2DNode);
+
+		// we create a Divide node to convert MaterialX angle in degrees to UE's angle which is a value between [0,1]
+		if(Input)
+		{
+			//copy "in" input into "in1" and create "amount" input as "in2" with the value 360 because UE's angle is a value between [0,1]
+			mx::NodePtr DivideNode = CreateNode(UpstreamNode->getParent()->asA<mx::NodeGraph>(),
+												UpstreamNode->getName().c_str(),
+												mx::Category::Divide,
+												{ {"in1", Input} }, // rename input to match <divide> input "in1"
+												{ {"in2", FAttributeValueArray{{"type", "float"}, {"value", "360"}}} });
+
+			// Input now points to the new node
+			Input->setNodeName(DivideNode->getName());
+		}
 	}
 
-	UInterchangeShaderPortsAPI::ConnectDefaultOuputToInput(ParentShaderNode, InputChannelName, Rotate2DNode->GetUniqueID());
+	//We insert a mask node because the output has to be a vector2
+	UInterchangeShaderNode * MaskNode = CreateMaskShaderNode(0b1100, UpstreamNode->getName().c_str() + FString(TEXT("_Mask")), NamesToShaderNodes, NodeContainer);
+	UInterchangeShaderPortsAPI::ConnectDefaultOuputToInput(MaskNode, TEXT("Input"), Rotate2DNode->GetUniqueID());
+
+	UInterchangeShaderPortsAPI::ConnectDefaultOuputToInput(ParentShaderNode, InputChannelName, MaskNode->GetUniqueID());
 }
 
 void UInterchangeMaterialXTranslator::ConnectRotate3DInputToOutput(mx::NodePtr UpstreamNode, UInterchangeShaderNode* ParentShaderNode, const FString& InputChannelName, TMap<FString, UInterchangeShaderNode*>& NamesToShaderNodes, UInterchangeBaseNodeContainer& NodeContainer) const
