MaterialX subsurface seems more appropriate using Subsurface Profile rather than Subsurface, see standard_surface_chess_set King/Queen_W/B

diff --git a/Engine/Plugins/Interchange/Runtime/Source/Pipelines/Private/InterchangeGenericMaterialPipeline.cpp b/Engine/Plugins/Interchange/Runtime/Source/Pipelines/Private/InterchangeGenericMaterialPipeline.cpp
--- a/Engine/Plugins/Interchange/Runtime/Source/Pipelines/Private/InterchangeGenericMaterialPipeline.cpp
+++ b/Engine/Plugins/Interchange/Runtime/Source/Pipelines/Private/InterchangeGenericMaterialPipeline.cpp
@@ -654,7 +654,7 @@ bool UInterchangeGenericMaterialPipeline::HandleStandardSurfaceModel(const UInte
 		else if(UInterchangeShaderPortsAPI::HasInput(ShaderGraphNode, StandardSurface::Parameters::Subsurface))
 		{
 			MaterialFactoryNode->ConnectOutputToSubsurface(FunctionCallExpressionUid, Subsurface::Parameters::SubsurfaceColor.ToString());
-			MaterialFactoryNode->SetCustomShadingModel(MSM_Subsurface);
+			MaterialFactoryNode->SetCustomShadingModel(MSM_SubsurfaceProfile); //MaterialX subsurface fits more with subsurface profile than subsurface, see: standard_surface_chess_set (King_B/W, Queen_B/W)
 		}
 		else if (UInterchangeShaderPortsAPI::HasInput(ShaderGraphNode, StandardSurface::Parameters::Transmission))
 		{
