Fixed compilation error when packaging a game using Interchange plugin
  - Root cause: The implementation of UTexture::GetMaximumDimensionOfNonVT is only in editor mode.

diff --git a/Engine/Plugins/Interchange/Runtime/Source/Import/Private/Texture/InterchangeTexturePayloadData.cpp b/Engine/Plugins/Interchange/Runtime/Source/Import/Private/Texture/InterchangeTexturePayloadData.cpp
--- a/Engine/Plugins/Interchange/Runtime/Source/Import/Private/Texture/InterchangeTexturePayloadData.cpp
+++ b/Engine/Plugins/Interchange/Runtime/Source/Import/Private/Texture/InterchangeTexturePayloadData.cpp
@@ -13,14 +13,17 @@ bool UE::Interchange::FImportImageHelper::IsImportResolutionValid(int32 Width, i
 	// VT res is currently limited by pixel count fitting in int32
 	const int64 MaximumSupportedVirtualTextureResolution = 32768;
 
-	// Get the non-VT size limit :
-	int64 MaximumSupportedResolutionNonVT = (int64)UTexture::GetMaximumDimensionOfNonVT();
-
 	// limit on current rendering RHI : == GetMax2DTextureDimension()
 	const int64 CurrentRHIMaxResolution = int64(1) << (GMaxTextureMipCount - 1);
 
 	// MaximumSupportedResolutionNonVT is only a popup/warning , not a hard limit
+#if WITH_EDITOR
+	// Get the non-VT size limit :
+	int64 MaximumSupportedResolutionNonVT = (int64)UTexture::GetMaximumDimensionOfNonVT();
 	MaximumSupportedResolutionNonVT = FMath::Min(MaximumSupportedResolutionNonVT, CurrentRHIMaxResolution);
+#else
+	int64 MaximumSupportedResolutionNonVT = CurrentRHIMaxResolution;
+#endif
 
 	// No zero-size textures :
 	if (Width == 0 || Height == 0)
