MX_TransmissionSurface:
- Expose the material function to the library
- Init the normal input with the correct value of [0,0,1] instead of [0,0,0]
Translator:
- Preload the MX_TransmissionSurface asset during the GameThread
Generic material pipeline:
- Depending on if a standard surface material has a transmission or not use the correct asset (MX_StandardSurface or MX_TransmissionSurface)

diff --git a/Engine/Plugins/Interchange/Runtime/Source/Import/Private/MaterialX/InterchangeMaterialXTranslator.cpp b/Engine/Plugins/Interchange/Runtime/Source/Import/Private/MaterialX/InterchangeMaterialXTranslator.cpp
--- a/Engine/Plugins/Interchange/Runtime/Source/Import/Private/MaterialX/InterchangeMaterialXTranslator.cpp
+++ b/Engine/Plugins/Interchange/Runtime/Source/Import/Private/MaterialX/InterchangeMaterialXTranslator.cpp
@@ -30,9 +30,8 @@ namespace UE::Interchange::MaterialX
 {
 	bool IsStandardSurfacePackageLoaded()
 	{
-		static const bool bStandardSurfacePackageLoaded = []() -> bool
+		auto IsPackageLoaded = [](const FString & TextPath) -> bool
 		{
-			const FString TextPath{ TEXT("MaterialFunction'/Interchange/Functions/MX_StandardSurface.MX_StandardSurface'") };
 			const FString FunctionPath{ FPackageName::ExportTextPathToObjectPath(TextPath) };
 			if(FPackageName::DoesPackageExist(FunctionPath))
 			{
@@ -51,7 +50,11 @@ namespace UE::Interchange::MaterialX
 			}
 
 			return false;
-		}();
+		};
+
+		static const bool bStandardSurfacePackageLoaded =
+			IsPackageLoaded(TEXT("MaterialFunction'/Interchange/Functions/MX_StandardSurface.MX_StandardSurface'")) &&
+			IsPackageLoaded(TEXT("MaterialFunction'/Interchange/Functions/MX_TransmissionSurface.MX_TransmissionSurface'"));
 
 		return bStandardSurfacePackageLoaded;
 	}
diff --git a/Engine/Plugins/Interchange/Runtime/Source/Pipelines/Private/InterchangeGenericMaterialPipeline.cpp b/Engine/Plugins/Interchange/Runtime/Source/Pipelines/Private/InterchangeGenericMaterialPipeline.cpp
--- a/Engine/Plugins/Interchange/Runtime/Source/Pipelines/Private/InterchangeGenericMaterialPipeline.cpp
+++ b/Engine/Plugins/Interchange/Runtime/Source/Pipelines/Private/InterchangeGenericMaterialPipeline.cpp
@@ -560,12 +560,19 @@ bool UInterchangeGenericMaterialPipeline::HandleStandardSurfaceModel(const UInte
 
 		const FName MaterialFunctionMemberName = GET_MEMBER_NAME_CHECKED(UMaterialExpressionMaterialFunctionCall, MaterialFunction);
 
+		using namespace UE::Interchange::Materials;
+		FString StandardSurfaceOrTransmissionPath{ TEXT("MaterialFunction'/Interchange/Functions/MX_StandardSurface.MX_StandardSurface'") };
+
+		if(UInterchangeShaderPortsAPI::HasInput(ShaderGraphNode, StandardSurface::Parameters::Transmission))
+		{
+			StandardSurfaceOrTransmissionPath = TEXT("MaterialFunction'/Interchange/Functions/MX_TransmissionSurface.MX_TransmissionSurface'");
+		}
+
 		FunctionCallExpression->AddStringAttribute(
 			MaterialFunctionMemberName,
-			TEXT("MaterialFunction'/Interchange/Functions/MX_StandardSurface.MX_StandardSurface'"));
+			StandardSurfaceOrTransmissionPath);
 		FunctionCallExpression->AddApplyAndFillDelegates<FString>(MaterialFunctionMemberName, UMaterialExpressionMaterialFunctionCall::StaticClass(), MaterialFunctionMemberName);
 
-		using namespace UE::Interchange::Materials;
 
 		auto ConnectMaterialExpressionOutputToInput = [&](const FString & InputName, EMaterialInputType InputType)
 		{
