Add ability to trim to whole frames (or range of frames) when importing animations

diff --git a/Engine/Plugins/Interchange/Runtime/Source/Pipelines/Private/InterchangeGenericAnimationPipeline.cpp b/Engine/Plugins/Interchange/Runtime/Source/Pipelines/Private/InterchangeGenericAnimationPipeline.cpp
--- a/Engine/Plugins/Interchange/Runtime/Source/Pipelines/Private/InterchangeGenericAnimationPipeline.cpp
+++ b/Engine/Plugins/Interchange/Runtime/Source/Pipelines/Private/InterchangeGenericAnimationPipeline.cpp
@@ -383,14 +383,55 @@ void UInterchangeGenericAnimationPipeline::ExecutePreImportPipeline(UInterchange
 
 			if (!FMath::IsNearlyZero(SubFrame, KINDA_SMALL_NUMBER) && !FMath::IsNearlyEqual(SubFrame, 1.0f, KINDA_SMALL_NUMBER))
 			{
-				UInterchangeResultError_Generic* Message = AddMessage<UInterchangeResultError_Generic>();
-				Message->SourceAssetName = SourceDatas[0]->GetFilename();
-				Message->DestinationAssetName = AnimSequenceName;
-				Message->AssetType = UAnimSequence::StaticClass();
-				Message->Text = FText::Format(NSLOCTEXT("UInterchangeGenericAnimationPipeline", "WrongSequenceLength", "Animation length {0} is not compatible with import frame-rate {1} (sub frame {2}), animation has to be frame-border aligned."),
-					FText::AsNumber(SequenceLength), FrameRate.ToPrettyText(), FText::AsNumber(SubFrame));
-				//Skip this anim sequence factory node
-				continue;
+			    if (bSnapToClosestFrameBoundary)
+			    {
+				    // Figure out whether start or stop has to be adjusted
+				    const FFrameTime StartFrameTime = FrameRate.AsFrameTime(RangeStart);
+				    const FFrameTime StopFrameTime = FrameRate.AsFrameTime(RangeStop);
+				    FFrameNumber StartFrameNumber = StartFrameTime.GetFrame().Value, StopFrameNumber = StopFrameTime.GetFrame().Value;
+				    double NewStartTime = RangeStart, NewStopTime = RangeStop;
+    
+				    if (!FMath::IsNearlyZero(StartFrameTime.GetSubFrame()))
+				    {
+					    StartFrameNumber = StartFrameTime.RoundToFrame();
+					    NewStartTime = FrameRate.AsSeconds(StartFrameNumber);
+				    }
+    
+				    if (!FMath::IsNearlyZero(StopFrameTime.GetSubFrame()))
+				    {
+					    StopFrameNumber = StopFrameTime.RoundToFrame();
+					    NewStopTime = FrameRate.AsSeconds(StopFrameNumber);
+				    }
+    
+				    UInterchangeResultWarning_Generic* Message = AddMessage<UInterchangeResultWarning_Generic>();
+				    Message->SourceAssetName = SourceDatas[0]->GetFilename();
+				    Message->DestinationAssetName = AnimSequenceName;
+				    Message->AssetType = UAnimSequence::StaticClass();
+				    Message->Text = FText::Format(NSLOCTEXT("UInterchangeGenericAnimationPipeline", "Info_ImportLengthSnap", "Animation length has been adjusted to align with frame borders using import frame-rate {0}.\n\nOriginal timings:\n\t\tStart: {1} ({2})\n\t\tStop: {3} ({4})\nAligned timings:\n\t\tStart: {5} ({6})\n\t\tStop: {7} ({8})"),
+					    FrameRate.ToPrettyText(),
+					    FText::AsNumber(RangeStart),
+					    FText::AsNumber(StartFrameTime.AsDecimal()),
+					    FText::AsNumber(RangeStop),
+					    FText::AsNumber(StopFrameTime.AsDecimal()),
+					    FText::AsNumber(NewStartTime),
+					    FText::AsNumber(StartFrameNumber.Value),
+					    FText::AsNumber(NewStopTime),
+					    FText::AsNumber(StopFrameNumber.Value));
+    
+				    RangeStart = NewStartTime;
+				    RangeStop = NewStopTime;
+			    }
+			    else
+			    {
+				    UInterchangeResultError_Generic* Message = AddMessage<UInterchangeResultError_Generic>();
+				    Message->SourceAssetName = SourceDatas[0]->GetFilename();
+				    Message->DestinationAssetName = AnimSequenceName;
+				    Message->AssetType = UAnimSequence::StaticClass();
+				    Message->Text = FText::Format(NSLOCTEXT("UInterchangeGenericAnimationPipeline", "WrongSequenceLength", "Animation length {0} is not compatible with import frame-rate {1} (sub frame {2}), animation has to be frame-border aligned."),
+					    FText::AsNumber(SequenceLength), FrameRate.ToPrettyText(), FText::AsNumber(SubFrame));
+				    //Skip this anim sequence factory node
+				    continue;
+			    }
 			}
 		}
 
@@ -734,14 +775,55 @@ void UInterchangeGenericAnimationPipeline::CreateAnimSequenceFactoryNode(UInterc
 
 		if (!FMath::IsNearlyZero(SubFrame, KINDA_SMALL_NUMBER) && !FMath::IsNearlyEqual(SubFrame, 1.0f, KINDA_SMALL_NUMBER))
 		{
-			UInterchangeResultError_Generic* Message = AddMessage<UInterchangeResultError_Generic>();
-			Message->SourceAssetName = SourceDatas[0]->GetFilename();
-			Message->DestinationAssetName = TrackNode.GetDisplayLabel();
-			Message->AssetType = UAnimSequence::StaticClass();
-			Message->Text = FText::Format(NSLOCTEXT("UInterchangeGenericAnimationPipeline", "WrongSequenceLength", "Animation length {0} is not compatible with import frame-rate {1} (sub frame {2}), animation has to be frame-border aligned."),
-				FText::AsNumber(SequenceLength), FrameRate.ToPrettyText(), FText::AsNumber(SubFrame));
-			//Skip this anim sequence factory node
-			return;
+			if (bSnapToClosestFrameBoundary)
+			{
+				// Figure out whether start or stop has to be adjusted
+				const FFrameTime StartFrameTime = FrameRate.AsFrameTime(StartTime);
+				const FFrameTime StopFrameTime = FrameRate.AsFrameTime(StopTime);
+				FFrameNumber StartFrameNumber = StartFrameTime.GetFrame().Value, StopFrameNumber = StopFrameTime.GetFrame().Value;
+				double NewStartTime = StartTime, NewStopTime = StopTime;
+
+				if (!FMath::IsNearlyZero(StartFrameTime.GetSubFrame()))
+				{
+					StartFrameNumber = StartFrameTime.RoundToFrame();
+					NewStartTime = FrameRate.AsSeconds(StartFrameNumber);
+				}
+
+				if (!FMath::IsNearlyZero(StopFrameTime.GetSubFrame()))
+				{
+					StopFrameNumber = StopFrameTime.RoundToFrame();
+					NewStopTime = FrameRate.AsSeconds(StopFrameNumber);
+				}
+			
+				UInterchangeResultWarning_Generic* Message = AddMessage<UInterchangeResultWarning_Generic>();
+				Message->SourceAssetName = SourceDatas[0]->GetFilename();
+				Message->DestinationAssetName = TrackNode.GetDisplayLabel();
+				Message->AssetType = UAnimSequence::StaticClass();
+				Message->Text = FText::Format(NSLOCTEXT("UInterchangeGenericAnimationPipeline", "Info_ImportLengthSnap", "Animation length has been adjusted to align with frame borders using import frame-rate {0}.\n\nOriginal timings:\n\t\tStart: {1} ({2})\n\t\tStop: {3} ({4})\nAligned timings:\n\t\tStart: {5} ({6})\n\t\tStop: {7} ({8})"),
+					FrameRate.ToPrettyText(),
+					FText::AsNumber(StartTime),
+					FText::AsNumber(StartFrameTime.AsDecimal()),
+					FText::AsNumber(StopTime),
+					FText::AsNumber(StopFrameTime.AsDecimal()),
+					FText::AsNumber(NewStartTime),
+					FText::AsNumber(StartFrameNumber.Value),
+					FText::AsNumber(NewStopTime),
+					FText::AsNumber(StopFrameNumber.Value));
+
+				StartTime = NewStartTime;
+				StopTime = NewStopTime;
+			}
+			else
+			{
+				UInterchangeResultError_Generic* Message = AddMessage<UInterchangeResultError_Generic>();
+				Message->SourceAssetName = SourceDatas[0]->GetFilename();
+				Message->DestinationAssetName = TrackNode.GetDisplayLabel();
+				Message->AssetType = UAnimSequence::StaticClass();
+				Message->Text = FText::Format(NSLOCTEXT("UInterchangeGenericAnimationPipeline", "WrongSequenceLength", "Animation length {0} is not compatible with import frame-rate {1} (sub frame {2}), animation has to be frame-border aligned."),
+					FText::AsNumber(SequenceLength), FrameRate.ToPrettyText(), FText::AsNumber(SubFrame));
+				//Skip this anim sequence factory node
+				return;
+			}
 		}
 	}
 
@@ -906,4 +988,4 @@ void UInterchangeGenericAnimationPipeline::CreateAnimSequenceFactoryNode(UInterc
 	TrackNode.AddTargetNodeUid(AnimSequenceFactoryNode->GetUniqueID());
 
 	BaseNodeContainer->AddNode(AnimSequenceFactoryNode);
-}
+}
\ No newline at end of file
diff --git a/Engine/Plugins/Interchange/Runtime/Source/Pipelines/Public/InterchangeGenericAnimationPipeline.h b/Engine/Plugins/Interchange/Runtime/Source/Pipelines/Public/InterchangeGenericAnimationPipeline.h
--- a/Engine/Plugins/Interchange/Runtime/Source/Pipelines/Public/InterchangeGenericAnimationPipeline.h
+++ b/Engine/Plugins/Interchange/Runtime/Source/Pipelines/Public/InterchangeGenericAnimationPipeline.h
@@ -66,6 +66,10 @@ public:
 	UPROPERTY(EditAnywhere, BlueprintReadWrite, Category = "Animations", meta = (EditCondition = "bImportBoneTracks && !bUse30HzToBakeBoneAnimation", ToolTip = "Sample fbx animation data at the specified sample rate, 0 find automaticaly the best sample rate", ClampMin = 0, UIMin = 0, ClampMax = 48000, UIMax = 60))
 	int32 CustomBoneAnimationSampleRate = 0;
 
+	/** If enabled, snaps the animation to the closest frame boundary using the import sampling rate */
+	UPROPERTY(EditAnywhere, BlueprintReadWrite, Category = "Animations")
+	bool bSnapToClosestFrameBoundary = false;
+
 	/** Import if custom attribute as a curve within the animation */
 	UPROPERTY(EditAnywhere, BlueprintReadWrite, Category = "Animations")
 	bool bImportCustomAttribute = true;
diff --git a/Engine/Source/Editor/UnrealEd/Classes/Factories/FbxAnimSequenceImportData.h b/Engine/Source/Editor/UnrealEd/Classes/Factories/FbxAnimSequenceImportData.h
--- a/Engine/Source/Editor/UnrealEd/Classes/Factories/FbxAnimSequenceImportData.h
+++ b/Engine/Source/Editor/UnrealEd/Classes/Factories/FbxAnimSequenceImportData.h
@@ -66,6 +66,9 @@ class UNREALED_API UFbxAnimSequenceImportData : public UFbxAssetImportData
 	UPROPERTY(EditAnywhere, AdvancedDisplay, config, Category = ImportSettings, meta = (EditCondition = "!bUseDefaultSampleRate", ToolTip = "Sample fbx animation data at the specified sample rate, 0 find automaticaly the best sample rate", ClampMin = 0, UIMin = 0, ClampMax = 48000, UIMax = 60))
 	int32 CustomSampleRate;
 
+	UPROPERTY(EditAnywhere, AdvancedDisplay, config, Category = ImportSettings, meta = (ToolTip = "If enabled, snaps the animation to the closest frame boundary using the import sampling rate"))
+	bool bSnapToClosestFrameBoundary;
+
 	/** Name of source animation that was imported, used to reimport correct animation from the FBX file */
 	UPROPERTY()
 	FString SourceAnimationName;
diff --git a/Engine/Source/Editor/UnrealEd/Classes/Factories/FbxSceneImportOptionsSkeletalMesh.h b/Engine/Source/Editor/UnrealEd/Classes/Factories/FbxSceneImportOptionsSkeletalMesh.h
--- a/Engine/Source/Editor/UnrealEd/Classes/Factories/FbxSceneImportOptionsSkeletalMesh.h
+++ b/Engine/Source/Editor/UnrealEd/Classes/Factories/FbxSceneImportOptionsSkeletalMesh.h
@@ -80,6 +80,9 @@ class UFbxSceneImportOptionsSkeletalMesh : public UObject
 	/** Use this option to specify a sample rate for the imported animation, a value of 0 use the best matching sample rate. */
 	UPROPERTY(EditAnywhere, AdvancedDisplay, config, Category = Animation, meta = (EditCondition = "!bUseDefaultSampleRate", ToolTip = "Sample fbx animation data at the specified sample rate, 0 find automaticaly the best sample rate"))
 	int32 CustomSampleRate;
+	
+	UPROPERTY(EditAnywhere, AdvancedDisplay, config, Category = Animation, meta = (ToolTip = "If enabled, snaps the animation to the closest frame boundary using the import sampling rate"))
+	bool bSnapToClosestFrameBoundary;
 
 	/** Import if custom attribute as a curve within the animation **/
 	UPROPERTY(EditAnywhere, AdvancedDisplay, config, Category = Animation)
diff --git a/Engine/Source/Editor/UnrealEd/Private/Fbx/FbxMainImport.cpp b/Engine/Source/Editor/UnrealEd/Private/Fbx/FbxMainImport.cpp
--- a/Engine/Source/Editor/UnrealEd/Private/Fbx/FbxMainImport.cpp
+++ b/Engine/Source/Editor/UnrealEd/Private/Fbx/FbxMainImport.cpp
@@ -491,6 +491,7 @@ void ApplyImportUIToImportOptions(UFbxImportUI* ImportUI, FBXImportOptions& InOu
 		// only re-sample if they don't want to use default sample rate
 		InOutImportOptions.bResample					= !ImportUI->AnimSequenceImportData->bUseDefaultSampleRate;
 		InOutImportOptions.ResampleRate					= ImportUI->AnimSequenceImportData->CustomSampleRate;
+		InOutImportOptions.bSnapToClosestFrameBoundary	= ImportUI->AnimSequenceImportData->bSnapToClosestFrameBoundary;
 		InOutImportOptions.bPreserveLocalTransform		= ImportUI->AnimSequenceImportData->bPreserveLocalTransform;
 		InOutImportOptions.bDeleteExistingMorphTargetCurves = ImportUI->AnimSequenceImportData->bDeleteExistingMorphTargetCurves;
 		InOutImportOptions.bRemoveRedundantKeys			= ImportUI->AnimSequenceImportData->bRemoveRedundantKeys;
@@ -1756,7 +1757,7 @@ bool FFbxImporter::ImportFromFile(const FString& Filename, const FString& Type,
 							Attribs.Add(FAnalyticsEventAttribute(TEXT("AnimOpt SetMaterialDriveParameterOnCustomAttribute"), CaptureImportOptions->bSetMaterialDriveParameterOnCustomAttribute));
 							Attribs.Add(FAnalyticsEventAttribute(TEXT("AnimOpt MaterialCurveSuffixes"), CaptureImportOptions->MaterialCurveSuffixes));
 							Attribs.Add(FAnalyticsEventAttribute(TEXT("AnimOpt ResampleRate"), CaptureImportOptions->ResampleRate));
-							
+							Attribs.Add(FAnalyticsEventAttribute(TEXT("AnimOpt SnapToClosestFrameBoundary"), CaptureImportOptions->bSnapToClosestFrameBoundary));
 						};
 						
 						if (ImportOptions->ImportType == FBXIT_SkeletalMesh)
diff --git a/Engine/Source/Editor/UnrealEd/Private/Fbx/ReimportFbxSceneFactory.cpp b/Engine/Source/Editor/UnrealEd/Private/Fbx/ReimportFbxSceneFactory.cpp
--- a/Engine/Source/Editor/UnrealEd/Private/Fbx/ReimportFbxSceneFactory.cpp
+++ b/Engine/Source/Editor/UnrealEd/Private/Fbx/ReimportFbxSceneFactory.cpp
@@ -1365,7 +1365,7 @@ EReimportResult::Type UReimportFbxSceneFactory::ReimportSkeletalMesh(void* VoidF
 					FbxAnimStack* CurAnimStack = FbxImporter->Scene->GetSrcObject<FbxAnimStack>(AnimStackIndex);
 
 					FbxTimeSpan AnimTimeSpan = FbxImporter->GetAnimationTimeSpan(SortedLinks[0], CurAnimStack);
-					bool bValidAnimStack = FbxImporter->ValidateAnimStack(SortedLinks, FBXMeshNodeArray, CurAnimStack, ResampleRate, GlobalImportSettings->bImportMorph, AnimTimeSpan);
+					bool bValidAnimStack = FbxImporter->ValidateAnimStack(SortedLinks, FBXMeshNodeArray, CurAnimStack, ResampleRate, GlobalImportSettings->bImportMorph, GlobalImportSettings->bSnapToClosestFrameBoundary, AnimTimeSpan);
 					// no animation
 					if (!bValidAnimStack)
 					{
@@ -1457,7 +1457,7 @@ EReimportResult::Type UReimportFbxSceneFactory::ReimportSkeletalMesh(void* VoidF
 						else
 						{
 							//Reimport in a existing sequence
-							if (FbxImporter->ValidateAnimStack(SortedLinks, FBXMeshNodeArray, CurAnimStack, ResampleRate, true, AnimTimeSpan))
+							if (FbxImporter->ValidateAnimStack(SortedLinks, FBXMeshNodeArray, CurAnimStack, ResampleRate, true, FbxImporter->ImportOptions->bSnapToClosestFrameBoundary, AnimTimeSpan))
 							{
 								FbxImporter->ImportAnimation(Mesh->GetSkeleton(), DestSeq, CurrentFilename, SortedLinks, FBXMeshNodeArray, CurAnimStack, ResampleRate, AnimTimeSpan);
 							}
diff --git a/Engine/Source/Editor/UnrealEd/Private/Fbx/SFbxSceneOptionWindow.cpp b/Engine/Source/Editor/UnrealEd/Private/Fbx/SFbxSceneOptionWindow.cpp
--- a/Engine/Source/Editor/UnrealEd/Private/Fbx/SFbxSceneOptionWindow.cpp
+++ b/Engine/Source/Editor/UnrealEd/Private/Fbx/SFbxSceneOptionWindow.cpp
@@ -1510,6 +1510,7 @@ void SFbxSceneOptionWindow::CopySkeletalMeshOptionsToFbxOptions(UnFbx::FBXImport
 	ImportSettings->bPreserveLocalTransform = SkeletalMeshOptions->bPreserveLocalTransform;
 	ImportSettings->bResample = !SkeletalMeshOptions->bUseDefaultSampleRate;
 	ImportSettings->ResampleRate = SkeletalMeshOptions->CustomSampleRate;
+	ImportSettings->bSnapToClosestFrameBoundary = SkeletalMeshOptions->bSnapToClosestFrameBoundary;	
 	ImportSettings->AnimationRange.X = SkeletalMeshOptions->FrameImportRange.Min;
 	ImportSettings->AnimationRange.Y = SkeletalMeshOptions->FrameImportRange.Max;
 }
diff --git a/Engine/Source/Editor/UnrealEd/Private/SkeletalMeshEdit.cpp b/Engine/Source/Editor/UnrealEd/Private/SkeletalMeshEdit.cpp
--- a/Engine/Source/Editor/UnrealEd/Private/SkeletalMeshEdit.cpp
+++ b/Engine/Source/Editor/UnrealEd/Private/SkeletalMeshEdit.cpp
@@ -280,7 +280,7 @@ bool UEditorEngine::ReimportFbxAnimation( USkeleton* Skeleton, UAnimSequence* An
 					}
 					FbxTimeSpan AnimTimeSpan = FbxImporter->GetAnimationTimeSpan(SortedLinks[0], CurAnimStack);
 					// for now it's not importing morph - in the future, this should be optional or saved with asset
-					if (FbxImporter->ValidateAnimStack(SortedLinks, FBXMeshNodeArray, CurAnimStack, ResampleRate, bImportMorphTracks, AnimTimeSpan))
+					if (FbxImporter->ValidateAnimStack(SortedLinks, FBXMeshNodeArray, CurAnimStack, ResampleRate, bImportMorphTracks, FbxImporter->ImportOptions->bSnapToClosestFrameBoundary, AnimTimeSpan))
 					{
 						AnimSequence->ImportResampleFramerate = ResampleRate;
 						FbxImporter->ImportAnimation(Skeleton, AnimSequence, Filename, SortedLinks, FBXMeshNodeArray, CurAnimStack, ResampleRate, AnimTimeSpan);
@@ -615,7 +615,7 @@ UAnimSequence * UnFbx::FFbxImporter::ImportAnimations(USkeleton* Skeleton, UObje
 		FbxAnimStack* CurAnimStack = Scene->GetSrcObject<FbxAnimStack>(AnimStackIndex);
 
 		FbxTimeSpan AnimTimeSpan = GetAnimationTimeSpan(SortedLinks[0], CurAnimStack);
-		bool bValidAnimStack = ValidateAnimStack(SortedLinks, NodeArray, CurAnimStack, ResampleRate, ImportOptions->bImportMorph, AnimTimeSpan);
+		bool bValidAnimStack = ValidateAnimStack(SortedLinks, NodeArray, CurAnimStack, ResampleRate, ImportOptions->bImportMorph, ImportOptions->bSnapToClosestFrameBoundary, AnimTimeSpan);
 		// no animation
 		if (!bValidAnimStack)
 		{
@@ -907,7 +907,7 @@ int32 UnFbx::FFbxImporter::GetMaxSampleRate(TArray<FbxNode*>& SortedLinks)
 	return DEFAULT_SAMPLERATE;
 }
 
-bool UnFbx::FFbxImporter::ValidateAnimStack(TArray<FbxNode*>& SortedLinks, TArray<FbxNode*>& NodeArray, FbxAnimStack* CurAnimStack, int32 ResampleRate, bool bImportMorph, FbxTimeSpan &AnimTimeSpan)
+bool UnFbx::FFbxImporter::ValidateAnimStack(TArray<FbxNode*>& SortedLinks, TArray<FbxNode*>& NodeArray, FbxAnimStack* CurAnimStack, int32 ResampleRate, bool bImportMorph, bool bSnapToClosestFrameBoundary, FbxTimeSpan &AnimTimeSpan)
 {
 	// set current anim stack
 	Scene->SetCurrentAnimationStack(CurAnimStack);
@@ -926,11 +926,52 @@ bool UnFbx::FFbxImporter::ValidateAnimStack(TArray<FbxNode*>& SortedLinks, TArra
 
 	const double SequenceLengthInSeconds = FGenericPlatformMath::Max<double>(AnimTimeSpan.GetDuration().GetSecondDouble(), MINIMUM_ANIMATION_LENGTH);
 	const FFrameRate TargetFrameRate(ResampleRate, 1);
-	const float SubFrame = TargetFrameRate.AsFrameTime(SequenceLengthInSeconds).GetSubFrame();
+	const FFrameTime LengthInFrameTime = TargetFrameRate.AsFrameTime(SequenceLengthInSeconds);
+	const float SubFrame = LengthInFrameTime.GetSubFrame();
 	if (!FMath::IsNearlyZero(SubFrame, KINDA_SMALL_NUMBER) && !FMath::IsNearlyEqual(SubFrame, 1.0f, KINDA_SMALL_NUMBER))
 	{
-		AddTokenizedErrorMessage(FTokenizedMessage::Create(EMessageSeverity::Error, FText::Format(LOCTEXT("Error_InvalidImportLength", "Animation length {0} is not compatible with import frame-rate {1} (sub frame {2}), animation has to be frame-border aligned."), FText::AsNumber(SequenceLengthInSeconds), TargetFrameRate.ToPrettyText(), FText::AsNumber(SubFrame))), FFbxErrors::Animation_InvalidData);
-		return false;
+		if (bSnapToClosestFrameBoundary)
+		{
+			// Figure out whether start or stop has to be adjusted
+			const FbxTime StartTime = AnimTimeSpan.GetStart();
+			const FbxTime StopTime = AnimTimeSpan.GetStop();
+			FbxTime NewStartTime;
+			FbxTime NewStopTime;
+
+			const FFrameTime StartFrameTime = TargetFrameRate.AsFrameTime(StartTime.GetSecondDouble());
+			const FFrameTime StopFrameTime = TargetFrameRate.AsFrameTime(StopTime.GetSecondDouble());
+			FFrameNumber StartFrameNumber, StopFrameNumber;
+
+			if (!FMath::IsNearlyZero(StartFrameTime.GetSubFrame()))
+			{
+				StartFrameNumber = StartFrameTime.RoundToFrame();
+				NewStartTime.SetSecondDouble(TargetFrameRate.AsSeconds(StartFrameNumber));
+				AnimTimeSpan.SetStart(NewStartTime);
+			}
+
+			if (!FMath::IsNearlyZero(StopFrameTime.GetSubFrame()))
+			{
+				StopFrameNumber = StopFrameTime.RoundToFrame();
+				NewStopTime.SetSecondDouble(TargetFrameRate.AsSeconds(StopFrameNumber));
+				AnimTimeSpan.SetStop(NewStopTime);
+			}
+			
+			AddTokenizedErrorMessage(FTokenizedMessage::Create(EMessageSeverity::Info, FText::Format(LOCTEXT("Info_ImportLengthSnap", "Animation length has been adjusted to align with frame borders using import frame-rate {0}.\n\nOriginal timings:\n\t\tStart: {1} ({2})\n\t\tStop: {3} ({4})\nAligned timings:\n\t\tStart: {5} ({6})\n\t\tStop: {7} ({8})"),
+				TargetFrameRate.ToPrettyText(),
+				FText::AsNumber(StartTime.GetSecondDouble()),
+				FText::AsNumber(StartFrameTime.AsDecimal()),
+				FText::AsNumber(StopTime.GetSecondDouble()),
+				FText::AsNumber(StopFrameTime.AsDecimal()),
+				FText::AsNumber(NewStartTime.GetSecondDouble()),
+				FText::AsNumber(StartFrameNumber.Value),
+				FText::AsNumber(NewStopTime.GetSecondDouble()),
+				FText::AsNumber(StopFrameNumber.Value))), FFbxErrors::Animation_InvalidData);
+		}
+		else
+		{
+			AddTokenizedErrorMessage(FTokenizedMessage::Create(EMessageSeverity::Error, FText::Format(LOCTEXT("Error_InvalidImportLength", "Animation length {0} is not compatible with import frame-rate {1} (sub frame {2}), animation has to be frame-border aligned. Either re-export animation or enable snap to closest frame boundary import option."), FText::AsNumber(SequenceLengthInSeconds), TargetFrameRate.ToPrettyText(), FText::AsNumber(SubFrame))), FFbxErrors::Animation_InvalidData);
+			return false;
+		}
 	}
 
 	const FBXImportOptions* ImportOption = GetImportOptions();
diff --git a/Engine/Source/Editor/UnrealEd/Public/FbxImporter.h b/Engine/Source/Editor/UnrealEd/Public/FbxImporter.h
--- a/Engine/Source/Editor/UnrealEd/Public/FbxImporter.h
+++ b/Engine/Source/Editor/UnrealEd/Public/FbxImporter.h
@@ -179,6 +179,7 @@ struct FBXImportOptions
 	bool bUpdateSkeletonReferencePose;
 	bool bResample;
 	int32 ResampleRate;
+	bool bSnapToClosestFrameBoundary;
 	bool bImportRigidMesh;
 	bool bUseT0AsRefPose;
 	bool bPreserveSmoothingGroups;
@@ -920,7 +921,7 @@ public:
 	 * @param ResampleRate	Resample Rate for data
 	 * @param AnimTimeSpan	AnimTimeSpan	 
 	 */	
-	bool ValidateAnimStack(TArray<FbxNode*>& SortedLinks, TArray<FbxNode*>& NodeArray, FbxAnimStack* CurAnimStack, int32 ResampleRate, bool bImportMorph, FbxTimeSpan &AnimTimeSpan);
+	bool ValidateAnimStack(TArray<FbxNode*>& SortedLinks, TArray<FbxNode*>& NodeArray, FbxAnimStack* CurAnimStack, int32 ResampleRate, bool bImportMorph, bool bSnapToClosestFrameBoundary, FbxTimeSpan &AnimTimeSpan);
 
 	/**
 	 * Import Fbx Morph object for the Skeletal Mesh.
