FBX Import Settings' Property Name Improvements

diff --git a/Engine/Plugins/Interchange/Runtime/Source/Pipelines/Public/InterchangeGenericAnimationPipeline.h b/Engine/Plugins/Interchange/Runtime/Source/Pipelines/Public/InterchangeGenericAnimationPipeline.h
--- a/Engine/Plugins/Interchange/Runtime/Source/Pipelines/Public/InterchangeGenericAnimationPipeline.h
+++ b/Engine/Plugins/Interchange/Runtime/Source/Pipelines/Public/InterchangeGenericAnimationPipeline.h
@@ -70,8 +70,8 @@ public:
 	UPROPERTY(EditAnywhere, BlueprintReadWrite, Category = "Animations")
 	bool bSnapToClosestFrameBoundary = false;
 
-	/** Import if custom attribute as a curve within the animation */
-	UPROPERTY(EditAnywhere, BlueprintReadWrite, Category = "Animations")
+	/** If true, import node attributes as either Animation Curves or Animation Attributes */
+	UPROPERTY(EditAnywhere, BlueprintReadWrite, Category = "Animations", meta = (DisplayName = "Import Attributes as Curves or Animation Attributes"))
 	bool bImportCustomAttribute = true;
 
 	/** Set Material Curve Type for all custom attributes that exists */
@@ -90,12 +90,12 @@ public:
 	UPROPERTY(EditAnywhere, BlueprintReadWrite, Category = "Animations", meta = (EditCondition = "bImportCustomAttribute", DisplayName = "Do not import curves with only 0 values"))
 	bool bDoNotImportCurveWithZero = false;
 
-	/** If true, all previous non-curve custom attributes will be deleted when doing a re-import. */
-	UPROPERTY(EditAnywhere, BlueprintReadWrite, Category = "Animations")
+	/** If true, all previous node attributes imported as Animation Attributes will be deleted when doing a re-import. */
+	UPROPERTY(EditAnywhere, BlueprintReadWrite, Category = "Animations", meta = (DisplayName = "Delete existing Animation Attributes"))
 	bool bDeleteExistingNonCurveCustomAttributes = false;
 
-	/** If true, all previous custom attribute curves will be deleted when doing a re-import. */
-	UPROPERTY(EditAnywhere, BlueprintReadWrite, Category = "Animations")
+	/** If true, all previous node attributes imported as Animation Curves will be deleted when doing a re-import. */
+	UPROPERTY(EditAnywhere, BlueprintReadWrite, Category = "Animations", meta = (DisplayName = "Delete existing Animation Curves"))
 	bool bDeleteExistingCustomAttributeCurves = false;
 
 	/** If true, all previous morph target curves will be deleted when doing a re-import */
diff --git a/Engine/Source/Editor/UnrealEd/Classes/Factories/FbxAnimSequenceImportData.h b/Engine/Source/Editor/UnrealEd/Classes/Factories/FbxAnimSequenceImportData.h
--- a/Engine/Source/Editor/UnrealEd/Classes/Factories/FbxAnimSequenceImportData.h
+++ b/Engine/Source/Editor/UnrealEd/Classes/Factories/FbxAnimSequenceImportData.h
@@ -73,16 +73,16 @@ class UNREALED_API UFbxAnimSequenceImportData : public UFbxAssetImportData
 	UPROPERTY()
 	FString SourceAnimationName;
 
-	/** Import if custom attribute as a curve within the animation */
-	UPROPERTY(EditAnywhere, AdvancedDisplay, config, Category = ImportSettings)
+	/** If true, import node attributes as either Animation Curves or Animation Attributes */
+	UPROPERTY(EditAnywhere, AdvancedDisplay, config, Category = ImportSettings, meta = (DisplayName = "Import Attributes as Curves or Animation Attributes"))
 	bool bImportCustomAttribute;
 
-	/** If true, all previous custom attribute curves will be deleted when doing a re-import. */
-	UPROPERTY(EditAnywhere, AdvancedDisplay, config, Category = ImportSettings)
+	/** If true, all previous node attributes imported as Animation Curves will be deleted when doing a re-import. */
+	UPROPERTY(EditAnywhere, AdvancedDisplay, config, Category = ImportSettings, meta = (DisplayName = "Delete existing Animation Curves"))
 	bool bDeleteExistingCustomAttributeCurves;
 
-	/** If true, all previous non-curve custom attributes will be deleted when doing a re-import. */
-	UPROPERTY(EditAnywhere, AdvancedDisplay, config, Category = ImportSettings)
+	/** If true, all previous node attributes imported as Animation Attributes will be deleted when doing a re-import. */
+	UPROPERTY(EditAnywhere, AdvancedDisplay, config, Category = ImportSettings, meta = (DisplayName = "Delete existing Animation Attributes"))
 	bool bDeleteExistingNonCurveCustomAttributes;
 	
 	/** Import bone transform tracks. If false, this will discard any bone transform tracks. (useful for curves only animations)*/
diff --git a/Engine/Source/Editor/UnrealEd/Classes/Factories/FbxSceneImportOptionsSkeletalMesh.h b/Engine/Source/Editor/UnrealEd/Classes/Factories/FbxSceneImportOptionsSkeletalMesh.h
--- a/Engine/Source/Editor/UnrealEd/Classes/Factories/FbxSceneImportOptionsSkeletalMesh.h
+++ b/Engine/Source/Editor/UnrealEd/Classes/Factories/FbxSceneImportOptionsSkeletalMesh.h
@@ -84,16 +84,16 @@ class UFbxSceneImportOptionsSkeletalMesh : public UObject
 	UPROPERTY(EditAnywhere, AdvancedDisplay, config, Category = Animation, meta = (ToolTip = "If enabled, snaps the animation to the closest frame boundary using the import sampling rate"))
 	bool bSnapToClosestFrameBoundary;
 
-	/** Import if custom attribute as a curve within the animation **/
-	UPROPERTY(EditAnywhere, AdvancedDisplay, config, Category = Animation)
+	/** If true, import node attributes as either Animation Curves or Animation Attributes */
+	UPROPERTY(EditAnywhere, AdvancedDisplay, config, Category = Animation, meta = (DisplayName = "Import Attributes as Curves or Animation Attributes"))
 	bool bImportCustomAttribute;
 
-	/** If true, all previous custom attribute curves will be deleted when doing a re-import. */
-	UPROPERTY(EditAnywhere, AdvancedDisplay, config, Category = Animation)
+	/** If true, all previous node attributes imported as Animation Curves will be deleted when doing a re-import. */
+	UPROPERTY(EditAnywhere, AdvancedDisplay, config, Category = Animation, meta = (DisplayName = "Delete existing Animation Curves"))
 	bool bDeleteExistingCustomAttributeCurves;
 
-	/** If true, all previous non-curve custom attributes will be deleted when doing a re-import. */
-	UPROPERTY(EditAnywhere, AdvancedDisplay, config, Category = ImportSettings)
+	/** If true, all previous node attributes imported as Animation Attributes will be deleted when doing a re-import. */
+	UPROPERTY(EditAnywhere, AdvancedDisplay, config, Category = Animation, meta = (DisplayName = "Delete existing Animation Attributes"))
 	bool bDeleteExistingNonCurveCustomAttributes;
 
 	/** Type of asset to import from the FBX file */
