Undo //UE5/Main/Engine/Plugins/Interchange/Runtime/... changelist 22071912

diff --git a/Engine/Plugins/Interchange/Runtime/Source/Import/Private/MaterialX/InterchangeMaterialXTranslator.cpp b/Engine/Plugins/Interchange/Runtime/Source/Import/Private/MaterialX/InterchangeMaterialXTranslator.cpp
--- a/Engine/Plugins/Interchange/Runtime/Source/Import/Private/MaterialX/InterchangeMaterialXTranslator.cpp
+++ b/Engine/Plugins/Interchange/Runtime/Source/Import/Private/MaterialX/InterchangeMaterialXTranslator.cpp
@@ -30,8 +30,9 @@ namespace UE::Interchange::MaterialX
 {
 	bool IsStandardSurfacePackageLoaded()
 	{
-		auto IsPackageLoaded = [](const FString & TextPath) -> bool
+		static const bool bStandardSurfacePackageLoaded = []() -> bool
 		{
+			const FString TextPath{ TEXT("MaterialFunction'/Interchange/Functions/MX_StandardSurface.MX_StandardSurface'") };
 			const FString FunctionPath{ FPackageName::ExportTextPathToObjectPath(TextPath) };
 			if(FPackageName::DoesPackageExist(FunctionPath))
 			{
@@ -50,11 +51,7 @@ namespace UE::Interchange::MaterialX
 			}
 
 			return false;
-		};
-
-		static const bool bStandardSurfacePackageLoaded =
-			IsPackageLoaded(TEXT("MaterialFunction'/Interchange/Functions/MX_StandardSurface.MX_StandardSurface'")) &&
-			IsPackageLoaded(TEXT("MaterialFunction'/Interchange/Functions/MX_TransmissionSurface.MX_TransmissionSurface'"));
+		}();
 
 		return bStandardSurfacePackageLoaded;
 	}
diff --git a/Engine/Plugins/Interchange/Runtime/Source/Pipelines/Private/InterchangeGenericMaterialPipeline.cpp b/Engine/Plugins/Interchange/Runtime/Source/Pipelines/Private/InterchangeGenericMaterialPipeline.cpp
--- a/Engine/Plugins/Interchange/Runtime/Source/Pipelines/Private/InterchangeGenericMaterialPipeline.cpp
+++ b/Engine/Plugins/Interchange/Runtime/Source/Pipelines/Private/InterchangeGenericMaterialPipeline.cpp
@@ -560,19 +560,12 @@ bool UInterchangeGenericMaterialPipeline::HandleStandardSurfaceModel(const UInte
 
 		const FName MaterialFunctionMemberName = GET_MEMBER_NAME_CHECKED(UMaterialExpressionMaterialFunctionCall, MaterialFunction);
 
-		using namespace UE::Interchange::Materials;
-		FString StandardSurfaceOrTransmissionPath{ TEXT("MaterialFunction'/Interchange/Functions/MX_StandardSurface.MX_StandardSurface'") };
-
-		if(UInterchangeShaderPortsAPI::HasInput(ShaderGraphNode, StandardSurface::Parameters::Transmission))
-		{
-			StandardSurfaceOrTransmissionPath = TEXT("MaterialFunction'/Interchange/Functions/MX_TransmissionSurface.MX_TransmissionSurface'");
-		}
-
 		FunctionCallExpression->AddStringAttribute(
 			MaterialFunctionMemberName,
-			StandardSurfaceOrTransmissionPath);
+			TEXT("MaterialFunction'/Interchange/Functions/MX_StandardSurface.MX_StandardSurface'"));
 		FunctionCallExpression->AddApplyAndFillDelegates<FString>(MaterialFunctionMemberName, UMaterialExpressionMaterialFunctionCall::StaticClass(), MaterialFunctionMemberName);
 
+		using namespace UE::Interchange::Materials;
 
 		auto ConnectMaterialExpressionOutputToInput = [&](const FString & InputName, EMaterialInputType InputType)
 		{
