Interchange GLTF Translator: MeshNode creation on Demand

diff --git a/Engine/Plugins/Interchange/Runtime/Source/Import/Private/Gltf/InterchangeGltfTranslator.cpp b/Engine/Plugins/Interchange/Runtime/Source/Import/Private/Gltf/InterchangeGltfTranslator.cpp
--- a/Engine/Plugins/Interchange/Runtime/Source/Import/Private/Gltf/InterchangeGltfTranslator.cpp
+++ b/Engine/Plugins/Interchange/Runtime/Source/Import/Private/Gltf/InterchangeGltfTranslator.cpp
@@ -127,7 +127,8 @@ namespace UE::Interchange::Gltf::Private
 	}
 }
 
-void UInterchangeGltfTranslator::HandleGltfNode( UInterchangeBaseNodeContainer& NodeContainer, const GLTF::FNode& GltfNode, const FString& ParentNodeUid, const int32 NodeIndex, bool &bHasVariants, TArray<int32>& SkinnedMeshNodes ) const
+void UInterchangeGltfTranslator::HandleGltfNode( UInterchangeBaseNodeContainer& NodeContainer, const GLTF::FNode& GltfNode, const FString& ParentNodeUid, const int32 NodeIndex, 
+	bool &bHasVariants, TArray<int32>& SkinnedMeshNodes, TSet<int>& UnusedMeshIndices ) const
 {
 	using namespace UE::Interchange::Gltf::Private;
 
@@ -168,6 +169,8 @@ void UInterchangeGltfTranslator::HandleGltfNode( UInterchangeBaseNodeContainer&
 		{
 			if ( GltfAsset.Meshes.IsValidIndex( GltfNode.MeshIndex ) )
 			{
+				HandleGltfMesh(NodeContainer, GltfAsset.Meshes[GltfNode.MeshIndex], GltfNode.MeshIndex, UnusedMeshIndices);
+
 				const FString MeshNodeUid = TEXT("\\Mesh\\") + GltfAsset.Meshes[ GltfNode.MeshIndex ].Name;
 				InterchangeSceneNode->SetCustomAssetInstanceUid( MeshNodeUid );
 
@@ -222,7 +225,7 @@ void UInterchangeGltfTranslator::HandleGltfNode( UInterchangeBaseNodeContainer&
 	{
 		if ( GltfAsset.Nodes.IsValidIndex( ChildIndex ) )
 		{
-			HandleGltfNode( NodeContainer, GltfAsset.Nodes[ ChildIndex ], NodeUid, ChildIndex, bHasVariants, SkinnedMeshNodes );
+			HandleGltfNode( NodeContainer, GltfAsset.Nodes[ ChildIndex ], NodeUid, ChildIndex, bHasVariants, SkinnedMeshNodes, UnusedMeshIndices );
 		}
 	}
 }
@@ -897,32 +900,13 @@ bool UInterchangeGltfTranslator::Translate( UInterchangeBaseNodeContainer& NodeC
 		}
 	}
 
+	TSet<int> UnusedGltfMeshIndices;
 	// Meshes
 	{
 		int32 MeshIndex = 0;
-		for ( const GLTF::FMesh& GltfMesh : GltfAsset.Meshes )
+		for (const GLTF::FMesh& GltfMesh : GltfAsset.Meshes)
 		{
-			UInterchangeMeshNode* MeshNode = NewObject< UInterchangeMeshNode >( &NodeContainer );
-			FString MeshNodeUid = TEXT("\\Mesh\\") + GltfMesh.Name;
-
-			MeshNode->InitializeNode( MeshNodeUid, GltfMesh.Name, EInterchangeNodeContainerType::TranslatedAsset );
-			MeshNode->SetPayLoadKey( LexToString( MeshIndex++ ) );
-
-			NodeContainer.AddNode( MeshNode );
-
-			
-			for (int32 PrimitiveCounter = 0; PrimitiveCounter < GltfMesh.Primitives.Num(); PrimitiveCounter++)
-			{
-				const GLTF::FPrimitive& Primitive = GltfMesh.Primitives[PrimitiveCounter];
-
-				// Assign materials
-				if (GltfAsset.Materials.IsValidIndex(Primitive.MaterialIndex))
-				{
-					const FString MaterialName = GltfAsset.Materials[Primitive.MaterialIndex].Name;
-					const FString ShaderGraphNodeUid = UInterchangeShaderGraphNode::MakeNodeUid(TEXT("Material_") + MaterialName);
-					MeshNode->SetSlotMaterialDependencyUid(MaterialName, ShaderGraphNodeUid);
-				}
-			}
+			UnusedGltfMeshIndices.Add(MeshIndex++);
 		}
 	}
 
@@ -1046,12 +1030,12 @@ bool UInterchangeGltfTranslator::Translate( UInterchangeBaseNodeContainer& NodeC
 			{
 				if ( GltfAsset.Nodes.IsValidIndex( NodeIndex ) )
 				{
-					HandleGltfNode( NodeContainer, GltfAsset.Nodes[ NodeIndex ], SceneNodeUid, NodeIndex, bHasVariants, SkinnedMeshNodes );
+					HandleGltfNode( NodeContainer, GltfAsset.Nodes[ NodeIndex ], SceneNodeUid, NodeIndex, bHasVariants, SkinnedMeshNodes, UnusedGltfMeshIndices );
 				}
 			}
 
 			// Skeletons:
-			HandleGltfSkeletons( NodeContainer, SceneNodeUid, SkinnedMeshNodes );
+			HandleGltfSkeletons( NodeContainer, SceneNodeUid, SkinnedMeshNodes, UnusedGltfMeshIndices );
 		}
 	}
 
@@ -1101,6 +1085,18 @@ bool UInterchangeGltfTranslator::Translate( UInterchangeBaseNodeContainer& NodeC
 		}
 	}
 
+	// Create any Mesh Nodes for meshes that have not been used and just in the gltf as an asset:
+	TSet<int> UnusedMeshIndices = UnusedGltfMeshIndices;
+	for (int UnusedMeshIndex : UnusedMeshIndices)
+	{
+		HandleGltfMesh(NodeContainer, GltfAsset.Meshes[UnusedMeshIndex], UnusedMeshIndex, UnusedGltfMeshIndices);
+	}
+
+	if (UnusedGltfMeshIndices.Num() != 0)
+	{
+		UE_LOG(LogInterchangeImport, Warning, TEXT("GLTF Mesh Import Warning. Gltf Mesh Usage expectation is not met."));
+	}
+
 	return true;
 }
 
@@ -1702,7 +1698,7 @@ TFuture<TOptional<UE::Interchange::FSkeletalMeshMorphTargetPayloadData>> UInterc
 	return Promise->GetFuture();
 }
 
-void UInterchangeGltfTranslator::HandleGltfSkeletons(UInterchangeBaseNodeContainer& NodeContainer, const FString& SceneNodeUid, const TArray<int32>& SkinnedMeshNodes) const
+void UInterchangeGltfTranslator::HandleGltfSkeletons(UInterchangeBaseNodeContainer& NodeContainer, const FString& SceneNodeUid, const TArray<int32>& SkinnedMeshNodes, TSet<int>& UnusedMeshIndices) const
 {
 	//If there are skeletons, set them as a root joint for the joint hierarchy as per documentation:
 	//The skeleton property (if present) points to the node that is the common root of a joints hierarchy or to a direct or indirect parent node of the common root.
@@ -1752,78 +1748,96 @@ void UInterchangeGltfTranslator::HandleGltfSkeletons(UInterchangeBaseNodeContain
 	{
 		const TMap<int32, TArray<int32>>& RootJointGroupedSkinnedMeshNodes = MeshIndexRootJointGroups.Value;
 
-		//Get original MeshNode:
-		const FString OriginalMeshNodeUid = TEXT("\\Mesh\\") + GltfAsset.Meshes[MeshIndexRootJointGroups.Key].Name;
-		if (const UInterchangeMeshNode* ConstOriginalMeshNode = Cast< UInterchangeMeshNode >(NodeContainer.GetNode(OriginalMeshNodeUid)))
+		int MeshIndex = MeshIndexRootJointGroups.Key;
+
+		//iterate through the groups:
+		//rootjoint , array<skinnedMeshNodes>
+		for (const TTuple<int32, TArray<int32>>& RootJointToSkinnedMeshNodes : RootJointGroupedSkinnedMeshNodes)
 		{
-			UInterchangeMeshNode* OriginalMeshNode = const_cast<UInterchangeMeshNode*>(ConstOriginalMeshNode);
+			//Duplicate MeshNode for each group:
+			int32 RootJointIndex = RootJointToSkinnedMeshNodes.Key;
+			UInterchangeMeshNode* SkeletalMeshNode = HandleGltfMesh(NodeContainer, GltfAsset.Meshes[MeshIndex], MeshIndex, UnusedMeshIndices, TEXT("_") + LexToString(RootJointIndex));
 
-			if (!ensure(OriginalMeshNode))
-			{
-				continue;
-			}
+			SkeletalMeshNode->SetSkinnedMesh(true);
 
-			//iterate through the groups:
-			//rootjoint , array<skinnedMeshNodes>
-			for (const TTuple<int32, TArray<int32>>& RootJointToSkinnedMeshNodes : RootJointGroupedSkinnedMeshNodes)
+			//generate payload key:
+			//of template:
+			//"LexToString(SkinnedMeshNode.MeshIndex | (SkinnedMeshNode.Skindex << 16))":"LexToString(SkinnedMeshNode.MeshIndex | (SkinnedMeshNode.Skindex << 16))".....
+			FString Payload = "";
+			for (int32 SkinnedMeshIndex : RootJointToSkinnedMeshNodes.Value)
 			{
-				//Duplicate MeshNode for each group:
-				UInterchangeMeshNode* DuplicatedMeshNode = NewObject< UInterchangeMeshNode >(&NodeContainer);
-				FString DuplicatedMeshNodeUid = OriginalMeshNode->GetUniqueID() + TEXT("_") + LexToString(RootJointToSkinnedMeshNodes.Key);
-				FString DuplicatedMeshDisplayLabel = OriginalMeshNode->GetDisplayLabel() + TEXT("_") + LexToString(RootJointToSkinnedMeshNodes.Key);
-				DuplicatedMeshNode->InitializeNode(DuplicatedMeshNodeUid, DuplicatedMeshDisplayLabel, EInterchangeNodeContainerType::TranslatedAsset);
-				TMap<FString, FString> MaterialDependencies;
-				OriginalMeshNode->GetSlotMaterialDependencies(MaterialDependencies);
-				for (const TTuple<FString, FString>& MaterialDependency : MaterialDependencies)
+				const GLTF::FNode& SkinnedMeshNode = GltfAsset.Nodes[SkinnedMeshIndex];
+				if (Payload.Len() > 0)
 				{
-					DuplicatedMeshNode->SetSlotMaterialDependencyUid(MaterialDependency.Key, MaterialDependency.Value);
+					Payload += ":";
 				}
+				Payload += LexToString(SkinnedMeshNode.MeshIndex | (SkinnedMeshNode.Skindex << 16));
+			}
+			SkeletalMeshNode->SetPayLoadKey(Payload);
 
-				DuplicatedMeshNode->SetSkinnedMesh(true);
+			//set the root joint node as the skeletonDependency:
+			int32 RootJointNodeIndex = RootJointToSkinnedMeshNodes.Key;
+			const GLTF::FNode& RootJointNode = GltfAsset.Nodes[RootJointNodeIndex];
+			const FString* SkeletonNodeUid = NodeUidMap.Find(&RootJointNode);
+			if (ensure(SkeletonNodeUid))
+			{
+				SkeletalMeshNode->SetSkeletonDependencyUid(*SkeletonNodeUid);
+			}
+				
 
-				//generate payload key:
-				//of template:
-				//"LexToString(SkinnedMeshNode.MeshIndex | (SkinnedMeshNode.Skindex << 16))":"LexToString(SkinnedMeshNode.MeshIndex | (SkinnedMeshNode.Skindex << 16))".....
-				FString Payload = "";
-				for (int32 SkinnedMeshIndex : RootJointToSkinnedMeshNodes.Value)
+			//set the mesh actor node's custom asset instance uid to the new duplicated mesh
+			//if there are more than one skins, then choose the topmost (root node of the collection, top most in a hierarchical tree term) occurance of SkinnedMeshIndex
+			int32 MeshActorNodeIndex = UE::Interchange::Gltf::Private::GetRootNodeIndex(GltfAsset, RootJointToSkinnedMeshNodes.Value);
+			const GLTF::FNode& MeshActorNode = GltfAsset.Nodes[MeshActorNodeIndex];
+			const FString* SceneMeshActorNodeUid = NodeUidMap.Find(&MeshActorNode);
+			if (const UInterchangeSceneNode* ConstSceneMeshActorNode = Cast< UInterchangeSceneNode >(NodeContainer.GetNode(*SceneMeshActorNodeUid)))
+			{
+				UInterchangeSceneNode* SceneMeshNode = const_cast<UInterchangeSceneNode*>(ConstSceneMeshActorNode);
+				if (ensure(SceneMeshNode))
 				{
-					const GLTF::FNode& SkinnedMeshNode = GltfAsset.Nodes[SkinnedMeshIndex];
-					if (Payload.Len() > 0)
-					{
-						Payload += ":";
-					}
-					Payload += LexToString(SkinnedMeshNode.MeshIndex | (SkinnedMeshNode.Skindex << 16));
+					SceneMeshNode->SetCustomAssetInstanceUid(SkeletalMeshNode->GetUniqueID());
 				}
-				DuplicatedMeshNode->SetPayLoadKey(Payload);
+			}
 
-				//set the root joint node as the skeletonDependency:
-				int32 RootJointNodeIndex = RootJointToSkinnedMeshNodes.Key;
-				const GLTF::FNode& RootJointNode = GltfAsset.Nodes[RootJointNodeIndex];
-				const FString* SkeletonNodeUid = NodeUidMap.Find(&RootJointNode);
-				if (ensure(SkeletonNodeUid))
-				{
-					DuplicatedMeshNode->SetSkeletonDependencyUid(*SkeletonNodeUid);
-				}
-				
+			NodeContainer.AddNode(SkeletalMeshNode);
+		}
+	}
+}
 
-				//set the mesh node's custom asset instance uid to the new duplicated mesh
-				//if there are more than one skins, then choose the topmost (root node of the collection, top most in a hierarchical tree term) occurance of SkinnedMeshIndex
-				int32 MeshNodeIndex = UE::Interchange::Gltf::Private::GetRootNodeIndex(GltfAsset, RootJointToSkinnedMeshNodes.Value);
-				const GLTF::FNode& MeshNode = GltfAsset.Nodes[MeshNodeIndex];
-				const FString* SceneMeshNodeUid = NodeUidMap.Find(&MeshNode);
-				if (const UInterchangeSceneNode* ConstSceneMeshNode = Cast< UInterchangeSceneNode >(NodeContainer.GetNode(*SceneMeshNodeUid)))
-				{
-					UInterchangeSceneNode* SceneMeshNode = const_cast<UInterchangeSceneNode*>(ConstSceneMeshNode);
-					if (ensure(SceneMeshNode))
-					{
-						SceneMeshNode->SetCustomAssetInstanceUid(DuplicatedMeshNodeUid);
-					}
-				}
+UInterchangeMeshNode* UInterchangeGltfTranslator::HandleGltfMesh(UInterchangeBaseNodeContainer& NodeContainer, 
+	const GLTF::FMesh& GltfMesh, int MeshIndex, 
+	TSet<int>& UnusedMeshIndices, 
+	FString AdditionalUniqueIdentifier/*If set it creates the mesh even if it was already creawted (for Skeletals)*/) const
+{
+	if (!UnusedMeshIndices.Contains(MeshIndex) && AdditionalUniqueIdentifier.Len() == 0)
+	{
+		return nullptr;
+	}
 
-				NodeContainer.AddNode(DuplicatedMeshNode);
-			}
+	UnusedMeshIndices.Remove(MeshIndex);
+
+	UInterchangeMeshNode* MeshNode = NewObject< UInterchangeMeshNode >(&NodeContainer);
+	FString MeshNodeUid = TEXT("\\Mesh\\") + GltfMesh.Name + AdditionalUniqueIdentifier;
+
+	MeshNode->InitializeNode(MeshNodeUid, GltfMesh.Name + AdditionalUniqueIdentifier, EInterchangeNodeContainerType::TranslatedAsset);
+	MeshNode->SetPayLoadKey(LexToString(MeshIndex));
+
+	NodeContainer.AddNode(MeshNode);
+
+	for (int32 PrimitiveCounter = 0; PrimitiveCounter < GltfMesh.Primitives.Num(); PrimitiveCounter++)
+	{
+		const GLTF::FPrimitive& Primitive = GltfMesh.Primitives[PrimitiveCounter];
+
+		// Assign materials
+		if (GltfAsset.Materials.IsValidIndex(Primitive.MaterialIndex))
+		{
+			const FString MaterialName = GltfAsset.Materials[Primitive.MaterialIndex].Name;
+			const FString ShaderGraphNodeUid = UInterchangeShaderGraphNode::MakeNodeUid(TEXT("Material_") + MaterialName);
+			MeshNode->SetSlotMaterialDependencyUid(MaterialName, ShaderGraphNodeUid);
 		}
 	}
+
+	return MeshNode;
 }
 
 #undef LOCTEXT_NAMESPACE
\ No newline at end of file
diff --git a/Engine/Plugins/Interchange/Runtime/Source/Import/Public/Gltf/InterchangeGltfTranslator.h b/Engine/Plugins/Interchange/Runtime/Source/Import/Public/Gltf/InterchangeGltfTranslator.h
--- a/Engine/Plugins/Interchange/Runtime/Source/Import/Public/Gltf/InterchangeGltfTranslator.h
+++ b/Engine/Plugins/Interchange/Runtime/Source/Import/Public/Gltf/InterchangeGltfTranslator.h
@@ -71,13 +71,16 @@ public:
 protected:
 	using FNodeUidMap = TMap<const GLTF::FNode*, FString>;
 
-	void HandleGltfSkeletons( UInterchangeBaseNodeContainer& NodeContainer, const FString& SceneNodeUid, const TArray<int32>& SkinnedMeshNodes ) const;
-	void HandleGltfNode( UInterchangeBaseNodeContainer& NodeContainer, const GLTF::FNode& GltfNode, const FString& ParentNodeUid, const int32 NodeIndex, bool &bHasVariants, TArray<int32>& SkinnedMeshNodes ) const;
+	void HandleGltfSkeletons( UInterchangeBaseNodeContainer& NodeContainer, const FString& SceneNodeUid, const TArray<int32>& SkinnedMeshNodes, TSet<int>& UnusedMeshIndices ) const;
+	void HandleGltfNode( UInterchangeBaseNodeContainer& NodeContainer, const GLTF::FNode& GltfNode, const FString& ParentNodeUid, const int32 NodeIndex, 
+		bool &bHasVariants, TArray<int32>& SkinnedMeshNodes, TSet<int>& UnusedMeshIndices ) const;
 	void HandleGltfMaterial( UInterchangeBaseNodeContainer& NodeContainer, const GLTF::FMaterial& GltfMaterial, UInterchangeShaderGraphNode& ShaderGraphNode ) const;
 	void HandleGltfMaterialParameter( UInterchangeBaseNodeContainer& NodeContainer, const GLTF::FTextureMap& TextureMap, UInterchangeShaderNode& ShaderNode,
 		const FString& MapName, const TVariant< FLinearColor, float >& MapFactor, const FString& OutputChannel, const bool bInverse = false, const bool bIsNormal = false ) const;
 	void HandleGltfAnimation(UInterchangeBaseNodeContainer& NodeContainer, int32 AnimationIndex) const;
 	void HandleGltfVariants(UInterchangeBaseNodeContainer& NodeContainer, const FString& FileName) const;
+	UInterchangeMeshNode* HandleGltfMesh(UInterchangeBaseNodeContainer& NodeContainer, const GLTF::FMesh& GltfMesh, int MeshIndex,
+		TSet<int>& UnusedMeshIndices, FString AdditionalUniqueIdentifier = "" /*If set it creates the mesh even if it was already creawted (for Skeletals)*/) const;
 
 	/** Support for KHR_materials_clearcoat */
 	void HandleGltfClearCoat( UInterchangeBaseNodeContainer& NodeContainer, const GLTF::FMaterial& GltfMaterial, UInterchangeShaderGraphNode& ShaderGraphNode, const bool bSwapNormalAndClearCoatNormal ) const;
