Interchange: Fix crash when importing obj at runtime

#rnx

diff --git a/Engine/Plugins/Interchange/Runtime/Source/Import/Private/Mesh/InterchangeOBJTranslator.cpp b/Engine/Plugins/Interchange/Runtime/Source/Import/Private/Mesh/InterchangeOBJTranslator.cpp
--- a/Engine/Plugins/Interchange/Runtime/Source/Import/Private/Mesh/InterchangeOBJTranslator.cpp
+++ b/Engine/Plugins/Interchange/Runtime/Source/Import/Private/Mesh/InterchangeOBJTranslator.cpp
@@ -273,7 +273,7 @@ FMeshDescription FObjData::MakeMeshDescriptionForGroup(const FString& GroupName)
 	}
 
 	// Create UVs and initialize values
-
+	MeshDescription.SetNumUVChannels(1);
 	if (UVIndexMapping.Num() > 0)
 	{
 		const int32 UVChannel = 0;
