Fixed warning when deleting assets [level sequences] imported with interchange

diff --git a/Engine/Plugins/Interchange/Runtime/Source/Import/Private/Animation/InterchangeAnimationTrackSetFactory.cpp b/Engine/Plugins/Interchange/Runtime/Source/Import/Private/Animation/InterchangeAnimationTrackSetFactory.cpp
--- a/Engine/Plugins/Interchange/Runtime/Source/Import/Private/Animation/InterchangeAnimationTrackSetFactory.cpp
+++ b/Engine/Plugins/Interchange/Runtime/Source/Import/Private/Animation/InterchangeAnimationTrackSetFactory.cpp
@@ -714,4 +714,19 @@ void UInterchangeAnimationTrackSetFactory::PreImportPreCompletedCallback(const F
 	// TODO: Talk with sequence team about adding AssetImportData to ULevelSequence for re-import 
 }
 
+void UInterchangeAnimationTrackSetFactory::PostImportPreCompletedCallback(const FImportPreCompletedCallbackParams& Arguments)
+{
+	check(IsInGameThread());
+	if (ULevelSequence* LevelSequence = Cast<ULevelSequence>(Arguments.ImportedObject))
+	{
+		// Ugly temporary workaround to Sequencer's design which has an implicit deferring system which
+		// holds onto UMovieSceneSignedObject objects when PreEditChange is called on such object (i.e. ULevelSequence)
+		// Only a call to open the Sequencer, which calls UMovieSceneSignedObject::ResetImplicitScopedModifyDefer, will release the hold
+		// See FDeferredSignedObjectChangeHandler, FSequencer::Tick
+		// Note: Sequencer team is aware of the issue and will look their design in a future release
+		UMovieSceneSignedObject::ResetImplicitScopedModifyDefer();
+	}
+	Super::PostImportPreCompletedCallback(Arguments);
+}
+
 #undef LOCTEXT_NAMESPACE
\ No newline at end of file
diff --git a/Engine/Plugins/Interchange/Runtime/Source/Import/Public/Animation/InterchangeAnimationTrackSetFactory.h b/Engine/Plugins/Interchange/Runtime/Source/Import/Public/Animation/InterchangeAnimationTrackSetFactory.h
--- a/Engine/Plugins/Interchange/Runtime/Source/Import/Public/Animation/InterchangeAnimationTrackSetFactory.h
+++ b/Engine/Plugins/Interchange/Runtime/Source/Import/Public/Animation/InterchangeAnimationTrackSetFactory.h
@@ -22,6 +22,7 @@ public:
 	virtual UObject* CreateEmptyAsset(const FCreateAssetParams& Arguments) override;
 	virtual UObject* CreateAsset(const FCreateAssetParams& Arguments) override;
 	virtual void PreImportPreCompletedCallback(const FImportPreCompletedCallbackParams& Arguments) override;
+	virtual void PostImportPreCompletedCallback(const FImportPreCompletedCallbackParams& Arguments) override;
 	virtual bool CanExecuteOnAnyThread() const override
 	{
 		//Currently we cannot use the anim sequence controller outside of the game thread
