GLTF import naming convention changes

diff --git a/Engine/Plugins/Editor/GLTFImporter/Source/GLTFCore/Private/GLTF/GLTFAsset.cpp b/Engine/Plugins/Editor/GLTFImporter/Source/GLTFCore/Private/GLTF/GLTFAsset.cpp
--- a/Engine/Plugins/Editor/GLTFImporter/Source/GLTFCore/Private/GLTF/GLTFAsset.cpp
+++ b/Engine/Plugins/Editor/GLTFImporter/Source/GLTFCore/Private/GLTF/GLTFAsset.cpp
@@ -17,10 +17,8 @@ namespace GLTF
 
 				if (Obj.Name.IsEmpty())
 				{
-					Obj.Name = Prefix;
+					Obj.Name = Prefix + FString::FromInt(Index);
 				}
-
-				Obj.Name = FString::FromInt(Index) + TEXT("_") + Obj.Name;
 			}
 		}
 	}
@@ -84,10 +82,10 @@ namespace GLTF
 				{
 					bool bIsJoint = Node.Type == FNode::EType::Joint;
 					Node.Name     = (bIsJoint ? JointPrefix : NodePrefix);
-				}
 
-				// Make sure node names are unique
-				Node.Name = FString::FromInt(NodeIndex) + TEXT("_") + Node.Name;
+					// Make sure node names are unique, in case Name is not empty it is presumed to be unique / handled by the processing pipelines.
+					Node.Name = Node.Name + FString::FromInt(NodeIndex);
+				}
 			}
 		}
 
@@ -111,13 +109,12 @@ namespace GLTF
 				}
 				else
 				{
-					Tex.Name = TexPrefix;
+					// GLTF texture name has decorative purpose, not guaranteed to be unique
+					// only its index is unique. Same with glTF image or its source file's basename
+					// So always include texture index into texture's name to increase probability that names are unique
+					Tex.Name = TexPrefix + FString::FromInt(TextureIndex);
 				}
 
-				// GLTF texture name has decorative purpose, not guaranteed to be unique
-				// only its index is unique. Same with glTF image or its source file's basename
-				// So always include texture index into texture's name to increase probability that names are unique
-				Tex.Name = FString::FromInt(TextureIndex) + TEXT("_") + Tex.Name;
 			}
 		}
 
@@ -135,11 +132,9 @@ namespace GLTF
 					}
 					else
 					{
-						Camera.Name = CameraPrefix;
+						Camera.Name = CameraPrefix + FString::FromInt(CameraIndex);
 					}
 				}
-
-				Camera.Name = FString::FromInt(CameraIndex) + TEXT("_") + Camera.Name;
 			}
 		}
 
@@ -157,11 +152,9 @@ namespace GLTF
 					}
 					else
 					{
-						Light.Name = LightPrefix;
+						Light.Name = LightPrefix + FString::FromInt(LightIndex);
 					}
 				}
-
-				Light.Name = FString::FromInt(LightIndex) + TEXT("_") + Light.Name;
 			}
 		}
 
@@ -171,11 +164,11 @@ namespace GLTF
 		// "0_MyFile_material" or "3_myfile_mesh".
 		// We want to do images after textures as the textures may want to reuse source image names, but we don't
 		// want to use those if we generated them ourselves.
-		GLTF::GenerateNames(Prefix + TEXT("_material"), Materials);
-		GLTF::GenerateNames(Prefix + TEXT("_skin"), Skins);
-		GLTF::GenerateNames(Prefix + TEXT("_animation"), Animations);
-		GLTF::GenerateNames(Prefix + TEXT("_image"), Images);
-		GLTF::GenerateNames(Prefix + TEXT("_mesh"), Meshes);
+		GLTF::GenerateNames(Prefix + TEXT("_material_"), Materials);
+		GLTF::GenerateNames(Prefix + TEXT("_skin_"), Skins);
+		GLTF::GenerateNames(Prefix + TEXT("_animation_"), Animations);
+		GLTF::GenerateNames(Prefix + TEXT("_image_"), Images);
+		GLTF::GenerateNames(Prefix + TEXT("_mesh_"), Meshes);
 	}
 
 	void FAsset::GetRootNodes(TArray<int32>& NodeIndices)
diff --git a/Engine/Plugins/Editor/GLTFImporter/Source/GLTFCore/Private/GLTF/GLTFReader.cpp b/Engine/Plugins/Editor/GLTFImporter/Source/GLTFCore/Private/GLTF/GLTFReader.cpp
--- a/Engine/Plugins/Editor/GLTFImporter/Source/GLTFCore/Private/GLTF/GLTFReader.cpp
+++ b/Engine/Plugins/Editor/GLTFImporter/Source/GLTFCore/Private/GLTF/GLTFReader.cpp
@@ -877,9 +877,6 @@ namespace GLTF
 		SetupObjects(TextureCount, TEXT("textures"), [this](const FJsonObject& Object) { SetupTexture(Object); });
 		SetupObjects(MaterialCount, TEXT("materials"), [this](const FJsonObject& Object) { SetupMaterial(Object); });
 
-		SetupObjects(MaterialCount, TEXT("materials"), [this](const FJsonObject& Object) { SetupMaterial(Object); });
-
-
 		/**
 		 * Returns true if all required extensions are supported.
 		 * According to gltf specification checking only the top-level extensionsRequired property in order to decide
diff --git a/Engine/Plugins/Interchange/Runtime/Source/Import/Private/Gltf/InterchangeGltfTranslator.cpp b/Engine/Plugins/Interchange/Runtime/Source/Import/Private/Gltf/InterchangeGltfTranslator.cpp
--- a/Engine/Plugins/Interchange/Runtime/Source/Import/Private/Gltf/InterchangeGltfTranslator.cpp
+++ b/Engine/Plugins/Interchange/Runtime/Source/Import/Private/Gltf/InterchangeGltfTranslator.cpp
@@ -879,9 +879,7 @@ bool UInterchangeGltfTranslator::Translate( UInterchangeBaseNodeContainer& NodeC
 		int32 TextureIndex = 0;
 		for ( const GLTF::FTexture& GltfTexture : GltfAsset.Textures )
 		{
-			const FString TextureName = GltfTexture.Name;
-
-			UInterchangeTexture2DNode* TextureNode = UInterchangeTexture2DNode::Create( &NodeContainer, TextureName );
+			UInterchangeTexture2DNode* TextureNode = UInterchangeTexture2DNode::Create( &NodeContainer, GltfTexture.Name);
 			TextureNode->SetCustomFilter(ConvertFilter(GltfTexture.Sampler.MinFilter));
 			TextureNode->SetPayLoadKey( LexToString( TextureIndex++ ) );
 
@@ -895,7 +893,7 @@ bool UInterchangeGltfTranslator::Translate( UInterchangeBaseNodeContainer& NodeC
 		int32 MaterialIndex = 0;
 		for ( const GLTF::FMaterial& GltfMaterial : GltfAsset.Materials )
 		{
-			UInterchangeShaderGraphNode* ShaderGraphNode = UInterchangeShaderGraphNode::Create( &NodeContainer, TEXT("Material_") + GltfMaterial.Name );
+			UInterchangeShaderGraphNode* ShaderGraphNode = UInterchangeShaderGraphNode::Create( &NodeContainer, GltfMaterial.Name );
 
 			HandleGltfMaterial( NodeContainer, GltfMaterial, *ShaderGraphNode );
 			++MaterialIndex;
@@ -1017,7 +1015,7 @@ bool UInterchangeGltfTranslator::Translate( UInterchangeBaseNodeContainer& NodeC
 			FString SceneName = GltfScene.Name;
 			if (SceneName.IsEmpty())
 			{
-				SceneName = LexToString(SceneIndex++) + TEXT("_scene_") + FileName;
+				SceneName = FileName + TEXT("_scene_") + LexToString(SceneIndex++);
 			}
 
 			FString SceneNodeUid = TEXT("\\Scene\\") + SceneName;
@@ -1235,9 +1233,7 @@ void UInterchangeGltfTranslator::HandleGltfAnimation(UInterchangeBaseNodeContain
 				{
 					TrackNode = NewObject< UInterchangeSkeletalAnimationTrackNode >(&NodeContainer);
 					FString TrackNodeUid = "\\SkeletalAnimation\\" + LexToString(AnimatedNode->RootJointIndex) + "_" + LexToString(AnimationIndex);
-					// AnimationIndex is already part of the GltfAnimation.Name (due to GenerateNames)
-					FString DisplayString = GltfAnimation.Name + "_" + LexToString(AnimatedNode->RootJointIndex);
-					TrackNode->InitializeNode(TrackNodeUid, DisplayString, EInterchangeNodeContainerType::TranslatedAsset);
+					TrackNode->InitializeNode(TrackNodeUid, GltfAnimation.Name, EInterchangeNodeContainerType::TranslatedAsset);
 					const GLTF::FNode& RootJointNode = GltfAsset.Nodes[AnimatedNode->RootJointIndex];
 					const FString* SkeletonNodeUid = NodeUidMap.Find(&RootJointNode);
 					TrackNode->SetCustomSkeletonNodeUid(*SkeletonNodeUid);
@@ -1552,7 +1548,7 @@ void UInterchangeGltfTranslator::HandleGltfVariants(UInterchangeBaseNodeContaine
 								}
 
 								const GLTF::FMaterial& GltfMaterial = Materials[VariantMapping.MaterialIndex];
-								const FString MaterialUid = UInterchangeShaderGraphNode::MakeNodeUid(TEXT("Material_") + GltfMaterial.Name);
+								const FString MaterialUid = UInterchangeShaderGraphNode::MakeNodeUid(GltfMaterial.Name);
 
 								VariantSetNode->AddCustomDependencyUid(MaterialUid);
 
@@ -1633,7 +1629,7 @@ bool UInterchangeGltfTranslator::GetVariantSetPayloadData(UE::Interchange::FVari
 						}
 
 						const GLTF::FMaterial& GltfMaterial = Materials[VariantMapping.MaterialIndex];
-						const FString MaterialUid = UInterchangeShaderGraphNode::MakeNodeUid(TEXT("Material_") + GltfMaterial.Name);
+						const FString MaterialUid = UInterchangeShaderGraphNode::MakeNodeUid(GltfMaterial.Name);
 
 						for (int32 VariantIndex : VariantMapping.VariantIndices)
 						{
@@ -1758,7 +1754,11 @@ void UInterchangeGltfTranslator::HandleGltfSkeletons(UInterchangeBaseNodeContain
 		{
 			//Duplicate MeshNode for each group:
 			int32 RootJointIndex = RootJointToSkinnedMeshNodes.Key;
-			UInterchangeMeshNode* SkeletalMeshNode = HandleGltfMesh(NodeContainer, GltfAsset.Meshes[MeshIndex], MeshIndex, UnusedMeshIndices, TEXT("_") + LexToString(RootJointIndex));
+
+			//Skeletal Mesh's naming policy: (Mesh.Name)_(RootJointNode.Name) naming policy:
+			FString SkeletalName = GltfAsset.Meshes[MeshIndex].Name + TEXT("_") + GltfAsset.Nodes[RootJointIndex].Name;
+
+			UInterchangeMeshNode* SkeletalMeshNode = HandleGltfMesh(NodeContainer, GltfAsset.Meshes[MeshIndex], MeshIndex, UnusedMeshIndices, SkeletalName);
 
 			SkeletalMeshNode->SetSkinnedMesh(true);
 
@@ -1809,19 +1809,19 @@ void UInterchangeGltfTranslator::HandleGltfSkeletons(UInterchangeBaseNodeContain
 UInterchangeMeshNode* UInterchangeGltfTranslator::HandleGltfMesh(UInterchangeBaseNodeContainer& NodeContainer, 
 	const GLTF::FMesh& GltfMesh, int MeshIndex, 
 	TSet<int>& UnusedMeshIndices, 
-	FString AdditionalUniqueIdentifier/*If set it creates the mesh even if it was already creawted (for Skeletals)*/) const
+	const FString& SkeletalName/*If set it creates the mesh even if it was already created (for Skeletals)*/) const
 {
-	if (!UnusedMeshIndices.Contains(MeshIndex) && AdditionalUniqueIdentifier.Len() == 0)
+	if (!UnusedMeshIndices.Contains(MeshIndex) && SkeletalName.Len() == 0)
 	{
 		return nullptr;
 	}
-
+	FString MeshName = SkeletalName.Len() ? SkeletalName : GltfMesh.Name;
 	UnusedMeshIndices.Remove(MeshIndex);
 
 	UInterchangeMeshNode* MeshNode = NewObject< UInterchangeMeshNode >(&NodeContainer);
-	FString MeshNodeUid = TEXT("\\Mesh\\") + GltfMesh.Name + AdditionalUniqueIdentifier;
+	FString MeshNodeUid = TEXT("\\Mesh\\") + MeshName;
 
-	MeshNode->InitializeNode(MeshNodeUid, GltfMesh.Name + AdditionalUniqueIdentifier, EInterchangeNodeContainerType::TranslatedAsset);
+	MeshNode->InitializeNode(MeshNodeUid, MeshName, EInterchangeNodeContainerType::TranslatedAsset);
 	MeshNode->SetPayLoadKey(LexToString(MeshIndex));
 
 	NodeContainer.AddNode(MeshNode);
@@ -1834,7 +1834,7 @@ UInterchangeMeshNode* UInterchangeGltfTranslator::HandleGltfMesh(UInterchangeBas
 		if (GltfAsset.Materials.IsValidIndex(Primitive.MaterialIndex))
 		{
 			const FString MaterialName = GltfAsset.Materials[Primitive.MaterialIndex].Name;
-			const FString ShaderGraphNodeUid = UInterchangeShaderGraphNode::MakeNodeUid(TEXT("Material_") + MaterialName);
+			const FString ShaderGraphNodeUid = UInterchangeShaderGraphNode::MakeNodeUid(MaterialName);
 			MeshNode->SetSlotMaterialDependencyUid(MaterialName, ShaderGraphNodeUid);
 		}
 	}
diff --git a/Engine/Plugins/Interchange/Runtime/Source/Import/Public/Gltf/InterchangeGltfTranslator.h b/Engine/Plugins/Interchange/Runtime/Source/Import/Public/Gltf/InterchangeGltfTranslator.h
--- a/Engine/Plugins/Interchange/Runtime/Source/Import/Public/Gltf/InterchangeGltfTranslator.h
+++ b/Engine/Plugins/Interchange/Runtime/Source/Import/Public/Gltf/InterchangeGltfTranslator.h
@@ -80,7 +80,7 @@ protected:
 	void HandleGltfAnimation(UInterchangeBaseNodeContainer& NodeContainer, int32 AnimationIndex) const;
 	void HandleGltfVariants(UInterchangeBaseNodeContainer& NodeContainer, const FString& FileName) const;
 	UInterchangeMeshNode* HandleGltfMesh(UInterchangeBaseNodeContainer& NodeContainer, const GLTF::FMesh& GltfMesh, int MeshIndex,
-		TSet<int>& UnusedMeshIndices, FString AdditionalUniqueIdentifier = "" /*If set it creates the mesh even if it was already creawted (for Skeletals)*/) const;
+		TSet<int>& UnusedMeshIndices, const FString& SkeletalName = "" /*If set it creates the mesh even if it was already created (for Skeletals)*/) const;
 
 	/** Support for KHR_materials_clearcoat */
 	void HandleGltfClearCoat( UInterchangeBaseNodeContainer& NodeContainer, const GLTF::FMaterial& GltfMaterial, UInterchangeShaderGraphNode& ShaderGraphNode, const bool bSwapNormalAndClearCoatNormal ) const;
