Interchange: Add the UseHighPrecisionSkinWeights to interchange pipeline.
#rnx

diff --git a/Engine/Plugins/Interchange/Runtime/Source/FactoryNodes/Private/InterchangeSkeletalMeshFactoryNode.cpp b/Engine/Plugins/Interchange/Runtime/Source/FactoryNodes/Private/InterchangeSkeletalMeshFactoryNode.cpp
--- a/Engine/Plugins/Interchange/Runtime/Source/FactoryNodes/Private/InterchangeSkeletalMeshFactoryNode.cpp
+++ b/Engine/Plugins/Interchange/Runtime/Source/FactoryNodes/Private/InterchangeSkeletalMeshFactoryNode.cpp
@@ -110,6 +110,23 @@ bool UInterchangeSkeletalMeshFactoryNode::SetCustomPhysicAssetSoftObjectPath(con
 	IMPLEMENT_NODE_ATTRIBUTE_SETTER_NODELEGATE(PhysicAssetSoftObjectPath, FSoftObjectPath)
 }
 
+bool UInterchangeSkeletalMeshFactoryNode::GetCustomUseHighPrecisionSkinWeights(bool& AttributeValue) const
+{
+	IMPLEMENT_NODE_ATTRIBUTE_GETTER(UseHighPrecisionSkinWeights, bool)
+}
+bool UInterchangeSkeletalMeshFactoryNode::SetCustomUseHighPrecisionSkinWeights(const bool& AttributeValue, bool bAddApplyDelegate)
+{
+	IMPLEMENT_NODE_ATTRIBUTE_SETTER_WITH_CUSTOM_DELEGATE(UInterchangeSkeletalMeshFactoryNode, UseHighPrecisionSkinWeights, bool);
+}
+bool UInterchangeSkeletalMeshFactoryNode::ApplyCustomUseHighPrecisionSkinWeightsToAsset(UObject* Asset) const
+{
+	IMPLEMENT_SKELETAL_BUILD_VALUE_TO_ASSET(UseHighPrecisionSkinWeights, bool, bUseHighPrecisionSkinWeights);
+}
+bool UInterchangeSkeletalMeshFactoryNode::FillCustomUseHighPrecisionSkinWeightsFromAsset(UObject* Asset)
+{
+	IMPLEMENT_SKELETALMESH_BUILD_ASSET_TO_VALUE(UseHighPrecisionSkinWeights, bUseHighPrecisionSkinWeights);
+}
+
 bool UInterchangeSkeletalMeshFactoryNode::GetCustomThresholdPosition(float& AttributeValue) const
 {
 	IMPLEMENT_NODE_ATTRIBUTE_GETTER(ThresholdPosition, float)
diff --git a/Engine/Plugins/Interchange/Runtime/Source/FactoryNodes/Public/InterchangeSkeletalMeshFactoryNode.h b/Engine/Plugins/Interchange/Runtime/Source/FactoryNodes/Public/InterchangeSkeletalMeshFactoryNode.h
--- a/Engine/Plugins/Interchange/Runtime/Source/FactoryNodes/Public/InterchangeSkeletalMeshFactoryNode.h
+++ b/Engine/Plugins/Interchange/Runtime/Source/FactoryNodes/Public/InterchangeSkeletalMeshFactoryNode.h
@@ -88,6 +88,14 @@ public:
 	UFUNCTION(BlueprintCallable, Category = "Interchange | Node | SkeletalMesh")
 	bool SetCustomImportContentType(const EInterchangeSkeletalMeshContentType& AttributeValue);
 
+	/** Query the skeletal mesh UseHighPrecisionSkinWeights. */
+	UFUNCTION(BlueprintCallable, Category = "Interchange | Node | SkeletalMesh")
+	bool GetCustomUseHighPrecisionSkinWeights(bool& AttributeValue) const;
+
+	/** Set the skeletal mesh UseHighPrecisionSkinWeights. */
+	UFUNCTION(BlueprintCallable, Category = "Interchange | Node | SkeletalMesh")
+	bool SetCustomUseHighPrecisionSkinWeights(const bool& AttributeValue, bool bAddApplyDelegate = true);
+		
 	/** Query the skeletal mesh threshold use to decide if two vertex position are equal. */
 	UFUNCTION(BlueprintCallable, Category = "Interchange | Node | SkeletalMesh")
 	bool GetCustomThresholdPosition(float& AttributeValue) const;
@@ -135,11 +143,14 @@ private:
 	const UE::Interchange::FAttributeKey Macro_CustomCreatePhysicsAssetKey = UE::Interchange::FAttributeKey(TEXT("CreatePhysicsAsset"));
 	const UE::Interchange::FAttributeKey Macro_CustomPhysicAssetSoftObjectPathKey = UE::Interchange::FAttributeKey(TEXT("PhysicAssetSoftObjectPath"));
 	const UE::Interchange::FAttributeKey Macro_CustomImportContentTypeKey = UE::Interchange::FAttributeKey(TEXT("ImportContentType"));
+	const UE::Interchange::FAttributeKey Macro_CustomUseHighPrecisionSkinWeightsKey = UE::Interchange::FAttributeKey(TEXT("UseHighPrecisionSkinWeights"));
 	const UE::Interchange::FAttributeKey Macro_CustomThresholdPositionKey = UE::Interchange::FAttributeKey(TEXT("ThresholdPosition"));
 	const UE::Interchange::FAttributeKey Macro_CustomThresholdTangentNormalKey = UE::Interchange::FAttributeKey(TEXT("ThresholdTangentNormal"));
 	const UE::Interchange::FAttributeKey Macro_CustomThresholdUVKey = UE::Interchange::FAttributeKey(TEXT("ThresholdUV"));
 	const UE::Interchange::FAttributeKey Macro_CustomMorphThresholdPositionKey = UE::Interchange::FAttributeKey(TEXT("MorphThresholdPosition"));
 
+	bool ApplyCustomUseHighPrecisionSkinWeightsToAsset(UObject* Asset) const;
+	bool FillCustomUseHighPrecisionSkinWeightsFromAsset(UObject* Asset);
 	bool ApplyCustomThresholdPositionToAsset(UObject* Asset) const;
 	bool FillCustomThresholdPositionFromAsset(UObject* Asset);
 	bool ApplyCustomThresholdTangentNormalToAsset(UObject* Asset) const;
diff --git a/Engine/Plugins/Interchange/Runtime/Source/Pipelines/Private/InterchangeGenericSkeletalMeshPipeline.cpp b/Engine/Plugins/Interchange/Runtime/Source/Pipelines/Private/InterchangeGenericSkeletalMeshPipeline.cpp
--- a/Engine/Plugins/Interchange/Runtime/Source/Pipelines/Private/InterchangeGenericSkeletalMeshPipeline.cpp
+++ b/Engine/Plugins/Interchange/Runtime/Source/Pipelines/Private/InterchangeGenericSkeletalMeshPipeline.cpp
@@ -504,6 +504,7 @@ UInterchangeSkeletalMeshFactoryNode* UInterchangeGenericMeshPipeline::CreateSkel
 	SkeletalMeshFactoryNode->SetCustomUseBackwardsCompatibleF16TruncUVs(CommonMeshesProperties->bUseBackwardsCompatibleF16TruncUVs);
 	SkeletalMeshFactoryNode->SetCustomRemoveDegenerates(CommonMeshesProperties->bRemoveDegenerates);
 	//Skeletal meshes build options
+	SkeletalMeshFactoryNode->SetCustomUseHighPrecisionSkinWeights(bUseHighPrecisionSkinWeights);
 	SkeletalMeshFactoryNode->SetCustomThresholdPosition(ThresholdPosition);
 	SkeletalMeshFactoryNode->SetCustomThresholdTangentNormal(ThresholdTangentNormal);
 	SkeletalMeshFactoryNode->SetCustomThresholdUV(ThresholdUV);
diff --git a/Engine/Plugins/Interchange/Runtime/Source/Pipelines/Public/InterchangeGenericMeshPipeline.h b/Engine/Plugins/Interchange/Runtime/Source/Pipelines/Public/InterchangeGenericMeshPipeline.h
--- a/Engine/Plugins/Interchange/Runtime/Source/Pipelines/Public/InterchangeGenericMeshPipeline.h
+++ b/Engine/Plugins/Interchange/Runtime/Source/Pipelines/Public/InterchangeGenericMeshPipeline.h
@@ -155,6 +155,10 @@ public:
 	UPROPERTY(EditAnywhere, BlueprintReadWrite, Category = "Skeletal Meshes", meta = (editcondition = "!bCreatePhysicsAsset"))
 	TWeakObjectPtr<UPhysicsAsset> PhysicsAsset;
 
+	/** If Checked, use 16 bits for skin weights instead of 8 bits. */
+	UPROPERTY(EditAnywhere, BlueprintReadWrite, Category = "Skeletal Meshes", meta = (SubCategory = "Build"))
+	bool bUseHighPrecisionSkinWeights = false;
+
 	/** Threshold use to decide if two vertex position are equal. */
 	UPROPERTY(EditAnywhere, BlueprintReadWrite, Category = "Skeletal Meshes", meta = (SubCategory = "Build"))
 	float ThresholdPosition = 0.00002f;
