Ensure that a UMaterialExpressionFunctionCall and UMaterialFunction are updated only once during import, we prevent a possible race condition when several material files are imported all at once, for example when importing the standard_surface_chess_set.mtlx that contains several materials in it

diff --git a/Engine/Plugins/Interchange/Runtime/Source/Import/Private/InterchangeImportModule.cpp b/Engine/Plugins/Interchange/Runtime/Source/Import/Private/InterchangeImportModule.cpp
--- a/Engine/Plugins/Interchange/Runtime/Source/Import/Private/InterchangeImportModule.cpp
+++ b/Engine/Plugins/Interchange/Runtime/Source/Import/Private/InterchangeImportModule.cpp
@@ -47,6 +47,8 @@ IMPLEMENT_MODULE(FInterchangeImportModule, InterchangeImport)
 
 void FInterchangeImportModule::StartupModule()
 {
+	FInterchangeImportMaterialAsyncHelper& InterchangeMaterialAsyncHelper = FInterchangeImportMaterialAsyncHelper::GetInstance();
+
 	auto RegisterItems = []()
 	{
 		UInterchangeManager& InterchangeManager = UInterchangeManager::GetInterchangeManager();
@@ -99,7 +101,8 @@ void FInterchangeImportModule::StartupModule()
 
 void FInterchangeImportModule::ShutdownModule()
 {
-
+	FInterchangeImportMaterialAsyncHelper& InterchangeMaterialAsyncHelper = FInterchangeImportMaterialAsyncHelper::GetInstance();
+	InterchangeMaterialAsyncHelper.CleanUp();
 }
 
 
diff --git a/Engine/Plugins/Interchange/Runtime/Source/Import/Private/Material/InterchangeMaterialFactory.cpp b/Engine/Plugins/Interchange/Runtime/Source/Import/Private/Material/InterchangeMaterialFactory.cpp
--- a/Engine/Plugins/Interchange/Runtime/Source/Import/Private/Material/InterchangeMaterialFactory.cpp
+++ b/Engine/Plugins/Interchange/Runtime/Source/Import/Private/Material/InterchangeMaterialFactory.cpp
@@ -19,7 +19,6 @@
 #include "Materials/MaterialExpressionTextureSample.h"
 #include "Materials/MaterialExpressionThinTranslucentMaterialOutput.h"
 #include "Materials/MaterialExpressionVectorParameter.h"
-#include "Materials/MaterialFunction.h"
 #include "Materials/MaterialInterface.h"
 #include "Materials/MaterialInstanceConstant.h"
 #include "Materials/MaterialInstanceDynamic.h"
@@ -116,6 +115,7 @@ namespace UE
 
 				void SetupFunctionCallExpression(const UInterchangeFactoryBase::FCreateAssetParams& Arguments, const UInterchangeMaterialExpressionFactoryNode& ExpressionNode, UMaterialExpressionMaterialFunctionCall* FunctionCallExpression)
 				{
+					FInterchangeImportMaterialAsyncHelper& AsyncHelper = FInterchangeImportMaterialAsyncHelper::GetInstance();
 					if (const UInterchangeMaterialFunctionCallExpressionFactoryNode* FunctionCallFactoryNode = Cast<UInterchangeMaterialFunctionCallExpressionFactoryNode>(&ExpressionNode))
 					{
 						FString MaterialFactoryNodeUid;
@@ -127,13 +127,13 @@ namespace UE
 							MaterialFactoryNode->GetCustomReferenceObject(ReferenceObject);
 							if (UMaterialFunctionInterface* MaterialFunction = Cast<UMaterialFunctionInterface>(ReferenceObject.TryLoad()))
 							{
-								MaterialFunction->UpdateFromFunctionResource();
+								AsyncHelper.UpdateFromFunctionResource(MaterialFunction);
 								FunctionCallExpression->SetMaterialFunction(MaterialFunction);
 							}
 						}
 					}
 
-					FunctionCallExpression->UpdateFromFunctionResource();
+					AsyncHelper.UpdateFromFunctionResource(FunctionCallExpression);
 				}
 
 				void SetupTextureExpression(const UInterchangeFactoryBase::FCreateAssetParams& Arguments, const UInterchangeMaterialExpressionFactoryNode* ExpressionNode, UMaterialExpressionTextureBase* TextureExpression)
@@ -1277,6 +1277,38 @@ void UInterchangeMaterialFunctionFactory::SetupMaterial(UMaterialFunction* Mater
 	MaterialFunction->UpdateDependentFunctionCandidates();
 }
 
+FInterchangeImportMaterialAsyncHelper& FInterchangeImportMaterialAsyncHelper::GetInstance()
+{
+	static FInterchangeImportMaterialAsyncHelper Instance;
+	return Instance;
+}
+
+void FInterchangeImportMaterialAsyncHelper::CleanUp()
+{
+	UpdatedMaterialFunctionCalls.Empty();
+	UpdatedMaterialFunctions.Empty();
+}
+
+void FInterchangeImportMaterialAsyncHelper::UpdateFromFunctionResource(UMaterialExpressionMaterialFunctionCall* MaterialFunctionCall)
+{
+	FScopeLock Lock(&UpdatedMaterialFunctionCallsLock);
+	if(UpdatedMaterialFunctionCalls.Find(MaterialFunctionCall) == nullptr)
+	{
+		UpdatedMaterialFunctionCalls.Add(MaterialFunctionCall);
+		MaterialFunctionCall->UpdateFromFunctionResource();
+	}
+}
+
+void FInterchangeImportMaterialAsyncHelper::UpdateFromFunctionResource(UMaterialFunctionInterface* MaterialFunction)
+{
+	FScopeLock Lock(&UpdatedMaterialFunctionsLock);
+	if(UpdatedMaterialFunctions.Find(MaterialFunction) == nullptr)
+	{
+		UpdatedMaterialFunctions.Add(MaterialFunction);
+		MaterialFunction->UpdateFromFunctionResource();
+	}
+}
+
 #endif // #if WITH_EDITOR
 
 #undef LOCTEXT_NAMESPACE
diff --git a/Engine/Plugins/Interchange/Runtime/Source/Import/Public/Material/InterchangeMaterialFactory.h b/Engine/Plugins/Interchange/Runtime/Source/Import/Public/Material/InterchangeMaterialFactory.h
--- a/Engine/Plugins/Interchange/Runtime/Source/Import/Public/Material/InterchangeMaterialFactory.h
+++ b/Engine/Plugins/Interchange/Runtime/Source/Import/Public/Material/InterchangeMaterialFactory.h
@@ -6,6 +6,7 @@
 #include "InterchangeFactoryBase.h"
 #include "UObject/Object.h"
 #include "UObject/ObjectMacros.h"
+#include "Materials/MaterialFunction.h"
 
 #include "InterchangeMaterialFactory.generated.h"
 
@@ -64,4 +65,33 @@ private:
 #endif
 };
 
+/**
+ * This class intends to avoid a race condition when importing 
+ * several material functions at once, for example when importing the MaterialX file standard_surface_chess_set
+ * Only one instance of a material function can update it
+ */
+class FInterchangeImportMaterialAsyncHelper
+{
+	FInterchangeImportMaterialAsyncHelper() = default;
+
+public:
+
+	FInterchangeImportMaterialAsyncHelper(const FInterchangeImportMaterialAsyncHelper&) = delete;
+	FInterchangeImportMaterialAsyncHelper& operator=(const FInterchangeImportMaterialAsyncHelper&) = delete;
+
+	static FInterchangeImportMaterialAsyncHelper& GetInstance();
+
+	void CleanUp();
+
+	void UpdateFromFunctionResource(UMaterialExpressionMaterialFunctionCall* MaterialFunctionCall);
+
+	void UpdateFromFunctionResource(UMaterialFunctionInterface* MaterialFunction);
+
+private:
+
+	FCriticalSection UpdatedMaterialFunctionCallsLock;
+	TSet<UMaterialExpressionMaterialFunctionCall*> UpdatedMaterialFunctionCalls;
 
+	FCriticalSection UpdatedMaterialFunctionsLock;
+	TSet<UMaterialFunctionInterface*> UpdatedMaterialFunctions;
+};
\ No newline at end of file
