compile fix for deprecation warnings of TObjectPtr.IsNull

diff --git a/Engine/Plugins/Animation/IKRig/Source/IKRigEditor/Private/RetargetEditor/IKRetargetEditPoseMode.cpp b/Engine/Plugins/Animation/IKRig/Source/IKRigEditor/Private/RetargetEditor/IKRetargetEditPoseMode.cpp
--- a/Engine/Plugins/Animation/IKRig/Source/IKRigEditor/Private/RetargetEditor/IKRetargetEditPoseMode.cpp
+++ b/Engine/Plugins/Animation/IKRig/Source/IKRigEditor/Private/RetargetEditor/IKRetargetEditPoseMode.cpp
@@ -255,7 +255,7 @@ bool FIKRetargetEditPoseMode::InputDelta(
 	}
 
 	const TObjectPtr<UIKRetargeterController> AssetController = Controller->AssetController;
-	if (AssetController.IsNull())
+	if (!AssetController)
 	{
 		return false; 
 	}
diff --git a/Engine/Plugins/Animation/IKRig/Source/IKRigEditor/Private/RetargetEditor/IKRetargeterController.cpp b/Engine/Plugins/Animation/IKRig/Source/IKRigEditor/Private/RetargetEditor/IKRetargeterController.cpp
--- a/Engine/Plugins/Animation/IKRig/Source/IKRigEditor/Private/RetargetEditor/IKRetargeterController.cpp
+++ b/Engine/Plugins/Animation/IKRig/Source/IKRigEditor/Private/RetargetEditor/IKRetargeterController.cpp
@@ -375,7 +375,7 @@ const URetargetChainSettings* UIKRetargeterController::GetChainMappingByTargetCh
 
 FName UIKRetargeterController::GetChainGoal(const TObjectPtr<URetargetChainSettings> ChainSettings) const
 {
-	if (ChainSettings.IsNull())
+	if (!ChainSettings)
 	{
 		return NAME_None;
 	}
diff --git a/Engine/Plugins/Experimental/Animation/PoseSearch/Source/Editor/Private/SPoseSearchDatabaseAssetList.cpp b/Engine/Plugins/Experimental/Animation/PoseSearch/Source/Editor/Private/SPoseSearchDatabaseAssetList.cpp
--- a/Engine/Plugins/Experimental/Animation/PoseSearch/Source/Editor/Private/SPoseSearchDatabaseAssetList.cpp
+++ b/Engine/Plugins/Experimental/Animation/PoseSearch/Source/Editor/Private/SPoseSearchDatabaseAssetList.cpp
@@ -558,12 +558,12 @@ namespace UE::PoseSearch
 			const FPoseSearchDatabaseSequence& B = Database->Sequences[SequenceIdxB];
 
 			//If its null add it to the end of the list 
-			if (B.Sequence.IsNull())
+			if (!B.Sequence)
 			{
 				return true;
 			}
 
-			if (A.Sequence.IsNull())
+			if (!A.Sequence)
 			{
 				return false;
 			}
@@ -598,12 +598,12 @@ namespace UE::PoseSearch
 			const FPoseSearchDatabaseBlendSpace& B = Database->BlendSpaces[BlendspaceIdxB];
 
 			//If its null add it to the end of the list 
-			if (B.BlendSpace.IsNull())
+			if (!B.BlendSpace)
 			{
 				return true;
 			}
 
-			if (A.BlendSpace.IsNull())
+			if (!A.BlendSpace)
 			{
 				return false;
 			}
diff --git a/Engine/Plugins/Interchange/Runtime/Source/Pipelines/Private/InterchangeGenericMaterialPipeline.cpp b/Engine/Plugins/Interchange/Runtime/Source/Pipelines/Private/InterchangeGenericMaterialPipeline.cpp
--- a/Engine/Plugins/Interchange/Runtime/Source/Pipelines/Private/InterchangeGenericMaterialPipeline.cpp
+++ b/Engine/Plugins/Interchange/Runtime/Source/Pipelines/Private/InterchangeGenericMaterialPipeline.cpp
@@ -101,7 +101,7 @@ void UInterchangeGenericMaterialPipeline::AdjustSettingsForContext(EInterchangeP
 	}
 
 	TArray<FString> HideCategories;
-	bool bIsObjectAMaterial = ReimportAsset.IsNull() ? false : ReimportAsset->IsA(UMaterialInterface::StaticClass());
+	bool bIsObjectAMaterial = !ReimportAsset ? false : ReimportAsset->IsA(UMaterialInterface::StaticClass());
 	if ((!bIsObjectAMaterial && ImportType == EInterchangePipelineContext::AssetReimport)
 		|| ImportType == EInterchangePipelineContext::AssetCustomLODImport
 		|| ImportType == EInterchangePipelineContext::AssetCustomLODReimport
@@ -131,7 +131,7 @@ void UInterchangeGenericMaterialPipeline::ExecutePreImportPipeline(UInterchangeB
 
 	//Set the result container to allow error message
 	//The parent Results container should be set at this point
-	ensure(!Results.IsNull());
+	ensure(Results);
 	{
 		if (TexturePipeline)
 		{
diff --git a/Engine/Plugins/Runtime/ComputeFramework/Source/ComputeFramework/Private/ComputeGraph.cpp b/Engine/Plugins/Runtime/ComputeFramework/Source/ComputeFramework/Private/ComputeGraph.cpp
--- a/Engine/Plugins/Runtime/ComputeFramework/Source/ComputeFramework/Private/ComputeGraph.cpp
+++ b/Engine/Plugins/Runtime/ComputeFramework/Source/ComputeFramework/Private/ComputeGraph.cpp
@@ -165,7 +165,7 @@ void UComputeGraph::CreateDataProviders(int32 InBindingIndex, TObjectPtr<UObject
 {
 	InOutDataProviders.SetNumZeroed(DataInterfaces.Num());
 
-	if (ensure(Bindings.IsValidIndex(InBindingIndex)) && !InBindingObject.IsNull() && InBindingObject.IsA(Bindings[InBindingIndex]))
+	if (ensure(Bindings.IsValidIndex(InBindingIndex)) && InBindingObject && InBindingObject.IsA(Bindings[InBindingIndex]))
 	{
 		for (int32 DataInterfaceIndex = 0; DataInterfaceIndex < DataInterfaces.Num(); ++DataInterfaceIndex)
 		{
diff --git a/Engine/Plugins/Runtime/MeshModelingToolset/Source/MeshModelingTools/Private/AddPrimitiveTool.cpp b/Engine/Plugins/Runtime/MeshModelingToolset/Source/MeshModelingTools/Private/AddPrimitiveTool.cpp
--- a/Engine/Plugins/Runtime/MeshModelingToolset/Source/MeshModelingTools/Private/AddPrimitiveTool.cpp
+++ b/Engine/Plugins/Runtime/MeshModelingToolset/Source/MeshModelingTools/Private/AddPrimitiveTool.cpp
@@ -367,7 +367,7 @@ void UAddPrimitiveTool::UpdatePreviewMesh() const
 
 bool UAddPrimitiveTool::SupportsWorldSpaceFocusBox()
 {
-	return !PreviewMesh.IsNull();
+	return PreviewMesh != nullptr;
 }
 
 FBox UAddPrimitiveTool::GetWorldSpaceFocusBox()
@@ -384,7 +384,7 @@ FBox UAddPrimitiveTool::GetWorldSpaceFocusBox()
 
 bool UAddPrimitiveTool::SupportsWorldSpaceFocusPoint()
 {
-	return !PreviewMesh.IsNull();
+	return PreviewMesh != nullptr;
 
 }
 
