Fixing linking error, caused by a misplaced WITH_EDITOR

diff --git a/Engine/Plugins/Interchange/Runtime/Source/Import/Public/Material/InterchangeMaterialFactory.h b/Engine/Plugins/Interchange/Runtime/Source/Import/Public/Material/InterchangeMaterialFactory.h
--- a/Engine/Plugins/Interchange/Runtime/Source/Import/Public/Material/InterchangeMaterialFactory.h
+++ b/Engine/Plugins/Interchange/Runtime/Source/Import/Public/Material/InterchangeMaterialFactory.h
@@ -83,9 +83,11 @@ public:
 
 	void CleanUp();
 
+#if WITH_EDITOR
 	void UpdateFromFunctionResource(UMaterialExpressionMaterialFunctionCall* MaterialFunctionCall);
 
 	void UpdateFromFunctionResource(UMaterialFunctionInterface* MaterialFunction);
+#endif
 
 private:
 
