Adding MaterialX material expressions for compositing node group:
- burn
- difference
- dodge
- disjointover
- in
- mask
- matte
- minus
- out
- over
- overlay
- plus
- premult
- screen
- unpremult

MaterialX Translator:
- Connecting the above material expressions
- Adding the support of the <inside> node, which is simply a multiply
- Adding the support of the <outside> node

diff --git a/Engine/Plugins/Interchange/Runtime/Source/Import/Private/MaterialX/InterchangeMaterialXDefinitions.h b/Engine/Plugins/Interchange/Runtime/Source/Import/Private/MaterialX/InterchangeMaterialXDefinitions.h
--- a/Engine/Plugins/Interchange/Runtime/Source/Import/Private/MaterialX/InterchangeMaterialXDefinitions.h
+++ b/Engine/Plugins/Interchange/Runtime/Source/Import/Private/MaterialX/InterchangeMaterialXDefinitions.h
@@ -206,9 +206,26 @@ MATERIALX_NAMESPACE_BEGIN
 		static constexpr const char* TransformNormal = "transformnormal";
 		static constexpr const char* TransformPoint = "transformpoint";
 		static constexpr const char* TransformVector = "transformvector";
-		static constexpr const char* Transpose = "transpose";
+		static constexpr const char* Transpose = "transpose";		
 		// Compositing nodes
+		static constexpr const char* Burn = "burn";
+		static constexpr const char* Difference = "difference";
+		static constexpr const char* Disjointover = "disjointover";
+		static constexpr const char* Dodge = "dodge";
+		static constexpr const char* In = "in";
+		static constexpr const char* Inside = "inside";
+		static constexpr const char* Mask = "mask";
+		static constexpr const char* Matte = "matte";
+		static constexpr const char* Minus = "minus";
 		static constexpr const char* Mix = "mix";
+		static constexpr const char* Out = "out";
+		static constexpr const char* Outside = "outside";
+		static constexpr const char* Over = "over";
+		static constexpr const char* Overlay = "overlay";
+		static constexpr const char* Plus = "plus";
+		static constexpr const char* Premult = "premult";
+		static constexpr const char* Screen = "screen";
+		static constexpr const char* Unpremult = "unpremult";
 		// Conditional nodes
 		static constexpr const char* IfGreater = "ifgreater";
 		static constexpr const char* IfGreaterEq = "ifgreatereq";
diff --git a/Engine/Plugins/Interchange/Runtime/Source/Import/Private/MaterialX/InterchangeMaterialXTranslator.cpp b/Engine/Plugins/Interchange/Runtime/Source/Import/Private/MaterialX/InterchangeMaterialXTranslator.cpp
--- a/Engine/Plugins/Interchange/Runtime/Source/Import/Private/MaterialX/InterchangeMaterialXTranslator.cpp
+++ b/Engine/Plugins/Interchange/Runtime/Source/Import/Private/MaterialX/InterchangeMaterialXTranslator.cpp
@@ -68,7 +68,7 @@ UInterchangeMaterialXTranslator::UInterchangeMaterialXTranslator()
 InputNamesMaterialX2UE
 {
 	{{TEXT(""),                         TEXT("bg")},		TEXT("B")},
-	{{TEXT(""),                         TEXT("center")},		TEXT("Center")},
+	{{TEXT(""),                         TEXT("center")},	TEXT("Center")},
 	{{TEXT(""),                         TEXT("fg")},		TEXT("A")},
 	{{TEXT(""),                         TEXT("high")},		TEXT("Max")},
 	{{TEXT(""),                         TEXT("in")},		TEXT("Input")},
@@ -80,7 +80,7 @@ InputNamesMaterialX2UE
 	{{TEXT(""),                         TEXT("inhigh")},    TEXT("InputHigh")},
 	{{TEXT(""),                         TEXT("low")},		TEXT("Min")},
 	{{TEXT(""),                         TEXT("lumacoeffs")},TEXT("LuminanceFactors")}, // for the moment not yet handled by Interchange, because of the attribute being an advanced pin
-	{{TEXT(""),                         TEXT("mix")},		TEXT("Factor")},
+	{{TEXT(""),                         TEXT("mix")},		TEXT("Alpha")},
 	{{TEXT(""),                         TEXT("texcoord")},  TEXT("Coordinates")},
 	{{TEXT(""),                         TEXT("outlow")},    TEXT("TargetLow")},
 	{{TEXT(""),                         TEXT("outhigh")},   TEXT("TargetHigh")},
@@ -99,13 +99,18 @@ InputNamesMaterialX2UE
 	{{MaterialX::Category::IfEqual,     TEXT("value1")},    TEXT("A")},
 	{{MaterialX::Category::IfEqual,     TEXT("value2")},    TEXT("B")},
 	{{MaterialX::Category::IfEqual,     TEXT("in1")},       TEXT("AEqualsB")},
-	{{MaterialX::Category::IfEqual,     TEXT("in2")},       TEXT("ALessThanB")}, //another input is added for the case greater see ConnectIfEqual
+	{{MaterialX::Category::IfEqual,     TEXT("in2")},       TEXT("ALessThanB")},  // another input is added for the case greater see ConnectIfEqual
+	{{MaterialX::Category::Inside,      TEXT("in")},        TEXT("A")},			  // Inside is treated as a Multiply node
+	{{MaterialX::Category::Inside,      TEXT("mask")},      TEXT("B")},			  // Inside is treated as a Multiply node
 	{{MaterialX::Category::Invert,      TEXT("amount")},    TEXT("A")},
 	{{MaterialX::Category::Invert,      TEXT("in")},        TEXT("B")},
 	{{MaterialX::Category::Magnitude,   TEXT("in")},        TEXT("A")},
 	{{MaterialX::Category::Mix,         TEXT("fg")},        TEXT("B")},
 	{{MaterialX::Category::Mix,         TEXT("bg")},        TEXT("A")},
+	{{MaterialX::Category::Mix,         TEXT("mix")},		TEXT("Factor")},
 	{{MaterialX::Category::Normalize,   TEXT("in")},        TEXT("VectorInput")},
+	{{MaterialX::Category::Outside,     TEXT("in")},        TEXT("A")},				// Outside is treated as Multiply node
+	{{MaterialX::Category::Outside,     TEXT("mask")},      TEXT("B")},				// Outside is treated as Multiply node
 	{{MaterialX::Category::Power,       TEXT("in1")},       TEXT("Base")},
 	{{MaterialX::Category::Power,	    TEXT("in2")},       TEXT("Exponent")},
 	{{MaterialX::Category::Rotate2D,    TEXT("amount")},    TEXT("RotationAngle")},
@@ -150,7 +155,23 @@ NodeNamesMaterialX2UE{
 	{MaterialX::Category::Sub,          TEXT("Subtract")},
 	{MaterialX::Category::Tan,          TEXT("Tangent")},
 	// Compositing nodes
+	{MaterialX::Category::Burn,         TEXT("Burn")},
+	{MaterialX::Category::Difference,   TEXT("Difference")},
+	{MaterialX::Category::Disjointover, TEXT("Disjointover")},
+	{MaterialX::Category::Dodge,        TEXT("Dodge")},
+	{MaterialX::Category::In,           TEXT("In")},
+	{MaterialX::Category::Inside,       TEXT("Multiply")},
+	{MaterialX::Category::Mask,         TEXT("Mask")},
+	{MaterialX::Category::Matte,        TEXT("Matte")},
+	{MaterialX::Category::Minus,        TEXT("Minus")},
 	{MaterialX::Category::Mix,          TEXT("Lerp")},
+	{MaterialX::Category::Out,          TEXT("Out")},
+	{MaterialX::Category::Over,         TEXT("Over")},
+	{MaterialX::Category::Overlay,      TEXT("Overlay")},
+	{MaterialX::Category::Plus,         TEXT("Plus")},
+	{MaterialX::Category::Premult,      TEXT("Premult")},
+	{MaterialX::Category::Screen,       TEXT("Screen")},
+	{MaterialX::Category::Unpremult,    TEXT("Unpremult")},
 	// Channel nodes
 	{MaterialX::Category::Combine2,     TEXT("AppendVector")},
 	{MaterialX::Category::Combine3,     TEXT("Append3Vector")},
@@ -168,6 +189,7 @@ NodeNamesMaterialX2UE{
 UEInputs
 {
     TEXT("A"),
+	TEXT("Alpha"),
 	TEXT("B"),
     TEXT("Base"),
     TEXT("C"),
@@ -985,6 +1007,10 @@ bool UInterchangeMaterialXTranslator::ConnectNodeGraphOutputToInput(MaterialX::I
 				{
 					ConnectIfEqualInputToOutput(UpstreamNode, ParentShaderNode, InputChannelName, NamesToShaderNodes, NodeContainer);
 				}
+				else if(UpstreamNode->getCategory() == mx::Category::Outside)
+				{
+					ConnectOutsideInputToOutput(UpstreamNode, ParentShaderNode, InputChannelName, NamesToShaderNodes, NodeContainer);
+				}
 				else
 				{
 					UInterchangeResultWarning_Generic* Message = AddMessage<UInterchangeResultWarning_Generic>();
@@ -1735,6 +1761,18 @@ void UInterchangeMaterialXTranslator::ConnectIfEqualInputToOutput(MaterialX::Nod
 	}
 }
 
+void UInterchangeMaterialXTranslator::ConnectOutsideInputToOutput(MaterialX::NodePtr UpstreamNode, UInterchangeShaderNode* ParentShaderNode, const FString& InputChannelName, TMap<FString, UInterchangeShaderNode*>& NamesToShaderNodes, UInterchangeBaseNodeContainer& NodeContainer) const
+{
+	//in * (1 - mask)
+	UInterchangeShaderNode* NodeMultiply = CreateShaderNode<UInterchangeShaderNode>(UpstreamNode->getName().c_str(), TEXT("Multiply"), NamesToShaderNodes, NodeContainer);
+	AddAttributeFromValueOrInterface(UpstreamNode->getInput("in"), TEXT("A"), NodeMultiply);
+	UInterchangeShaderPortsAPI::ConnectDefaultOuputToInput(ParentShaderNode, InputChannelName, NodeMultiply->GetUniqueID());
+
+	UInterchangeShaderNode* NodeOneMinus = CreateShaderNode<UInterchangeShaderNode>(UpstreamNode->getName().c_str() + FString(TEXT("_OneMinus")), TEXT("OneMinus"), NamesToShaderNodes, NodeContainer);
+	AddAttributeFromValueOrInterface(UpstreamNode->getInput("mask"), TEXT("Input"), NodeOneMinus);
+	UInterchangeShaderPortsAPI::ConnectDefaultOuputToInput(NodeMultiply, TEXT("B"), NodeOneMinus->GetUniqueID());
+}
+
 UInterchangeShaderNode* UInterchangeMaterialXTranslator::CreateMaskShaderNode(uint8 RGBA, const FString& NodeName, TMap<FString, UInterchangeShaderNode*>& NamesToShaderNodes, UInterchangeBaseNodeContainer& NodeContainer) const
 {
 	bool bR = (0b1000 & RGBA) >> 3;
diff --git a/Engine/Plugins/Interchange/Runtime/Source/Import/Private/MaterialX/MaterialExpressionBurn.cpp b/Engine/Plugins/Interchange/Runtime/Source/Import/Private/MaterialX/MaterialExpressionBurn.cpp
new file mode 100644
--- /dev/null
+++ b/Engine/Plugins/Interchange/Runtime/Source/Import/Private/MaterialX/MaterialExpressionBurn.cpp
@@ -0,0 +1,61 @@
+// Copyright Epic Games, Inc. All Rights Reserved.
+#include "MaterialX/MaterialExpressionBurn.h"
+#include "MaterialCompiler.h"
+
+#define LOCTEXT_NAMESPACE "MaterialExpressionBurn"
+
+UMaterialExpressionBurn::UMaterialExpressionBurn(const FObjectInitializer& ObjectInitializer)
+	: Super(ObjectInitializer)
+{
+	// Structure to hold one-time initialization
+	struct FConstructorStatics
+	{
+		FText NAME_MaterialX;
+		FText NAME_Compositing;
+		FConstructorStatics() :
+			NAME_MaterialX(LOCTEXT("MaterialX", "MaterialX")),
+			NAME_Compositing(LOCTEXT("Compositing", "Compositing"))
+		{}
+	};
+	static FConstructorStatics ConstructorStatics;
+
+#if WITH_EDITORONLY_DATA
+	MenuCategories.Add(ConstructorStatics.NAME_MaterialX);
+	MenuCategories.Add(ConstructorStatics.NAME_Compositing);
+#endif
+}
+
+#if WITH_EDITOR
+int32 UMaterialExpressionBurn::Compile(FMaterialCompiler* Compiler, int32 OutputIndex)
+{
+	if(!A.GetTracedInput().Expression)
+	{
+		return Compiler->Errorf(TEXT("Missing A input"));
+	}
+
+	if(!B.GetTracedInput().Expression)
+	{
+		return Compiler->Errorf(TEXT("Missing B input"));
+	}
+
+	int32 IndexAlpha = Alpha.GetTracedInput().Expression ? Alpha.Compile(Compiler) : Compiler->Constant(ConstAlpha);
+
+	int32 OneIndex = Compiler->Constant(1.f);
+	auto OneMinus = [&](int32 index)
+	{
+		return Compiler->Sub(OneIndex, index);
+	};
+
+	int32 IndexB = B.Compile(Compiler);
+	int32 Burn = OneMinus(Compiler->Div(OneMinus(IndexB), A.Compile(Compiler)));
+
+	return Compiler->Lerp(IndexB, Burn, IndexAlpha);
+}
+
+void UMaterialExpressionBurn::GetCaption(TArray<FString>& OutCaptions) const
+{
+	OutCaptions.Add(TEXT("Burn"));
+}
+#endif
+
+#undef LOCTEXT_NAMESPACE 
\ No newline at end of file
diff --git a/Engine/Plugins/Interchange/Runtime/Source/Import/Private/MaterialX/MaterialExpressionBurn.h b/Engine/Plugins/Interchange/Runtime/Source/Import/Private/MaterialX/MaterialExpressionBurn.h
new file mode 100644
--- /dev/null
+++ b/Engine/Plugins/Interchange/Runtime/Source/Import/Private/MaterialX/MaterialExpressionBurn.h
@@ -0,0 +1,41 @@
+// Copyright Epic Games, Inc. All Rights Reserved.
+#pragma once
+
+#include "CoreMinimal.h"
+#include "Materials/MaterialExpression.h"
+
+#include "MaterialExpressionBurn.generated.h"
+
+/**
+ * Blend nodes take two 1-4 channel inputs and apply the same operator to all channels.
+ * Blend nodes support an optional float input mix , which can be used
+ * to mix the original B value with the result of the blend operation.
+ * Operation: 1-(1-B)/A
+ * Result: Lerp(B, 1-(1-B)/A, Alpha)
+ */
+UCLASS(MinimalAPI, collapsecategories, hidecategories = Object)
+class UMaterialExpressionBurn : public UMaterialExpression
+{
+	GENERATED_UCLASS_BODY()
+
+	UPROPERTY()
+	FExpressionInput A;
+
+	UPROPERTY()
+	FExpressionInput B;
+
+	UPROPERTY(meta = (RequiredInput = "false", ToolTip = "Defaults to 'ConstAlpha' if not specified"))
+	FExpressionInput Alpha;
+
+	/** only used if Alpha is not hooked up */
+	UPROPERTY(EditAnywhere, Category = MaterialExpressionLinearInterpolate, meta = (OverridingInputProperty = "Alpha"))
+	float ConstAlpha = 1.f;
+
+	//~ Begin UMaterialExpression Interface
+#if WITH_EDITOR
+	virtual int32 Compile(class FMaterialCompiler* Compiler, int32 OutputIndex) override;
+	virtual void GetCaption(TArray<FString>& OutCaptions) const override;
+#endif
+	//~ End UMaterialExpression Interface
+};
+
diff --git a/Engine/Plugins/Interchange/Runtime/Source/Import/Private/MaterialX/MaterialExpressionDifference.cpp b/Engine/Plugins/Interchange/Runtime/Source/Import/Private/MaterialX/MaterialExpressionDifference.cpp
new file mode 100644
--- /dev/null
+++ b/Engine/Plugins/Interchange/Runtime/Source/Import/Private/MaterialX/MaterialExpressionDifference.cpp
@@ -0,0 +1,54 @@
+// Copyright Epic Games, Inc. All Rights Reserved.
+#include "MaterialX/MaterialExpressionDifference.h"
+#include "MaterialCompiler.h"
+
+#define LOCTEXT_NAMESPACE "MaterialExpressionDifference"
+
+UMaterialExpressionDifference::UMaterialExpressionDifference(const FObjectInitializer& ObjectInitializer)
+	: Super(ObjectInitializer)
+{
+	// Structure to hold one-time initialization
+	struct FConstructorStatics
+	{
+		FText NAME_MaterialX;
+		FText NAME_Compositing;
+		FConstructorStatics() :
+			NAME_MaterialX(LOCTEXT("MaterialX", "MaterialX")),
+			NAME_Compositing(LOCTEXT("Compositing", "Compositing"))
+		{}
+	};
+	static FConstructorStatics ConstructorStatics;
+
+#if WITH_EDITORONLY_DATA
+	MenuCategories.Add(ConstructorStatics.NAME_MaterialX);
+	MenuCategories.Add(ConstructorStatics.NAME_Compositing);
+#endif
+}
+
+#if WITH_EDITOR
+int32 UMaterialExpressionDifference::Compile(FMaterialCompiler* Compiler, int32 OutputIndex)
+{
+	if(!A.GetTracedInput().Expression)
+	{
+		return Compiler->Errorf(TEXT("Missing A input"));
+	}
+
+	if(!B.GetTracedInput().Expression)
+	{
+		return Compiler->Errorf(TEXT("Missing B input"));
+	}
+
+	int32 IndexAlpha = Alpha.GetTracedInput().Expression ? Alpha.Compile(Compiler) : Compiler->Constant(ConstAlpha);
+
+	int32 IndexB = B.Compile(Compiler);
+	int32 Diff = Compiler->Abs(Compiler->Sub(A.Compile(Compiler), IndexB));
+	return Compiler->Lerp(IndexB, Diff, IndexAlpha);
+}
+
+void UMaterialExpressionDifference::GetCaption(TArray<FString>& OutCaptions) const
+{
+	OutCaptions.Add(TEXT("Difference"));
+}
+#endif
+
+#undef LOCTEXT_NAMESPACE 
\ No newline at end of file
diff --git a/Engine/Plugins/Interchange/Runtime/Source/Import/Private/MaterialX/MaterialExpressionDifference.h b/Engine/Plugins/Interchange/Runtime/Source/Import/Private/MaterialX/MaterialExpressionDifference.h
new file mode 100644
--- /dev/null
+++ b/Engine/Plugins/Interchange/Runtime/Source/Import/Private/MaterialX/MaterialExpressionDifference.h
@@ -0,0 +1,41 @@
+// Copyright Epic Games, Inc. All Rights Reserved.
+#pragma once
+
+#include "CoreMinimal.h"
+#include "Materials/MaterialExpression.h"
+
+#include "MaterialExpressionDifference.generated.h"
+
+/**
+ * Blend nodes take two 1-4 channel inputs and apply the same operator to all channels.
+ * Blend nodes support an optional float input mix , which can be used
+ * to mix the original B value with the result of the blend operation.
+ * Operation: abs(A - B)
+ * Result: Lerp(B, abs(A - B), Alpha)
+ */
+UCLASS(MinimalAPI, collapsecategories, hidecategories = Object)
+class UMaterialExpressionDifference : public UMaterialExpression
+{
+	GENERATED_UCLASS_BODY()
+
+	UPROPERTY()
+	FExpressionInput A;
+
+	UPROPERTY()
+	FExpressionInput B;
+
+	UPROPERTY(meta = (RequiredInput = "false", ToolTip = "Defaults to 'ConstAlpha' if not specified"))
+	FExpressionInput Alpha;
+
+	/** only used if Alpha is not hooked up */
+	UPROPERTY(EditAnywhere, Category = MaterialExpressionLinearInterpolate, meta = (OverridingInputProperty = "Alpha"))
+	float ConstAlpha = 1.f;
+
+	//~ Begin UMaterialExpression Interface
+#if WITH_EDITOR
+	virtual int32 Compile(class FMaterialCompiler* Compiler, int32 OutputIndex) override;
+	virtual void GetCaption(TArray<FString>& OutCaptions) const override;
+#endif
+	//~ End UMaterialExpression Interface
+};
+
diff --git a/Engine/Plugins/Interchange/Runtime/Source/Import/Private/MaterialX/MaterialExpressionDisjointOver.cpp b/Engine/Plugins/Interchange/Runtime/Source/Import/Private/MaterialX/MaterialExpressionDisjointOver.cpp
new file mode 100644
--- /dev/null
+++ b/Engine/Plugins/Interchange/Runtime/Source/Import/Private/MaterialX/MaterialExpressionDisjointOver.cpp
@@ -0,0 +1,71 @@
+// Copyright Epic Games, Inc. All Rights Reserved.
+#include "MaterialX/MaterialExpressionDisjointOver.h"
+#include "MaterialCompiler.h"
+
+#define LOCTEXT_NAMESPACE "MaterialExpressionDisjointOver"
+
+UMaterialExpressionDisjointOver::UMaterialExpressionDisjointOver(const FObjectInitializer& ObjectInitializer)
+	: Super(ObjectInitializer)
+{
+	// Structure to hold one-time initialization
+	struct FConstructorStatics
+	{
+		FText NAME_MaterialX;
+		FText NAME_Compositing;
+		FConstructorStatics() :
+			NAME_MaterialX(LOCTEXT("MaterialX", "MaterialX")),
+			NAME_Compositing(LOCTEXT("Compositing", "Compositing"))
+		{}
+	};
+	static FConstructorStatics ConstructorStatics;
+
+#if WITH_EDITORONLY_DATA
+	MenuCategories.Add(ConstructorStatics.NAME_MaterialX);
+	MenuCategories.Add(ConstructorStatics.NAME_Compositing);
+#endif
+}
+
+#if WITH_EDITOR
+int32 UMaterialExpressionDisjointOver::Compile(FMaterialCompiler* Compiler, int32 OutputIndex)
+{
+	if(!A.GetTracedInput().Expression)
+	{
+		return Compiler->Errorf(TEXT("Missing A input"));
+	}
+
+	if(!B.GetTracedInput().Expression)
+	{
+		return Compiler->Errorf(TEXT("Missing B input"));
+	}
+
+	int32 IndexAlpha = Alpha.GetTracedInput().Expression ? Alpha.Compile(Compiler) : Compiler->Constant(ConstAlpha);
+	
+	int32 IndexA = A.Compile(Compiler);
+	int32 IndexB = B.Compile(Compiler);
+
+	int32 IndexAlphaA = Compiler->ComponentMask(IndexA, false, false, false, true);
+	int32 IndexAlphaB = Compiler->ComponentMask(IndexB, false, false, false, true);
+	int32 IndexSumAlpha = Compiler->Add(IndexAlphaA, IndexAlphaB);
+
+	int32 IndexRgbA = Compiler->ComponentMask(IndexA, true, true, true, false);
+	int32 IndexRgbB = Compiler->ComponentMask(IndexB, true, true, true, false);
+	int32 IndexSumRgb = Compiler->Add(IndexRgbA, IndexRgbB);
+	
+	int32 IndexOne = Compiler->Constant(1.f);
+	int32 X = Compiler->Div(Compiler->Sub(IndexOne, IndexAlphaA), IndexAlphaB);
+	int32 Disjoint = Compiler->Add(IndexRgbA, Compiler->Mul(IndexRgbB, X));
+
+	int32 Result = Compiler->If(IndexSumAlpha, IndexOne, Disjoint, IndexSumRgb, IndexSumRgb, Compiler->Constant(0.00001f));
+	int32 ResultAlpha = Compiler->Min(IndexSumAlpha, IndexOne);
+
+	int32 IndexAppend = Compiler->AppendVector(Result, ResultAlpha);
+	return Compiler->Lerp(IndexB, IndexAppend, IndexAlpha);
+}
+
+void UMaterialExpressionDisjointOver::GetCaption(TArray<FString>& OutCaptions) const
+{
+	OutCaptions.Add(TEXT("DisjointOver"));
+}
+#endif
+
+#undef LOCTEXT_NAMESPACE 
\ No newline at end of file
diff --git a/Engine/Plugins/Interchange/Runtime/Source/Import/Private/MaterialX/MaterialExpressionDisjointOver.h b/Engine/Plugins/Interchange/Runtime/Source/Import/Private/MaterialX/MaterialExpressionDisjointOver.h
new file mode 100644
--- /dev/null
+++ b/Engine/Plugins/Interchange/Runtime/Source/Import/Private/MaterialX/MaterialExpressionDisjointOver.h
@@ -0,0 +1,46 @@
+// Copyright Epic Games, Inc. All Rights Reserved.
+#pragma once
+
+#include "CoreMinimal.h"
+#include "Materials/MaterialExpression.h"
+
+#include "MaterialExpressionDisjointOver.generated.h"
+
+/**
+ * Merge nodes take two 4-channel (color4) inputs and use the built-in alpha channel(s) to control the
+ * compositing of the A and B inputs. "A" and "B" refer to the non-alpha channels of the A and B inputs respectively, 
+ * and "a" and "b" refer to the alpha channels of the A and B inputs.
+ * Merge nodes are only defined for float4 inputs
+ * Merge nodes support an optional float input Alpha , which can be used to mix the
+ * original B value with the result of the blend operation.
+ * 
+ * Operation: A + B        if a + b <= 1;
+ *            A + B(1-a)/b if a + b > 1
+ * Result: Lerp(B, Op, Alpha)
+ */
+UCLASS(MinimalAPI, collapsecategories, hidecategories = Object)
+class UMaterialExpressionDisjointOver : public UMaterialExpression
+{
+	GENERATED_UCLASS_BODY()
+
+	UPROPERTY()
+	FExpressionInput A;
+
+	UPROPERTY()
+	FExpressionInput B;
+
+	UPROPERTY(meta = (RequiredInput = "false", ToolTip = "Defaults to 'ConstAlpha' if not specified"))
+	FExpressionInput Alpha;
+
+	/** only used if Alpha is not hooked up */
+	UPROPERTY(EditAnywhere, Category = MaterialExpressionLinearInterpolate, meta = (OverridingInputProperty = "Alpha"))
+	float ConstAlpha = 1.f;
+
+	//~ Begin UMaterialExpression Interface
+#if WITH_EDITOR
+	virtual int32 Compile(class FMaterialCompiler* Compiler, int32 OutputIndex) override;
+	virtual void GetCaption(TArray<FString>& OutCaptions) const override;
+#endif
+	//~ End UMaterialExpression Interface
+};
+
diff --git a/Engine/Plugins/Interchange/Runtime/Source/Import/Private/MaterialX/MaterialExpressionDodge.cpp b/Engine/Plugins/Interchange/Runtime/Source/Import/Private/MaterialX/MaterialExpressionDodge.cpp
new file mode 100644
--- /dev/null
+++ b/Engine/Plugins/Interchange/Runtime/Source/Import/Private/MaterialX/MaterialExpressionDodge.cpp
@@ -0,0 +1,59 @@
+// Copyright Epic Games, Inc. All Rights Reserved.
+#include "MaterialX/MaterialExpressionDodge.h"
+#include "MaterialCompiler.h"
+
+#define LOCTEXT_NAMESPACE "MaterialExpressionDodge"
+
+UMaterialExpressionDodge::UMaterialExpressionDodge(const FObjectInitializer& ObjectInitializer)
+	: Super(ObjectInitializer)
+{
+	// Structure to hold one-time initialization
+	struct FConstructorStatics
+	{
+		FText NAME_MaterialX;
+		FText NAME_Compositing;
+		FConstructorStatics() :
+			NAME_MaterialX(LOCTEXT("MaterialX", "MaterialX")),
+			NAME_Compositing(LOCTEXT("Compositing", "Compositing"))
+		{}
+	};
+	static FConstructorStatics ConstructorStatics;
+
+#if WITH_EDITORONLY_DATA
+	MenuCategories.Add(ConstructorStatics.NAME_MaterialX);
+	MenuCategories.Add(ConstructorStatics.NAME_Compositing);
+#endif
+}
+
+#if WITH_EDITOR
+int32 UMaterialExpressionDodge::Compile(FMaterialCompiler* Compiler, int32 OutputIndex)
+{
+	if(!A.GetTracedInput().Expression)
+	{
+		return Compiler->Errorf(TEXT("Missing A input"));
+	}
+
+	if(!B.GetTracedInput().Expression)
+	{
+		return Compiler->Errorf(TEXT("Missing B input"));
+	}
+
+	int32 IndexAlpha = Alpha.GetTracedInput().Expression ? Alpha.Compile(Compiler) : Compiler->Constant(ConstAlpha);
+	int32 IndexB = B.Compile(Compiler);
+	int32 Sub = Compiler->Sub(Compiler->Constant(1.f), A.Compile(Compiler));
+
+	int32 Zero = Compiler->Constant(0.f);
+	int32 Dodge = Compiler->Div(IndexB, Sub);
+
+	int32 Result = Compiler->Lerp(IndexB, Dodge, IndexAlpha);
+
+	return Compiler->If(Compiler->Abs(Sub), Zero, Result, Result, Zero, Compiler->Constant(0.00001f));
+}
+
+void UMaterialExpressionDodge::GetCaption(TArray<FString>& OutCaptions) const
+{
+	OutCaptions.Add(TEXT("Dodge"));
+}
+#endif
+
+#undef LOCTEXT_NAMESPACE 
\ No newline at end of file
diff --git a/Engine/Plugins/Interchange/Runtime/Source/Import/Private/MaterialX/MaterialExpressionDodge.h b/Engine/Plugins/Interchange/Runtime/Source/Import/Private/MaterialX/MaterialExpressionDodge.h
new file mode 100644
--- /dev/null
+++ b/Engine/Plugins/Interchange/Runtime/Source/Import/Private/MaterialX/MaterialExpressionDodge.h
@@ -0,0 +1,41 @@
+// Copyright Epic Games, Inc. All Rights Reserved.
+#pragma once
+
+#include "CoreMinimal.h"
+#include "Materials/MaterialExpression.h"
+
+#include "MaterialExpressionDodge.generated.h"
+
+/**
+ * Blend nodes take two 1-4 channel inputs and apply the same operator to all channels.
+ * Blend nodes support an optional float input mix , which can be used
+ * to mix the original B value with the result of the blend operation.
+ * Operation: B/(1-A)
+ * Result: Lerp(B, B/(1-A), Alpha)
+ */
+UCLASS(MinimalAPI, collapsecategories, hidecategories = Object)
+class UMaterialExpressionDodge: public UMaterialExpression
+{
+	GENERATED_UCLASS_BODY()
+
+	UPROPERTY()
+	FExpressionInput A;
+
+	UPROPERTY()
+	FExpressionInput B;
+
+	UPROPERTY(meta = (RequiredInput = "false", ToolTip = "Defaults to 'ConstAlpha' if not specified"))
+	FExpressionInput Alpha;
+
+	/** only used if Alpha is not hooked up */
+	UPROPERTY(EditAnywhere, Category = MaterialExpressionLinearInterpolate, meta = (OverridingInputProperty = "Alpha"))
+	float ConstAlpha = 1.f;
+
+	//~ Begin UMaterialExpression Interface
+#if WITH_EDITOR
+	virtual int32 Compile(class FMaterialCompiler* Compiler, int32 OutputIndex) override;
+	virtual void GetCaption(TArray<FString>& OutCaptions) const override;
+#endif
+	//~ End UMaterialExpression Interface
+};
+
diff --git a/Engine/Plugins/Interchange/Runtime/Source/Import/Private/MaterialX/MaterialExpressionIn.cpp b/Engine/Plugins/Interchange/Runtime/Source/Import/Private/MaterialX/MaterialExpressionIn.cpp
new file mode 100644
--- /dev/null
+++ b/Engine/Plugins/Interchange/Runtime/Source/Import/Private/MaterialX/MaterialExpressionIn.cpp
@@ -0,0 +1,57 @@
+// Copyright Epic Games, Inc. All Rights Reserved.
+#include "MaterialX/MaterialExpressionIn.h"
+#include "MaterialCompiler.h"
+
+#define LOCTEXT_NAMESPACE "MaterialExpressionIn"
+
+UMaterialExpressionIn::UMaterialExpressionIn(const FObjectInitializer& ObjectInitializer)
+	: Super(ObjectInitializer)
+{
+	// Structure to hold one-time initialization
+	struct FConstructorStatics
+	{
+		FText NAME_MaterialX;
+		FText NAME_Compositing;
+		FConstructorStatics() :
+			NAME_MaterialX(LOCTEXT("MaterialX", "MaterialX")),
+			NAME_Compositing(LOCTEXT("Compositing", "Compositing"))
+		{}
+	};
+	static FConstructorStatics ConstructorStatics;
+
+#if WITH_EDITORONLY_DATA
+	MenuCategories.Add(ConstructorStatics.NAME_MaterialX);
+	MenuCategories.Add(ConstructorStatics.NAME_Compositing);
+#endif
+}
+
+#if WITH_EDITOR
+int32 UMaterialExpressionIn::Compile(FMaterialCompiler* Compiler, int32 OutputIndex)
+{
+	if(!A.GetTracedInput().Expression)
+	{
+		return Compiler->Errorf(TEXT("Missing A input"));
+	}
+
+	if(!B.GetTracedInput().Expression)
+	{
+		return Compiler->Errorf(TEXT("Missing B input"));
+	}
+
+	int32 IndexAlpha = Alpha.GetTracedInput().Expression ? Alpha.Compile(Compiler) : Compiler->Constant(ConstAlpha);
+
+	int32 IndexA = A.Compile(Compiler);
+	int32 IndexB = B.Compile(Compiler);
+
+	int32 IndexIn = Compiler->Mul(IndexA, Compiler->ComponentMask(IndexB, false, false, false, true));
+
+	return Compiler->Lerp(IndexB, IndexIn, IndexAlpha);
+}
+
+void UMaterialExpressionIn::GetCaption(TArray<FString>& OutCaptions) const
+{
+	OutCaptions.Add(TEXT("In"));
+}
+#endif
+
+#undef LOCTEXT_NAMESPACE 
\ No newline at end of file
diff --git a/Engine/Plugins/Interchange/Runtime/Source/Import/Private/MaterialX/MaterialExpressionIn.h b/Engine/Plugins/Interchange/Runtime/Source/Import/Private/MaterialX/MaterialExpressionIn.h
new file mode 100644
--- /dev/null
+++ b/Engine/Plugins/Interchange/Runtime/Source/Import/Private/MaterialX/MaterialExpressionIn.h
@@ -0,0 +1,45 @@
+// Copyright Epic Games, Inc. All Rights Reserved.
+#pragma once
+
+#include "CoreMinimal.h"
+#include "Materials/MaterialExpression.h"
+
+#include "MaterialExpressionIn.generated.h"
+
+/**
+ * Merge nodes take two 4-channel (color4) inputs and use the built-in alpha channel(s) to control the
+ * compositing of the A and B inputs. "A" and "B" refer to the non-alpha channels of the A and B inputs respectively,
+ * and "a" and "b" refer to the alpha channels of the A and B inputs.
+ * Merge nodes are only defined for float4 inputs
+ * Merge nodes support an optional float input Alpha , which can be used to mix the
+ * original B value with the result of the blend operation.
+ *
+ * Operation: Ab
+ * Result: Lerp(B, Ab, Alpha)
+ */
+UCLASS(MinimalAPI, collapsecategories, hidecategories = Object)
+class UMaterialExpressionIn: public UMaterialExpression
+{
+	GENERATED_UCLASS_BODY()
+
+	UPROPERTY()
+	FExpressionInput A;
+
+	UPROPERTY()
+	FExpressionInput B;
+
+	UPROPERTY(meta = (RequiredInput = "false", ToolTip = "Defaults to 'ConstAlpha' if not specified"))
+	FExpressionInput Alpha;
+
+	/** only used if Alpha is not hooked up */
+	UPROPERTY(EditAnywhere, Category = MaterialExpressionLinearInterpolate, meta = (OverridingInputProperty = "Alpha"))
+	float ConstAlpha = 1.f;
+
+	//~ Begin UMaterialExpression Interface
+#if WITH_EDITOR
+	virtual int32 Compile(class FMaterialCompiler* Compiler, int32 OutputIndex) override;
+	virtual void GetCaption(TArray<FString>& OutCaptions) const override;
+#endif
+	//~ End UMaterialExpression Interface
+};
+
diff --git a/Engine/Plugins/Interchange/Runtime/Source/Import/Private/MaterialX/MaterialExpressionMask.cpp b/Engine/Plugins/Interchange/Runtime/Source/Import/Private/MaterialX/MaterialExpressionMask.cpp
new file mode 100644
--- /dev/null
+++ b/Engine/Plugins/Interchange/Runtime/Source/Import/Private/MaterialX/MaterialExpressionMask.cpp
@@ -0,0 +1,57 @@
+// Copyright Epic Games, Inc. All Rights Reserved.
+#include "MaterialX/MaterialExpressionMask.h"
+#include "MaterialCompiler.h"
+
+#define LOCTEXT_NAMESPACE "MaterialExpressionMask"
+
+UMaterialExpressionMask::UMaterialExpressionMask(const FObjectInitializer& ObjectInitializer)
+	: Super(ObjectInitializer)
+{
+	// Structure to hold one-time initialization
+	struct FConstructorStatics
+	{
+		FText NAME_MaterialX;
+		FText NAME_Compositing;
+		FConstructorStatics() :
+			NAME_MaterialX(LOCTEXT("MaterialX", "MaterialX")),
+			NAME_Compositing(LOCTEXT("Compositing", "Compositing"))
+		{}
+	};
+	static FConstructorStatics ConstructorStatics;
+
+#if WITH_EDITORONLY_DATA
+	MenuCategories.Add(ConstructorStatics.NAME_MaterialX);
+	MenuCategories.Add(ConstructorStatics.NAME_Compositing);
+#endif
+}
+
+#if WITH_EDITOR
+int32 UMaterialExpressionMask::Compile(FMaterialCompiler* Compiler, int32 OutputIndex)
+{
+	if(!A.GetTracedInput().Expression)
+	{
+		return Compiler->Errorf(TEXT("Missing A input"));
+	}
+
+	if(!B.GetTracedInput().Expression)
+	{
+		return Compiler->Errorf(TEXT("Missing B input"));
+	}
+
+	int32 IndexAlpha = Alpha.GetTracedInput().Expression ? Alpha.Compile(Compiler) : Compiler->Constant(ConstAlpha);
+
+	int32 IndexA = A.Compile(Compiler);
+	int32 IndexB = B.Compile(Compiler);
+
+	int32 Mask = Compiler->Mul(IndexB, Compiler->ComponentMask(IndexA, false, false, false, true));
+
+	return Compiler->Lerp(IndexB, Mask, IndexAlpha);
+}
+
+void UMaterialExpressionMask::GetCaption(TArray<FString>& OutCaptions) const
+{
+	OutCaptions.Add(TEXT("Mask"));
+}
+#endif
+
+#undef LOCTEXT_NAMESPACE 
\ No newline at end of file
diff --git a/Engine/Plugins/Interchange/Runtime/Source/Import/Private/MaterialX/MaterialExpressionMask.h b/Engine/Plugins/Interchange/Runtime/Source/Import/Private/MaterialX/MaterialExpressionMask.h
new file mode 100644
--- /dev/null
+++ b/Engine/Plugins/Interchange/Runtime/Source/Import/Private/MaterialX/MaterialExpressionMask.h
@@ -0,0 +1,45 @@
+// Copyright Epic Games, Inc. All Rights Reserved.
+#pragma once
+
+#include "CoreMinimal.h"
+#include "Materials/MaterialExpression.h"
+
+#include "MaterialExpressionMask.generated.h"
+
+/**
+ * Merge nodes take two 4-channel (color4) inputs and use the built-in alpha channel(s) to control the
+ * compositing of the A and B inputs. "A" and "B" refer to the non-alpha channels of the A and B inputs respectively,
+ * and "a" and "b" refer to the alpha channels of the A and B inputs.
+ * Merge nodes are only defined for float4 inputs
+ * Merge nodes support an optional float input Alpha , which can be used to mix the
+ * original B value with the result of the blend operation.
+ *
+ * Operation: B*a
+ * Result: Lerp(B, B*a, Alpha)
+ */
+UCLASS(MinimalAPI, collapsecategories, hidecategories = Object)
+class UMaterialExpressionMask : public UMaterialExpression
+{
+	GENERATED_UCLASS_BODY()
+
+	UPROPERTY()
+	FExpressionInput A;
+
+	UPROPERTY()
+	FExpressionInput B;
+
+	UPROPERTY(meta = (RequiredInput = "false", ToolTip = "Defaults to 'ConstAlpha' if not specified"))
+	FExpressionInput Alpha;
+
+	/** only used if Alpha is not hooked up */
+	UPROPERTY(EditAnywhere, Category = MaterialExpressionLinearInterpolate, meta = (OverridingInputProperty = "Alpha"))
+	float ConstAlpha = 1.f;
+
+	//~ Begin UMaterialExpression Interface
+#if WITH_EDITOR
+	virtual int32 Compile(class FMaterialCompiler* Compiler, int32 OutputIndex) override;
+	virtual void GetCaption(TArray<FString>& OutCaptions) const override;
+#endif
+	//~ End UMaterialExpression Interface
+};
+
diff --git a/Engine/Plugins/Interchange/Runtime/Source/Import/Private/MaterialX/MaterialExpressionMatte.cpp b/Engine/Plugins/Interchange/Runtime/Source/Import/Private/MaterialX/MaterialExpressionMatte.cpp
new file mode 100644
--- /dev/null
+++ b/Engine/Plugins/Interchange/Runtime/Source/Import/Private/MaterialX/MaterialExpressionMatte.cpp
@@ -0,0 +1,66 @@
+// Copyright Epic Games, Inc. All Rights Reserved.
+#include "MaterialX/MaterialExpressionMatte.h"
+#include "MaterialCompiler.h"
+
+#define LOCTEXT_NAMESPACE "MaterialExpressionMatte"
+
+UMaterialExpressionMatte::UMaterialExpressionMatte(const FObjectInitializer& ObjectInitializer)
+	: Super(ObjectInitializer)
+{
+	// Structure to hold one-time initialization
+	struct FConstructorStatics
+	{
+		FText NAME_MaterialX;
+		FText NAME_Compositing;
+		FConstructorStatics() :
+			NAME_MaterialX(LOCTEXT("MaterialX", "MaterialX")),
+			NAME_Compositing(LOCTEXT("Compositing", "Compositing"))
+		{}
+	};
+	static FConstructorStatics ConstructorStatics;
+
+#if WITH_EDITORONLY_DATA
+	MenuCategories.Add(ConstructorStatics.NAME_MaterialX);
+	MenuCategories.Add(ConstructorStatics.NAME_Compositing);
+#endif
+}
+
+#if WITH_EDITOR
+int32 UMaterialExpressionMatte::Compile(FMaterialCompiler* Compiler, int32 OutputIndex)
+{
+	if(!A.GetTracedInput().Expression)
+	{
+		return Compiler->Errorf(TEXT("Missing A input"));
+	}
+
+	if(!B.GetTracedInput().Expression)
+	{
+		return Compiler->Errorf(TEXT("Missing B input"));
+	}
+
+	int32 IndexAlpha = Alpha.GetTracedInput().Expression ? Alpha.Compile(Compiler) : Compiler->Constant(ConstAlpha);
+
+	int32 IndexA = A.Compile(Compiler);
+	int32 IndexB = B.Compile(Compiler);
+
+	int32 IndexAlphaA = Compiler->ComponentMask(IndexA, false, false, false, true);
+	int32 IndexRgbA = Compiler->ComponentMask(IndexA, true, true, true, false);
+	int32 IndexAlphaB = Compiler->ComponentMask(IndexB, false, false, false, true);
+	int32 IndexRgbB = Compiler->ComponentMask(IndexB, true, true, true, false);
+
+	int32 IndexOneMinusAlphaA = Compiler->Sub(Compiler->Constant(1.f), IndexAlphaA);
+	int32 XYZ = Compiler->Add(Compiler->Mul(IndexRgbA, IndexAlphaA), Compiler->Mul(IndexRgbB, IndexOneMinusAlphaA));
+	int32 W = Compiler->Add(IndexAlphaA, Compiler->Mul(IndexAlphaB, IndexOneMinusAlphaA));
+
+	int32 Matte = Compiler->AppendVector(XYZ, W);
+
+	return Compiler->Lerp(IndexB, Matte, IndexAlpha);
+}
+
+void UMaterialExpressionMatte::GetCaption(TArray<FString>& OutCaptions) const
+{
+	OutCaptions.Add(TEXT("Matte"));
+}
+#endif
+
+#undef LOCTEXT_NAMESPACE 
\ No newline at end of file
diff --git a/Engine/Plugins/Interchange/Runtime/Source/Import/Private/MaterialX/MaterialExpressionMatte.h b/Engine/Plugins/Interchange/Runtime/Source/Import/Private/MaterialX/MaterialExpressionMatte.h
new file mode 100644
--- /dev/null
+++ b/Engine/Plugins/Interchange/Runtime/Source/Import/Private/MaterialX/MaterialExpressionMatte.h
@@ -0,0 +1,45 @@
+// Copyright Epic Games, Inc. All Rights Reserved.
+#pragma once
+
+#include "CoreMinimal.h"
+#include "Materials/MaterialExpression.h"
+
+#include "MaterialExpressionMatte.generated.h"
+
+/**
+ * Merge nodes take two 4-channel (color4) inputs and use the built-in alpha channel(s) to control the
+ * compositing of the A and B inputs. "A" and "B" refer to the non-alpha channels of the A and B inputs respectively,
+ * and "a" and "b" refer to the alpha channels of the A and B inputs.
+ * Merge nodes are only defined for float4 inputs
+ * Merge nodes support an optional float input Alpha , which can be used to mix the
+ * original B value with the result of the blend operation.
+ *
+ * Operation: A*a + B*(1-a)
+ * Result: Lerp(B, A*a + B*(1-a), Alpha)
+ */
+UCLASS(MinimalAPI, collapsecategories, hidecategories = Object)
+class UMaterialExpressionMatte : public UMaterialExpression
+{
+	GENERATED_UCLASS_BODY()
+
+	UPROPERTY()
+	FExpressionInput A;
+
+	UPROPERTY()
+	FExpressionInput B;
+
+	UPROPERTY(meta = (RequiredInput = "false", ToolTip = "Defaults to 'ConstAlpha' if not specified"))
+	FExpressionInput Alpha;
+
+	/** only used if Alpha is not hooked up */
+	UPROPERTY(EditAnywhere, Category = MaterialExpressionLinearInterpolate, meta = (OverridingInputProperty = "Alpha"))
+	float ConstAlpha = 1.f;
+
+	//~ Begin UMaterialExpression Interface
+#if WITH_EDITOR
+	virtual int32 Compile(class FMaterialCompiler* Compiler, int32 OutputIndex) override;
+	virtual void GetCaption(TArray<FString>& OutCaptions) const override;
+#endif
+	//~ End UMaterialExpression Interface
+};
+
diff --git a/Engine/Plugins/Interchange/Runtime/Source/Import/Private/MaterialX/MaterialExpressionMinus.cpp b/Engine/Plugins/Interchange/Runtime/Source/Import/Private/MaterialX/MaterialExpressionMinus.cpp
new file mode 100644
--- /dev/null
+++ b/Engine/Plugins/Interchange/Runtime/Source/Import/Private/MaterialX/MaterialExpressionMinus.cpp
@@ -0,0 +1,54 @@
+// Copyright Epic Games, Inc. All Rights Reserved.
+#include "MaterialX/MaterialExpressionMinus.h"
+#include "MaterialCompiler.h"
+
+#define LOCTEXT_NAMESPACE "MaterialExpressionMinus"
+
+UMaterialExpressionMinus::UMaterialExpressionMinus(const FObjectInitializer& ObjectInitializer)
+	: Super(ObjectInitializer)
+{
+	// Structure to hold one-time initialization
+	struct FConstructorStatics
+	{
+		FText NAME_MaterialX;
+		FText NAME_Compositing;
+		FConstructorStatics() :
+			NAME_MaterialX(LOCTEXT("MaterialX", "MaterialX")),
+			NAME_Compositing(LOCTEXT("Compositing", "Compositing"))
+		{}
+	};
+	static FConstructorStatics ConstructorStatics;
+
+#if WITH_EDITORONLY_DATA
+	MenuCategories.Add(ConstructorStatics.NAME_MaterialX);
+	MenuCategories.Add(ConstructorStatics.NAME_Compositing);
+#endif
+}
+
+#if WITH_EDITOR
+int32 UMaterialExpressionMinus::Compile(FMaterialCompiler* Compiler, int32 OutputIndex)
+{
+	if(!A.GetTracedInput().Expression)
+	{
+		return Compiler->Errorf(TEXT("Missing A input"));
+	}
+
+	if(!B.GetTracedInput().Expression)
+	{
+		return Compiler->Errorf(TEXT("Missing B input"));
+	}
+
+	int32 IndexAlpha = Alpha.GetTracedInput().Expression ? Alpha.Compile(Compiler) : Compiler->Constant(ConstAlpha);
+
+	int32 IndexB = B.Compile(Compiler);
+	int32 Sub = Compiler->Sub(A.Compile(Compiler), IndexB);
+	return Compiler->Lerp(IndexB, Sub, IndexAlpha);
+}
+
+void UMaterialExpressionMinus::GetCaption(TArray<FString>& OutCaptions) const
+{
+	OutCaptions.Add(TEXT("Minus"));
+}
+#endif
+
+#undef LOCTEXT_NAMESPACE 
\ No newline at end of file
diff --git a/Engine/Plugins/Interchange/Runtime/Source/Import/Private/MaterialX/MaterialExpressionMinus.h b/Engine/Plugins/Interchange/Runtime/Source/Import/Private/MaterialX/MaterialExpressionMinus.h
new file mode 100644
--- /dev/null
+++ b/Engine/Plugins/Interchange/Runtime/Source/Import/Private/MaterialX/MaterialExpressionMinus.h
@@ -0,0 +1,41 @@
+// Copyright Epic Games, Inc. All Rights Reserved.
+#pragma once
+
+#include "CoreMinimal.h"
+#include "Materials/MaterialExpression.h"
+
+#include "MaterialExpressionMinus.generated.h"
+
+/**
+ * Blend nodes take two 1-4 channel inputs and apply the same operator to all channels.
+ * Blend nodes support an optional float input mix , which can be used
+ * to mix the original B value with the result of the blend operation.
+ * Operation: A - B
+ * Result: Lerp(B, A - B, Alpha)
+ */
+UCLASS(MinimalAPI, collapsecategories, hidecategories = Object)
+class UMaterialExpressionMinus : public UMaterialExpression
+{
+	GENERATED_UCLASS_BODY()
+
+	UPROPERTY()
+	FExpressionInput A;
+
+	UPROPERTY()
+	FExpressionInput B;
+
+	UPROPERTY(meta = (RequiredInput = "false", ToolTip = "Defaults to 'ConstAlpha' if not specified"))
+	FExpressionInput Alpha;
+
+	/** only used if Alpha is not hooked up */
+	UPROPERTY(EditAnywhere, Category = MaterialExpressionLinearInterpolate, meta = (OverridingInputProperty = "Alpha"))
+	float ConstAlpha = 1.f;
+
+	//~ Begin UMaterialExpression Interface
+#if WITH_EDITOR
+	virtual int32 Compile(class FMaterialCompiler* Compiler, int32 OutputIndex) override;
+	virtual void GetCaption(TArray<FString>& OutCaptions) const override;
+#endif
+	//~ End UMaterialExpression Interface
+};
+
diff --git a/Engine/Plugins/Interchange/Runtime/Source/Import/Private/MaterialX/MaterialExpressionOut.cpp b/Engine/Plugins/Interchange/Runtime/Source/Import/Private/MaterialX/MaterialExpressionOut.cpp
new file mode 100644
--- /dev/null
+++ b/Engine/Plugins/Interchange/Runtime/Source/Import/Private/MaterialX/MaterialExpressionOut.cpp
@@ -0,0 +1,60 @@
+// Copyright Epic Games, Inc. All Rights Reserved.
+#include "MaterialX/MaterialExpressionOut.h"
+#include "MaterialCompiler.h"
+
+#define LOCTEXT_NAMESPACE "MaterialExpressionOut"
+
+UMaterialExpressionOut::UMaterialExpressionOut(const FObjectInitializer& ObjectInitializer)
+	: Super(ObjectInitializer)
+{
+	// Structure to hold one-time initialization
+	struct FConstructorStatics
+	{
+		FText NAME_MaterialX;
+		FText NAME_Compositing;
+		FConstructorStatics() :
+			NAME_MaterialX(LOCTEXT("MaterialX", "MaterialX")),
+			NAME_Compositing(LOCTEXT("Compositing", "Compositing"))
+		{}
+	};
+	static FConstructorStatics ConstructorStatics;
+
+#if WITH_EDITORONLY_DATA
+	MenuCategories.Add(ConstructorStatics.NAME_MaterialX);
+	MenuCategories.Add(ConstructorStatics.NAME_Compositing);
+#endif
+}
+
+#if WITH_EDITOR
+int32 UMaterialExpressionOut::Compile(FMaterialCompiler* Compiler, int32 OutputIndex)
+{
+	if(!A.GetTracedInput().Expression)
+	{
+		return Compiler->Errorf(TEXT("Missing A input"));
+	}
+
+	if(!B.GetTracedInput().Expression)
+	{
+		return Compiler->Errorf(TEXT("Missing B input"));
+	}
+
+	int32 IndexAlpha = Alpha.GetTracedInput().Expression ? Alpha.Compile(Compiler) : Compiler->Constant(ConstAlpha);
+
+	int32 IndexA = A.Compile(Compiler);
+	int32 IndexB = B.Compile(Compiler);
+
+	int32 IndexAlphaB = Compiler->ComponentMask(IndexB, false, false, false, true);
+
+	int32 IndexOneMinusAlphaB = Compiler->Sub(Compiler->Constant(1.f), IndexAlphaB);
+	int32 Out = Compiler->Mul(IndexA, IndexOneMinusAlphaB);
+
+	return Compiler->Lerp(IndexB, Out, IndexAlpha);
+}
+
+void UMaterialExpressionOut::GetCaption(TArray<FString>& OutCaptions) const
+{
+	OutCaptions.Add(TEXT("Out"));
+}
+#endif
+
+#undef LOCTEXT_NAMESPACE 
\ No newline at end of file
diff --git a/Engine/Plugins/Interchange/Runtime/Source/Import/Private/MaterialX/MaterialExpressionOut.h b/Engine/Plugins/Interchange/Runtime/Source/Import/Private/MaterialX/MaterialExpressionOut.h
new file mode 100644
--- /dev/null
+++ b/Engine/Plugins/Interchange/Runtime/Source/Import/Private/MaterialX/MaterialExpressionOut.h
@@ -0,0 +1,45 @@
+// Copyright Epic Games, Inc. All Rights Reserved.
+#pragma once
+
+#include "CoreMinimal.h"
+#include "Materials/MaterialExpression.h"
+
+#include "MaterialExpressionOut.generated.h"
+
+/**
+ * Merge nodes take two 4-channel (color4) inputs and use the built-in alpha channel(s) to control the
+ * compositing of the A and B inputs. "A" and "B" refer to the non-alpha channels of the A and B inputs respectively,
+ * and "a" and "b" refer to the alpha channels of the A and B inputs.
+ * Merge nodes are only defined for float4 inputs
+ * Merge nodes support an optional float input Alpha , which can be used to mix the
+ * original B value with the result of the blend operation.
+ *
+ * Operation: A(1-b)
+ * Result: Lerp(B, A(1-b), Alpha)
+ */
+UCLASS(MinimalAPI, collapsecategories, hidecategories = Object)
+class UMaterialExpressionOut : public UMaterialExpression
+{
+	GENERATED_UCLASS_BODY()
+
+	UPROPERTY()
+	FExpressionInput A;
+
+	UPROPERTY()
+	FExpressionInput B;
+
+	UPROPERTY(meta = (RequiredInput = "false", ToolTip = "Defaults to 'ConstAlpha' if not specified"))
+	FExpressionInput Alpha;
+
+	/** only used if Alpha is not hooked up */
+	UPROPERTY(EditAnywhere, Category = MaterialExpressionLinearInterpolate, meta = (OverridingInputProperty = "Alpha"))
+	float ConstAlpha = 1.f;
+
+	//~ Begin UMaterialExpression Interface
+#if WITH_EDITOR
+	virtual int32 Compile(class FMaterialCompiler* Compiler, int32 OutputIndex) override;
+	virtual void GetCaption(TArray<FString>& OutCaptions) const override;
+#endif
+	//~ End UMaterialExpression Interface
+};
+
diff --git a/Engine/Plugins/Interchange/Runtime/Source/Import/Private/MaterialX/MaterialExpressionOver.cpp b/Engine/Plugins/Interchange/Runtime/Source/Import/Private/MaterialX/MaterialExpressionOver.cpp
new file mode 100644
--- /dev/null
+++ b/Engine/Plugins/Interchange/Runtime/Source/Import/Private/MaterialX/MaterialExpressionOver.cpp
@@ -0,0 +1,60 @@
+// Copyright Epic Games, Inc. All Rights Reserved.
+#include "MaterialX/MaterialExpressionOver.h"
+#include "MaterialCompiler.h"
+
+#define LOCTEXT_NAMESPACE "MaterialExpressionOver"
+
+UMaterialExpressionOver::UMaterialExpressionOver(const FObjectInitializer& ObjectInitializer)
+	: Super(ObjectInitializer)
+{
+	// Structure to hold one-time initialization
+	struct FConstructorStatics
+	{
+		FText NAME_MaterialX;
+		FText NAME_Compositing;
+		FConstructorStatics() :
+			NAME_MaterialX(LOCTEXT("MaterialX", "MaterialX")),
+			NAME_Compositing(LOCTEXT("Compositing", "Compositing"))
+		{}
+	};
+	static FConstructorStatics ConstructorStatics;
+
+#if WITH_EDITORONLY_DATA
+	MenuCategories.Add(ConstructorStatics.NAME_MaterialX);
+	MenuCategories.Add(ConstructorStatics.NAME_Compositing);
+#endif
+}
+
+#if WITH_EDITOR
+int32 UMaterialExpressionOver::Compile(FMaterialCompiler* Compiler, int32 OutputIndex)
+{
+	if(!A.GetTracedInput().Expression)
+	{
+		return Compiler->Errorf(TEXT("Missing A input"));
+	}
+
+	if(!B.GetTracedInput().Expression)
+	{
+		return Compiler->Errorf(TEXT("Missing B input"));
+	}
+
+	int32 IndexAlpha = Alpha.GetTracedInput().Expression ? Alpha.Compile(Compiler) : Compiler->Constant(ConstAlpha);
+
+	int32 IndexA = A.Compile(Compiler);
+	int32 IndexB = B.Compile(Compiler);
+
+	int32 IndexAlphaA = Compiler->ComponentMask(IndexA, false, false, false, true);
+
+	int32 IndexOneMinusAlphaA = Compiler->Sub(Compiler->Constant(1.f), IndexAlphaA);
+	int32 Over = Compiler->Add(IndexA, Compiler->Mul(IndexB, IndexOneMinusAlphaA));
+
+	return Compiler->Lerp(IndexB, Over, IndexAlpha);
+}
+
+void UMaterialExpressionOver::GetCaption(TArray<FString>& OutCaptions) const
+{
+	OutCaptions.Add(TEXT("Over"));
+}
+#endif
+
+#undef LOCTEXT_NAMESPACE 
\ No newline at end of file
diff --git a/Engine/Plugins/Interchange/Runtime/Source/Import/Private/MaterialX/MaterialExpressionOver.h b/Engine/Plugins/Interchange/Runtime/Source/Import/Private/MaterialX/MaterialExpressionOver.h
new file mode 100644
--- /dev/null
+++ b/Engine/Plugins/Interchange/Runtime/Source/Import/Private/MaterialX/MaterialExpressionOver.h
@@ -0,0 +1,45 @@
+// Copyright Epic Games, Inc. All Rights Reserved.
+#pragma once
+
+#include "CoreMinimal.h"
+#include "Materials/MaterialExpression.h"
+
+#include "MaterialExpressionOver.generated.h"
+
+/**
+ * Merge nodes take two 4-channel (color4) inputs and use the built-in alpha channel(s) to control the
+ * compositing of the A and B inputs. "A" and "B" refer to the non-alpha channels of the A and B inputs respectively,
+ * and "a" and "b" refer to the alpha channels of the A and B inputs.
+ * Merge nodes are only defined for float4 inputs
+ * Merge nodes support an optional float input Alpha , which can be used to mix the
+ * original B value with the result of the blend operation.
+ *
+ * Operation: A + B(1-a)
+ * Result: Lerp(B, A + B(1-a), Alpha)
+ */
+UCLASS(MinimalAPI, collapsecategories, hidecategories = Object)
+class UMaterialExpressionOver : public UMaterialExpression
+{
+	GENERATED_UCLASS_BODY()
+
+	UPROPERTY()
+	FExpressionInput A;
+
+	UPROPERTY()
+	FExpressionInput B;
+
+	UPROPERTY(meta = (RequiredInput = "false", ToolTip = "Defaults to 'ConstAlpha' if not specified"))
+	FExpressionInput Alpha;
+
+	/** only used if Alpha is not hooked up */
+	UPROPERTY(EditAnywhere, Category = MaterialExpressionLinearInterpolate, meta = (OverridingInputProperty = "Alpha"))
+	float ConstAlpha = 1.f;
+
+	//~ Begin UMaterialExpression Interface
+#if WITH_EDITOR
+	virtual int32 Compile(class FMaterialCompiler* Compiler, int32 OutputIndex) override;
+	virtual void GetCaption(TArray<FString>& OutCaptions) const override;
+#endif
+	//~ End UMaterialExpression Interface
+};
+
diff --git a/Engine/Plugins/Interchange/Runtime/Source/Import/Private/MaterialX/MaterialExpressionOverlay.cpp b/Engine/Plugins/Interchange/Runtime/Source/Import/Private/MaterialX/MaterialExpressionOverlay.cpp
new file mode 100644
--- /dev/null
+++ b/Engine/Plugins/Interchange/Runtime/Source/Import/Private/MaterialX/MaterialExpressionOverlay.cpp
@@ -0,0 +1,69 @@
+// Copyright Epic Games, Inc. All Rights Reserved.
+#include "MaterialX/MaterialExpressionOverlay.h"
+#include "MaterialCompiler.h"
+
+#define LOCTEXT_NAMESPACE "MaterialExpressionOverlay"
+
+UMaterialExpressionOverlay::UMaterialExpressionOverlay(const FObjectInitializer& ObjectInitializer)
+	: Super(ObjectInitializer)
+{
+	// Structure to hold one-time initialization
+	struct FConstructorStatics
+	{
+		FText NAME_MaterialX;
+		FText NAME_Compositing;
+		FConstructorStatics() :
+			NAME_MaterialX(LOCTEXT("MaterialX", "MaterialX")),
+			NAME_Compositing(LOCTEXT("Compositing", "Compositing"))
+		{}
+	};
+	static FConstructorStatics ConstructorStatics;
+
+#if WITH_EDITORONLY_DATA
+	MenuCategories.Add(ConstructorStatics.NAME_MaterialX);
+	MenuCategories.Add(ConstructorStatics.NAME_Compositing);
+#endif
+}
+
+#if WITH_EDITOR
+int32 UMaterialExpressionOverlay::Compile(FMaterialCompiler* Compiler, int32 OutputIndex)
+{
+	if(!A.GetTracedInput().Expression)
+	{
+		return Compiler->Errorf(TEXT("Missing A input"));
+	}
+
+	if(!B.GetTracedInput().Expression)
+	{
+		return Compiler->Errorf(TEXT("Missing B input"));
+	}
+
+	int32 IndexAlpha = Alpha.GetTracedInput().Expression ? Alpha.Compile(Compiler) : Compiler->Constant(ConstAlpha);
+
+	int32 OneIndex = Compiler->Constant(1.f);
+	auto OneMinus = [&](int32 index)
+	{
+		return Compiler->Sub(OneIndex, index);
+	};
+
+	int32 IndexA = A.Compile(Compiler);
+	int32 IndexB = B.Compile(Compiler);
+	int32 Screen = OneMinus(Compiler->Mul(OneMinus(IndexA), OneMinus(IndexB)));
+
+	int32 Overlay = Compiler->If(IndexA,
+								 Compiler->Constant(0.5f),
+								 Screen,
+								 Screen,
+								 Compiler->Mul(Compiler->Mul(Compiler->Constant(2.f), IndexA), IndexB),
+								 Compiler->Constant(0.00001f));
+
+	return Compiler->Lerp(IndexB, Overlay, IndexAlpha);
+}
+
+void UMaterialExpressionOverlay::GetCaption(TArray<FString>& OutCaptions) const
+{
+	OutCaptions.Add(TEXT("Overlay"));
+}
+#endif
+
+#undef LOCTEXT_NAMESPACE 
\ No newline at end of file
diff --git a/Engine/Plugins/Interchange/Runtime/Source/Import/Private/MaterialX/MaterialExpressionOverlay.h b/Engine/Plugins/Interchange/Runtime/Source/Import/Private/MaterialX/MaterialExpressionOverlay.h
new file mode 100644
--- /dev/null
+++ b/Engine/Plugins/Interchange/Runtime/Source/Import/Private/MaterialX/MaterialExpressionOverlay.h
@@ -0,0 +1,42 @@
+// Copyright Epic Games, Inc. All Rights Reserved.
+#pragma once
+
+#include "CoreMinimal.h"
+#include "Materials/MaterialExpression.h"
+
+#include "MaterialExpressionOverlay.generated.h"
+
+/**
+ * Blend nodes take two 1-4 channel inputs and apply the same operator to all channels.
+ * Blend nodes support an optional float input mix , which can be used
+ * to mix the original B value with the result of the blend operation.
+ * Operation: 2*A*B          if A < 0.5;
+ *            1-(1-A)(1-B) if A >= 0.5
+ * Result: Lerp(B, Op, Alpha)
+ */
+UCLASS(MinimalAPI, collapsecategories, hidecategories = Object)
+class UMaterialExpressionOverlay : public UMaterialExpression
+{
+	GENERATED_UCLASS_BODY()
+
+	UPROPERTY()
+	FExpressionInput A;
+
+	UPROPERTY()
+	FExpressionInput B;
+
+	UPROPERTY(meta = (RequiredInput = "false", ToolTip = "Defaults to 'ConstAlpha' if not specified"))
+	FExpressionInput Alpha;
+
+	/** only used if Alpha is not hooked up */
+	UPROPERTY(EditAnywhere, Category = MaterialExpressionLinearInterpolate, meta = (OverridingInputProperty = "Alpha"))
+	float ConstAlpha = 1.f;
+
+	//~ Begin UMaterialExpression Interface
+#if WITH_EDITOR
+	virtual int32 Compile(class FMaterialCompiler* Compiler, int32 OutputIndex) override;
+	virtual void GetCaption(TArray<FString>& OutCaptions) const override;
+#endif
+	//~ End UMaterialExpression Interface
+};
+
diff --git a/Engine/Plugins/Interchange/Runtime/Source/Import/Private/MaterialX/MaterialExpressionPlus.cpp b/Engine/Plugins/Interchange/Runtime/Source/Import/Private/MaterialX/MaterialExpressionPlus.cpp
new file mode 100644
--- /dev/null
+++ b/Engine/Plugins/Interchange/Runtime/Source/Import/Private/MaterialX/MaterialExpressionPlus.cpp
@@ -0,0 +1,54 @@
+// Copyright Epic Games, Inc. All Rights Reserved.
+#include "MaterialX/MaterialExpressionPlus.h"
+#include "MaterialCompiler.h"
+
+#define LOCTEXT_NAMESPACE "MaterialExpressionPlus"
+
+UMaterialExpressionPlus::UMaterialExpressionPlus(const FObjectInitializer& ObjectInitializer)
+	: Super(ObjectInitializer)
+{
+	// Structure to hold one-time initialization
+	struct FConstructorStatics
+	{
+		FText NAME_MaterialX;
+		FText NAME_Compositing;
+		FConstructorStatics() :
+			NAME_MaterialX(LOCTEXT("MaterialX", "MaterialX")),
+			NAME_Compositing(LOCTEXT("Compositing", "Compositing"))
+		{}
+	};
+	static FConstructorStatics ConstructorStatics;
+
+#if WITH_EDITORONLY_DATA
+	MenuCategories.Add(ConstructorStatics.NAME_MaterialX);
+	MenuCategories.Add(ConstructorStatics.NAME_Compositing);
+#endif
+}
+
+#if WITH_EDITOR
+int32 UMaterialExpressionPlus::Compile(FMaterialCompiler* Compiler, int32 OutputIndex)
+{
+	if(!A.GetTracedInput().Expression)
+	{
+		return Compiler->Errorf(TEXT("Missing A input"));
+	}
+
+	if(!B.GetTracedInput().Expression)
+	{
+		return Compiler->Errorf(TEXT("Missing B input"));
+	}
+
+	int32 IndexAlpha = Alpha.GetTracedInput().Expression ? Alpha.Compile(Compiler) : Compiler->Constant(ConstAlpha);
+
+	int32 IndexB = B.Compile(Compiler);
+	int32 Add = Compiler->Add(A.Compile(Compiler), IndexB);
+	return Compiler->Lerp(IndexB, Add, IndexAlpha);
+}
+
+void UMaterialExpressionPlus::GetCaption(TArray<FString>& OutCaptions) const
+{
+	OutCaptions.Add(TEXT("Plus"));
+}
+#endif
+
+#undef LOCTEXT_NAMESPACE 
\ No newline at end of file
diff --git a/Engine/Plugins/Interchange/Runtime/Source/Import/Private/MaterialX/MaterialExpressionPlus.h b/Engine/Plugins/Interchange/Runtime/Source/Import/Private/MaterialX/MaterialExpressionPlus.h
new file mode 100644
--- /dev/null
+++ b/Engine/Plugins/Interchange/Runtime/Source/Import/Private/MaterialX/MaterialExpressionPlus.h
@@ -0,0 +1,41 @@
+// Copyright Epic Games, Inc. All Rights Reserved.
+#pragma once
+
+#include "CoreMinimal.h"
+#include "Materials/MaterialExpression.h"
+
+#include "MaterialExpressionPlus.generated.h"
+
+/**
+ * Blend nodes take two 1-4 channel inputs and apply the same operator to all channels.
+ * Blend nodes support an optional float input mix , which can be used
+ * to mix the original B value with the result of the blend operation.
+ * Operation: A + B
+ * Result: Lerp(B, A + B, Alpha)
+ */
+UCLASS(MinimalAPI, collapsecategories, hidecategories = Object)
+class UMaterialExpressionPlus : public UMaterialExpression
+{
+	GENERATED_UCLASS_BODY()
+
+	UPROPERTY()
+	FExpressionInput A;
+
+	UPROPERTY()
+	FExpressionInput B;
+
+	UPROPERTY(meta = (RequiredInput = "false", ToolTip = "Defaults to 'ConstAlpha' if not specified"))
+	FExpressionInput Alpha;
+
+	/** only used if Alpha is not hooked up */
+	UPROPERTY(EditAnywhere, Category = MaterialExpressionLinearInterpolate, meta = (OverridingInputProperty = "Alpha"))
+	float ConstAlpha = 1.f;
+
+	//~ Begin UMaterialExpression Interface
+#if WITH_EDITOR
+	virtual int32 Compile(class FMaterialCompiler* Compiler, int32 OutputIndex) override;
+	virtual void GetCaption(TArray<FString>& OutCaptions) const override;
+#endif
+	//~ End UMaterialExpression Interface
+};
+
diff --git a/Engine/Plugins/Interchange/Runtime/Source/Import/Private/MaterialX/MaterialExpressionPremult.cpp b/Engine/Plugins/Interchange/Runtime/Source/Import/Private/MaterialX/MaterialExpressionPremult.cpp
new file mode 100644
--- /dev/null
+++ b/Engine/Plugins/Interchange/Runtime/Source/Import/Private/MaterialX/MaterialExpressionPremult.cpp
@@ -0,0 +1,45 @@
+// Copyright Epic Games, Inc. All Rights Reserved.
+#include "MaterialX/MaterialExpressionPremult.h"
+#include "MaterialCompiler.h"
+
+#define LOCTEXT_NAMESPACE "MaterialExpressionPremult"
+
+UMaterialExpressionPremult::UMaterialExpressionPremult(const FObjectInitializer& ObjectInitializer)
+	: Super(ObjectInitializer)
+{
+	// Structure to hold one-time initialization
+	struct FConstructorStatics
+	{
+		FText NAME_Utility;
+		FConstructorStatics()
+			: NAME_Utility(LOCTEXT("Compositing", "Compositing"))
+		{}
+	};
+	static FConstructorStatics ConstructorStatics;
+
+#if WITH_EDITORONLY_DATA
+	MenuCategories.Add(ConstructorStatics.NAME_Utility);
+#endif
+}
+
+#if WITH_EDITOR
+int32 UMaterialExpressionPremult::Compile(FMaterialCompiler* Compiler, int32 OutputIndex)
+{
+	if(!Input.GetTracedInput().Expression)
+	{
+		return Compiler->Errorf(TEXT("Missing Premult input"));
+	}
+
+	int32 IndexInput = Input.Compile(Compiler);
+	int32 RGB = Compiler->ComponentMask(IndexInput, true, true, true, false);
+	int32 Alpha = Compiler->ComponentMask(IndexInput, false, false, false, true);
+	return Compiler->AppendVector(Compiler->Mul(RGB, Alpha), Alpha);
+}
+
+void UMaterialExpressionPremult::GetCaption(TArray<FString>& OutCaptions) const
+{
+	OutCaptions.Add(TEXT("Premult"));
+}
+#endif
+
+#undef LOCTEXT_NAMESPACE 
\ No newline at end of file
diff --git a/Engine/Plugins/Interchange/Runtime/Source/Import/Private/MaterialX/MaterialExpressionPremult.h b/Engine/Plugins/Interchange/Runtime/Source/Import/Private/MaterialX/MaterialExpressionPremult.h
new file mode 100644
--- /dev/null
+++ b/Engine/Plugins/Interchange/Runtime/Source/Import/Private/MaterialX/MaterialExpressionPremult.h
@@ -0,0 +1,28 @@
+// Copyright Epic Games, Inc. All Rights Reserved.
+#pragma once
+
+#include "CoreMinimal.h"
+#include "Materials/MaterialExpression.h"
+
+#include "MaterialExpressionPremult.generated.h"
+
+/**
+ * Multiply the RGB channels of the input by the Alpha channel of the input.
+ * Input must be of type float4
+ */
+UCLASS(collapsecategories, hidecategories = Object, MinimalAPI)
+class UMaterialExpressionPremult : public UMaterialExpression
+{
+	GENERATED_UCLASS_BODY()
+
+	UPROPERTY()
+	FExpressionInput Input;
+
+	//~ Begin UMaterialExpression Interface
+#if WITH_EDITOR
+	virtual int32 Compile(class FMaterialCompiler* Compiler, int32 OutputIndex) override;
+	virtual void GetCaption(TArray<FString>& OutCaptions) const override;
+#endif
+	//~ End UMaterialExpression Interface
+};
+
diff --git a/Engine/Plugins/Interchange/Runtime/Source/Import/Private/MaterialX/MaterialExpressionScreen.cpp b/Engine/Plugins/Interchange/Runtime/Source/Import/Private/MaterialX/MaterialExpressionScreen.cpp
new file mode 100644
--- /dev/null
+++ b/Engine/Plugins/Interchange/Runtime/Source/Import/Private/MaterialX/MaterialExpressionScreen.cpp
@@ -0,0 +1,62 @@
+// Copyright Epic Games, Inc. All Rights Reserved.
+#include "MaterialX/MaterialExpressionScreen.h"
+#include "MaterialCompiler.h"
+
+#define LOCTEXT_NAMESPACE "MaterialExpressionScreen"
+
+UMaterialExpressionScreen::UMaterialExpressionScreen(const FObjectInitializer& ObjectInitializer)
+	: Super(ObjectInitializer)
+{
+	// Structure to hold one-time initialization
+	struct FConstructorStatics
+	{
+		FText NAME_MaterialX;
+		FText NAME_Compositing;
+		FConstructorStatics() :
+			NAME_MaterialX(LOCTEXT("MaterialX", "MaterialX")),
+			NAME_Compositing(LOCTEXT("Compositing", "Compositing"))
+		{}
+	};
+	static FConstructorStatics ConstructorStatics;
+
+#if WITH_EDITORONLY_DATA
+	MenuCategories.Add(ConstructorStatics.NAME_MaterialX);
+	MenuCategories.Add(ConstructorStatics.NAME_Compositing);
+#endif
+}
+
+#if WITH_EDITOR
+int32 UMaterialExpressionScreen::Compile(FMaterialCompiler* Compiler, int32 OutputIndex)
+{
+	if(!A.GetTracedInput().Expression)
+	{
+		return Compiler->Errorf(TEXT("Missing A input"));
+	}
+
+	if(!B.GetTracedInput().Expression)
+	{
+		return Compiler->Errorf(TEXT("Missing B input"));
+	}
+
+	int32 IndexAlpha = Alpha.GetTracedInput().Expression ? Alpha.Compile(Compiler) : Compiler->Constant(ConstAlpha);
+
+	int32 OneIndex = Compiler->Constant(1.f);
+	auto OneMinus = [&](int32 index)
+	{
+		return Compiler->Sub(OneIndex, index);
+	};
+
+	int32 IndexA = A.Compile(Compiler);
+	int32 IndexB = B.Compile(Compiler);
+	int32 Screen = OneMinus(Compiler->Mul(OneMinus(IndexA), OneMinus(IndexB)));
+
+	return Compiler->Lerp(IndexB, Screen, IndexAlpha);
+}
+
+void UMaterialExpressionScreen::GetCaption(TArray<FString>& OutCaptions) const
+{
+	OutCaptions.Add(TEXT("Screen"));
+}
+#endif
+
+#undef LOCTEXT_NAMESPACE 
\ No newline at end of file
diff --git a/Engine/Plugins/Interchange/Runtime/Source/Import/Private/MaterialX/MaterialExpressionScreen.h b/Engine/Plugins/Interchange/Runtime/Source/Import/Private/MaterialX/MaterialExpressionScreen.h
new file mode 100644
--- /dev/null
+++ b/Engine/Plugins/Interchange/Runtime/Source/Import/Private/MaterialX/MaterialExpressionScreen.h
@@ -0,0 +1,41 @@
+// Copyright Epic Games, Inc. All Rights Reserved.
+#pragma once
+
+#include "CoreMinimal.h"
+#include "Materials/MaterialExpression.h"
+
+#include "MaterialExpressionScreen.generated.h"
+
+/**
+ * Blend nodes take two 1-4 channel inputs and apply the same operator to all channels.
+ * Blend nodes support an optional float input mix , which can be used
+ * to mix the original B value with the result of the blend operation.
+ * Operation: 1-(1-A)(1-B)
+ * Result: Lerp(B, 1-(1-F)(1-B), Alpha)
+ */
+UCLASS(MinimalAPI, collapsecategories, hidecategories = Object)
+class UMaterialExpressionScreen : public UMaterialExpression
+{
+	GENERATED_UCLASS_BODY()
+
+	UPROPERTY()
+	FExpressionInput A;
+
+	UPROPERTY()
+	FExpressionInput B;
+
+	UPROPERTY(meta = (RequiredInput = "false", ToolTip = "Defaults to 'ConstAlpha' if not specified"))
+	FExpressionInput Alpha;
+
+	/** only used if Alpha is not hooked up */
+	UPROPERTY(EditAnywhere, Category = MaterialExpressionLinearInterpolate, meta = (OverridingInputProperty = "Alpha"))
+	float ConstAlpha = 1.f;
+
+	//~ Begin UMaterialExpression Interface
+#if WITH_EDITOR
+	virtual int32 Compile(class FMaterialCompiler* Compiler, int32 OutputIndex) override;
+	virtual void GetCaption(TArray<FString>& OutCaptions) const override;
+#endif
+	//~ End UMaterialExpression Interface
+};
+
diff --git a/Engine/Plugins/Interchange/Runtime/Source/Import/Private/MaterialX/MaterialExpressionUnpremult.cpp b/Engine/Plugins/Interchange/Runtime/Source/Import/Private/MaterialX/MaterialExpressionUnpremult.cpp
new file mode 100644
--- /dev/null
+++ b/Engine/Plugins/Interchange/Runtime/Source/Import/Private/MaterialX/MaterialExpressionUnpremult.cpp
@@ -0,0 +1,50 @@
+// Copyright Epic Games, Inc. All Rights Reserved.
+#include "MaterialX/MaterialExpressionUnpremult.h"
+#include "MaterialCompiler.h"
+
+#define LOCTEXT_NAMESPACE "MaterialExpressionUnpremult"
+
+UMaterialExpressionUnpremult::UMaterialExpressionUnpremult(const FObjectInitializer& ObjectInitializer)
+	: Super(ObjectInitializer)
+{
+	// Structure to hold one-time initialization
+	struct FConstructorStatics
+	{
+		FText NAME_MaterialX;
+		FText NAME_Compositing;
+		FConstructorStatics() :
+			NAME_MaterialX(LOCTEXT("MaterialX", "MaterialX")),
+			NAME_Compositing(LOCTEXT("Compositing", "Compositing"))
+		{}
+	};
+	static FConstructorStatics ConstructorStatics;
+
+#if WITH_EDITORONLY_DATA
+	MenuCategories.Add(ConstructorStatics.NAME_MaterialX);
+	MenuCategories.Add(ConstructorStatics.NAME_Compositing);
+#endif
+}
+
+#if WITH_EDITOR
+int32 UMaterialExpressionUnpremult::Compile(FMaterialCompiler* Compiler, int32 OutputIndex)
+{
+	if(!Input.GetTracedInput().Expression)
+	{
+		return Compiler->Errorf(TEXT("Missing Unpremult input"));
+	}
+
+	int32 IndexInput = Input.Compile(Compiler);
+	int32 RGB = Compiler->ComponentMask(IndexInput, true, true, true, false);
+	int32 Alpha = Compiler->ComponentMask(IndexInput, false, false, false, true);
+
+	int32 Unpremult = Compiler->AppendVector(Compiler->Div(RGB, Alpha), Alpha);
+	return Compiler->If(Alpha, Compiler->Constant(0.f), Unpremult, IndexInput, Unpremult, Compiler->Constant(0.00001f));
+}
+
+void UMaterialExpressionUnpremult::GetCaption(TArray<FString>& OutCaptions) const
+{
+	OutCaptions.Add(TEXT("Unpremult"));
+}
+#endif
+
+#undef LOCTEXT_NAMESPACE 
\ No newline at end of file
diff --git a/Engine/Plugins/Interchange/Runtime/Source/Import/Private/MaterialX/MaterialExpressionUnpremult.h b/Engine/Plugins/Interchange/Runtime/Source/Import/Private/MaterialX/MaterialExpressionUnpremult.h
new file mode 100644
--- /dev/null
+++ b/Engine/Plugins/Interchange/Runtime/Source/Import/Private/MaterialX/MaterialExpressionUnpremult.h
@@ -0,0 +1,29 @@
+// Copyright Epic Games, Inc. All Rights Reserved.
+#pragma once
+
+#include "CoreMinimal.h"
+#include "Materials/MaterialExpression.h"
+
+#include "MaterialExpressionUnpremult.generated.h"
+
+/**
+ * Divide the RGB channels of the input by the Alpha channel of the input. If the
+ * Alpha value is zero, the original color4 input value is passed through unchanged.
+ * Input must be of type float4
+ */
+UCLASS(collapsecategories, hidecategories = Object, MinimalAPI)
+class UMaterialExpressionUnpremult : public UMaterialExpression
+{
+	GENERATED_UCLASS_BODY()
+
+	UPROPERTY()
+	FExpressionInput Input;
+
+	//~ Begin UMaterialExpression Interface
+#if WITH_EDITOR
+	virtual int32 Compile(class FMaterialCompiler* Compiler, int32 OutputIndex) override;
+	virtual void GetCaption(TArray<FString>& OutCaptions) const override;
+#endif
+	//~ End UMaterialExpression Interface
+};
+
diff --git a/Engine/Plugins/Interchange/Runtime/Source/Import/Public/MaterialX/InterchangeMaterialXTranslator.h b/Engine/Plugins/Interchange/Runtime/Source/Import/Public/MaterialX/InterchangeMaterialXTranslator.h
--- a/Engine/Plugins/Interchange/Runtime/Source/Import/Public/MaterialX/InterchangeMaterialXTranslator.h
+++ b/Engine/Plugins/Interchange/Runtime/Source/Import/Public/MaterialX/InterchangeMaterialXTranslator.h
@@ -346,6 +346,7 @@ protected:
 	void ConnectIfGreaterInputToOutput(MaterialX::NodePtr UpstreamNode, UInterchangeShaderNode* ParentShaderNode, const FString& InputChannelName, TMap<FString, UInterchangeShaderNode*>& NamesToShaderNodes, UInterchangeBaseNodeContainer& NodeContainer) const;
 	void ConnectIfGreaterEqInputToOutput(MaterialX::NodePtr UpstreamNode, UInterchangeShaderNode* ParentShaderNode, const FString& InputChannelName, TMap<FString, UInterchangeShaderNode*>& NamesToShaderNodes, UInterchangeBaseNodeContainer& NodeContainer) const;
 	void ConnectIfEqualInputToOutput(MaterialX::NodePtr UpstreamNode, UInterchangeShaderNode* ParentShaderNode, const FString& InputChannelName, TMap<FString, UInterchangeShaderNode*>& NamesToShaderNodes, UInterchangeBaseNodeContainer& NodeContainer) const;
+	void ConnectOutsideInputToOutput(MaterialX::NodePtr UpstreamNode, UInterchangeShaderNode* ParentShaderNode, const FString& InputChannelName, TMap<FString, UInterchangeShaderNode*>& NamesToShaderNodes, UInterchangeBaseNodeContainer& NodeContainer) const;
 
 	UInterchangeShaderNode* CreateMaskShaderNode(uint8 RGBA, const FString & NodeName, TMap<FString, UInterchangeShaderNode*>& NamesToShaderNodes, UInterchangeBaseNodeContainer& NodeContainer) const;
 
