Fixed warning about invalid value on  'EInterchangeMaterialImportOption' during cooking

diff --git a/Engine/Plugins/Interchange/Runtime/Config/BaseInterchange.ini b/Engine/Plugins/Interchange/Runtime/Config/BaseInterchange.ini
new file mode 100644
--- /dev/null
+++ b/Engine/Plugins/Interchange/Runtime/Config/BaseInterchange.ini
@@ -0,0 +1,2 @@
+[CoreRedirects]
++EnumRedirects=(OldName="/Script/InterchangePipelines.EInterchangeMaterialImportOption",ValueChanges=(("DoNotImport","ImportAsMaterials"))
