Interchange (hierarchy) baking fixes
+ Interchange GLTF Translator Skeleton handling fix
+ Fixing Interchange AnimSequence transform baking
+ Changed Approach for Mesh Baking, from Complete to RootJoint baking (with bBakeMeshes option to bake to the root level node)

diff --git a/Engine/Plugins/Interchange/Runtime/Source/FactoryNodes/Private/InterchangeCommonPipelineDataFactoryNode.cpp b/Engine/Plugins/Interchange/Runtime/Source/FactoryNodes/Private/InterchangeCommonPipelineDataFactoryNode.cpp
--- a/Engine/Plugins/Interchange/Runtime/Source/FactoryNodes/Private/InterchangeCommonPipelineDataFactoryNode.cpp
+++ b/Engine/Plugins/Interchange/Runtime/Source/FactoryNodes/Private/InterchangeCommonPipelineDataFactoryNode.cpp
@@ -47,3 +47,12 @@ bool UInterchangeCommonPipelineDataFactoryNode::SetCustomGlobalOffsetTransform(c
 	UInterchangeSceneNode::ResetAllGlobalTransformCaches(NodeContainer);
 }
 
+bool UInterchangeCommonPipelineDataFactoryNode::GetBakeMeshes(bool& AttributeValue) const
+{
+	IMPLEMENT_NODE_ATTRIBUTE_GETTER(BakeMeshes, bool);
+}
+
+bool UInterchangeCommonPipelineDataFactoryNode::SetBakeMeshes(const UInterchangeBaseNodeContainer* NodeContainer, const bool& AttributeValue)
+{
+	IMPLEMENT_NODE_ATTRIBUTE_SETTER_NODELEGATE(BakeMeshes, bool);
+}
\ No newline at end of file
diff --git a/Engine/Plugins/Interchange/Runtime/Source/FactoryNodes/Public/InterchangeCommonPipelineDataFactoryNode.h b/Engine/Plugins/Interchange/Runtime/Source/FactoryNodes/Public/InterchangeCommonPipelineDataFactoryNode.h
--- a/Engine/Plugins/Interchange/Runtime/Source/FactoryNodes/Public/InterchangeCommonPipelineDataFactoryNode.h
+++ b/Engine/Plugins/Interchange/Runtime/Source/FactoryNodes/Public/InterchangeCommonPipelineDataFactoryNode.h
@@ -30,8 +30,17 @@ public:
 	UFUNCTION(BlueprintCallable, Category = "Interchange | Node | Common Pipeline Data")
 	bool SetCustomGlobalOffsetTransform(const UInterchangeBaseNodeContainer* NodeContainer, const FTransform& AttributeValue);
 
+	/** Return the Bake Meshes set by the pipelines. */
+	UFUNCTION(BlueprintCallable, Category = "Interchange | Node | Common Pipeline Data")
+	bool GetBakeMeshes(bool& AttributeValue) const;
+
+	/** Pipeline can set Bake Meshes, factories will use this to identify if Global transforms should be applied to Meshes/Skeletals. */
+	UFUNCTION(BlueprintCallable, Category = "Interchange | Node | Common Pipeline Data")
+	bool SetBakeMeshes(const UInterchangeBaseNodeContainer* NodeContainer, const bool& AttributeValue);
+
 private:
 	UInterchangeCommonPipelineDataFactoryNode() {};
 
 	const UE::Interchange::FAttributeKey Macro_CustomGlobalOffsetTransformKey = UE::Interchange::FAttributeKey(TEXT("GlobalOffsetTransform"));
+	const UE::Interchange::FAttributeKey Macro_CustomBakeMeshesKey = UE::Interchange::FAttributeKey(TEXT("BakeMeshes"));
 };
\ No newline at end of file
diff --git a/Engine/Plugins/Interchange/Runtime/Source/FactoryNodes/Public/InterchangeSkeletonFactoryNode.h b/Engine/Plugins/Interchange/Runtime/Source/FactoryNodes/Public/InterchangeSkeletonFactoryNode.h
--- a/Engine/Plugins/Interchange/Runtime/Source/FactoryNodes/Public/InterchangeSkeletonFactoryNode.h
+++ b/Engine/Plugins/Interchange/Runtime/Source/FactoryNodes/Public/InterchangeSkeletonFactoryNode.h
@@ -100,6 +100,18 @@ public:
 		IMPLEMENT_NODE_ATTRIBUTE_SETTER_NODELEGATE(UseTimeZeroForBindPose, bool)
 	}
 
+	UFUNCTION(BlueprintCallable, Category = "Interchange | Node | Skeleton")
+	bool GetCustomSkeletalMeshFactoryNodeUid(FString& AttributeValue) const
+	{
+		IMPLEMENT_NODE_ATTRIBUTE_GETTER(SkeletalMeshFactoryNodeUid, FString);
+	}
+
+	UFUNCTION(BlueprintCallable, Category = "Interchange | Node | Skeleton")
+	bool SetCustomSkeletalMeshFactoryNodeUid(const FString& AttributeValue)
+	{
+		IMPLEMENT_NODE_ATTRIBUTE_SETTER_NODELEGATE(SkeletalMeshFactoryNodeUid, FString)
+	}
+
 private:
 
 	void FillAssetClassFromAttribute()
@@ -139,6 +151,7 @@ private:
 	//Skeleton
 	const UE::Interchange::FAttributeKey Macro_CustomRootJointIDKey = UE::Interchange::FAttributeKey(TEXT("RootJointID"));
 	const UE::Interchange::FAttributeKey Macro_CustomUseTimeZeroForBindPoseKey = UE::Interchange::FAttributeKey(TEXT("UseTimeZeroForBindPose"));
+	const UE::Interchange::FAttributeKey Macro_CustomSkeletalMeshFactoryNodeUidKey = UE::Interchange::FAttributeKey(TEXT("SkeletalMeshFactoryNodeUid"));
 	
 protected:
 #if WITH_ENGINE
diff --git a/Engine/Plugins/Interchange/Runtime/Source/Import/Private/Animation/InterchangeAnimSequenceFactory.cpp b/Engine/Plugins/Interchange/Runtime/Source/Import/Private/Animation/InterchangeAnimSequenceFactory.cpp
--- a/Engine/Plugins/Interchange/Runtime/Source/Import/Private/Animation/InterchangeAnimSequenceFactory.cpp
+++ b/Engine/Plugins/Interchange/Runtime/Source/Import/Private/Animation/InterchangeAnimSequenceFactory.cpp
@@ -336,14 +336,15 @@ namespace UE::Interchange::Private
 			Controller.SetFrameRate(ResampleFrameRate);
 			Controller.SetNumberOfFrames(FGenericPlatformMath::Max<int32>(NumberOfFrames.Value, 1));
 
-			FTransform3f GlobalOffsetTransform;
+			FTransform GlobalOffsetTransform;
+			bool bBakeMeshes = false;
 			{
-				FTransform TempTransform = FTransform::Identity;
+				GlobalOffsetTransform = FTransform::Identity;
 				if (UInterchangeCommonPipelineDataFactoryNode* CommonPipelineDataFactoryNode = UInterchangeCommonPipelineDataFactoryNode::GetUniqueInstance(NodeContainer))
 				{
-					CommonPipelineDataFactoryNode->GetCustomGlobalOffsetTransform(TempTransform);
+					CommonPipelineDataFactoryNode->GetCustomGlobalOffsetTransform(GlobalOffsetTransform);
+					CommonPipelineDataFactoryNode->GetBakeMeshes(bBakeMeshes);
 				}
-				GlobalOffsetTransform = FTransform3f(TempTransform);
 			}
 
 			for (const TTuple< const UInterchangeSceneNode*, TFuture<TOptional<UE::Interchange::FAnimationBakeTransformPayloadData>>>& AnimationPayload : AnimationPayloads)
@@ -392,7 +393,22 @@ namespace UE::Interchange::Private
 					FTransform3f AnimKeyTransform = FTransform3f(AnimationTransformPayload->Transforms[TransformIndex]);
 					if (bApplyGlobalOffset)
 					{
-						AnimKeyTransform = AnimKeyTransform * GlobalOffsetTransform;
+						if (bBakeMeshes)
+						{
+							const UInterchangeSceneNode* RootJointNode = Cast<UInterchangeSceneNode>(NodeContainer->GetNode(SkeletonRootUid));
+							if (RootJointNode)
+							{
+								FString RootJointParentNodeUid = RootJointNode->GetParentUid();
+
+								const UInterchangeSceneNode* RootJointParentNode = Cast<UInterchangeSceneNode>(NodeContainer->GetNode(RootJointParentNodeUid));
+								if (RootJointParentNode)
+								{
+									FTransform GlobalTransform;
+									RootJointParentNode->GetCustomGlobalTransform(NodeContainer, GlobalOffsetTransform, GlobalTransform);
+									AnimKeyTransform = AnimKeyTransform * FTransform3f(GlobalTransform);
+								}
+							}
+						}
 					}
 					//Default value to identity
 					FVector3f Position(0.0f);
diff --git a/Engine/Plugins/Interchange/Runtime/Source/Import/Private/Gltf/InterchangeGltfTranslator.cpp b/Engine/Plugins/Interchange/Runtime/Source/Import/Private/Gltf/InterchangeGltfTranslator.cpp
--- a/Engine/Plugins/Interchange/Runtime/Source/Import/Private/Gltf/InterchangeGltfTranslator.cpp
+++ b/Engine/Plugins/Interchange/Runtime/Source/Import/Private/Gltf/InterchangeGltfTranslator.cpp
@@ -1698,21 +1698,6 @@ TFuture<TOptional<UE::Interchange::FSkeletalMeshMorphTargetPayloadData>> UInterc
 
 void UInterchangeGltfTranslator::HandleGltfSkeletons(UInterchangeBaseNodeContainer& NodeContainer, const FString& SceneNodeUid, const TArray<int32>& SkinnedMeshNodes, TSet<int>& UnusedMeshIndices) const
 {
-	//If there are skeletons, set them as a root joint for the joint hierarchy as per documentation:
-	//The skeleton property (if present) points to the node that is the common root of a joints hierarchy or to a direct or indirect parent node of the common root.
-	for (const GLTF::FSkinInfo& Skin : GltfAsset.Skins)
-	{
-		if (GltfAsset.Nodes.IsValidIndex(Skin.Skeleton))
-		{
-			const FString* SkeletonNodeUid = NodeUidMap.Find(&GltfAsset.Nodes[Skin.Skeleton]);
-			if (SkeletonNodeUid)
-			{
-				NodeContainer.SetNodeParentUid(*SkeletonNodeUid, SceneNodeUid);
-			}
-		}
-	}
-
-
 	TMap<int32, TMap<int32, TArray<int32>>> MeshIndexToRootJointGroupedSkinnedMeshNodesMap;
 
 	//group SkinnedMeshNodes based on Joint Root Parents and Mesh indices
diff --git a/Engine/Plugins/Interchange/Runtime/Source/Import/Private/Mesh/InterchangeSkeletalMeshFactory.cpp b/Engine/Plugins/Interchange/Runtime/Source/Import/Private/Mesh/InterchangeSkeletalMeshFactory.cpp
--- a/Engine/Plugins/Interchange/Runtime/Source/Import/Private/Mesh/InterchangeSkeletalMeshFactory.cpp
+++ b/Engine/Plugins/Interchange/Runtime/Source/Import/Private/Mesh/InterchangeSkeletalMeshFactory.cpp
@@ -996,9 +996,11 @@ UObject* UInterchangeSkeletalMeshFactory::CreateAsset(const FCreateAssetParams&
 	const bool bIsReImport = (Arguments.ReimportObject != nullptr) || (SkeletalMesh->GetLODNum() > 0);
 
 	FTransform GlobalOffsetTransform = FTransform::Identity;
+	bool bBakeMeshes = false;
 	if (UInterchangeCommonPipelineDataFactoryNode* CommonPipelineDataFactoryNode = UInterchangeCommonPipelineDataFactoryNode::GetUniqueInstance(Arguments.NodeContainer))
 	{
 		CommonPipelineDataFactoryNode->GetCustomGlobalOffsetTransform(GlobalOffsetTransform);
+		CommonPipelineDataFactoryNode->GetBakeMeshes(bBakeMeshes);
 	}
 
 	//Dirty the DDC Key for any imported Skeletal Mesh
@@ -1157,6 +1159,18 @@ UObject* UInterchangeSkeletalMeshFactory::CreateAsset(const FCreateAssetParams&
 			UE_LOG(LogInterchangeImport, Warning, TEXT("Invalid Skeleton LOD Root Joint when importing SkeletalMesh asset %s"), *Arguments.AssetName);
 			continue;
 		}
+		
+		const UInterchangeSceneNode* RootJointNode = Cast<UInterchangeSceneNode>(Arguments.NodeContainer->GetNode(RootJointNodeId));
+		if (!RootJointNode)
+		{
+			UE_LOG(LogInterchangeImport, Warning, TEXT("Invalid Skeleton RootJointNode."));
+			continue;
+		}
+		FTransform RootJointNodeGlobalTransform;
+		ensure(RootJointNode->GetCustomGlobalTransform(Arguments.NodeContainer, GlobalOffsetTransform, RootJointNodeGlobalTransform));
+		FTransform RootJointNodeLocalTransform;
+		ensure(RootJointNode->GetCustomLocalTransform(RootJointNodeLocalTransform));
+		FTransform BakeToRootJointTransfromModifier = RootJointNodeGlobalTransform.Inverse() * RootJointNodeLocalTransform;
 
 		int32 SkeletonDepth = 0;
 		TArray<SkeletalMeshImportData::FBone> RefBonesBinary;
@@ -1192,17 +1206,21 @@ UObject* UInterchangeSkeletalMeshFactory::CreateAsset(const FCreateAssetParams&
 					MeshReference.SceneNode->GetCustomAssetInstanceUid(MeshDependencyUid);
 					MeshReference.MeshNode = Cast<UInterchangeMeshNode>(Arguments.NodeContainer->GetNode(MeshDependencyUid));
 					//Cache the scene node global matrix, we will use this matrix to bake the vertices, add the node geometric mesh offset to this matrix to bake it properly
-					FTransform SceneNodeGlobalTransform;
-					if (!bUseTimeZeroAsBindPose || !MeshReference.SceneNode->GetCustomTimeZeroGlobalTransform(Arguments.NodeContainer, GlobalOffsetTransform, SceneNodeGlobalTransform))
+					FTransform SceneNodeTransform;
+					if (!bUseTimeZeroAsBindPose || !MeshReference.SceneNode->GetCustomTimeZeroGlobalTransform(Arguments.NodeContainer, GlobalOffsetTransform, SceneNodeTransform))
 					{
-						ensure(MeshReference.SceneNode->GetCustomGlobalTransform(Arguments.NodeContainer, GlobalOffsetTransform, SceneNodeGlobalTransform));
+						ensure(MeshReference.SceneNode->GetCustomGlobalTransform(Arguments.NodeContainer, GlobalOffsetTransform, SceneNodeTransform));
+						if (!bBakeMeshes)
+						{
+							SceneNodeTransform *= BakeToRootJointTransfromModifier;
+						}
 					}
 					FTransform SceneNodeGeometricTransform;
 					if(MeshReference.SceneNode->GetCustomGeometricTransform(SceneNodeGeometricTransform))
 					{
-						SceneNodeGlobalTransform *= SceneNodeGeometricTransform;
+						SceneNodeTransform *= SceneNodeGeometricTransform;
 					}
-					MeshReference.SceneGlobalTransform = SceneNodeGlobalTransform;
+					MeshReference.SceneGlobalTransform = SceneNodeTransform;
 				}
 				else
 				{
diff --git a/Engine/Plugins/Interchange/Runtime/Source/Import/Private/Mesh/InterchangeSkeletonHelper.cpp b/Engine/Plugins/Interchange/Runtime/Source/Import/Private/Mesh/InterchangeSkeletonHelper.cpp
--- a/Engine/Plugins/Interchange/Runtime/Source/Import/Private/Mesh/InterchangeSkeletonHelper.cpp
+++ b/Engine/Plugins/Interchange/Runtime/Source/Import/Private/Mesh/InterchangeSkeletonHelper.cpp
@@ -192,22 +192,28 @@ namespace UE::Interchange::Private
 		bool bHasTimeZeroTransform = false;
 		FTransform BindPoseLocalTransform;
 		bool bHasBindPoseTransform = false;
+
+		ensure(JointNode->GetCustomLocalTransform(LocalTransform));
+		bHasTimeZeroTransform = JointNode->GetCustomTimeZeroLocalTransform(TimeZeroLocalTransform);
+		bHasBindPoseTransform = JointNode->GetCustomBindPoseLocalTransform(BindPoseLocalTransform);
+
 		if (ParentIndex == INDEX_NONE)
 		{
 			FTransform GlobalOffsetTransform = FTransform::Identity;
+			bool bBakeMeshes = false;
 			if (UInterchangeCommonPipelineDataFactoryNode* CommonPipelineDataFactoryNode = UInterchangeCommonPipelineDataFactoryNode::GetUniqueInstance(NodeContainer))
 			{
 				CommonPipelineDataFactoryNode->GetCustomGlobalOffsetTransform(GlobalOffsetTransform);
+				CommonPipelineDataFactoryNode->GetBakeMeshes(bBakeMeshes);
+			}
+
+			if (bBakeMeshes)
+			{
+				LocalTransform = FTransform::Identity;
+				ensure(JointNode->GetCustomGlobalTransform(NodeContainer, GlobalOffsetTransform, LocalTransform));
+				bHasTimeZeroTransform = JointNode->GetCustomTimeZeroGlobalTransform(NodeContainer, GlobalOffsetTransform, TimeZeroLocalTransform);
+				bHasBindPoseTransform = JointNode->GetCustomBindPoseGlobalTransform(NodeContainer, GlobalOffsetTransform, BindPoseLocalTransform);
 			}
-			ensure(JointNode->GetCustomGlobalTransform(NodeContainer, GlobalOffsetTransform, LocalTransform));
-			bHasTimeZeroTransform = JointNode->GetCustomTimeZeroGlobalTransform(NodeContainer, GlobalOffsetTransform, TimeZeroLocalTransform);
-			bHasBindPoseTransform = JointNode->GetCustomBindPoseGlobalTransform(NodeContainer, GlobalOffsetTransform, BindPoseLocalTransform);
-		}
-		else
-		{
-			ensure(JointNode->GetCustomLocalTransform(LocalTransform));
-			bHasTimeZeroTransform = JointNode->GetCustomTimeZeroLocalTransform(TimeZeroLocalTransform);
-			bHasBindPoseTransform = JointNode->GetCustomBindPoseLocalTransform(BindPoseLocalTransform);
 		}
 
 		Info.LocalTransform = bHasBindPoseTransform ? BindPoseLocalTransform : LocalTransform;
diff --git a/Engine/Plugins/Interchange/Runtime/Source/Import/Private/Mesh/InterchangeStaticMeshFactory.cpp b/Engine/Plugins/Interchange/Runtime/Source/Import/Private/Mesh/InterchangeStaticMeshFactory.cpp
--- a/Engine/Plugins/Interchange/Runtime/Source/Import/Private/Mesh/InterchangeStaticMeshFactory.cpp
+++ b/Engine/Plugins/Interchange/Runtime/Source/Import/Private/Mesh/InterchangeStaticMeshFactory.cpp
@@ -539,9 +539,11 @@ TArray<UInterchangeStaticMeshFactory::FMeshPayload> UInterchangeStaticMeshFactor
 	}
 
 	FTransform GlobalOffsetTransform = FTransform::Identity;
+	bool bBakeMeshes = false;
 	if (UInterchangeCommonPipelineDataFactoryNode* CommonPipelineDataFactoryNode = UInterchangeCommonPipelineDataFactoryNode::GetUniqueInstance(Arguments.NodeContainer))
 	{
 		CommonPipelineDataFactoryNode->GetCustomGlobalOffsetTransform(GlobalOffsetTransform);
+		CommonPipelineDataFactoryNode->GetBakeMeshes(bBakeMeshes);
 	}
 
 	for (const FString& MeshUid : MeshUids)
@@ -560,11 +562,14 @@ TArray<UInterchangeStaticMeshFactory::FMeshPayload> UInterchangeStaticMeshFactor
 				continue;
 			}
 
-			// Get the transform from the scene node
-			FTransform SceneNodeGlobalTransform;
-			if (SceneNode->GetCustomGlobalTransform(Arguments.NodeContainer, GlobalOffsetTransform, SceneNodeGlobalTransform))
+			if (bBakeMeshes)
 			{
-				Payload.Transform = SceneNodeGlobalTransform;
+				// Get the transform from the scene node
+				FTransform SceneNodeGlobalTransform;
+				if (SceneNode->GetCustomGlobalTransform(Arguments.NodeContainer, GlobalOffsetTransform, SceneNodeGlobalTransform))
+				{
+					Payload.Transform = SceneNodeGlobalTransform;
+				}
 			}
 
 			// And get the mesh node which it references
@@ -574,8 +579,11 @@ TArray<UInterchangeStaticMeshFactory::FMeshPayload> UInterchangeStaticMeshFactor
 		}
 		else
 		{
-			//If we have a mesh that is not reference by a scene node, we must apply the global offset.
-			Payload.Transform = GlobalOffsetTransform;
+			if (bBakeMeshes)
+			{
+				//If we have a mesh that is not reference by a scene node, we must apply the global offset.
+				Payload.Transform = GlobalOffsetTransform;
+			}
 		}
 
 		if (!ensure(MeshNode))
@@ -1409,9 +1417,11 @@ bool UInterchangeStaticMeshFactory::ImportSockets(const FCreateAssetParams& Argu
 	TSet<FName> ImportedSocketNames;
 
 	FTransform GlobalOffsetTransform = FTransform::Identity;
+	bool bBakeMeshes = false;
 	if (UInterchangeCommonPipelineDataFactoryNode* CommonPipelineDataFactoryNode = UInterchangeCommonPipelineDataFactoryNode::GetUniqueInstance(Arguments.NodeContainer))
 	{
 		CommonPipelineDataFactoryNode->GetCustomGlobalOffsetTransform(GlobalOffsetTransform);
+		CommonPipelineDataFactoryNode->GetBakeMeshes(bBakeMeshes);
 	}
 
 	for (const FString& SocketUid : SocketUids)
@@ -1427,8 +1437,11 @@ bool UInterchangeStaticMeshFactory::ImportSockets(const FCreateAssetParams& Argu
 			FName SocketName = FName(NodeDisplayName);
 			ImportedSocketNames.Add(SocketName);
 
-			FTransform GlobalTransform;
-			SceneNode->GetCustomGlobalTransform(Arguments.NodeContainer, GlobalOffsetTransform, GlobalTransform);
+			FTransform Transform;
+			if (bBakeMeshes)
+			{
+				SceneNode->GetCustomGlobalTransform(Arguments.NodeContainer, GlobalOffsetTransform, Transform);
+			}
 
 			UStaticMeshSocket* Socket = StaticMesh->FindSocket(SocketName);
 			if (!Socket)
@@ -1442,9 +1455,9 @@ bool UInterchangeStaticMeshFactory::ImportSockets(const FCreateAssetParams& Argu
 				StaticMesh->AddSocket(Socket);
 			}
 
-			Socket->RelativeLocation = GlobalTransform.GetLocation();
-			Socket->RelativeRotation = GlobalTransform.GetRotation().Rotator();
-			Socket->RelativeScale = GlobalTransform.GetScale3D();
+			Socket->RelativeLocation = Transform.GetLocation();
+			Socket->RelativeRotation = Transform.GetRotation().Rotator();
+			Socket->RelativeScale = Transform.GetScale3D();
 		}
 	}
 
diff --git a/Engine/Plugins/Interchange/Runtime/Source/Import/Private/Scene/InterchangeSkeletalMeshActorFactory.cpp b/Engine/Plugins/Interchange/Runtime/Source/Import/Private/Scene/InterchangeSkeletalMeshActorFactory.cpp
--- a/Engine/Plugins/Interchange/Runtime/Source/Import/Private/Scene/InterchangeSkeletalMeshActorFactory.cpp
+++ b/Engine/Plugins/Interchange/Runtime/Source/Import/Private/Scene/InterchangeSkeletalMeshActorFactory.cpp
@@ -4,6 +4,8 @@
 #include "InterchangeActorFactoryNode.h"
 #include "InterchangeMeshActorFactoryNode.h"
 #include "Scene/InterchangeActorHelper.h"
+#include "InterchangeSceneNode.h"
+#include "InterchangeSkeletalMeshFactoryNode.h"
 
 #include "Animation/SkeletalMeshActor.h"
 #include "Components/SkeletalMeshComponent.h"
@@ -52,17 +54,22 @@ void UInterchangeSkeletalMeshActorFactory::PostImportPreCompletedCallback(const
 	{
 		if (USkeletalMeshComponent * SkeletalMeshComponent = SkeletalMeshActor->GetSkeletalMeshComponent())
 		{
-			if (const UInterchangeFactoryBaseNode* MeshNode = UE::Interchange::ActorHelper::FindAssetInstanceFactoryNode(Arguments.NodeContainer, Arguments.FactoryNode))
 			{
-				FSoftObjectPath ReferenceObject;
-				MeshNode->GetCustomReferenceObject(ReferenceObject);
-				if (USkeletalMesh* SkeletalMesh = Cast<USkeletalMesh>(ReferenceObject.TryLoad()))
+				TArray<FString> TargetNodeUids;
+				Arguments.FactoryNode->GetTargetNodeUids(TargetNodeUids);
+				const UInterchangeSkeletalMeshFactoryNode* SkeletalMeshFactoryNode = TargetNodeUids.IsEmpty() ? nullptr : Cast<UInterchangeSkeletalMeshFactoryNode>(Arguments.NodeContainer->GetFactoryNode(TargetNodeUids[0]));
+				if (SkeletalMeshFactoryNode)
 				{
-					SkeletalMeshComponent->SetSkeletalMesh(SkeletalMesh);
-
-					if (const UInterchangeMeshActorFactoryNode* MeshActorFactoryNode = Cast<UInterchangeMeshActorFactoryNode>(Arguments.FactoryNode))
+					FSoftObjectPath ReferenceObject;
+					SkeletalMeshFactoryNode->GetCustomReferenceObject(ReferenceObject);
+					if (USkeletalMesh* SkeletalMesh = Cast<USkeletalMesh>(ReferenceObject.TryLoad()))
 					{
-						UE::Interchange::ActorHelper::ApplySlotMaterialDependencies(*Arguments.NodeContainer, *MeshActorFactoryNode, *SkeletalMeshComponent);
+						SkeletalMeshComponent->SetSkeletalMesh(SkeletalMesh);
+
+						if (const UInterchangeMeshActorFactoryNode* MeshActorFactoryNode = Cast<UInterchangeMeshActorFactoryNode>(Arguments.FactoryNode))
+						{
+							UE::Interchange::ActorHelper::ApplySlotMaterialDependencies(*Arguments.NodeContainer, *MeshActorFactoryNode, *SkeletalMeshComponent);
+						}
 					}
 				}
 			}
diff --git a/Engine/Plugins/Interchange/Runtime/Source/Pipelines/Private/InterchangeGenericAssetsPipeline.cpp b/Engine/Plugins/Interchange/Runtime/Source/Pipelines/Private/InterchangeGenericAssetsPipeline.cpp
--- a/Engine/Plugins/Interchange/Runtime/Source/Pipelines/Private/InterchangeGenericAssetsPipeline.cpp
+++ b/Engine/Plugins/Interchange/Runtime/Source/Pipelines/Private/InterchangeGenericAssetsPipeline.cpp
@@ -187,7 +187,9 @@ void UInterchangeGenericAssetsPipeline::ExecutePreImportPipeline(UInterchangeBas
 
 		UInterchangeCommonPipelineDataFactoryNode* CommonPipelineDataFactoryNode = UInterchangeCommonPipelineDataFactoryNode::FindOrCreateUniqueInstance(BaseNodeContainer);
 		CommonPipelineDataFactoryNode->SetCustomGlobalOffsetTransform(BaseNodeContainer, ImportOffsetTransform);
-		
+
+		// In case all mesh types are forced to Static/Skeletal we bake the scene instance hierarchy transforms
+		CommonPipelineDataFactoryNode->SetBakeMeshes(BaseNodeContainer, CommonMeshesProperties->ForceAllMeshAsType != EInterchangeForceMeshType::IFMT_None || CommonMeshesProperties->bBakeMeshes);
 	}
 
 	SourceDatas.Empty(InSourceDatas.Num());
diff --git a/Engine/Plugins/Interchange/Runtime/Source/Pipelines/Private/InterchangeGenericScenesPipeline.cpp b/Engine/Plugins/Interchange/Runtime/Source/Pipelines/Private/InterchangeGenericScenesPipeline.cpp
--- a/Engine/Plugins/Interchange/Runtime/Source/Pipelines/Private/InterchangeGenericScenesPipeline.cpp
+++ b/Engine/Plugins/Interchange/Runtime/Source/Pipelines/Private/InterchangeGenericScenesPipeline.cpp
@@ -16,6 +16,7 @@
 #include "InterchangeSceneVariantSetsFactoryNode.h"
 #include "InterchangeVariantSetNode.h"
 #include "Nodes/InterchangeUserDefinedAttribute.h"
+#include "InterchangeSkeletonFactoryNode.h"
 
 #include "Animation/SkeletalMeshActor.h"
 #include "CineCameraActor.h"
@@ -71,8 +72,26 @@ void UInterchangeGenericLevelPipeline::ExecutePreImportPipeline(UInterchangeBase
 				SceneNode->GetSpecializedTypes(SpecializeTypes);
 				if (!SpecializeTypes.Contains(UE::Interchange::FSceneNodeStaticData::GetTransformSpecializeTypeString()))
 				{
-					//Skip any scene node that have specialized types but not the "Transform" type.
-					continue;
+					bool bSkipNode = true;
+					if (SpecializeTypes.Contains(UE::Interchange::FSceneNodeStaticData::GetJointSpecializeTypeString()))
+					{
+						//check if its the rootjoint (we want to create an actor for the rootjoint)
+						FString CurrentNodesParentUid = SceneNode->GetParentUid();
+						const UInterchangeBaseNode* ParentNode = BaseNodeContainer->GetNode(CurrentNodesParentUid);
+						if (const UInterchangeSceneNode* ParentSceneNode = Cast<UInterchangeSceneNode>(ParentNode))
+						{
+							if (!ParentSceneNode->IsSpecializedTypeContains(UE::Interchange::FSceneNodeStaticData::GetJointSpecializeTypeString()))
+							{
+								bSkipNode = false;
+							}
+						}
+					}
+
+					if (bSkipNode)
+					{
+						//Skip any scene node that have specialized types but not the "Transform" type.
+						continue;
+					}
 				}
 			}
 			ExecuteSceneNodePreImport(GlobalOffsetTransform, SceneNode);
@@ -82,19 +101,19 @@ void UInterchangeGenericLevelPipeline::ExecutePreImportPipeline(UInterchangeBase
 	//Find all translated scene variant sets
 	TArray<UInterchangeSceneVariantSetsNode*> SceneVariantSetNodes;
 
-	InBaseNodeContainer->IterateNodesOfType<UInterchangeSceneVariantSetsNode>([&SceneVariantSetNodes](const FString& NodeUid, UInterchangeSceneVariantSetsNode* Node)
-		{
-			SceneVariantSetNodes.Add(Node);
-		});
+InBaseNodeContainer->IterateNodesOfType<UInterchangeSceneVariantSetsNode>([&SceneVariantSetNodes](const FString& NodeUid, UInterchangeSceneVariantSetsNode* Node)
+	{
+		SceneVariantSetNodes.Add(Node);
+	});
 
-	for (const UInterchangeSceneVariantSetsNode* SceneVariantSetNode : SceneVariantSetNodes)
+for (const UInterchangeSceneVariantSetsNode* SceneVariantSetNode : SceneVariantSetNodes)
+{
+	if (SceneVariantSetNode)
 	{
-		if (SceneVariantSetNode)
-		{
-			ExecuteSceneVariantSetNodePreImport(*SceneVariantSetNode);
-		}
+		ExecuteSceneVariantSetNodePreImport(*SceneVariantSetNode);
 	}
 }
+}
 
 void UInterchangeGenericLevelPipeline::ExecuteSceneNodePreImport(const FTransform& GlobalOffsetTransform, const UInterchangeSceneNode* SceneNode)
 {
@@ -104,11 +123,41 @@ void UInterchangeGenericLevelPipeline::ExecuteSceneNodePreImport(const FTransfor
 	}
 
 	const UInterchangeBaseNode* TranslatedAssetNode = nullptr;
+	bool bRootJointNode = SceneNode->IsSpecializedTypeContains(UE::Interchange::FSceneNodeStaticData::GetJointSpecializeTypeString());
+	FString SkeletalMeshFactoryNodeUid;
 
-	FString AssetInstanceUid;
-	if (SceneNode->GetCustomAssetInstanceUid(AssetInstanceUid))
+	if (bRootJointNode)
+	{
+		FString SkeletonFactoryNodeUid = UInterchangeFactoryBaseNode::BuildFactoryNodeUid(SceneNode->GetUniqueID());
+		const UInterchangeSkeletonFactoryNode* SkeletonFactoryNode = Cast<UInterchangeSkeletonFactoryNode>(BaseNodeContainer->GetFactoryNode(SkeletonFactoryNodeUid));
+		if (SkeletonFactoryNode)
+		{
+			if (SkeletonFactoryNode->GetCustomSkeletalMeshFactoryNodeUid(SkeletalMeshFactoryNodeUid))
+			{
+				if (const UInterchangeFactoryBaseNode* SkeletalMeshFactoryNode = BaseNodeContainer->GetFactoryNode(SkeletalMeshFactoryNodeUid))
+				{
+					TArray<FString> NodeUids;
+					SkeletalMeshFactoryNode->GetTargetNodeUids(NodeUids);
+
+					if (NodeUids.Num() > 0)
+					{
+						TranslatedAssetNode = BaseNodeContainer->GetNode(NodeUids[0]);
+					}
+					else
+					{
+						TranslatedAssetNode = nullptr;
+					}
+				}
+			}
+		}
+	}
+	else
 	{
-		TranslatedAssetNode = BaseNodeContainer->GetNode(AssetInstanceUid);
+		FString AssetInstanceUid;
+		if (SceneNode->GetCustomAssetInstanceUid(AssetInstanceUid))
+		{
+			TranslatedAssetNode = BaseNodeContainer->GetNode(AssetInstanceUid);
+		}
 	}
 
 	UInterchangeActorFactoryNode* ActorFactoryNode = CreateActorFactoryNode(SceneNode, TranslatedAssetNode);
@@ -118,7 +167,10 @@ void UInterchangeGenericLevelPipeline::ExecuteSceneNodePreImport(const FTransfor
 		return;
 	}
 
-	ActorFactoryNode->InitializeNode(UInterchangeFactoryBaseNode::BuildFactoryNodeUid(SceneNode->GetUniqueID()), SceneNode->GetDisplayLabel(), EInterchangeNodeContainerType::FactoryData);
+	FString NodeUid = SceneNode->GetUniqueID() + (bRootJointNode ? TEXT("_SkeletonNode") : TEXT(""));
+	FString FactoryNodeUid = UInterchangeFactoryBaseNode::BuildFactoryNodeUid(NodeUid);
+
+	ActorFactoryNode->InitializeNode(FactoryNodeUid, SceneNode->GetDisplayLabel(), EInterchangeNodeContainerType::FactoryData);
 	const FString ActorFactoryNodeUid = BaseNodeContainer->AddNode(ActorFactoryNode);
 	if (!SceneNode->GetParentUid().IsEmpty())
 	{
@@ -127,16 +179,36 @@ void UInterchangeGenericLevelPipeline::ExecuteSceneNodePreImport(const FTransfor
 		ActorFactoryNode->AddFactoryDependencyUid(ParentFactoryNodeUid);
 	}
 
-	ActorFactoryNode->AddTargetNodeUid(SceneNode->GetUniqueID());
-	SceneNode->AddTargetNodeUid(ActorFactoryNode->GetUniqueID());
+	if (bRootJointNode)
+	{
+		ActorFactoryNode->AddTargetNodeUid(SkeletalMeshFactoryNodeUid);
+	}
+	else
+	{
+		ActorFactoryNode->AddTargetNodeUid(SceneNode->GetUniqueID());
+		SceneNode->AddTargetNodeUid(ActorFactoryNode->GetUniqueID());
+	}
 
 	//TODO move this code to the factory, a stack over pipeline can change the global offset transform which will affect this value.
 	FTransform GlobalTransform;
 	if (SceneNode->GetCustomGlobalTransform(BaseNodeContainer, GlobalOffsetTransform, GlobalTransform))
 	{
+		if (bRootJointNode)
+		{
+			GlobalTransform = FTransform::Identity;
+			//LocalTransform of RootjointNode is already baked into the Skeletal and animation.
+			//due to that we acquire the Parent SceneNode and get its GlobalTransform:
+			if (!SceneNode->GetParentUid().IsEmpty())
+			{
+				if (const UInterchangeSceneNode* ParentSceneNode = Cast<UInterchangeSceneNode>(BaseNodeContainer->GetNode(SceneNode->GetParentUid())))
+				{
+					ParentSceneNode->GetCustomGlobalTransform(BaseNodeContainer, GlobalOffsetTransform, GlobalTransform);
+				}
+			}
+		}
 		ActorFactoryNode->SetCustomGlobalTransform(GlobalTransform);
 	}
-	
+
 	ActorFactoryNode->SetCustomMobility(EComponentMobility::Static);
 
 	if (TranslatedAssetNode)
@@ -194,6 +266,12 @@ void UInterchangeGenericLevelPipeline::SetUpFactoryNode(UInterchangeActorFactory
 	{
 		if (MeshNode->IsSkinnedMesh())
 		{
+			bool bRootJointNode = SceneNode->IsSpecializedTypeContains(UE::Interchange::FSceneNodeStaticData::GetJointSpecializeTypeString());
+			if (!bRootJointNode)
+			{
+				return;
+			}
+
 			ActorFactoryNode->SetCustomActorClassName(ASkeletalMeshActor::StaticClass()->GetPathName());
 			ActorFactoryNode->SetCustomMobility(EComponentMobility::Movable);
 		}
diff --git a/Engine/Plugins/Interchange/Runtime/Source/Pipelines/Private/InterchangeGenericSkeletalMeshPipeline.cpp b/Engine/Plugins/Interchange/Runtime/Source/Pipelines/Private/InterchangeGenericSkeletalMeshPipeline.cpp
--- a/Engine/Plugins/Interchange/Runtime/Source/Pipelines/Private/InterchangeGenericSkeletalMeshPipeline.cpp
+++ b/Engine/Plugins/Interchange/Runtime/Source/Pipelines/Private/InterchangeGenericSkeletalMeshPipeline.cpp
@@ -168,20 +168,16 @@ void UInterchangeGenericMeshPipeline::ExecutePreImportPipelineSkeletalMesh()
 			return bFoundInstances;
 		};
 
-		bool bFoundMeshes = false;
-		if (CommonMeshesProperties->bBakeMeshes)
-		{
-			PipelineMeshesUtilities->GetCombinedSkinnedMeshInstances(BaseNodeContainer, MeshUidsPerSkeletonRootUid, bConvertStaticMeshToSkeletalMesh);
-			const bool bUseMeshInstance = true;
-			bFoundMeshes = CreatePerSkeletonRootUidCombinedSkinnedMesh(bUseMeshInstance);
-		}
+		PipelineMeshesUtilities->GetCombinedSkinnedMeshInstances(BaseNodeContainer, MeshUidsPerSkeletonRootUid, bConvertStaticMeshToSkeletalMesh);
+		bool bUseMeshInstance = true;
+		bool bFoundMeshes = CreatePerSkeletonRootUidCombinedSkinnedMesh(bUseMeshInstance);
 
 		if (!bFoundMeshes)
 		{
 			//Fall back to support mesh not reference by any scene node
 			MeshUidsPerSkeletonRootUid.Empty();
 			PipelineMeshesUtilities->GetCombinedSkinnedMeshGeometries(MeshUidsPerSkeletonRootUid);
-			const bool bUseMeshInstance = false;
+			bUseMeshInstance = false;
 			CreatePerSkeletonRootUidCombinedSkinnedMesh(bUseMeshInstance);
 		}
 	}
@@ -242,20 +238,16 @@ void UInterchangeGenericMeshPipeline::ExecutePreImportPipelineSkeletalMesh()
 			}
 			return bFoundInstances;
 		};
-			
-		bool bFoundMeshes = false;
-		if (CommonMeshesProperties->bBakeMeshes)
-		{
-			PipelineMeshesUtilities->GetAllSkinnedMeshInstance(MeshUids, bConvertStaticMeshToSkeletalMesh);
-			const bool bUseMeshInstance = true;
-			bFoundMeshes = CreatePerSkeletonRootUidSkinnedMesh(bUseMeshInstance);
-		}
+		
+		PipelineMeshesUtilities->GetAllSkinnedMeshInstance(MeshUids, bConvertStaticMeshToSkeletalMesh);
+		bool bUseMeshInstance = true;
+		bool bFoundMeshes = CreatePerSkeletonRootUidSkinnedMesh(bUseMeshInstance);
 
 		if (!bFoundMeshes)
 		{
 			MeshUids.Empty();
 			PipelineMeshesUtilities->GetAllSkinnedMeshGeometry(MeshUids);
-			const bool bUseMeshInstance = false;
+			bUseMeshInstance = false;
 			CreatePerSkeletonRootUidSkinnedMesh(bUseMeshInstance);
 		}
 	}
@@ -418,6 +410,8 @@ UInterchangeSkeletalMeshFactoryNode* UInterchangeGenericMeshPipeline::CreateSkel
 	SkeletalMeshFactoryNode->AddFactoryDependencyUid(SkeletonUid);
 	BaseNodeContainer->AddNode(SkeletalMeshFactoryNode);
 
+	SkeletonFactoryNode->SetCustomSkeletalMeshFactoryNodeUid(SkeletalMeshFactoryNode->GetUniqueID());
+
 	AddLodDataToSkeletalMesh(SkeletonFactoryNode, SkeletalMeshFactoryNode, MeshUidsPerLodIndex);
 	SkeletalMeshFactoryNode->SetCustomImportMorphTarget(bImportMorphTargets);
 
diff --git a/Engine/Plugins/Interchange/Runtime/Source/Pipelines/Private/InterchangeGenericStaticMeshPipeline.cpp b/Engine/Plugins/Interchange/Runtime/Source/Pipelines/Private/InterchangeGenericStaticMeshPipeline.cpp
--- a/Engine/Plugins/Interchange/Runtime/Source/Pipelines/Private/InterchangeGenericStaticMeshPipeline.cpp
+++ b/Engine/Plugins/Interchange/Runtime/Source/Pipelines/Private/InterchangeGenericStaticMeshPipeline.cpp
@@ -159,7 +159,6 @@ void UInterchangeGenericMeshPipeline::ExecutePreImportPipelineStaticMesh()
 			// Combine all the static meshes
 
 			bool bFoundMeshes = false;
-			if (CommonMeshesProperties->bBakeMeshes)
 			{
 				// If baking transforms, get all the static mesh instance nodes, and group them by LOD
 				TArray<FString> MeshUids;
@@ -222,8 +221,6 @@ void UInterchangeGenericMeshPipeline::ExecutePreImportPipelineStaticMesh()
 			// Do not combine static meshes
 
 			bool bFoundMeshes = false;
-
-			if (CommonMeshesProperties->bBakeMeshes)
 			{
 				TArray<FString> MeshUids;
 				PipelineMeshesUtilities->GetAllStaticMeshInstance(MeshUids, bConvertSkeletalMeshToStaticMesh);
@@ -363,9 +360,7 @@ bool UInterchangeGenericMeshPipeline::MakeMeshFactoryNodeUidAndDisplayLabel(cons
 								DisplayLabel = SceneNode->GetDisplayLabel();
 							}
 
-							// Use the first scene node uid this LOD references, adding backslash since this uid is not asset typed (\\Mesh\\) like a mesh node
-							// @TODO: change this so that scene nodes get a type prefix of their own?
-							NewNodeUid = TEXT("\\") + Uid;
+							NewNodeUid = RefMeshUid;
 							return true;
 						}
 					}
diff --git a/Engine/Plugins/Interchange/Runtime/Source/Pipelines/Public/InterchangeGenericAssetsPipelineSharedSettings.h b/Engine/Plugins/Interchange/Runtime/Source/Pipelines/Public/InterchangeGenericAssetsPipelineSharedSettings.h
--- a/Engine/Plugins/Interchange/Runtime/Source/Pipelines/Public/InterchangeGenericAssetsPipelineSharedSettings.h
+++ b/Engine/Plugins/Interchange/Runtime/Source/Pipelines/Public/InterchangeGenericAssetsPipelineSharedSettings.h
@@ -56,7 +56,7 @@ public:
  	EInterchangeForceMeshType ForceAllMeshAsType = EInterchangeForceMeshType::IFMT_None;
 
 	/** If enable, meshes LODs will be imported. Note that it required the advanced bBakeMesh property to be enabled. */
-	UPROPERTY(EditAnywhere, BlueprintReadWrite, Category = "Common Meshes", meta = (editcondition = "bBakeMeshes"))
+	UPROPERTY(EditAnywhere, BlueprintReadWrite, Category = "Common Meshes")
 	bool bImportLods = true;
 
 	/** If enable, meshes will be baked with the scene instance hierarchy transform. */
