Interchange texture import
- Check the r.VT.EnableAutoImport flag to decide if a texture should be imported as a virtual texture.
- Set a texture as virtual if its enabled and its bigger than the max non-vt size but smaller than the auto VT threshold.

This is equivalent to the behavior in UTextureFactory.

diff --git a/Engine/Plugins/Interchange/Runtime/Source/Import/Private/Texture/InterchangeTextureFactory.cpp b/Engine/Plugins/Interchange/Runtime/Source/Import/Private/Texture/InterchangeTextureFactory.cpp
--- a/Engine/Plugins/Interchange/Runtime/Source/Import/Private/Texture/InterchangeTextureFactory.cpp
+++ b/Engine/Plugins/Interchange/Runtime/Source/Import/Private/Texture/InterchangeTextureFactory.cpp
@@ -1341,7 +1341,10 @@ namespace UE::Interchange::Private::InterchangeTextureFactory
 		static const TConsoleVariableData<int32>* CVarVirtualTexturesEnabled = IConsoleManager::Get().FindTConsoleVariableDataInt(TEXT("r.VirtualTextures"));
 		check(CVarVirtualTexturesEnabled);
 
-		if (CVarVirtualTexturesEnabled->GetValueOnGameThread())
+		static const auto CVarVirtualTexturesAutoImportEnabled = IConsoleManager::Get().FindTConsoleVariableDataInt(TEXT("r.VT.EnableAutoImport"));
+		check(CVarVirtualTexturesAutoImportEnabled);
+
+		if (CVarVirtualTexturesEnabled->GetValueOnGameThread() && CVarVirtualTexturesAutoImportEnabled->GetValueOnGameThread())
 		{
 			const int VirtualTextureAutoEnableThreshold = GetDefault<UTextureImportSettings>()->AutoVTSize;
 			const int VirtualTextureAutoEnableThresholdPixels = VirtualTextureAutoEnableThreshold * VirtualTextureAutoEnableThreshold;
@@ -1351,7 +1354,9 @@ namespace UE::Interchange::Private::InterchangeTextureFactory
 			// however for a new texture platform data may not be generated yet, and for an reimport of a texture this is the size of the
 			// old texture. 
 			// Using source size gives one small caveat. It looks at the size before mipmap power of two padding adjustment.
-			if (Texture2D->Source.GetSizeX() * Texture2D->Source.GetSizeY() >= VirtualTextureAutoEnableThresholdPixels)
+			if (Texture2D->Source.GetSizeX() * Texture2D->Source.GetSizeY() >= VirtualTextureAutoEnableThresholdPixels ||
+				Texture2D->Source.GetSizeX() > UTexture::GetMaximumDimensionOfNonVT() ||
+				Texture2D->Source.GetSizeY() > UTexture::GetMaximumDimensionOfNonVT())
 			{
 				Texture2D->VirtualTextureStreaming = true;
 			}
