Fixed the invalid reimport behavior for the textures when using interchange.

#rnx

diff --git a/Engine/Plugins/Interchange/Runtime/Source/Import/Private/Texture/InterchangeTextureFactory.cpp b/Engine/Plugins/Interchange/Runtime/Source/Import/Private/Texture/InterchangeTextureFactory.cpp
--- a/Engine/Plugins/Interchange/Runtime/Source/Import/Private/Texture/InterchangeTextureFactory.cpp
+++ b/Engine/Plugins/Interchange/Runtime/Source/Import/Private/Texture/InterchangeTextureFactory.cpp
@@ -644,10 +644,10 @@ namespace UE::Interchange::Private::InterchangeTextureFactory
 		return {};
 	}
 
-	void SetupTextureSourceDataFromBulkData(UTexture* Texture, const FImportImage& Image, UE::Serialization::FEditorBulkData::FSharedBufferWithID&& BufferAndId);
+	void SetupTextureSourceDataFromBulkData(UTexture* Texture, const FImportImage& Image, UE::Serialization::FEditorBulkData::FSharedBufferWithID&& BufferAndId, bool bIsReimport);
 
 #if WITH_EDITOR
-	void SetupTextureSourceDataFromBulkData_Editor(UTexture* Texture, const FImportImage& Image, UE::Serialization::FEditorBulkData::FSharedBufferWithID&& BufferAndId)
+	void SetupTextureSourceDataFromBulkData_Editor(UTexture* Texture, const FImportImage& Image, UE::Serialization::FEditorBulkData::FSharedBufferWithID&& BufferAndId, bool bIsReimport)
 	{
 		Texture->Source.InitWithCompressedSourceData(
 			Image.SizeX,
@@ -658,18 +658,21 @@ namespace UE::Interchange::Private::InterchangeTextureFactory
 			Image.RawDataCompressionFormat
 		);
 
-		Texture->CompressionSettings = Image.CompressionSettings;
-		Texture->SRGB = Image.bSRGB;
-
-		//If the MipGenSettings was set by the translator, we must apply it before the build
-		if (Image.MipGenSettings.IsSet())
+		if (!bIsReimport)
 		{
-			// if the source has mips we keep the mips by default, unless the user changes that
-			Texture->MipGenSettings = Image.MipGenSettings.GetValue();
+			Texture->CompressionSettings = Image.CompressionSettings;
+			Texture->SRGB = Image.bSRGB;
+
+			//If the MipGenSettings was set by the translator, we must apply it before the build
+			if (Image.MipGenSettings.IsSet())
+			{
+				// if the source has mips we keep the mips by default, unless the user changes that
+				Texture->MipGenSettings = Image.MipGenSettings.GetValue();
+			}
 		}
 	}
 
-	void SetupTexture2DSourceDataFromBulkData_Editor(UTexture2D* Texture2D, const FImportBlockedImage& BlockedImage, UE::Serialization::FEditorBulkData::FSharedBufferWithID&& BufferAndId)
+	void SetupTexture2DSourceDataFromBulkData_Editor(UTexture2D* Texture2D, const FImportBlockedImage& BlockedImage, UE::Serialization::FEditorBulkData::FSharedBufferWithID&& BufferAndId, bool bIsReimport)
 	{
 		if (BlockedImage.BlocksData.Num() > 1)
 		{
@@ -681,14 +684,17 @@ namespace UE::Interchange::Private::InterchangeTextureFactory
 				MoveTemp(BufferAndId)
 				);
 
-			Texture2D->CompressionSettings = BlockedImage.CompressionSettings;
-			Texture2D->SRGB = BlockedImage.bSRGB;
-			Texture2D->VirtualTextureStreaming = true;
-
-			if (BlockedImage.MipGenSettings.IsSet())
+			if (!bIsReimport)
 			{
-				// if the source has mips we keep the mips by default, unless the user changes that
-				Texture2D->MipGenSettings = BlockedImage.MipGenSettings.GetValue();
+				Texture2D->CompressionSettings = BlockedImage.CompressionSettings;
+				Texture2D->SRGB = BlockedImage.bSRGB;
+				Texture2D->VirtualTextureStreaming = true;
+
+				if (BlockedImage.MipGenSettings.IsSet())
+				{
+					// if the source has mips we keep the mips by default, unless the user changes that
+					Texture2D->MipGenSettings = BlockedImage.MipGenSettings.GetValue();
+				}
 			}
 		}
 		else
@@ -705,11 +711,11 @@ namespace UE::Interchange::Private::InterchangeTextureFactory
 			Image.SizeY = Block.SizeY;
 			Image.NumMips = Block.NumMips;
 
-			SetupTextureSourceDataFromBulkData_Editor(Texture2D, Image, MoveTemp(BufferAndId));
+			SetupTextureSourceDataFromBulkData_Editor(Texture2D, Image, MoveTemp(BufferAndId), bIsReimport);
 		}
 	}
 
-	void SetupTextureSourceDataFromBulkData_Editor(UTexture* Texture, const FImportSlicedImage& SlicedImage, UE::Serialization::FEditorBulkData::FSharedBufferWithID&& BufferAndId)
+	void SetupTextureSourceDataFromBulkData_Editor(UTexture* Texture, const FImportSlicedImage& SlicedImage, UE::Serialization::FEditorBulkData::FSharedBufferWithID&& BufferAndId,  bool bIsReimport)
 	{
 		Texture->Source.InitLayered(
 			SlicedImage.SizeX,
@@ -721,13 +727,16 @@ namespace UE::Interchange::Private::InterchangeTextureFactory
 			MoveTemp(BufferAndId)
 			);
 
-		Texture->CompressionSettings = SlicedImage.CompressionSettings;
-		Texture->SRGB = SlicedImage.bSRGB;
-
-		if (SlicedImage.MipGenSettings.IsSet())
+		if (!bIsReimport)
 		{
-			// if the source has mips we keep the mips by default, unless the user changes that
-			Texture->MipGenSettings = SlicedImage.MipGenSettings.GetValue();
+			Texture->CompressionSettings = SlicedImage.CompressionSettings;
+			Texture->SRGB = SlicedImage.bSRGB;
+
+			if (SlicedImage.MipGenSettings.IsSet())
+			{
+				// if the source has mips we keep the mips by default, unless the user changes that
+				Texture->MipGenSettings = SlicedImage.MipGenSettings.GetValue();
+			}
 		}
 	}
 
@@ -760,7 +769,7 @@ namespace UE::Interchange::Private::InterchangeTextureFactory
 		return false;
 	}
 
-	void SetupTextureCubeSourceDataFromPayload(UTextureCube* TextureCube, FProcessedPayload& ProcessedPayload)
+	void SetupTextureCubeSourceDataFromPayload(UTextureCube* TextureCube, FProcessedPayload& ProcessedPayload, bool bIsReimport)
 	{
 		if (TOptional<FImportSlicedImage>* OptionalSlicedImage = ProcessedPayload.SettingsFromPayload.TryGet<TOptional<FImportSlicedImage>>())
 		{
@@ -771,7 +780,7 @@ namespace UE::Interchange::Private::InterchangeTextureFactory
 				// Cube texture always have six slice
 				if (SlicedImage.NumSlice == 6)
 				{
-					SetupTextureSourceDataFromBulkData_Editor(TextureCube, SlicedImage, MoveTemp(ProcessedPayload.PayloadAndId));
+					SetupTextureSourceDataFromBulkData_Editor(TextureCube, SlicedImage, MoveTemp(ProcessedPayload.PayloadAndId), bIsReimport);
 				}
 			}
 		}
@@ -780,7 +789,7 @@ namespace UE::Interchange::Private::InterchangeTextureFactory
 			if (OptionalImage->IsSet())
 			{
 				FImportImage& Image = OptionalImage->GetValue();
-				SetupTextureSourceDataFromBulkData_Editor(TextureCube, Image, MoveTemp(ProcessedPayload.PayloadAndId));
+				SetupTextureSourceDataFromBulkData_Editor(TextureCube, Image, MoveTemp(ProcessedPayload.PayloadAndId), bIsReimport);
 			}
 		}
 		else if (TOptional<FImportLightProfile>* OptionalLightProfile = ProcessedPayload.SettingsFromPayload.TryGet<TOptional<FImportLightProfile>>())
@@ -788,7 +797,7 @@ namespace UE::Interchange::Private::InterchangeTextureFactory
 			if (OptionalLightProfile->IsSet())
 			{
 				FImportLightProfile& LightProfile = OptionalLightProfile->GetValue();
-				SetupTextureSourceDataFromBulkData_Editor(TextureCube, LightProfile, MoveTemp(ProcessedPayload.PayloadAndId));
+				SetupTextureSourceDataFromBulkData_Editor(TextureCube, LightProfile, MoveTemp(ProcessedPayload.PayloadAndId), bIsReimport);
 			}
 		}
 		else
@@ -827,7 +836,7 @@ namespace UE::Interchange::Private::InterchangeTextureFactory
 		return false;
 	}
 
-	void SetupTextureCubeArraySourceDataFromPayload(UTextureCubeArray* TextureCubeArray, FProcessedPayload& ProcessedPayload)
+	void SetupTextureCubeArraySourceDataFromPayload(UTextureCubeArray* TextureCubeArray, FProcessedPayload& ProcessedPayload, bool bIsReimport)
 	{
 		if (TOptional<FImportSlicedImage>* OptionalSlicedImage = ProcessedPayload.SettingsFromPayload.TryGet<TOptional<FImportSlicedImage>>())
 		{
@@ -837,7 +846,7 @@ namespace UE::Interchange::Private::InterchangeTextureFactory
 
 				if (SlicedImage.NumSlice > 6)
 				{
-					SetupTextureSourceDataFromBulkData_Editor(TextureCubeArray, SlicedImage, MoveTemp(ProcessedPayload.PayloadAndId));
+					SetupTextureSourceDataFromBulkData_Editor(TextureCubeArray, SlicedImage, MoveTemp(ProcessedPayload.PayloadAndId), bIsReimport);
 				}
 			}
 		}
@@ -846,7 +855,7 @@ namespace UE::Interchange::Private::InterchangeTextureFactory
 			if (OptionalImage->IsSet())
 			{
 				FImportImage& Image = OptionalImage->GetValue();
-				SetupTextureSourceDataFromBulkData_Editor(TextureCubeArray, Image, MoveTemp(ProcessedPayload.PayloadAndId));
+				SetupTextureSourceDataFromBulkData_Editor(TextureCubeArray, Image, MoveTemp(ProcessedPayload.PayloadAndId), bIsReimport);
 			}
 		}
 		else if (TOptional<FImportLightProfile>* OptionalLightProfile = ProcessedPayload.SettingsFromPayload.TryGet<TOptional<FImportLightProfile>>())
@@ -854,7 +863,7 @@ namespace UE::Interchange::Private::InterchangeTextureFactory
 			if (OptionalLightProfile->IsSet())
 			{
 				FImportLightProfile& LightProfile = OptionalLightProfile->GetValue();
-				SetupTextureSourceDataFromBulkData_Editor(TextureCubeArray, LightProfile, MoveTemp(ProcessedPayload.PayloadAndId));
+				SetupTextureSourceDataFromBulkData_Editor(TextureCubeArray, LightProfile, MoveTemp(ProcessedPayload.PayloadAndId), bIsReimport);
 			}
 		}
 		else
@@ -893,7 +902,7 @@ namespace UE::Interchange::Private::InterchangeTextureFactory
 		return false;
 	}
 
-	void SetupTexture2DArraySourceDataFromPayload(UTexture2DArray* Texture2DArray, FProcessedPayload& ProcessedPayload)
+	void SetupTexture2DArraySourceDataFromPayload(UTexture2DArray* Texture2DArray, FProcessedPayload& ProcessedPayload, bool bIsReimport)
 	{
 		if (TOptional<FImportSlicedImage>* OptionalSlicedImage = ProcessedPayload.SettingsFromPayload.TryGet<TOptional<FImportSlicedImage>>())
 		{
@@ -901,7 +910,7 @@ namespace UE::Interchange::Private::InterchangeTextureFactory
 			{
 				FImportSlicedImage& SlicedImage = OptionalSlicedImage->GetValue();
 
-				SetupTextureSourceDataFromBulkData_Editor(Texture2DArray, SlicedImage, MoveTemp(ProcessedPayload.PayloadAndId));
+				SetupTextureSourceDataFromBulkData_Editor(Texture2DArray, SlicedImage, MoveTemp(ProcessedPayload.PayloadAndId), bIsReimport);
 			}
 		}
 		else if (TOptional<FImportImage>* OptionalImage = ProcessedPayload.SettingsFromPayload.TryGet<TOptional<FImportImage>>())
@@ -909,7 +918,7 @@ namespace UE::Interchange::Private::InterchangeTextureFactory
 			if (OptionalImage->IsSet())
 			{
 				FImportImage& Image = OptionalImage->GetValue();
-				SetupTextureSourceDataFromBulkData_Editor(Texture2DArray, Image, MoveTemp(ProcessedPayload.PayloadAndId));
+				SetupTextureSourceDataFromBulkData_Editor(Texture2DArray, Image, MoveTemp(ProcessedPayload.PayloadAndId), bIsReimport);
 			}
 		}
 		else if (TOptional<FImportLightProfile>* OptionalLightProfile = ProcessedPayload.SettingsFromPayload.TryGet<TOptional<FImportLightProfile>>())
@@ -917,7 +926,7 @@ namespace UE::Interchange::Private::InterchangeTextureFactory
 			if (OptionalLightProfile->IsSet())
 			{
 				FImportLightProfile& LightProfile = OptionalLightProfile->GetValue();
-				SetupTextureSourceDataFromBulkData_Editor(Texture2DArray, LightProfile, MoveTemp(ProcessedPayload.PayloadAndId));
+				SetupTextureSourceDataFromBulkData_Editor(Texture2DArray, LightProfile, MoveTemp(ProcessedPayload.PayloadAndId), bIsReimport);
 			}
 		}
 		else
@@ -956,7 +965,7 @@ namespace UE::Interchange::Private::InterchangeTextureFactory
 		return false;
 	}
 
-	void SetupVolumeTextureSourceDataFromPayload(UVolumeTexture* VolumeTexture, FProcessedPayload& ProcessedPayload)
+	void SetupVolumeTextureSourceDataFromPayload(UVolumeTexture* VolumeTexture, FProcessedPayload& ProcessedPayload, bool bIsReimport)
 	{
 		if (TOptional<FImportSlicedImage>* OptionalSlicedImage = ProcessedPayload.SettingsFromPayload.TryGet<TOptional<FImportSlicedImage>>())
 		{
@@ -964,7 +973,7 @@ namespace UE::Interchange::Private::InterchangeTextureFactory
 			{
 				FImportSlicedImage& SlicedImage = OptionalSlicedImage->GetValue();
 
-				SetupTextureSourceDataFromBulkData_Editor(VolumeTexture, SlicedImage, MoveTemp(ProcessedPayload.PayloadAndId));
+				SetupTextureSourceDataFromBulkData_Editor(VolumeTexture, SlicedImage, MoveTemp(ProcessedPayload.PayloadAndId), bIsReimport);
 			}
 		}
 		else if (TOptional<FImportImage>* OptionalImage = ProcessedPayload.SettingsFromPayload.TryGet<TOptional<FImportImage>>())
@@ -972,7 +981,7 @@ namespace UE::Interchange::Private::InterchangeTextureFactory
 			if (OptionalImage->IsSet())
 			{
 				FImportImage& Image = OptionalImage->GetValue();
-				SetupTextureSourceDataFromBulkData_Editor(VolumeTexture, Image, MoveTemp(ProcessedPayload.PayloadAndId));
+				SetupTextureSourceDataFromBulkData_Editor(VolumeTexture, Image, MoveTemp(ProcessedPayload.PayloadAndId), bIsReimport);
 			}
 		}
 		else if (TOptional<FImportLightProfile>* OptionalLightProfile = ProcessedPayload.SettingsFromPayload.TryGet<TOptional<FImportLightProfile>>())
@@ -980,7 +989,7 @@ namespace UE::Interchange::Private::InterchangeTextureFactory
 			if (OptionalLightProfile->IsSet())
 			{
 				FImportLightProfile& LightProfile = OptionalLightProfile->GetValue();
-				SetupTextureSourceDataFromBulkData_Editor(VolumeTexture, LightProfile, MoveTemp(ProcessedPayload.PayloadAndId));
+				SetupTextureSourceDataFromBulkData_Editor(VolumeTexture, LightProfile, MoveTemp(ProcessedPayload.PayloadAndId), bIsReimport);
 			}
 		}
 		else
@@ -1270,23 +1279,23 @@ namespace UE::Interchange::Private::InterchangeTextureFactory
 		return false;
 	}
 
-	void SetupTextureSourceDataFromBulkData(UTexture* Texture, const FImportImage& Image, UE::Serialization::FEditorBulkData::FSharedBufferWithID&& BufferAndId)
+	void SetupTextureSourceDataFromBulkData(UTexture* Texture, const FImportImage& Image, UE::Serialization::FEditorBulkData::FSharedBufferWithID&& BufferAndId, bool bIsReimport)
 	{
 #if WITH_EDITOR
-		SetupTextureSourceDataFromBulkData_Editor(Texture, Image, MoveTemp(BufferAndId));
+		SetupTextureSourceDataFromBulkData_Editor(Texture, Image, MoveTemp(BufferAndId), bIsReimport);
 #else
 		SetupTexture2DSourceDataFromBulkData_Runtime(Cast<UTexture2D>(Texture), Image, MoveTemp(BufferAndId));
 #endif
 	}
 
-	void SetupTexture2DSourceDataFromPayload(UTexture2D* Texture2D, FProcessedPayload& ProcessedPayload)
+	void SetupTexture2DSourceDataFromPayload(UTexture2D* Texture2D, FProcessedPayload& ProcessedPayload, bool bIsReimport)
 	{
 #if WITH_EDITOR
 		if (TOptional<FImportBlockedImage>* BlockedImage = ProcessedPayload.SettingsFromPayload.TryGet<TOptional<FImportBlockedImage>>())
 		{
 			if (BlockedImage->IsSet())
 			{
-				SetupTexture2DSourceDataFromBulkData_Editor(Texture2D, BlockedImage->GetValue(), MoveTemp(ProcessedPayload.PayloadAndId));
+				SetupTexture2DSourceDataFromBulkData_Editor(Texture2D, BlockedImage->GetValue(), MoveTemp(ProcessedPayload.PayloadAndId), bIsReimport);
 			}
 		}
 		else
@@ -1295,7 +1304,7 @@ namespace UE::Interchange::Private::InterchangeTextureFactory
 		{
 			if (Image->IsSet())
 			{
-				SetupTextureSourceDataFromBulkData(Texture2D, Image->GetValue(), MoveTemp(ProcessedPayload.PayloadAndId));
+				SetupTextureSourceDataFromBulkData(Texture2D, Image->GetValue(), MoveTemp(ProcessedPayload.PayloadAndId), bIsReimport);
 			}
 		}
 		else if (TOptional<FImportLightProfile>* OptionalLightProfile = ProcessedPayload.SettingsFromPayload.TryGet<TOptional<FImportLightProfile>>())
@@ -1306,7 +1315,7 @@ namespace UE::Interchange::Private::InterchangeTextureFactory
 
 				if (UTextureLightProfile* TextureLightProfile = Cast<UTextureLightProfile>(Texture2D))
 				{
-					SetupTextureSourceDataFromBulkData(TextureLightProfile, LightProfile, MoveTemp(ProcessedPayload.PayloadAndId));
+					SetupTextureSourceDataFromBulkData(TextureLightProfile, LightProfile, MoveTemp(ProcessedPayload.PayloadAndId), bIsReimport);
 
 					if (const TOptional<FImportLightProfile>* ImportLightProfilePtr = ProcessedPayload.SettingsFromPayload.TryGet<TOptional<FImportLightProfile>>())
 					{
@@ -1314,6 +1323,7 @@ namespace UE::Interchange::Private::InterchangeTextureFactory
 						{
 							const UE::Interchange::FImportLightProfile& ImportLightProfile = ImportLightProfilePtr->GetValue();
 
+							// The legacy factory updated these setting on each reimport (leave it as it is for 5.1 but we should look into moving that decision into the pipelines)
 							TextureLightProfile->Brightness = ImportLightProfile.Brightness;
 							TextureLightProfile->TextureMultiplier = ImportLightProfile.TextureMultiplier;
 						}
@@ -1321,7 +1331,7 @@ namespace UE::Interchange::Private::InterchangeTextureFactory
 				}
 				else
 				{
-					SetupTextureSourceDataFromBulkData(Texture2D, LightProfile, MoveTemp(ProcessedPayload.PayloadAndId));
+					SetupTextureSourceDataFromBulkData(Texture2D, LightProfile, MoveTemp(ProcessedPayload.PayloadAndId), bIsReimport);
 				}
 			}
 		}
@@ -1364,28 +1374,28 @@ namespace UE::Interchange::Private::InterchangeTextureFactory
 #endif // WITH_EDITORONLY_DATA
 	}
 
-	void SetupTextureSourceDataFromPayload(UTexture* Texture, FProcessedPayload& ProcessedPayload)
+	void SetupTextureSourceDataFromPayload(UTexture* Texture, FProcessedPayload& ProcessedPayload, bool bIsReimport)
 	{
 		if (UTexture2D* Texture2D = Cast<UTexture2D>(Texture))
 		{
-			SetupTexture2DSourceDataFromPayload(Texture2D, ProcessedPayload);
+			SetupTexture2DSourceDataFromPayload(Texture2D, ProcessedPayload, bIsReimport);
 		}
 #if WITH_EDITOR
 		else if (UTextureCube* TextureCube = Cast<UTextureCube>(Texture))
 		{
-			SetupTextureCubeSourceDataFromPayload(TextureCube, ProcessedPayload);
+			SetupTextureCubeSourceDataFromPayload(TextureCube, ProcessedPayload, bIsReimport);
 		}
 		else if (UTextureCubeArray* TextureCubeArray = Cast<UTextureCubeArray>(Texture))
 		{
-			SetupTextureCubeArraySourceDataFromPayload(TextureCubeArray, ProcessedPayload);
+			SetupTextureCubeArraySourceDataFromPayload(TextureCubeArray, ProcessedPayload, bIsReimport);
 		}
 		else if (UTexture2DArray* Texture2DArray = Cast<UTexture2DArray>(Texture))
 		{
-			SetupTexture2DArraySourceDataFromPayload(Texture2DArray, ProcessedPayload);
+			SetupTexture2DArraySourceDataFromPayload(Texture2DArray, ProcessedPayload, bIsReimport);
 		}
 		else if (UVolumeTexture* VolumeTexture = Cast<UVolumeTexture>(Texture))
 		{
-			SetupVolumeTextureSourceDataFromPayload(VolumeTexture, ProcessedPayload);
+			SetupVolumeTextureSourceDataFromPayload(VolumeTexture, ProcessedPayload, bIsReimport);
 		}
 #endif // WITH_EDITOR
 		else
@@ -1654,7 +1664,7 @@ void UInterchangeTextureFactory::PreImportPreCompletedCallback(const FImportPreC
 		using namespace UE::Interchange::Private::InterchangeTextureFactory;
 
 		// Setup source data
-		SetupTextureSourceDataFromPayload(Texture, ProcessedPayload);
+		SetupTextureSourceDataFromPayload(Texture, ProcessedPayload, Arguments.bIsReimport);
 
 #if WITH_EDITOR
 		UInterchangeFactoryBaseNode* TextureFactoryNode = Arguments.FactoryNode;
diff --git a/Engine/Source/Runtime/Interchange/Engine/Private/Tasks/InterchangeTaskCompletion.cpp b/Engine/Source/Runtime/Interchange/Engine/Private/Tasks/InterchangeTaskCompletion.cpp
--- a/Engine/Source/Runtime/Interchange/Engine/Private/Tasks/InterchangeTaskCompletion.cpp
+++ b/Engine/Source/Runtime/Interchange/Engine/Private/Tasks/InterchangeTaskCompletion.cpp
@@ -81,6 +81,7 @@ void UE::Interchange::FTaskPreCompletion::DoTask(ENamedThreads::Type CurrentThre
 					Arguments.NodeContainer = AsyncHelper->BaseNodeContainers[SourceIndex].Get();
 					Arguments.Pipelines = AsyncHelper->Pipelines;
 					Arguments.OriginalPipelines = AsyncHelper->OriginalPipelines;
+					Arguments.bIsReimport = ObjectInfo.bIsReimport;
 					ObjectInfo.Factory->PreImportPreCompletedCallback(Arguments);
 				}
 
