Interchange fix "Failed to import PNG" error messages for all texture formats

diff --git a/Engine/Plugins/Interchange/Runtime/Source/Import/Private/Texture/InterchangeImageWrapperTranslator.cpp b/Engine/Plugins/Interchange/Runtime/Source/Import/Private/Texture/InterchangeImageWrapperTranslator.cpp
--- a/Engine/Plugins/Interchange/Runtime/Source/Import/Private/Texture/InterchangeImageWrapperTranslator.cpp
+++ b/Engine/Plugins/Interchange/Runtime/Source/Import/Private/Texture/InterchangeImageWrapperTranslator.cpp
@@ -128,7 +128,7 @@ TOptional<UE::Interchange::FImportImage> UInterchangeImageWrapperTranslator::Get
 
 	if (!GetSourceData())
 	{
-		UE_LOG(LogInterchangeImport, Error, TEXT("Failed to import PNG, bad source data."));
+		UE_LOG(LogInterchangeImport, Error, TEXT("Failed to import Texture, bad source data."));
 		return TOptional<UE::Interchange::FImportImage>();
 	}
 
@@ -138,19 +138,19 @@ TOptional<UE::Interchange::FImportImage> UInterchangeImageWrapperTranslator::Get
 	//Make sure the key fit the filename, The key should always be valid
 	if (!Filename.Equals(PayLoadKey))
 	{
-		UE_LOG(LogInterchangeImport, Error, TEXT("Failed to import PNG, wrong payload key. [%s]"), *Filename);
+		UE_LOG(LogInterchangeImport, Error, TEXT("Failed to import Texture, wrong payload key. [%s]"), *Filename);
 		return TOptional<UE::Interchange::FImportImage>();
 	}
 
 	if (!FPaths::FileExists(Filename))
 	{
-		UE_LOG(LogInterchangeImport, Error, TEXT("Failed to import PNG, cannot open file. [%s]"), *Filename);
+		UE_LOG(LogInterchangeImport, Error, TEXT("Failed to import Texture, cannot open file. [%s]"), *Filename);
 		return TOptional<UE::Interchange::FImportImage>();
 	}
 
 	if (!FFileHelper::LoadFileToArray(SourceDataBuffer, *Filename))
 	{
-		UE_LOG(LogInterchangeImport, Error, TEXT("Failed to import PNG, cannot load file content into an array. [%s]"), *Filename);
+		UE_LOG(LogInterchangeImport, Error, TEXT("Failed to import Texture, cannot load file content into an array. [%s]"), *Filename);
 		return TOptional<UE::Interchange::FImportImage>();
 	}
 
