Interchange:
- Fixed crash when trying to load an asset outside the grame thread.
- Fixed textures being listed in Import into level

diff --git a/Engine/Plugins/Interchange/Runtime/Source/Import/Private/MaterialX/InterchangeMaterialXTranslator.cpp b/Engine/Plugins/Interchange/Runtime/Source/Import/Private/MaterialX/InterchangeMaterialXTranslator.cpp
--- a/Engine/Plugins/Interchange/Runtime/Source/Import/Private/MaterialX/InterchangeMaterialXTranslator.cpp
+++ b/Engine/Plugins/Interchange/Runtime/Source/Import/Private/MaterialX/InterchangeMaterialXTranslator.cpp
@@ -28,40 +28,6 @@ static FAutoConsoleVariableRef CCvarInterchangeEnableMaterialXImport(
 
 namespace mx = MaterialX;
 
-namespace UE::Interchange::MaterialX
-{
-	bool IsStandardSurfacePackageLoaded()
-	{
-		auto IsPackageLoaded = [](const FString& TextPath) -> bool
-		{
-			const FString FunctionPath{ FPackageName::ExportTextPathToObjectPath(TextPath) };
-			if(FPackageName::DoesPackageExist(FunctionPath))
-			{
-				if(FSoftObjectPath(FunctionPath).TryLoad())
-				{
-					return true;
-				}
-				else
-				{
-					UE_LOG(LogInterchangeImport, Warning, TEXT("Couldn't load %s"), *FunctionPath);
-				}
-			}
-			else
-			{
-				UE_LOG(LogInterchangeImport, Warning, TEXT("Couldn't find %s"), *FunctionPath);
-			}
-
-			return false;
-		};
-
-		static const bool bStandardSurfacePackageLoaded =
-			IsPackageLoaded(TEXT("MaterialFunction'/Interchange/Functions/MX_StandardSurface.MX_StandardSurface'")) &&
-			IsPackageLoaded(TEXT("MaterialFunction'/Interchange/Functions/MX_TransmissionSurface.MX_TransmissionSurface'"));
-
-		return bStandardSurfacePackageLoaded;
-	}
-}
-
 UInterchangeMaterialXTranslator::UInterchangeMaterialXTranslator()
 #if WITH_EDITOR
 	:
@@ -234,7 +200,7 @@ TArray<FString> UInterchangeMaterialXTranslator::GetSupportedFormats() const
 		return TArray<FString>{};
 	}
 
-	return UE::Interchange::MaterialX::IsStandardSurfacePackageLoaded() ? TArray<FString>{ TEXT("mtlx;MaterialX File Format") } : TArray<FString>{};
+	return TArray<FString>{ TEXT("mtlx;MaterialX File Format") };
 }
 
 bool UInterchangeMaterialXTranslator::Translate(UInterchangeBaseNodeContainer& BaseNodeContainer) const
diff --git a/Engine/Plugins/Interchange/Runtime/Source/Pipelines/Private/InterchangeGenericMaterialPipeline.cpp b/Engine/Plugins/Interchange/Runtime/Source/Pipelines/Private/InterchangeGenericMaterialPipeline.cpp
--- a/Engine/Plugins/Interchange/Runtime/Source/Pipelines/Private/InterchangeGenericMaterialPipeline.cpp
+++ b/Engine/Plugins/Interchange/Runtime/Source/Pipelines/Private/InterchangeGenericMaterialPipeline.cpp
@@ -66,6 +66,55 @@ FString LexToString(UInterchangeGenericMaterialPipeline::EMaterialInputType Valu
 	}
 }
 
+namespace UE::Interchange::InterchangeGenericMaterialPipeline::Private
+{
+	bool AreRequiredPackagesLoaded()
+	{
+		auto ArePackagesLoaded = [](const TArray<FString>& PackagePaths) -> bool
+		{
+			bool bAllLoaded = true;
+
+			for (const FString& PackagePath : PackagePaths)
+			{
+				const FString ObjectPath(FPackageName::ExportTextPathToObjectPath(PackagePath));
+
+				if (FPackageName::DoesPackageExist(ObjectPath))
+				{
+					if (FSoftObjectPath(ObjectPath).TryLoad())
+					{
+						continue;
+					}
+					else
+					{
+						UE_LOG(LogInterchangePipeline, Warning, TEXT("Couldn't load %s"), *PackagePath);
+					}
+				}
+				else
+				{
+					UE_LOG(LogInterchangePipeline, Warning, TEXT("Couldn't find %s"), *PackagePath);
+				}
+
+				bAllLoaded = false;
+			}
+
+			return bAllLoaded;
+		};
+
+		TArray<FString> RequiredPackages = {
+			TEXT("MaterialFunction'/Interchange/Functions/MX_StandardSurface.MX_StandardSurface'"),
+			TEXT("MaterialFunction'/Interchange/Functions/MX_TransmissionSurface.MX_TransmissionSurface'"),
+			TEXT("MaterialFunction'/Engine/Functions/Engine_MaterialFunctions01/Shading/ConvertFromDiffSpec.ConvertFromDiffSpec'"),
+			TEXT("MaterialFunction'/Engine/Functions/Engine_MaterialFunctions01/Texturing/FlattenNormal.FlattenNormal'"),
+			TEXT("MaterialFunction'/Engine/Functions/Engine_MaterialFunctions02/Utility/MakeFloat3.MakeFloat3'"),
+			TEXT("MaterialFunction'/Engine/Functions/Engine_MaterialFunctions02/Texturing/CustomRotator.CustomRotator'"),
+		};
+
+		static const bool bRequiredPackagesLoaded = ArePackagesLoaded(RequiredPackages);
+
+		return bRequiredPackagesLoaded;
+	}
+}
+
 UInterchangeGenericMaterialPipeline::UInterchangeGenericMaterialPipeline()
 {
 	TexturePipeline = CreateDefaultSubobject<UInterchangeGenericTexturePipeline>("TexturePipeline");
@@ -83,12 +132,14 @@ void UInterchangeGenericMaterialPipeline::PreDialogCleanup(const FName PipelineS
 
 bool UInterchangeGenericMaterialPipeline::IsSettingsAreValid(TOptional<FText>& OutInvalidReason) const
 {
+	using namespace UE::Interchange;
+
 	if (TexturePipeline && !TexturePipeline->IsSettingsAreValid(OutInvalidReason))
 	{
 		return false;
 	}
 
-	return true;
+	return InterchangeGenericMaterialPipeline::Private::AreRequiredPackagesLoaded() &&  Super::IsSettingsAreValid(OutInvalidReason);
 }
 
 void UInterchangeGenericMaterialPipeline::AdjustSettingsForContext(EInterchangePipelineContext ImportType, TObjectPtr<UObject> ReimportAsset)
diff --git a/Engine/Source/Runtime/Interchange/Engine/Private/InterchangeManager.cpp b/Engine/Source/Runtime/Interchange/Engine/Private/InterchangeManager.cpp
--- a/Engine/Source/Runtime/Interchange/Engine/Private/InterchangeManager.cpp
+++ b/Engine/Source/Runtime/Interchange/Engine/Private/InterchangeManager.cpp
@@ -732,7 +732,7 @@ TArray<FString> UInterchangeManager::GetSupportedFormats(const EInterchangeTrans
 	{
 		const UInterchangeTranslatorBase* TranslatorBaseCDO = TranslatorClass->GetDefaultObject<UInterchangeTranslatorBase>();
 
-		if (EnumHasAnyFlags(TranslatorBaseCDO->GetTranslatorType(), ForTranslatorType))
+		if (EnumHasAllFlags(TranslatorBaseCDO->GetTranslatorType(), ForTranslatorType))
 		{
 			FileExtensions.Append(TranslatorBaseCDO->GetSupportedFormats());
 		}
