MaterialX
- Adding the support of geometric nodes: position, normal, tangent, bitangent, geomcolor
- Adding the support of the time node
Generic Material pipeline
- Handling TransformPosition, TransformVector and Time Material expression

diff --git a/Engine/Plugins/Interchange/Runtime/Source/Import/Private/MaterialX/InterchangeMaterialXDefinitions.h b/Engine/Plugins/Interchange/Runtime/Source/Import/Private/MaterialX/InterchangeMaterialXDefinitions.h
--- a/Engine/Plugins/Interchange/Runtime/Source/Import/Private/MaterialX/InterchangeMaterialXDefinitions.h
+++ b/Engine/Plugins/Interchange/Runtime/Source/Import/Private/MaterialX/InterchangeMaterialXDefinitions.h
@@ -251,7 +251,12 @@ MATERIALX_NAMESPACE_BEGIN
 		// Texture nodes 
 		static constexpr const char* Image = "image";
 		static constexpr const char* TiledImage = "tiledimage";
-		// Geometric nodes 
+		// Geometric nodes
+		static constexpr const char* Bitangent = "bitangent";
+		static constexpr const char* GeomColor = "geomcolor";
+		static constexpr const char* Normal = "normal";
+		static constexpr const char* Position = "position";
+		static constexpr const char* Tangent = "tangent";
 		static constexpr const char* TexCoord = "texcoord";
 		// Light nodes
 		static constexpr const char* DirectionalLight = "directional_light";
@@ -264,6 +269,8 @@ MATERIALX_NAMESPACE_BEGIN
 		static constexpr const char* RgbToHsv = "rgbtohsv";
 		static constexpr const char* Saturate = "saturate";
 		static constexpr const char* Smoothstep = "smoothstep";
+		// Application
+		static constexpr const char* Time= "time";
 		// PBR
 		static constexpr const char* StandardSurface = "standard_surface";
 	}
diff --git a/Engine/Plugins/Interchange/Runtime/Source/Import/Private/MaterialX/InterchangeMaterialXTranslator.cpp b/Engine/Plugins/Interchange/Runtime/Source/Import/Private/MaterialX/InterchangeMaterialXTranslator.cpp
--- a/Engine/Plugins/Interchange/Runtime/Source/Import/Private/MaterialX/InterchangeMaterialXTranslator.cpp
+++ b/Engine/Plugins/Interchange/Runtime/Source/Import/Private/MaterialX/InterchangeMaterialXTranslator.cpp
@@ -14,6 +14,9 @@
 #include "Nodes/InterchangeSourceNode.h"
 #include "UObject/GCObjectScopeGuard.h"
 
+#include "Materials/MaterialExpressionTransformPosition.h"
+#include "Materials/MaterialExpressionTransform.h"
+
 #include UE_INLINE_GENERATED_CPP_BY_NAME(InterchangeMaterialXTranslator)
 
 #define LOCTEXT_NAMESPACE "InterchangeMaterialXTranslator"
@@ -143,6 +146,7 @@ NodeNamesMaterialX2UE{
 	{MaterialX::Category::Combine3,     TEXT("Append3Vector")},
 	{MaterialX::Category::Combine4,     TEXT("Append4Vector")},
 	// Geometric nodes 
+	{MaterialX::Category::GeomColor,	TEXT("VertexColor")},
 	{MaterialX::Category::TexCoord,		TEXT("TextureCoordinate")},
 	// Adjustment nodes,
 	{MaterialX::Category::HsvToRgb,		TEXT("HsvToRgb")},
@@ -977,6 +981,26 @@ bool UInterchangeMaterialXTranslator::ConnectNodeGraphOutputToInput(MaterialX::I
 				{
 					ConnectOutsideInputToOutput(UpstreamNode, ParentShaderNode, InputChannelName, NamesToShaderNodes, NodeContainer);
 				}
+				else if(UpstreamNode->getCategory() == mx::Category::Position)
+				{
+					ConnectPositionInputToOutput(UpstreamNode, ParentShaderNode, InputChannelName, NamesToShaderNodes, NodeContainer);
+				}
+				else if(UpstreamNode->getCategory() == mx::Category::Normal)
+				{
+					ConnectNormalInputToOutput(UpstreamNode, ParentShaderNode, InputChannelName, NamesToShaderNodes, NodeContainer);
+				}
+				else if(UpstreamNode->getCategory() == mx::Category::Tangent)
+				{
+					ConnectTangentInputToOutput(UpstreamNode, ParentShaderNode, InputChannelName, NamesToShaderNodes, NodeContainer);
+				}
+				else if(UpstreamNode->getCategory() == mx::Category::Bitangent)
+				{
+					ConnectBitangentInputToOutput(UpstreamNode, ParentShaderNode, InputChannelName, NamesToShaderNodes, NodeContainer);
+				}
+				else if(UpstreamNode->getCategory() == mx::Category::Time)
+				{
+					ConnectTimeInputToOutput(UpstreamNode, ParentShaderNode, InputChannelName, NamesToShaderNodes, NodeContainer);
+				}
 				else
 				{
 					UInterchangeResultWarning_Generic* Message = AddMessage<UInterchangeResultWarning_Generic>();
@@ -1739,6 +1763,184 @@ void UInterchangeMaterialXTranslator::ConnectOutsideInputToOutput(MaterialX::Nod
 	UInterchangeShaderPortsAPI::ConnectDefaultOuputToInput(NodeMultiply, TEXT("B"), NodeOneMinus->GetUniqueID());
 }
 
+void UInterchangeMaterialXTranslator::ConnectPositionInputToOutput(MaterialX::NodePtr UpstreamNode, UInterchangeShaderNode* ParentShaderNode, const FString& InputChannelName, TMap<FString, UInterchangeShaderNode*>& NamesToShaderNodes, UInterchangeBaseNodeContainer& NodeContainer) const
+{
+	// MaterialX defines the space as: object, model, world
+	// model: The local coordinate space of the geometry, before any local deformations or global transforms have been applied.
+	// object: The local coordinate space of the geometry, after local deformations have been applied, but before any global transforms.
+	// world : The global coordinate space of the geometry, after local deformationsand global transforms have been applied.
+
+	// For the moment we don't have the distinction between model/object, so let's just create an UMaterialExpressionObjectPositionWS
+	// In case of model/object we need to add a TransformPoint from world to local space
+	UInterchangeShaderNode* PositionNode = CreateShaderNode<UInterchangeShaderNode>(UpstreamNode->getName().c_str(), TEXT("ObjectPositionWS"), NamesToShaderNodes, NodeContainer);
+	UInterchangeShaderNode* NodeToConnectTo = ParentShaderNode;
+	FString InputToConnectTo = InputChannelName;
+
+	std::string Space = "world";
+	mx::InputPtr InputSpace = UpstreamNode->getInput("space");
+
+	if(InputSpace)
+	{
+		Space = InputSpace->getValueString();
+	}
+
+	//the default space defined by the nodedef is "object"
+	if(Space != "world" || !InputSpace)
+	{
+		using namespace UE::Interchange::Materials::Standard::Nodes;
+
+		UInterchangeShaderNode* TransformNode = CreateShaderNode<UInterchangeShaderNode>(UpstreamNode->getName().c_str() + FString(TEXT("_Transform")), TransformPosition::Name.ToString(), NamesToShaderNodes, NodeContainer);
+		NodeToConnectTo = TransformNode;
+		InputToConnectTo = TransformPosition::Inputs::Input.ToString();
+		TransformNode->AddInt32Attribute(TransformPosition::Attributes::TransformSourceType, EMaterialPositionTransformSource::TRANSFORMPOSSOURCE_World);
+		TransformNode->AddInt32Attribute(TransformPosition::Attributes::TransformType, EMaterialPositionTransformSource::TRANSFORMPOSSOURCE_Local);
+		UInterchangeShaderPortsAPI::ConnectDefaultOuputToInput(ParentShaderNode, InputChannelName, TransformNode->GetUniqueID());
+	}
+
+	UInterchangeShaderPortsAPI::ConnectDefaultOuputToInput(NodeToConnectTo, InputToConnectTo, PositionNode->GetUniqueID());
+}
+
+void UInterchangeMaterialXTranslator::ConnectNormalInputToOutput(MaterialX::NodePtr UpstreamNode, UInterchangeShaderNode* ParentShaderNode, const FString& InputChannelName, TMap<FString, UInterchangeShaderNode*>& NamesToShaderNodes, UInterchangeBaseNodeContainer& NodeContainer) const
+{
+	// MaterialX defines the space as: object, model, world
+	// model: The local coordinate space of the geometry, before any local deformations or global transforms have been applied.
+	// object: The local coordinate space of the geometry, after local deformations have been applied, but before any global transforms.
+	// world : The global coordinate space of the geometry, after local deformationsand global transforms have been applied.
+
+	// For the moment we don't have the distinction between model/object, so let's just create an UMaterialExpressionVertexNormalWS
+	// In case of model/object we need to add a TransformVector from world to local space
+	UInterchangeShaderNode* NormalNode = CreateShaderNode<UInterchangeShaderNode>(UpstreamNode->getName().c_str(), TEXT("VertexNormalWS"), NamesToShaderNodes, NodeContainer);
+	UInterchangeShaderNode* NodeToConnectTo = ParentShaderNode;
+	FString InputToConnectTo = InputChannelName;
+
+	std::string Space = "world";
+	mx::InputPtr InputSpace = UpstreamNode->getInput("space");
+
+	if(InputSpace)
+	{
+		Space = InputSpace->getValueString();
+	}
+
+	//the default space defined by the nodedef is "object"
+	if(Space != "world" || !InputSpace)
+	{
+		using namespace UE::Interchange::Materials::Standard::Nodes;
+
+		UInterchangeShaderNode* TransformNode = CreateShaderNode<UInterchangeShaderNode>(UpstreamNode->getName().c_str() + FString(TEXT("_Transform")), TransformVector::Name.ToString(), NamesToShaderNodes, NodeContainer);
+		NodeToConnectTo = TransformNode;
+		InputToConnectTo = TransformVector::Inputs::Input.ToString();
+		TransformNode->AddInt32Attribute(TransformVector::Attributes::TransformSourceType, EMaterialVectorCoordTransformSource::TRANSFORMSOURCE_World);
+		TransformNode->AddInt32Attribute(TransformVector::Attributes::TransformType, EMaterialVectorCoordTransform::TRANSFORM_Local);
+		UInterchangeShaderPortsAPI::ConnectDefaultOuputToInput(ParentShaderNode, InputChannelName, TransformNode->GetUniqueID());
+	}
+
+	UInterchangeShaderPortsAPI::ConnectDefaultOuputToInput(NodeToConnectTo, InputToConnectTo, NormalNode->GetUniqueID());
+}
+
+void UInterchangeMaterialXTranslator::ConnectTangentInputToOutput(MaterialX::NodePtr UpstreamNode, UInterchangeShaderNode* ParentShaderNode, const FString& InputChannelName, TMap<FString, UInterchangeShaderNode*>& NamesToShaderNodes, UInterchangeBaseNodeContainer& NodeContainer) const
+{
+	// MaterialX defines the space as: object, model, world
+	// model: The local coordinate space of the geometry, before any local deformations or global transforms have been applied.
+	// object: The local coordinate space of the geometry, after local deformations have been applied, but before any global transforms.
+	// world : The global coordinate space of the geometry, after local deformationsand global transforms have been applied.
+
+	// For the moment we don't have the distinction between model/object, so let's just create an UMaterialExpressionVertexTangentWS
+	// In case of model/object we need to add a TransformVector from world to local space
+	UInterchangeShaderNode* TangentNode = CreateShaderNode<UInterchangeShaderNode>(UpstreamNode->getName().c_str(), TEXT("VertexTangentWS"), NamesToShaderNodes, NodeContainer);
+	UInterchangeShaderNode* NodeToConnectTo = ParentShaderNode;
+	FString InputToConnectTo = InputChannelName;
+
+	std::string Space = "world";
+	mx::InputPtr InputSpace = UpstreamNode->getInput("space");
+
+	if(InputSpace)
+	{
+		Space = InputSpace->getValueString();
+	}
+
+	//the default space defined by the nodedef is "object"
+	if(Space != "world" || !InputSpace)
+	{
+		using namespace UE::Interchange::Materials::Standard::Nodes;
+
+		UInterchangeShaderNode* TransformNode = CreateShaderNode<UInterchangeShaderNode>(UpstreamNode->getName().c_str() + FString(TEXT("_Transform")), TransformVector::Name.ToString(), NamesToShaderNodes, NodeContainer);
+		NodeToConnectTo = TransformNode;
+		InputToConnectTo = TransformVector::Inputs::Input.ToString();
+		TransformNode->AddInt32Attribute(TransformVector::Attributes::TransformSourceType, EMaterialVectorCoordTransformSource::TRANSFORMSOURCE_World);
+		TransformNode->AddInt32Attribute(TransformVector::Attributes::TransformType, EMaterialVectorCoordTransform::TRANSFORM_Local);
+		UInterchangeShaderPortsAPI::ConnectDefaultOuputToInput(ParentShaderNode, InputChannelName, TransformNode->GetUniqueID());
+	}
+
+	UInterchangeShaderPortsAPI::ConnectDefaultOuputToInput(NodeToConnectTo, InputToConnectTo, TangentNode->GetUniqueID());
+}
+
+void UInterchangeMaterialXTranslator::ConnectBitangentInputToOutput(MaterialX::NodePtr UpstreamNode, UInterchangeShaderNode* ParentShaderNode, const FString& InputChannelName, TMap<FString, UInterchangeShaderNode*>& NamesToShaderNodes, UInterchangeBaseNodeContainer& NodeContainer) const
+{
+	// MaterialX defines the space as: object, model, world
+	// model: The local coordinate space of the geometry, before any local deformations or global transforms have been applied.
+	// object: The local coordinate space of the geometry, after local deformations have been applied, but before any global transforms.
+	// world : The global coordinate space of the geometry, after local deformationsand global transforms have been applied.
+
+	// For the moment we don't have the distinction between model/object, so let's just do the cross product between the normal and the tangent
+	// In case of model/object we need to add a TransformVector from world to local space
+	UInterchangeShaderNode* NormalNode = CreateShaderNode<UInterchangeShaderNode>(UpstreamNode->getName().c_str() + FString(TEXT("_Normal")), TEXT("VertexNormalWS"), NamesToShaderNodes, NodeContainer);
+	UInterchangeShaderNode* TangentNode = CreateShaderNode<UInterchangeShaderNode>(UpstreamNode->getName().c_str() + FString(TEXT("_Tangent")), TEXT("VertexTangentWS"), NamesToShaderNodes, NodeContainer);
+	UInterchangeShaderNode* BitangentNode = CreateShaderNode<UInterchangeShaderNode>(UpstreamNode->getName().c_str(), TEXT("CrossProduct"), NamesToShaderNodes, NodeContainer);
+	
+	UInterchangeShaderPortsAPI::ConnectDefaultOuputToInput(BitangentNode, TEXT("A"), NormalNode->GetUniqueID());
+	UInterchangeShaderPortsAPI::ConnectDefaultOuputToInput(BitangentNode, TEXT("B"), TangentNode->GetUniqueID());
+
+
+	UInterchangeShaderNode* NodeToConnectTo = ParentShaderNode;
+	FString InputToConnectTo = InputChannelName;
+
+	std::string Space = "world";
+	mx::InputPtr InputSpace = UpstreamNode->getInput("space");
+
+	if(InputSpace)
+	{
+		Space = InputSpace->getValueString();
+	}
+
+	//the default space defined by the nodedef is "object"
+	if(Space != "world" || !InputSpace)
+	{
+		using namespace UE::Interchange::Materials::Standard::Nodes;
+
+		UInterchangeShaderNode* TransformNode = CreateShaderNode<UInterchangeShaderNode>(UpstreamNode->getName().c_str() + FString(TEXT("_Transform")), TransformVector::Name.ToString(), NamesToShaderNodes, NodeContainer);
+		NodeToConnectTo = TransformNode;
+		InputToConnectTo = TransformVector::Inputs::Input.ToString();
+		TransformNode->AddInt32Attribute(TransformVector::Attributes::TransformSourceType, EMaterialVectorCoordTransformSource::TRANSFORMSOURCE_World);
+		TransformNode->AddInt32Attribute(TransformVector::Attributes::TransformType, EMaterialVectorCoordTransform::TRANSFORM_Local);
+		UInterchangeShaderPortsAPI::ConnectDefaultOuputToInput(ParentShaderNode, InputChannelName, TransformNode->GetUniqueID());
+	}
+
+	UInterchangeShaderPortsAPI::ConnectDefaultOuputToInput(NodeToConnectTo, InputToConnectTo, BitangentNode->GetUniqueID());
+}
+
+void UInterchangeMaterialXTranslator::ConnectTimeInputToOutput(MaterialX::NodePtr UpstreamNode, UInterchangeShaderNode* ParentShaderNode, const FString& InputChannelName, TMap<FString, UInterchangeShaderNode*>& NamesToShaderNodes, UInterchangeBaseNodeContainer& NodeContainer) const
+{
+	using namespace UE::Interchange::Materials::Standard::Nodes;
+	UInterchangeShaderNode* TimeNode = CreateShaderNode<UInterchangeShaderNode>(UpstreamNode->getName().c_str(), TEXT("Time"), NamesToShaderNodes, NodeContainer);
+	TimeNode->AddBooleanAttribute(Time::Attributes::OverridePeriod, true);
+
+	float FPS; 
+	mx::InputPtr Input = UpstreamNode->getInput("fps");
+
+	//Take the default value from the node definition
+	if(!Input)
+	{
+		Input = UpstreamNode->getNodeDef()->getInput("fps");
+	}
+
+	FPS = mx::fromValueString<float>(Input->getValueString());
+
+	//UE is a period
+	TimeNode->AddFloatAttribute(Time::Attributes::Period, 1.f / FPS);
+
+	UInterchangeShaderPortsAPI::ConnectDefaultOuputToInput(ParentShaderNode, InputChannelName, TimeNode->GetUniqueID());
+}
+
 UInterchangeShaderNode* UInterchangeMaterialXTranslator::CreateMaskShaderNode(uint8 RGBA, const FString& NodeName, TMap<FString, UInterchangeShaderNode*>& NamesToShaderNodes, UInterchangeBaseNodeContainer& NodeContainer) const
 {
 	bool bR = (0b1000 & RGBA) >> 3;
diff --git a/Engine/Plugins/Interchange/Runtime/Source/Import/Public/MaterialX/InterchangeMaterialXTranslator.h b/Engine/Plugins/Interchange/Runtime/Source/Import/Public/MaterialX/InterchangeMaterialXTranslator.h
--- a/Engine/Plugins/Interchange/Runtime/Source/Import/Public/MaterialX/InterchangeMaterialXTranslator.h
+++ b/Engine/Plugins/Interchange/Runtime/Source/Import/Public/MaterialX/InterchangeMaterialXTranslator.h
@@ -347,6 +347,11 @@ protected:
 	void ConnectIfGreaterEqInputToOutput(MaterialX::NodePtr UpstreamNode, UInterchangeShaderNode* ParentShaderNode, const FString& InputChannelName, TMap<FString, UInterchangeShaderNode*>& NamesToShaderNodes, UInterchangeBaseNodeContainer& NodeContainer) const;
 	void ConnectIfEqualInputToOutput(MaterialX::NodePtr UpstreamNode, UInterchangeShaderNode* ParentShaderNode, const FString& InputChannelName, TMap<FString, UInterchangeShaderNode*>& NamesToShaderNodes, UInterchangeBaseNodeContainer& NodeContainer) const;
 	void ConnectOutsideInputToOutput(MaterialX::NodePtr UpstreamNode, UInterchangeShaderNode* ParentShaderNode, const FString& InputChannelName, TMap<FString, UInterchangeShaderNode*>& NamesToShaderNodes, UInterchangeBaseNodeContainer& NodeContainer) const;
+	void ConnectPositionInputToOutput(MaterialX::NodePtr UpstreamNode, UInterchangeShaderNode* ParentShaderNode, const FString& InputChannelName, TMap<FString, UInterchangeShaderNode*>& NamesToShaderNodes, UInterchangeBaseNodeContainer& NodeContainer) const;
+	void ConnectNormalInputToOutput(MaterialX::NodePtr UpstreamNode, UInterchangeShaderNode* ParentShaderNode, const FString& InputChannelName, TMap<FString, UInterchangeShaderNode*>& NamesToShaderNodes, UInterchangeBaseNodeContainer& NodeContainer) const;
+	void ConnectTangentInputToOutput(MaterialX::NodePtr UpstreamNode, UInterchangeShaderNode* ParentShaderNode, const FString& InputChannelName, TMap<FString, UInterchangeShaderNode*>& NamesToShaderNodes, UInterchangeBaseNodeContainer& NodeContainer) const;
+	void ConnectBitangentInputToOutput(MaterialX::NodePtr UpstreamNode, UInterchangeShaderNode* ParentShaderNode, const FString& InputChannelName, TMap<FString, UInterchangeShaderNode*>& NamesToShaderNodes, UInterchangeBaseNodeContainer& NodeContainer) const;
+	void ConnectTimeInputToOutput(MaterialX::NodePtr UpstreamNode, UInterchangeShaderNode* ParentShaderNode, const FString& InputChannelName, TMap<FString, UInterchangeShaderNode*>& NamesToShaderNodes, UInterchangeBaseNodeContainer& NodeContainer) const;
 
 	UInterchangeShaderNode* CreateMaskShaderNode(uint8 RGBA, const FString & NodeName, TMap<FString, UInterchangeShaderNode*>& NamesToShaderNodes, UInterchangeBaseNodeContainer& NodeContainer) const;
 
diff --git a/Engine/Plugins/Interchange/Runtime/Source/Nodes/Public/InterchangeMaterialDefinitions.h b/Engine/Plugins/Interchange/Runtime/Source/Nodes/Public/InterchangeMaterialDefinitions.h
--- a/Engine/Plugins/Interchange/Runtime/Source/Nodes/Public/InterchangeMaterialDefinitions.h
+++ b/Engine/Plugins/Interchange/Runtime/Source/Nodes/Public/InterchangeMaterialDefinitions.h
@@ -132,6 +132,50 @@ namespace Interchange
 						const FName RGBA = TEXT("RGBA"); // Type: linear color
 					}
 				}
+
+				namespace Time
+				{
+					const FName Name = TEXT("Time");
+
+					namespace Attributes
+					{
+						const FName IgnorePause= TEXT("bIgnorePause");
+						const FName OverridePeriod = TEXT("bOverride_Period");
+						const FName Period = TEXT("Period");
+					}
+				}
+
+				namespace TransformPosition
+				{
+					const FName Name = TEXT("TransformPosition");
+
+					namespace Attributes
+					{
+						const FName TransformSourceType = TEXT("TransformSourceType");
+						const FName TransformType = TEXT("TransformType");
+					}
+
+					namespace Inputs
+					{
+						const FName Input = TEXT("Input");
+					}
+				}
+
+				namespace TransformVector
+				{
+					const FName Name = TEXT("Transform");
+
+					namespace Attributes
+					{
+						const FName TransformSourceType = TEXT("TransformSourceType");
+						const FName TransformType = TEXT("TransformType");
+					}
+
+					namespace Inputs
+					{
+						const FName Input = TEXT("Input");
+					}
+				}
 			}
 		}
 
diff --git a/Engine/Plugins/Interchange/Runtime/Source/Pipelines/Private/InterchangeGenericMaterialPipeline.cpp b/Engine/Plugins/Interchange/Runtime/Source/Pipelines/Private/InterchangeGenericMaterialPipeline.cpp
--- a/Engine/Plugins/Interchange/Runtime/Source/Pipelines/Private/InterchangeGenericMaterialPipeline.cpp
+++ b/Engine/Plugins/Interchange/Runtime/Source/Pipelines/Private/InterchangeGenericMaterialPipeline.cpp
@@ -31,6 +31,9 @@
 #include "Materials/MaterialExpressionTextureSampleParameter2D.h"
 #include "Materials/MaterialExpressionTextureSampleParameter2DArray.h"
 #include "Materials/MaterialExpressionTextureSampleParameterCube.h"
+#include "Materials/MaterialExpressionTime.h"
+#include "Materials/MaterialExpressionTransformPosition.h"
+#include "Materials/MaterialExpressionTransform.h"
 #include "Materials/MaterialExpressionVectorParameter.h"
 #include "Materials/MaterialInstanceConstant.h"
 #include "Materials/MaterialInstanceDynamic.h"
@@ -1472,6 +1475,107 @@ void UInterchangeGenericMaterialPipeline::HandleMaskNode(const UInterchangeShade
 	}
 }
 
+void UInterchangeGenericMaterialPipeline::HandleTimeNode(const UInterchangeShaderNode* ShaderNode, UInterchangeBaseMaterialFactoryNode* MaterialFactoryNode, UInterchangeMaterialExpressionFactoryNode* TimeFactoryNode)
+{
+	using namespace UE::Interchange::Materials::Standard::Nodes;
+
+	TimeFactoryNode->SetCustomExpressionClassName(UMaterialExpressionTime::StaticClass()->GetName());
+
+	// IgnorePause
+	if(bool bIgnorePause; ShaderNode->GetBooleanAttribute(Time::Attributes::IgnorePause, bIgnorePause))
+	{
+		const FName IgnorePauseMemberName = GET_MEMBER_NAME_CHECKED(UMaterialExpressionTime, bIgnorePause);
+		TimeFactoryNode->AddBooleanAttribute(IgnorePauseMemberName, bIgnorePause);
+		TimeFactoryNode->AddApplyAndFillDelegates<bool>(IgnorePauseMemberName, UMaterialExpressionTime::StaticClass(), IgnorePauseMemberName);
+	}
+
+	// OverridePeriod
+	if(bool bOverridePeriod; ShaderNode->GetBooleanAttribute(Time::Attributes::OverridePeriod, bOverridePeriod))
+	{
+		const FName OverridePeriodMemberName = GET_MEMBER_NAME_CHECKED(UMaterialExpressionTime, bOverride_Period);
+		TimeFactoryNode->AddBooleanAttribute(OverridePeriodMemberName, bOverridePeriod);
+		TimeFactoryNode->AddApplyAndFillDelegates<bool>(OverridePeriodMemberName, UMaterialExpressionTime::StaticClass(), OverridePeriodMemberName);
+	}
+
+	// Period
+	if(float Period; ShaderNode->GetFloatAttribute(Time::Attributes::Period, Period))
+	{
+		const FName PeriodMemberName = GET_MEMBER_NAME_CHECKED(UMaterialExpressionTime, Period);
+		TimeFactoryNode->AddFloatAttribute(PeriodMemberName, Period);
+		TimeFactoryNode->AddApplyAndFillDelegates<float>(PeriodMemberName, UMaterialExpressionTime::StaticClass(), PeriodMemberName);
+	}
+}
+
+void UInterchangeGenericMaterialPipeline::HandleTransformPositionNode(const UInterchangeShaderNode* ShaderNode, UInterchangeBaseMaterialFactoryNode* MaterialFactoryNode, UInterchangeMaterialExpressionFactoryNode* TransformPositionFactoryNode)
+{
+	using namespace UE::Interchange::Materials::Standard::Nodes;
+
+	TransformPositionFactoryNode->SetCustomExpressionClassName(UMaterialExpressionTransformPosition::StaticClass()->GetName());
+
+	// Input
+	{
+		TTuple<UInterchangeMaterialExpressionFactoryNode*, FString> InputExpression =
+			CreateMaterialExpressionForInput(MaterialFactoryNode, ShaderNode, TransformPosition::Inputs::Input.ToString(), TransformPositionFactoryNode->GetUniqueID());
+
+		if(InputExpression.Get<0>())
+		{
+			UInterchangeShaderPortsAPI::ConnectOuputToInput(TransformPositionFactoryNode, GET_MEMBER_NAME_CHECKED(UMaterialExpressionTransformPosition, Input).ToString(),
+															InputExpression.Get<0>()->GetUniqueID(), InputExpression.Get<1>());
+		}
+	}
+
+	// TransformSourceType
+	if(int32 TransformSourceType; ShaderNode->GetInt32Attribute(TransformPosition::Attributes::TransformSourceType, TransformSourceType))
+	{
+		const FName TransformSourceTypeMemberName = GET_MEMBER_NAME_CHECKED(UMaterialExpressionTransformPosition, TransformSourceType);
+		TransformPositionFactoryNode->AddInt32Attribute(TransformSourceTypeMemberName,TransformSourceType);
+		TransformPositionFactoryNode->AddApplyAndFillDelegates<int32>(TransformSourceTypeMemberName, UMaterialExpressionTransformPosition::StaticClass(), TransformSourceTypeMemberName);
+	}
+
+	// TransformType
+	if(int32 TransformType; ShaderNode->GetInt32Attribute(TransformPosition::Attributes::TransformType, TransformType))
+	{
+		const FName TransformTypeMemberName = GET_MEMBER_NAME_CHECKED(UMaterialExpressionTransformPosition, TransformType);
+		TransformPositionFactoryNode->AddInt32Attribute(TransformTypeMemberName, TransformType);
+		TransformPositionFactoryNode->AddApplyAndFillDelegates<int32>(TransformTypeMemberName, UMaterialExpressionTransformPosition::StaticClass(), TransformTypeMemberName);
+	}
+}
+
+void UInterchangeGenericMaterialPipeline::HandleTransformVectorNode(const UInterchangeShaderNode* ShaderNode, UInterchangeBaseMaterialFactoryNode* MaterialFactoryNode, UInterchangeMaterialExpressionFactoryNode* TransformVectorFactoryNode)
+{
+	using namespace UE::Interchange::Materials::Standard::Nodes;
+
+	TransformVectorFactoryNode->SetCustomExpressionClassName(UMaterialExpressionTransform::StaticClass()->GetName());
+
+	// Input
+	{
+		TTuple<UInterchangeMaterialExpressionFactoryNode*, FString> InputExpression =
+			CreateMaterialExpressionForInput(MaterialFactoryNode, ShaderNode, TransformVector::Inputs::Input.ToString(), TransformVectorFactoryNode->GetUniqueID());
+
+		if(InputExpression.Get<0>())
+		{
+			UInterchangeShaderPortsAPI::ConnectOuputToInput(TransformVectorFactoryNode, GET_MEMBER_NAME_CHECKED(UMaterialExpressionTransform, Input).ToString(),
+															InputExpression.Get<0>()->GetUniqueID(), InputExpression.Get<1>());
+		}
+	}
+
+	// TransformSourceType
+	if(int32 TransformSourceType; ShaderNode->GetInt32Attribute(TransformVector::Attributes::TransformSourceType, TransformSourceType))
+	{
+		const FName TransformSourceTypeMemberName = GET_MEMBER_NAME_CHECKED(UMaterialExpressionTransform, TransformSourceType);
+		TransformVectorFactoryNode->AddInt32Attribute(TransformSourceTypeMemberName, TransformSourceType);
+		TransformVectorFactoryNode->AddApplyAndFillDelegates<int32>(TransformSourceTypeMemberName, UMaterialExpressionTransform::StaticClass(), TransformSourceTypeMemberName);
+	}
+
+	// TransformType
+	if(int32 TransformType; ShaderNode->GetInt32Attribute(TransformVector::Attributes::TransformType, TransformType))
+	{
+		const FName TransformTypeMemberName = GET_MEMBER_NAME_CHECKED(UMaterialExpressionTransform, TransformType);
+		TransformVectorFactoryNode->AddInt32Attribute(TransformTypeMemberName, TransformType);
+		TransformVectorFactoryNode->AddApplyAndFillDelegates<int32>(TransformTypeMemberName, UMaterialExpressionTransform::StaticClass(), TransformTypeMemberName);
+	}
+}
+
 UInterchangeMaterialExpressionFactoryNode* UInterchangeGenericMaterialPipeline::CreateMaterialExpressionForShaderNode(UInterchangeBaseMaterialFactoryNode* MaterialFactoryNode,
 	const UInterchangeShaderNode* ShaderNode, const FString& ParentUid)
 {
@@ -1540,6 +1644,18 @@ UInterchangeMaterialExpressionFactoryNode* UInterchangeGenericMaterialPipeline::
 	{
 		HandleTextureSampleNode(ShaderNode, MaterialFactoryNode, MaterialExpression);
 	}
+	else if(*ShaderType == Nodes::Time::Name)
+	{
+		HandleTimeNode(ShaderNode, MaterialFactoryNode, MaterialExpression);
+	}
+	else if(*ShaderType == Nodes::TransformPosition::Name)
+	{
+		HandleTransformPositionNode(ShaderNode, MaterialFactoryNode, MaterialExpression);
+	}
+	else if(*ShaderType == Nodes::TransformVector::Name)
+	{
+		HandleTransformVectorNode(ShaderNode, MaterialFactoryNode, MaterialExpression);
+	}
 	else
 	{
 		const FString ExpressionClassName = TEXT("MaterialExpression") + ShaderType;
diff --git a/Engine/Plugins/Interchange/Runtime/Source/Pipelines/Public/InterchangeGenericMaterialPipeline.h b/Engine/Plugins/Interchange/Runtime/Source/Pipelines/Public/InterchangeGenericMaterialPipeline.h
--- a/Engine/Plugins/Interchange/Runtime/Source/Pipelines/Public/InterchangeGenericMaterialPipeline.h
+++ b/Engine/Plugins/Interchange/Runtime/Source/Pipelines/Public/InterchangeGenericMaterialPipeline.h
@@ -124,6 +124,10 @@ private:
 	void HandleTextureCoordinateNode(const UInterchangeShaderNode* ShaderNode, UInterchangeBaseMaterialFactoryNode* MaterialFactoryNode, UInterchangeMaterialExpressionFactoryNode*& TextureSampleFactoryNode);
 	void HandleLerpNode(const UInterchangeShaderNode* ShaderNode, UInterchangeBaseMaterialFactoryNode* MaterialFactoryNode, UInterchangeMaterialExpressionFactoryNode* LerpFactoryNode);
 	void HandleMaskNode(const UInterchangeShaderNode* ShaderNode, UInterchangeBaseMaterialFactoryNode* MaterialFactoryNode, UInterchangeMaterialExpressionFactoryNode* MaskFactoryNode);
+	void HandleTimeNode(const UInterchangeShaderNode* ShaderNode, UInterchangeBaseMaterialFactoryNode* MaterialFactoryNode, UInterchangeMaterialExpressionFactoryNode* TimeFactoryNode);
+	void HandleTransformPositionNode(const UInterchangeShaderNode* ShaderNode, UInterchangeBaseMaterialFactoryNode* MaterialFactoryNode, UInterchangeMaterialExpressionFactoryNode* TransformPositionFactoryNode);
+	void HandleTransformVectorNode(const UInterchangeShaderNode* ShaderNode, UInterchangeBaseMaterialFactoryNode* MaterialFactoryNode, UInterchangeMaterialExpressionFactoryNode* TransformVectorFactoryNode);
+
 
 	UInterchangeMaterialExpressionFactoryNode* CreateMaterialExpressionForShaderNode(UInterchangeBaseMaterialFactoryNode* MaterialFactoryNode, const UInterchangeShaderNode* ShaderNode, const FString& ParentUid);
 	TTuple<UInterchangeMaterialExpressionFactoryNode*, FString> CreateMaterialExpressionForInput(UInterchangeBaseMaterialFactoryNode* MaterialFactoryNode, const UInterchangeShaderNode* ShaderNode, const FString& InputName, const FString& ParentUid);
