Fixed a crash when importing an uncompressed 8 bits psd file.

diff --git a/Engine/Plugins/Interchange/Runtime/Source/Import/Private/Texture/InterchangePSDTranslator.cpp b/Engine/Plugins/Interchange/Runtime/Source/Import/Private/Texture/InterchangePSDTranslator.cpp
--- a/Engine/Plugins/Interchange/Runtime/Source/Import/Private/Texture/InterchangePSDTranslator.cpp
+++ b/Engine/Plugins/Interchange/Runtime/Source/Import/Private/Texture/InterchangePSDTranslator.cpp
@@ -183,18 +183,18 @@ namespace UE
 						{
 							if (Info.nChannels == 1)
 							{
-								Dest[Pixel].R = pCur[NPixels + Pixel];
-								Dest[Pixel].G = pCur[NPixels + Pixel];
-								Dest[Pixel].B = pCur[NPixels + Pixel];
+								Dest[Pixel].R = pCur[Pixel];
+								Dest[Pixel].G = pCur[Pixel];
+								Dest[Pixel].B = pCur[Pixel];
 							}
 							else
 							{
-								Dest[Pixel].R = pCur[NPixels * 0 + Pixel];
-								Dest[Pixel].G = pCur[NPixels * 1 + Pixel];
-								Dest[Pixel].B = pCur[NPixels * 2 + Pixel];
+								Dest[Pixel].R = pCur[Pixel];
+								Dest[Pixel].G = pCur[Pixel + 1];
+								Dest[Pixel].B = pCur[Pixel + 2];
 								if (Info.nChannels == 4)
 								{
-									Dest[Pixel].A = pCur[NPixels * 3 + Pixel];
+									Dest[Pixel].A = pCur[Pixel + 3];
 								}
 							}
 						}
diff --git a/Engine/Source/Editor/UnrealEd/Private/Factories/EditorFactories.cpp b/Engine/Source/Editor/UnrealEd/Private/Factories/EditorFactories.cpp
--- a/Engine/Source/Editor/UnrealEd/Private/Factories/EditorFactories.cpp
+++ b/Engine/Source/Editor/UnrealEd/Private/Factories/EditorFactories.cpp
@@ -2535,18 +2535,18 @@ static bool psd_ReadData( uint8* pOut, const uint8*& pBuffer, FPSDFileHeader& In
 			{
 				if (Info.nChannels == 1)
 				{
-					Dest[Pixel].R = pCur[NPixels + Pixel];
-					Dest[Pixel].G = pCur[NPixels + Pixel];
-					Dest[Pixel].B = pCur[NPixels + Pixel];
+					Dest[Pixel].R = pCur[Pixel];
+					Dest[Pixel].G = pCur[Pixel];
+					Dest[Pixel].B = pCur[Pixel];
 				}
 				else
 				{
-					Dest[Pixel].R = pCur[NPixels * 0 + Pixel];
-					Dest[Pixel].G = pCur[NPixels * 1 + Pixel];
-					Dest[Pixel].B = pCur[NPixels * 2 + Pixel];
+					Dest[Pixel].R = pCur[Pixel];
+					Dest[Pixel].G = pCur[Pixel + 1];
+					Dest[Pixel].B = pCur[Pixel + 2];
 					if (Info.nChannels == 4)
 					{
-						Dest[Pixel].A = pCur[NPixels * 3 + Pixel];
+						Dest[Pixel].A = pCur[Pixel + 3];
 					}
 				}
 			}
