Interchange: Fixed OBJ translator crashing when a face provides a normal index but the file doesn't have a normal buffer

diff --git a/Engine/Plugins/Interchange/Runtime/Source/Import/Private/Mesh/InterchangeOBJTranslator.cpp b/Engine/Plugins/Interchange/Runtime/Source/Import/Private/Mesh/InterchangeOBJTranslator.cpp
--- a/Engine/Plugins/Interchange/Runtime/Source/Import/Private/Mesh/InterchangeOBJTranslator.cpp
+++ b/Engine/Plugins/Interchange/Runtime/Source/Import/Private/Mesh/InterchangeOBJTranslator.cpp
@@ -310,12 +310,12 @@ FMeshDescription FObjData::MakeMeshDescriptionForGroup(const FString& GroupName)
 			FVertexInstanceID VertexInstanceID = MeshDescription.CreateVertexInstance(VertexID);
 			VertexInstanceIDs.Add(VertexInstanceID);
 
-			if (VertexData.NormalIndex != INDEX_NONE)
+			if (VertexData.NormalIndex != INDEX_NONE && Normals.IsValidIndex(VertexData.NormalIndex))
 			{
 				Attributes.GetVertexInstanceNormals()[VertexInstanceID] = NormalToUEBasis(Normals[VertexData.NormalIndex]);
 			}
 
-			if (VertexData.UVIndex != INDEX_NONE)
+			if (VertexData.UVIndex != INDEX_NONE && UVs.IsValidIndex(VertexData.UVIndex))
 			{
 				Attributes.GetVertexInstanceUVs()[VertexInstanceID] = UVToUEBasis(UVs[VertexData.UVIndex]);
 			}
