Fixed an issue caused by last submit. It broke the import of uncompressed 8 bits psd that had 3 or 4 channel

robomerge[FN_EngineMerge] -UE5-RES

diff --git a/Engine/Plugins/Interchange/Runtime/Source/Import/Private/Texture/InterchangePSDTranslator.cpp b/Engine/Plugins/Interchange/Runtime/Source/Import/Private/Texture/InterchangePSDTranslator.cpp
--- a/Engine/Plugins/Interchange/Runtime/Source/Import/Private/Texture/InterchangePSDTranslator.cpp
+++ b/Engine/Plugins/Interchange/Runtime/Source/Import/Private/Texture/InterchangePSDTranslator.cpp
@@ -189,12 +189,13 @@ namespace UE
 							}
 							else
 							{
+								// Each channel live in a separate plane
 								Dest[Pixel].R = pCur[Pixel];
-								Dest[Pixel].G = pCur[Pixel + 1];
-								Dest[Pixel].B = pCur[Pixel + 2];
+								Dest[Pixel].G = pCur[NPixels + Pixel];
+								Dest[Pixel].B = pCur[NPixels * 2 + Pixel];
 								if (Info.nChannels == 4)
 								{
-									Dest[Pixel].A = pCur[Pixel + 3];
+									Dest[Pixel].A = pCur[NPixels * 3 + Pixel];
 								}
 							}
 						}
diff --git a/Engine/Source/Editor/UnrealEd/Private/Factories/EditorFactories.cpp b/Engine/Source/Editor/UnrealEd/Private/Factories/EditorFactories.cpp
--- a/Engine/Source/Editor/UnrealEd/Private/Factories/EditorFactories.cpp
+++ b/Engine/Source/Editor/UnrealEd/Private/Factories/EditorFactories.cpp
@@ -2541,12 +2541,13 @@ static bool psd_ReadData( uint8* pOut, const uint8*& pBuffer, FPSDFileHeader& In
 				}
 				else
 				{
+					// Each channel live in a separate plane
 					Dest[Pixel].R = pCur[Pixel];
-					Dest[Pixel].G = pCur[Pixel + 1];
-					Dest[Pixel].B = pCur[Pixel + 2];
+					Dest[Pixel].G = pCur[NPixels + Pixel];
+					Dest[Pixel].B = pCur[NPixels * 2 + Pixel];
 					if (Info.nChannels == 4)
 					{
-						Dest[Pixel].A = pCur[Pixel + 3];
+						Dest[Pixel].A = pCur[NPixels * 3 + Pixel];
 					}
 				}
 			}
