Fixed test failure when loading MaterialX files

diff --git a/Engine/Plugins/Interchange/Runtime/Source/Pipelines/Private/InterchangeGenericMaterialPipeline.cpp b/Engine/Plugins/Interchange/Runtime/Source/Pipelines/Private/InterchangeGenericMaterialPipeline.cpp
--- a/Engine/Plugins/Interchange/Runtime/Source/Pipelines/Private/InterchangeGenericMaterialPipeline.cpp
+++ b/Engine/Plugins/Interchange/Runtime/Source/Pipelines/Private/InterchangeGenericMaterialPipeline.cpp
@@ -136,14 +136,12 @@ void UInterchangeGenericMaterialPipeline::PreDialogCleanup(const FName PipelineS
 
 bool UInterchangeGenericMaterialPipeline::IsSettingsAreValid(TOptional<FText>& OutInvalidReason) const
 {
-	using namespace UE::Interchange;
-
 	if (TexturePipeline && !TexturePipeline->IsSettingsAreValid(OutInvalidReason))
 	{
 		return false;
 	}
 
-	return InterchangeGenericMaterialPipeline::Private::AreRequiredPackagesLoaded() &&  Super::IsSettingsAreValid(OutInvalidReason);
+	return Super::IsSettingsAreValid(OutInvalidReason);
 }
 
 void UInterchangeGenericMaterialPipeline::AdjustSettingsForContext(EInterchangePipelineContext ImportType, TObjectPtr<UObject> ReimportAsset)
@@ -174,6 +172,13 @@ void UInterchangeGenericMaterialPipeline::AdjustSettingsForContext(EInterchangeP
 			HidePropertiesOfCategory(OuterMostPipeline, this, HideCategoryName);
 		}
 	}
+
+	using namespace UE::Interchange;
+
+	if (!InterchangeGenericMaterialPipeline::Private::AreRequiredPackagesLoaded())
+	{
+		UE_LOG(LogInterchangePipeline, Warning, TEXT("UInterchangeGenericMaterialPipeline: Some required packages are missing. Material import might be wrong"));
+	}
 }
 
 void UInterchangeGenericMaterialPipeline::ExecutePreImportPipeline(UInterchangeBaseNodeContainer* InBaseNodeContainer, const TArray<UInterchangeSourceData*>& InSourceDatas)
