Material expression
- Adding a material expression class to create fractals using a 3D Perlin Noise

MaterialX
- Connect fractal3d nodes
- Adding a 0.001 scale to the position node, empirical value to match MaterialX output files especially the standard_surface_marble_solid one
- Adding the support of integer attribute from MaterialX (it's created as a float in UE, since Interchange doesn't create a scale material expression for it)

diff --git a/Engine/Plugins/Interchange/Runtime/Source/Import/Private/MaterialX/InterchangeMaterialXDefinitions.h b/Engine/Plugins/Interchange/Runtime/Source/Import/Private/MaterialX/InterchangeMaterialXDefinitions.h
--- a/Engine/Plugins/Interchange/Runtime/Source/Import/Private/MaterialX/InterchangeMaterialXDefinitions.h
+++ b/Engine/Plugins/Interchange/Runtime/Source/Import/Private/MaterialX/InterchangeMaterialXDefinitions.h
@@ -246,6 +246,11 @@ MATERIALX_NAMESPACE_BEGIN
 		static constexpr const char* RampTB = "ramptb";
 		static constexpr const char* SplitLR = "splitlr";
 		static constexpr const char* SplitTB = "splittb";
+		// Procedural3D nodes 
+		static constexpr const char* CellNoise3D = "cellnoise3d";
+		static constexpr const char* Fractal3D = "fractal3d";
+		static constexpr const char* Noise3D = "noise3d";
+		static constexpr const char* WorleyNoise3D = "worleynoise3d";
 		// Organization nodes 
 		static constexpr const char* Dot = "dot";
 		// Texture nodes 
diff --git a/Engine/Plugins/Interchange/Runtime/Source/Import/Private/MaterialX/InterchangeMaterialXTranslator.cpp b/Engine/Plugins/Interchange/Runtime/Source/Import/Private/MaterialX/InterchangeMaterialXTranslator.cpp
--- a/Engine/Plugins/Interchange/Runtime/Source/Import/Private/MaterialX/InterchangeMaterialXTranslator.cpp
+++ b/Engine/Plugins/Interchange/Runtime/Source/Import/Private/MaterialX/InterchangeMaterialXTranslator.cpp
@@ -36,21 +36,25 @@ UInterchangeMaterialXTranslator::UInterchangeMaterialXTranslator()
 	:
 InputNamesMaterialX2UE
 {
-	{{TEXT(""),                         TEXT("bg")},		TEXT("B")},
-	{{TEXT(""),                         TEXT("center")},	TEXT("Center")},
-	{{TEXT(""),                         TEXT("fg")},		TEXT("A")},
-	{{TEXT(""),                         TEXT("high")},		TEXT("Max")},
-	{{TEXT(""),                         TEXT("in")},		TEXT("Input")},
-	{{TEXT(""),                         TEXT("in1")},		TEXT("A")},
-	{{TEXT(""),                         TEXT("in2")},		TEXT("B")},
-	{{TEXT(""),                         TEXT("in3")},		TEXT("C")},
-	{{TEXT(""),                         TEXT("in4")},		TEXT("D")},
+	{{TEXT(""),                         TEXT("bg")},        TEXT("B")},
+	{{TEXT(""),                         TEXT("center")},    TEXT("Center")},
+	{{TEXT(""),                         TEXT("diminish")},  TEXT("Diminish")},
+	{{TEXT(""),                         TEXT("fg")},        TEXT("A")},
+	{{TEXT(""),                         TEXT("high")},      TEXT("Max")},
+	{{TEXT(""),                         TEXT("in")},        TEXT("Input")},
+	{{TEXT(""),                         TEXT("in1")},       TEXT("A")},
+	{{TEXT(""),                         TEXT("in2")},       TEXT("B")},
+	{{TEXT(""),                         TEXT("in3")},       TEXT("C")},
+	{{TEXT(""),                         TEXT("in4")},       TEXT("D")},
 	{{TEXT(""),                         TEXT("inlow")},     TEXT("InputLow")},
 	{{TEXT(""),                         TEXT("inhigh")},    TEXT("InputHigh")},
-	{{TEXT(""),                         TEXT("low")},		TEXT("Min")},
+	{{TEXT(""),                         TEXT("lacunarity")},TEXT("Lacunarity")},
+	{{TEXT(""),                         TEXT("low")},       TEXT("Min")},
 	{{TEXT(""),                         TEXT("lumacoeffs")},TEXT("LuminanceFactors")}, // for the moment not yet handled by Interchange, because of the attribute being an advanced pin
-	{{TEXT(""),                         TEXT("mix")},		TEXT("Alpha")},
+	{{TEXT(""),                         TEXT("mix")},       TEXT("Alpha")},
+	{{TEXT(""),                         TEXT("position")},  TEXT("Position")},
 	{{TEXT(""),                         TEXT("texcoord")},  TEXT("Coordinates")},
+	{{TEXT(""),                         TEXT("octaves")},   TEXT("Octaves")},
 	{{TEXT(""),                         TEXT("outlow")},    TEXT("TargetLow")},
 	{{TEXT(""),                         TEXT("outhigh")},   TEXT("TargetHigh")},
 	{{TEXT(""),                         TEXT("valuel")},    TEXT("A")},
@@ -145,6 +149,8 @@ NodeNamesMaterialX2UE{
 	{MaterialX::Category::Combine2,     TEXT("AppendVector")},
 	{MaterialX::Category::Combine3,     TEXT("Append3Vector")},
 	{MaterialX::Category::Combine4,     TEXT("Append4Vector")},
+	// Procedural3D nodes
+	{MaterialX::Category::Fractal3D,    TEXT("Fractal3D")}, 
 	// Geometric nodes 
 	{MaterialX::Category::GeomColor,	TEXT("VertexColor")},
 	{MaterialX::Category::TexCoord,		TEXT("TextureCoordinate")},
@@ -160,6 +166,7 @@ UEInputs
 {
     TEXT("A"),
 	TEXT("Alpha"),
+	TEXT("Amplitude"),
 	TEXT("B"),
     TEXT("Base"),
     TEXT("C"),
@@ -171,9 +178,12 @@ UEInputs
     TEXT("Input"),
 	TEXT("InputLow"),
 	TEXT("InputHigh"),
+	TEXT("Lacunarity"),
 	TEXT("LuminanceFactors"),
     TEXT("Max"),
     TEXT("Min"),
+	TEXT("Octaves"),
+	TEXT("Position"),
 	TEXT("TargetLow"),
 	TEXT("TargetHigh"),
 	TEXT("Value"),
@@ -1212,8 +1222,11 @@ bool UInterchangeMaterialXTranslator::AddAttribute(MaterialX::InputPtr Input, co
 		{
 			return ShaderNode->AddFloatAttribute(UInterchangeShaderPortsAPI::MakeInputValueKey(InputChannelName), mx::fromValueString<float>(Input->getValueString()));
 		}
-
-		if(Input->getType() == mx::Type::Color3 || Input->getType() == mx::Type::Color4)
+		else if(Input->getType() == mx::Type::Integer) //Let's add it a Float attribute, because Interchange doesn't create a scalar if it's an int
+		{
+			return ShaderNode->AddFloatAttribute(UInterchangeShaderPortsAPI::MakeInputValueKey(InputChannelName), mx::fromValueString<int32>(Input->getValueString()));
+		}
+		else if(Input->getType() == mx::Type::Color3 || Input->getType() == mx::Type::Color4)
 		{
 			FLinearColor LinearColor = GetLinearColor(Input);
 			return ShaderNode->AddLinearColorAttribute(UInterchangeShaderPortsAPI::MakeInputValueKey(InputChannelName), LinearColor);
@@ -1770,9 +1783,16 @@ void UInterchangeMaterialXTranslator::ConnectPositionInputToOutput(MaterialX::No
 	// object: The local coordinate space of the geometry, after local deformations have been applied, but before any global transforms.
 	// world : The global coordinate space of the geometry, after local deformationsand global transforms have been applied.
 
-	// For the moment we don't have the distinction between model/object, so let's just create an UMaterialExpressionObjectPositionWS
+	// For the moment we don't have the distinction between model/object, so let's just create an UMaterialExpressionWorldPosition
 	// In case of model/object we need to add a TransformPoint from world to local space
-	UInterchangeShaderNode* PositionNode = CreateShaderNode<UInterchangeShaderNode>(UpstreamNode->getName().c_str(), TEXT("ObjectPositionWS"), NamesToShaderNodes, NodeContainer);
+	UInterchangeShaderNode* PositionNode = CreateShaderNode<UInterchangeShaderNode>(UpstreamNode->getName().c_str() + FString(TEXT("_Position")), TEXT("WorldPosition"), NamesToShaderNodes, NodeContainer);
+
+	// In case of the position node, it seems that the unit is different, we assume for now a conversion from mm -> m, even if UE by default is cm
+	// See standard_surface_marble_solid file, especially on the fractal3d node
+	UInterchangeShaderNode* UnitNode = CreateShaderNode<UInterchangeShaderNode>(UpstreamNode->getName().c_str(), TEXT("Multiply"), NamesToShaderNodes, NodeContainer);
+	UnitNode->AddFloatAttribute(UInterchangeShaderPortsAPI::MakeInputValueKey(TEXT("B")), 0.001f);
+	UInterchangeShaderPortsAPI::ConnectDefaultOuputToInput(UnitNode, TEXT("A"), PositionNode->GetUniqueID());
+
 	UInterchangeShaderNode* NodeToConnectTo = ParentShaderNode;
 	FString InputToConnectTo = InputChannelName;
 
@@ -1797,7 +1817,7 @@ void UInterchangeMaterialXTranslator::ConnectPositionInputToOutput(MaterialX::No
 		UInterchangeShaderPortsAPI::ConnectDefaultOuputToInput(ParentShaderNode, InputChannelName, TransformNode->GetUniqueID());
 	}
 
-	UInterchangeShaderPortsAPI::ConnectDefaultOuputToInput(NodeToConnectTo, InputToConnectTo, PositionNode->GetUniqueID());
+	UInterchangeShaderPortsAPI::ConnectDefaultOuputToInput(NodeToConnectTo, InputToConnectTo, UnitNode->GetUniqueID());
 }
 
 void UInterchangeMaterialXTranslator::ConnectNormalInputToOutput(MaterialX::NodePtr UpstreamNode, UInterchangeShaderNode* ParentShaderNode, const FString& InputChannelName, TMap<FString, UInterchangeShaderNode*>& NamesToShaderNodes, UInterchangeBaseNodeContainer& NodeContainer) const
diff --git a/Engine/Plugins/Interchange/Runtime/Source/Import/Private/MaterialX/MaterialExpressionFractal3D.cpp b/Engine/Plugins/Interchange/Runtime/Source/Import/Private/MaterialX/MaterialExpressionFractal3D.cpp
new file mode 100644
--- /dev/null
+++ b/Engine/Plugins/Interchange/Runtime/Source/Import/Private/MaterialX/MaterialExpressionFractal3D.cpp
@@ -0,0 +1,97 @@
+// Copyright Epic Games, Inc. All Rights Reserved.
+
+#include "MaterialX/MaterialExpressionFractal3D.h"
+#include "MaterialCompiler.h"
+#include "Materials/MaterialExpressionCustom.h"
+
+#define LOCTEXT_NAMESPACE "MaterialExpressionFractal3D"
+
+UMaterialExpressionFractal3D::UMaterialExpressionFractal3D(const FObjectInitializer& ObjectInitializer)
+	: Super(ObjectInitializer)
+{
+	// Structure to hold one-time initialization
+	struct FConstructorStatics
+	{
+		FText NAME_MaterialX;
+		FText NAME_Procedural;
+		FConstructorStatics() :
+			NAME_MaterialX(LOCTEXT("MaterialX", "MaterialX")),
+			NAME_Procedural(LOCTEXT("Procedural", "Procedural"))
+		{}
+	};
+	static FConstructorStatics ConstructorStatics;
+
+#if WITH_EDITORONLY_DATA
+	MenuCategories.Add(ConstructorStatics.NAME_MaterialX);
+	MenuCategories.Add(ConstructorStatics.NAME_Procedural);
+#endif
+}
+
+#if WITH_EDITOR
+int32 UMaterialExpressionFractal3D::Compile(FMaterialCompiler* Compiler, int32 OutputIndex)
+{
+
+	UMaterialExpressionCustom* MaterialExpressionCustom = NewObject<UMaterialExpressionCustom>();
+	MaterialExpressionCustom->Inputs[0].InputName = TEXT("Position");
+	MaterialExpressionCustom->Inputs.Add({ TEXT("Octaves") });
+	MaterialExpressionCustom->Inputs.Add({ TEXT("Lacunarity") });
+	MaterialExpressionCustom->Inputs.Add({ TEXT("Diminish") });
+	
+	MaterialExpressionCustom->OutputType = ECustomMaterialOutputType::CMOT_Float3;
+	MaterialExpressionCustom->Code = 
+		TEXT("float Scale =") + FString::SanitizeFloat(Scale) + TEXT(";") +
+		TEXT(R"(int Quality = 1;
+             int Function = 2; // Fast Gradient
+             bool bTurbulence =)") + FString::FromInt(bTurbulence) + TEXT(";") +
+		TEXT("uint Levels =") + FString::FromInt(Levels) + TEXT(";") +
+		TEXT("float OutputMin =") + FString::SanitizeFloat(OutputMin) + TEXT(";") +
+		TEXT("float OutputMax =") + FString::SanitizeFloat(OutputMax) + TEXT(";") +
+		TEXT(R"(float LevelScale = 2;
+             float FilterWidth = 0;
+             bool bTiling = false;
+             float RepeatSize = 512;
+             
+             float3 result = float3(0.0,0.0,0.0);
+             float amplitude = 1.0;
+             for (int i = 0;  i < Octaves; ++i)
+             {
+                 result += amplitude * 
+                 MaterialExpressionNoise(
+                     Position,
+                     Scale,
+                     Quality,
+                     Function,
+                     bTurbulence,
+                     Levels,
+                     OutputMin,
+                     OutputMax,
+                     LevelScale,
+                     FilterWidth,
+                     bTiling,
+                     RepeatSize);
+             
+                 amplitude *= Diminish;
+                 Position *= Lacunarity;
+             }
+             return result;)");
+
+	int32 IndexPosition = Position.GetTracedInput().Expression ? Position.Compile(Compiler) : Compiler->TransformPosition(EMaterialCommonBasis::MCB_World, EMaterialCommonBasis::MCB_Local, Compiler->WorldPosition(WPT_Default));
+	int32 IndexOctaves = Octaves.GetTracedInput().Expression ? Octaves.Compile(Compiler) : Compiler->Constant(ConstOctaves);
+	int32 IndexLacunarity = Lacunarity.GetTracedInput().Expression ? Lacunarity.Compile(Compiler) : Compiler->Constant(ConstLacunarity);
+	int32 IndexDiminish = Diminish.GetTracedInput().Expression ? Diminish.Compile(Compiler) : Compiler->Constant(ConstDiminish);
+
+	TArray<int32> Inputs{ IndexPosition, IndexOctaves, IndexLacunarity, IndexDiminish };
+	int32 IndexFractal = Compiler->CustomExpression(MaterialExpressionCustom, 0, Inputs);
+
+	int32 IndexAmplitude = Amplitude.GetTracedInput().Expression ? Amplitude.Compile(Compiler) : Compiler->Constant(ConstAmplitude);
+
+	return Compiler->Mul(IndexFractal, IndexAmplitude);
+}
+
+void UMaterialExpressionFractal3D::GetCaption(TArray<FString>& OutCaptions) const
+{
+	OutCaptions.Add(TEXT("Fractal3D"));
+}
+#endif
+
+#undef LOCTEXT_NAMESPACE 
\ No newline at end of file
diff --git a/Engine/Plugins/Interchange/Runtime/Source/Import/Private/MaterialX/MaterialExpressionFractal3D.h b/Engine/Plugins/Interchange/Runtime/Source/Import/Private/MaterialX/MaterialExpressionFractal3D.h
new file mode 100644
--- /dev/null
+++ b/Engine/Plugins/Interchange/Runtime/Source/Import/Private/MaterialX/MaterialExpressionFractal3D.h
@@ -0,0 +1,79 @@
+// Copyright Epic Games, Inc. All Rights Reserved.
+#pragma once
+
+#include "Materials/MaterialExpression.h"
+
+#include "MaterialExpressionFractal3D.generated.h"
+
+/**
+ * Zero-centered 3D Fractal noise in 1, 2, 3 or 4 channels, created by summing several
+ * octaves of 3D Perlin noise, increasing the frequency and decreasing the amplitude at each octave.
+ */
+
+UCLASS(collapsecategories, hidecategories = Object, MinimalAPI)
+class UMaterialExpressionFractal3D : public UMaterialExpression
+{
+	GENERATED_UCLASS_BODY()
+
+	/** The name of a vector3-type node specifying the 3D position at which the noise is evaluated. By default the vector is in local space*/
+	UPROPERTY()
+	FExpressionInput Position;	
+
+	/** Center-to-peak amplitude of the noise (peak-to-peak amplitude is 2x this value).*/
+	UPROPERTY(meta = (RequiredInput = "false", ToolTip = "Defaults to 'ConstAmplitude' if not specified"))
+	FExpressionInput Amplitude;
+
+	/** only used if Amplitude is not hooked up */
+	UPROPERTY(EditAnywhere, Category = MaterialExpressionFractal, meta = (OverridingInputProperty = "Amplitude"))
+	float ConstAmplitude = 1;
+
+	/** The number of octaves of noise to be summed. */
+	UPROPERTY(meta = (RequiredInput = "false", ToolTip = "Defaults to 'ConstOctaves' if not specified"))
+	FExpressionInput Octaves;
+
+	/** only used if Octaves is not hooked up */
+	UPROPERTY(EditAnywhere, Category = MaterialExpressionFractal, meta = (OverridingInputProperty = "Octaves"))
+	int32 ConstOctaves = 3;
+
+	/** The exponential scale between successive octaves of noise. */
+	UPROPERTY(meta = (RequiredInput = "false", ToolTip = "Defaults to 'ConstLacunarity' if not specified"))
+	FExpressionInput Lacunarity;
+
+	/** only used if Lacunarity is not hooked up */
+	UPROPERTY(EditAnywhere, Category = MaterialExpressionFractal, meta = (OverridingInputProperty = "Lacunarity"))
+	float ConstLacunarity = 2;
+
+	/** The rate at which noise amplitude is diminished for each octave. Should be between 0.0 and 1.0 */
+	UPROPERTY(meta = (RequiredInput = "false", ToolTip = "Defaults to 'ConstDiminish' if not specified"))
+	FExpressionInput Diminish;
+
+	/** only used if Diminish is not hooked up */
+	UPROPERTY(EditAnywhere, Category = MaterialExpressionFractal, meta = (OverridingInputProperty = "Diminish"))
+	float ConstDiminish = 0.5;
+
+	/** can also be done with a multiply on the Position */
+	UPROPERTY(EditAnywhere, Category = MaterialExpressionNoise)
+	float Scale = 1.f;
+
+	/** How multiple frequencies are getting combined */
+	UPROPERTY(EditAnywhere, Category = MaterialExpressionNoise)
+	bool bTurbulence = false;
+
+	/** 1 = fast but little detail, .. larger numbers cost more performance */
+	UPROPERTY(EditAnywhere, Category = MaterialExpressionNoise, meta = (UIMin = "1", UIMax = "10"))
+	int32 Levels = 6;
+
+	UPROPERTY(EditAnywhere, Category = MaterialExpressionNoise)
+	float OutputMin = 0.f;
+
+	UPROPERTY(EditAnywhere, Category = MaterialExpressionNoise)
+	float OutputMax = 1.f;
+
+	//~ Begin UMaterialExpression Interface
+#if WITH_EDITOR
+	virtual int32 Compile(class FMaterialCompiler* Compiler, int32 OutputIndex) override;
+	virtual void GetCaption(TArray<FString>& OutCaptions) const override;
+#endif
+	//~ End UMaterialExpression Interface
+};
+
