Added the missing include and ordered the includes

diff --git a/Engine/Plugins/Interchange/Runtime/Source/Import/Private/Texture/InterchangeTextureFactory.cpp b/Engine/Plugins/Interchange/Runtime/Source/Import/Private/Texture/InterchangeTextureFactory.cpp
--- a/Engine/Plugins/Interchange/Runtime/Source/Import/Private/Texture/InterchangeTextureFactory.cpp
+++ b/Engine/Plugins/Interchange/Runtime/Source/Import/Private/Texture/InterchangeTextureFactory.cpp
@@ -22,17 +22,17 @@
 #include "InterchangeTexture2DArrayNode.h"
 #include "InterchangeTexture2DFactoryNode.h"
 #include "InterchangeTexture2DNode.h"
-#include "InterchangeTextureCubeFactoryNode.h"
-#include "InterchangeTextureCubeNode.h"
 #include "InterchangeTextureCubeArrayFactoryNode.h"
 #include "InterchangeTextureCubeArrayNode.h"
+#include "InterchangeTextureCubeFactoryNode.h"
+#include "InterchangeTextureCubeNode.h"
 #include "InterchangeTextureFactoryNode.h"
 #include "InterchangeTextureLightProfileFactoryNode.h"
 #include "InterchangeTextureLightProfileNode.h"
 #include "InterchangeTextureNode.h"
 #include "InterchangeTranslatorBase.h"
-#include "InterchangeVolumeTextureNode.h"
 #include "InterchangeVolumeTextureFactoryNode.h"
+#include "InterchangeVolumeTextureNode.h"
 #include "Misc/CoreStats.h"
 #include "Nodes/InterchangeBaseNode.h"
 #include "ProfilingDebugging/CpuProfilerTrace.h"
@@ -43,6 +43,7 @@
 #include "Texture/InterchangeTexturePayloadInterface.h"
 #include "TextureCompiler.h"
 #include "TextureImportSettings.h"
+#include "TextureResource.h"
 #include "UDIMUtilities.h"
 #include "UObject/ObjectMacros.h"
 
