* InterchangeActorHelper - Added include to solve compile error
* GeometryCollectionRenderLevelSetActor - Removed using namespace Chaos in global namespace

Attempt at fixing Mac build errors

diff --git a/Engine/Plugins/Interchange/Runtime/Source/Import/Private/Scene/InterchangeActorHelper.cpp b/Engine/Plugins/Interchange/Runtime/Source/Import/Private/Scene/InterchangeActorHelper.cpp
--- a/Engine/Plugins/Interchange/Runtime/Source/Import/Private/Scene/InterchangeActorHelper.cpp
+++ b/Engine/Plugins/Interchange/Runtime/Source/Import/Private/Scene/InterchangeActorHelper.cpp
@@ -2,6 +2,7 @@
 
 #include "Scene/InterchangeActorHelper.h"
 
+#include "Engine/Engine.h"
 #include "InterchangeActorFactoryNode.h"
 #include "InterchangeMaterialFactoryNode.h"
 #include "InterchangeMeshActorFactoryNode.h"
