[Interchange]
* InterchangeActorHelper - Fixed non-unity/pch compile error

diff --git a/Engine/Plugins/Interchange/Runtime/Source/Import/Private/Scene/InterchangeActorHelper.cpp b/Engine/Plugins/Interchange/Runtime/Source/Import/Private/Scene/InterchangeActorHelper.cpp
--- a/Engine/Plugins/Interchange/Runtime/Source/Import/Private/Scene/InterchangeActorHelper.cpp
+++ b/Engine/Plugins/Interchange/Runtime/Source/Import/Private/Scene/InterchangeActorHelper.cpp
@@ -3,6 +3,7 @@
 #include "Scene/InterchangeActorHelper.h"
 
 #include "Engine/Engine.h"
+#include "Engine/Level.h"
 #include "InterchangeActorFactoryNode.h"
 #include "InterchangeMaterialFactoryNode.h"
 #include "InterchangeMeshActorFactoryNode.h"
