Interchange: Obj translator now verify the material file extension and make sure its a .mtl file.

#rnx

diff --git a/Engine/Plugins/Interchange/Runtime/Source/Import/Private/Mesh/InterchangeOBJTranslator.cpp b/Engine/Plugins/Interchange/Runtime/Source/Import/Private/Mesh/InterchangeOBJTranslator.cpp
--- a/Engine/Plugins/Interchange/Runtime/Source/Import/Private/Mesh/InterchangeOBJTranslator.cpp
+++ b/Engine/Plugins/Interchange/Runtime/Source/Import/Private/Mesh/InterchangeOBJTranslator.cpp
@@ -643,6 +643,22 @@ namespace ObjParser
 		return true;
 	}
 
+	static bool ParseObjectName(FObjData& ObjData, FStringView Line)
+	{
+		FStringView ObjectName = GetToken(Line);
+		if (ObjectName.IsEmpty())
+		{
+			UE_LOG(LogTemp, Warning, TEXT("Missing object name on o keyword"));
+		}
+
+		if (!Line.IsEmpty())
+		{
+			UE_LOG(LogTemp, Warning, TEXT("Unexpected extra arguments on o keyword"));
+		}
+
+		//We do not do anything yet with the object name.
+		return true;
+	}
 
 	static bool ParseGroup(FObjData& ObjData, FStringView Line)
 	{
@@ -661,6 +677,23 @@ namespace ObjParser
 		return true;
 	}
 
+	static bool ParseSmoothing(FObjData& ObjData, FStringView Line)
+	{
+		FString SmoothingValue = FString(GetToken(Line));
+		if (SmoothingValue.IsEmpty())
+		{
+			UE_LOG(LogTemp, Warning, TEXT("Missing smoothing value on s keyword"));
+		}
+
+		if (!Line.IsEmpty())
+		{
+			UE_LOG(LogTemp, Warning, TEXT("Unexpected extra arguments on s keyword"));
+		}
+
+		//We do not do anything yet with smoothing value
+		return true;
+	}
+	
 
 	static bool ParseLine(FObjData& ObjData, FStringView Line, const FKeywordMap& KeywordMap)
 	{
@@ -806,7 +839,18 @@ namespace ObjParser
 		// Parse the file referenced by the .obj
 		// This will set the current material name
 		FString FileToParse = FPaths::GetPath(ObjData.ObjFilename) / FString(MtlFilename);
-		bool bSuccess = ParseFile(ObjData, *FileToParse, MtlKeywordMap);
+
+		bool bSuccess = true;
+
+		if (FPaths::GetExtension(FileToParse).Equals(TEXT("mtl"), ESearchCase::IgnoreCase))
+		{
+			bSuccess &= ParseFile(ObjData, *FileToParse, MtlKeywordMap);
+		}
+		else
+		{
+			//We support only .mtl material file
+			UE_LOG(LogTemp, Warning, TEXT("Unsupported material file: "), *FileToParse);
+		}
 
 		// After parsing, clear the current material name
 		ObjData.MaterialBeingDefined.Empty();
@@ -909,6 +953,7 @@ bool UInterchangeOBJTranslator::Translate(UInterchangeBaseNodeContainer& BaseNod
 
 	static ObjParser::FKeywordMap ObjKeywordMap =
 	{
+		{ TEXT("o"),      ObjParser::ParseObjectName },
 		{ TEXT("v"),      ObjParser::ParseVertexPosition },
 		{ TEXT("vt"),     ObjParser::ParseTextureCoordinate },
 		{ TEXT("vn"),     ObjParser::ParseNormalVector },
@@ -916,7 +961,7 @@ bool UInterchangeOBJTranslator::Translate(UInterchangeBaseNodeContainer& BaseNod
 		{ TEXT("g"),      ObjParser::ParseGroup },
 		{ TEXT("mtllib"), ObjParser::ParseMaterialLib },
 		{ TEXT("usemtl"), ObjParser::ParseUseMaterial },
-//		{ TEXT("s"),      ObjParser::ParseSmoothing },
+		{ TEXT("s"),      ObjParser::ParseSmoothing },
 	};
 
 	bool bSuccess = ObjParser::ParseFile(*ObjDataPtr.Get(), *Filename, ObjKeywordMap);
