Interchange Plugin: Change word from error to fault in log so automated tests do not count it as an error
#rnx

diff --git a/Engine/Plugins/Interchange/Runtime/Source/Dispatcher/Private/InterchangeWorkerHandler.cpp b/Engine/Plugins/Interchange/Runtime/Source/Dispatcher/Private/InterchangeWorkerHandler.cpp
--- a/Engine/Plugins/Interchange/Runtime/Source/Dispatcher/Private/InterchangeWorkerHandler.cpp
+++ b/Engine/Plugins/Interchange/Runtime/Source/Dispatcher/Private/InterchangeWorkerHandler.cpp
@@ -213,7 +213,7 @@ namespace UE
 			WorkerState = EWorkerState::Uninitialized;
 			RunInternal();
 			WorkerState = EWorkerState::Terminated;
-			UE_CLOG(ErrorState != EWorkerErrorState::Ok, LogInterchangeDispatcher, Display, TEXT("Handler ended with error: %s"), EWorkerErrorStateAsString(ErrorState));
+			UE_CLOG(ErrorState != EWorkerErrorState::Ok, LogInterchangeDispatcher, Display, TEXT("Handler ended with fault: %s"), EWorkerErrorStateAsString(ErrorState));
 		}
 
 		void FInterchangeWorkerHandler::KillAllCurrentTasks()
