Interchange: fix obj translator message to not spam the log and also downgrade unknown Keyword log to a Display message since there is a couple of obj and mtl property we do not parse.

#rnx

diff --git a/Engine/Plugins/Interchange/Runtime/Source/Import/Private/Mesh/InterchangeOBJTranslator.cpp b/Engine/Plugins/Interchange/Runtime/Source/Import/Private/Mesh/InterchangeOBJTranslator.cpp
--- a/Engine/Plugins/Interchange/Runtime/Source/Import/Private/Mesh/InterchangeOBJTranslator.cpp
+++ b/Engine/Plugins/Interchange/Runtime/Source/Import/Private/Mesh/InterchangeOBJTranslator.cpp
@@ -33,6 +33,13 @@ static FAutoConsoleVariableRef CCvarInterchangeEnableOBJImport(
 	TEXT("Whether OBJ support is enabled."),
 	ECVF_Default);
 
+
+#define INTERCHANGE_OBJ_TRANSLATOR_LOG_ONCE(Verbosity, Format, ...) \
+{ \
+	static bool bLogged = false; \
+	UE_CLOG(!bLogged, Verbosity, Format, ##__VA_ARGS__); \
+	bLogged = true; \
+}
 /**
  * This is a binary representation of the .obj file 
  */
@@ -398,7 +405,7 @@ FMeshDescription FObjData::MakeMeshDescriptionForGroup(const FString& GroupName)
 			}
 			else
 			{
-				ensureMsgf(false, TEXT("Tried to apply UV data that did not match the MeshDescription."));
+				ensureMsgf(false, TEXT("Interchange Obj translator: Tried to apply UV data that did not match the MeshDescription."));
 			}
 		}
 	}
@@ -534,7 +541,7 @@ namespace ObjParser
 			// @todo: emit warning through Interchange results container
 			// @todo: get new API working through InterchangeManager (keyed multiple times on translator, pipeline, factory objects)
 			// We have a pointer to the Translator inside the FObjData, so we can get the results container from there.
-			UE_LOG(LogTemp, Warning, TEXT("Unexpected extra arguments on v keyword"));
+			INTERCHANGE_OBJ_TRANSLATOR_LOG_ONCE(LogTemp, Warning, TEXT("Interchange Obj translator: Unexpected extra arguments on v keyword"));
 		}
 
 		return bSuccess;
@@ -550,7 +557,7 @@ namespace ObjParser
 
 		if (!Line.IsEmpty())
 		{
-			UE_LOG(LogTemp, Warning, TEXT("Unexpected extra arguments on vt keyword"));
+			INTERCHANGE_OBJ_TRANSLATOR_LOG_ONCE(LogTemp, Warning, TEXT("Interchange Obj translator: Unexpected extra arguments on vt keyword"));
 		}
 
 		return bSuccess;
@@ -566,7 +573,7 @@ namespace ObjParser
 
 		if (!Line.IsEmpty())
 		{
-			UE_LOG(LogTemp, Warning, TEXT("Unexpected extra arguments on vn keyword"));
+			INTERCHANGE_OBJ_TRANSLATOR_LOG_ONCE(LogTemp, Warning, TEXT("Interchange Obj translator: Unexpected extra arguments on vn keyword"));
 		}
 
 		return bSuccess;
@@ -616,7 +623,7 @@ namespace ObjParser
 		}
 		else
 		{
-			UE_LOG(LogTemp, Warning, TEXT("Too few vertices on f keyword"));
+			INTERCHANGE_OBJ_TRANSLATOR_LOG_ONCE(LogTemp, Warning, TEXT("Interchange Obj translator: Too few vertices on f keyword"));
 			return false;
 		}
 
@@ -630,12 +637,12 @@ namespace ObjParser
 
 		if (MaterialName.IsEmpty())
 		{
-			UE_LOG(LogTemp, Warning, TEXT("Missing material name on usemtl keyword"));
+			INTERCHANGE_OBJ_TRANSLATOR_LOG_ONCE(LogTemp, Warning, TEXT("Interchange Obj translator: Missing material name on usemtl keyword"));
 		}
 
 		if (!Line.IsEmpty())
 		{
-			UE_LOG(LogTemp, Warning, TEXT("Unexpected extra arguments on newmtl keyword"));
+			INTERCHANGE_OBJ_TRANSLATOR_LOG_ONCE(LogTemp, Warning, TEXT("Interchange Obj translator: Unexpected extra arguments on newmtl keyword"));
 		}
 
 		ObjData.CurrentMaterial = FString(MaterialName);
@@ -648,12 +655,12 @@ namespace ObjParser
 		FStringView ObjectName = GetToken(Line);
 		if (ObjectName.IsEmpty())
 		{
-			UE_LOG(LogTemp, Warning, TEXT("Missing object name on o keyword"));
+			INTERCHANGE_OBJ_TRANSLATOR_LOG_ONCE(LogTemp, Warning, TEXT("Interchange Obj translator: Missing object name on o keyword"));
 		}
 
 		if (!Line.IsEmpty())
 		{
-			UE_LOG(LogTemp, Warning, TEXT("Unexpected extra arguments on o keyword"));
+			INTERCHANGE_OBJ_TRANSLATOR_LOG_ONCE(LogTemp, Warning, TEXT("Interchange Obj translator: Unexpected extra arguments on o keyword"));
 		}
 
 		//We do not do anything yet with the object name.
@@ -665,40 +672,24 @@ namespace ObjParser
 		FStringView GroupName = GetToken(Line);
 		if (GroupName.IsEmpty())
 		{
-			UE_LOG(LogTemp, Warning, TEXT("Missing group name on g keyword"));
+			INTERCHANGE_OBJ_TRANSLATOR_LOG_ONCE(LogTemp, Warning, TEXT("Interchange Obj translator: Missing group name on g keyword"));
 		}
 
 		if (!Line.IsEmpty())
 		{
-			UE_LOG(LogTemp, Warning, TEXT("Unexpected extra arguments on g keyword"));
+			INTERCHANGE_OBJ_TRANSLATOR_LOG_ONCE(LogTemp, Warning, TEXT("Interchange Obj translator: Unexpected extra arguments on g keyword"));
 		}
 
 		ObjData.CurrentGroup = GroupName;
 		return true;
 	}
 
-	static bool ParseSmoothing(FObjData& ObjData, FStringView Line)
-	{
-		FString SmoothingValue = FString(GetToken(Line));
-		if (SmoothingValue.IsEmpty())
-		{
-			UE_LOG(LogTemp, Warning, TEXT("Missing smoothing value on s keyword"));
-		}
-
-		if (!Line.IsEmpty())
-		{
-			UE_LOG(LogTemp, Warning, TEXT("Unexpected extra arguments on s keyword"));
-		}
-
-		//We do not do anything yet with smoothing value
-		return true;
-	}
-	
-
 	static bool ParseLine(FObjData& ObjData, FStringView Line, const FKeywordMap& KeywordMap)
 	{
-		// Dispatch by keyword to appropriate parsing function
+		//Issue unknown keyword only once per keyword
+		static TArray<FString> UnknownKeywords;
 
+		// Dispatch by keyword to appropriate parsing function
 		FString Keyword = FString(GetToken(Line));
 		if (Keyword.IsEmpty())
 		{
@@ -709,12 +700,14 @@ namespace ObjParser
 		{
 			return (*Handler)(ObjData, Line);
 		}
-		else
+		else if (!UnknownKeywords.Contains(Keyword))
 		{
-			// Unknown keyword, report it and move on
-			UE_LOG(LogTemp, Warning, TEXT("Unknown keyword: %s"), *Keyword);
-			return true;
+			// Unknown keyword, report it once
+			UE_LOG(LogTemp, Display, TEXT("Interchange Obj translator: Unknown keyword: %s"), *Keyword);
+			UnknownKeywords.Add(Keyword);
 		}
+
+		return true;
 	}
 
 
@@ -744,12 +737,12 @@ namespace ObjParser
 
 		if (MaterialName.IsEmpty())
 		{
-			UE_LOG(LogTemp, Warning, TEXT("Missing material name on newmtl keyword"));
+			INTERCHANGE_OBJ_TRANSLATOR_LOG_ONCE(LogTemp, Warning, TEXT("Interchange Obj translator: Missing material name on newmtl keyword"));
 		}
 
 		if (!Line.IsEmpty())
 		{
-			UE_LOG(LogTemp, Warning, TEXT("Unexpected extra arguments on newmtl keyword"));
+			INTERCHANGE_OBJ_TRANSLATOR_LOG_ONCE(LogTemp, Warning, TEXT("Interchange Obj translator: Unexpected extra arguments on newmtl keyword"));
 		}
 
 		ObjData.Materials.Add(FString(MaterialName));
@@ -775,14 +768,14 @@ namespace ObjParser
 			}
 			else
 			{
-				UE_LOG(LogTemp, Warning, TEXT("Missing newmtl keyword"));
+				INTERCHANGE_OBJ_TRANSLATOR_LOG_ONCE(LogTemp, Warning, TEXT("Interchange Obj translator: Missing newmtl keyword"));
 				return false;
 			}
 		}
 
 		if (!Line.IsEmpty())
 		{
-			UE_LOG(LogTemp, Warning, TEXT("Unexpected extra arguments on illum keyword"));
+			INTERCHANGE_OBJ_TRANSLATOR_LOG_ONCE(LogTemp, Warning, TEXT("Interchange Obj translator: Unexpected extra arguments on illum keyword"));
 		}
 
 		return true;
@@ -794,7 +787,7 @@ namespace ObjParser
 	{
 		if (Line.IsEmpty())
 		{
-			UE_LOG(LogTemp, Warning, TEXT("Missing texture filename on map_Kd keyword"));
+			INTERCHANGE_OBJ_TRANSLATOR_LOG_ONCE(LogTemp, Warning, TEXT("Interchange Obj translator: Missing texture filename on map_Kd keyword"));
 		}
 
 		if (FObjData::FMaterialData* MaterialData = ObjData.Materials.Find(ObjData.MaterialBeingDefined))
@@ -803,7 +796,7 @@ namespace ObjParser
 		}
 		else
 		{
-			UE_LOG(LogTemp, Warning, TEXT("Missing newmtl keyword"));
+			INTERCHANGE_OBJ_TRANSLATOR_LOG_ONCE(LogTemp, Warning, TEXT("Interchange Obj translator: Missing newmtl keyword"));
 			return false;
 		}
 
@@ -833,7 +826,7 @@ namespace ObjParser
 		FStringView MtlFilename = Line; // (instead of GetToken(Line);)
 		if (MtlFilename.IsEmpty())
 		{
-			UE_LOG(LogTemp, Warning, TEXT("Missing filename on mtllib keyword"));
+			INTERCHANGE_OBJ_TRANSLATOR_LOG_ONCE(LogTemp, Warning, TEXT("Interchange Obj translator: Missing filename on mtllib keyword"));
 		}
 
 		// Parse the file referenced by the .obj
@@ -849,7 +842,7 @@ namespace ObjParser
 		else
 		{
 			//We support only .mtl material file
-			UE_LOG(LogTemp, Warning, TEXT("Unsupported material file: "), *FileToParse);
+			UE_LOG(LogTemp, Warning, TEXT("Interchange Obj translator: Unsupported material file: %s"), *FileToParse);
 		}
 
 		// After parsing, clear the current material name
@@ -960,8 +953,7 @@ bool UInterchangeOBJTranslator::Translate(UInterchangeBaseNodeContainer& BaseNod
 		{ TEXT("f"),      ObjParser::ParseFace },
 		{ TEXT("g"),      ObjParser::ParseGroup },
 		{ TEXT("mtllib"), ObjParser::ParseMaterialLib },
-		{ TEXT("usemtl"), ObjParser::ParseUseMaterial },
-		{ TEXT("s"),      ObjParser::ParseSmoothing },
+		{ TEXT("usemtl"), ObjParser::ParseUseMaterial }
 	};
 
 	bool bSuccess = ObjParser::ParseFile(*ObjDataPtr.Get(), *Filename, ObjKeywordMap);
