Generic Material pipeline:
- Adding the possibility for users to override the predefined material functions built for MaterialX. Currently we only support the Autodesk's standard surface bxdf.
- Adding the support of noise and vector material expressions

MaterialX Translator
- Adding the support of 3D noise functions: noise3d (perlin noise), worleynoise (voronoi) and cellnoise
- Adding a method to allow a node to connect to a different Interchange Shader node during the traversal of the node graph

diff --git a/Engine/Plugins/Interchange/Runtime/Source/Import/Private/MaterialX/InterchangeMaterialXDefinitions.h b/Engine/Plugins/Interchange/Runtime/Source/Import/Private/MaterialX/InterchangeMaterialXDefinitions.h
--- a/Engine/Plugins/Interchange/Runtime/Source/Import/Private/MaterialX/InterchangeMaterialXDefinitions.h
+++ b/Engine/Plugins/Interchange/Runtime/Source/Import/Private/MaterialX/InterchangeMaterialXDefinitions.h
@@ -163,6 +163,7 @@ MATERIALX_NAMESPACE_BEGIN
 	{
 		static constexpr const char* IsVisited = "UE:IsVisited";
 		static constexpr const char* NewName = "UE:NewName";
+		static constexpr const char* ParentName = "UE:ParentName";
 	}
 
 	namespace Category
diff --git a/Engine/Plugins/Interchange/Runtime/Source/Import/Private/MaterialX/InterchangeMaterialXTranslator.cpp b/Engine/Plugins/Interchange/Runtime/Source/Import/Private/MaterialX/InterchangeMaterialXTranslator.cpp
--- a/Engine/Plugins/Interchange/Runtime/Source/Import/Private/MaterialX/InterchangeMaterialXTranslator.cpp
+++ b/Engine/Plugins/Interchange/Runtime/Source/Import/Private/MaterialX/InterchangeMaterialXTranslator.cpp
@@ -14,8 +14,10 @@
 #include "Nodes/InterchangeSourceNode.h"
 #include "UObject/GCObjectScopeGuard.h"
 
-#include "Materials/MaterialExpressionTransformPosition.h"
+#include "Materials/MaterialExpressionNoise.h"
 #include "Materials/MaterialExpressionTransform.h"
+#include "Materials/MaterialExpressionTransformPosition.h"
+#include "Materials/MaterialExpressionVectorNoise.h"
 
 #include UE_INLINE_GENERATED_CPP_BY_NAME(InterchangeMaterialXTranslator)
 
@@ -36,6 +38,7 @@ UInterchangeMaterialXTranslator::UInterchangeMaterialXTranslator()
 	:
 InputNamesMaterialX2UE
 {
+	{{TEXT(""),                         TEXT("amplitude")}, TEXT("Amplitude")},
 	{{TEXT(""),                         TEXT("bg")},        TEXT("B")},
 	{{TEXT(""),                         TEXT("center")},    TEXT("Center")},
 	{{TEXT(""),                         TEXT("diminish")},  TEXT("Diminish")},
@@ -81,6 +84,8 @@ InputNamesMaterialX2UE
 	{{MaterialX::Category::Mix,         TEXT("fg")},        TEXT("B")},
 	{{MaterialX::Category::Mix,         TEXT("bg")},        TEXT("A")},
 	{{MaterialX::Category::Mix,         TEXT("mix")},		TEXT("Factor")},
+	{{MaterialX::Category::Noise3D,     TEXT("amplitude")}, TEXT("B")},            // The amplitude of the noise is connected to a multiply node
+	{{MaterialX::Category::Noise3D,     TEXT("pivot")},     TEXT("B")},            // The pivot of the noise is connected to a add node
 	{{MaterialX::Category::Normalize,   TEXT("in")},        TEXT("VectorInput")},
 	{{MaterialX::Category::Outside,     TEXT("in")},        TEXT("A")},				// Outside is treated as Multiply node
 	{{MaterialX::Category::Outside,     TEXT("mask")},      TEXT("B")},				// Outside is treated as Multiply node
@@ -923,7 +928,7 @@ bool UInterchangeMaterialXTranslator::ConnectNodeGraphOutputToInput(MaterialX::I
 
 				if(mx::NodePtr DownstreamNode = Edge.getDownstreamElement()->asA<mx::Node>())
 				{
-					if(UInterchangeShaderNode** FoundNode = NamesToShaderNodes.Find(DownstreamNode->getName().c_str()))
+					if(UInterchangeShaderNode** FoundNode = NamesToShaderNodes.Find(GetAttributeParentName(DownstreamNode)))// DownstreamNode->getName().c_str()))
 					{
 						ParentShaderNode = *FoundNode;
 					}
@@ -1011,6 +1016,18 @@ bool UInterchangeMaterialXTranslator::ConnectNodeGraphOutputToInput(MaterialX::I
 				{
 					ConnectTimeInputToOutput(UpstreamNode, ParentShaderNode, InputChannelName, NamesToShaderNodes, NodeContainer);
 				}
+				else if(UpstreamNode->getCategory() == mx::Category::Noise3D)
+				{
+					ConnectNoise3DInputToOutput(UpstreamNode, ParentShaderNode, InputChannelName, NamesToShaderNodes, NodeContainer);
+				}
+				else if(UpstreamNode->getCategory() == mx::Category::CellNoise3D)
+				{
+					ConnectCellNoise3DInputToOutput(UpstreamNode, ParentShaderNode, InputChannelName, NamesToShaderNodes, NodeContainer);
+				}
+				else if(UpstreamNode->getCategory() == mx::Category::WorleyNoise3D)
+				{
+					ConnectWorleyNoise3DInputToOutput(UpstreamNode, ParentShaderNode, InputChannelName, NamesToShaderNodes, NodeContainer);
+				}				
 				else
 				{
 					UInterchangeResultWarning_Generic* Message = AddMessage<UInterchangeResultWarning_Generic>();
@@ -1961,6 +1978,76 @@ void UInterchangeMaterialXTranslator::ConnectTimeInputToOutput(MaterialX::NodePt
 	UInterchangeShaderPortsAPI::ConnectDefaultOuputToInput(ParentShaderNode, InputChannelName, TimeNode->GetUniqueID());
 }
 
+void UInterchangeMaterialXTranslator::ConnectNoise3DInputToOutput(MaterialX::NodePtr UpstreamNode, UInterchangeShaderNode* ParentShaderNode, const FString& InputChannelName, TMap<FString, UInterchangeShaderNode*>& NamesToShaderNodes, UInterchangeBaseNodeContainer& NodeContainer) const
+{
+	using namespace UE::Interchange::Materials::Standard::Nodes;
+
+	// MaterialX defines the Noise3d as Perlin Noise which is multiplied by the Amplitude then Added to Pivot
+	UInterchangeShaderNode* NoiseNode = CreateShaderNode<UInterchangeShaderNode>(UpstreamNode->getName().c_str(), Noise::Name.ToString(), NamesToShaderNodes, NodeContainer);
+	NoiseNode->AddInt32Attribute(Noise::Attributes::Function, ENoiseFunction::NOISEFUNCTION_GradientTex);
+	NoiseNode->AddBooleanAttribute(Noise::Attributes::Turbulence, false);
+	NoiseNode->AddFloatAttribute(Noise::Attributes::OutputMin, 0);
+
+	UInterchangeShaderNode* NodeToConnect = NoiseNode;
+
+	// Multiply Node
+	auto ConnectNodeToInput = [&](mx::InputPtr Input, UInterchangeShaderNode * NodeToConnectTo, const FString & ShaderType, int32 IndexAttrib) -> UInterchangeShaderNode*
+	{
+		if(!Input)
+		{
+			return nullptr;
+		}
+
+		const FString ShaderNodeName = UpstreamNode->getName().c_str() + FString{ TEXT("_") } + ShaderType;
+		UInterchangeShaderNode* ShaderNode = CreateShaderNode<UInterchangeShaderNode>(ShaderNodeName, ShaderType, NamesToShaderNodes, NodeContainer);
+
+		UInterchangeShaderPortsAPI::ConnectDefaultOuputToInput(ShaderNode, TEXT("A"), NodeToConnectTo->GetUniqueID());
+
+		// Connect the amplitude node to the shader node not the noise
+		// it will be handle during the upstream-downstream connection phase
+		// The index is here for the unicity of the attribute
+		UpstreamNode->setAttribute(mx::Attributes::ParentName + std::to_string(IndexAttrib), TCHAR_TO_ANSI(*ShaderNodeName));
+		AddAttributeFromValueOrInterface(Input, TEXT("B"), ShaderNode);
+
+		return ShaderNode;
+	};
+
+	if(UInterchangeShaderNode* MultiplyNode = ConnectNodeToInput(UpstreamNode->getInput("amplitude"), NoiseNode, TEXT("Multiply"), 0))
+	{
+		NodeToConnect = MultiplyNode;
+	}
+
+
+	if(UInterchangeShaderNode* AddNode = ConnectNodeToInput(UpstreamNode->getInput("pivot"), NodeToConnect, TEXT("Add"), 1))
+	{
+		NodeToConnect = AddNode;
+	}
+
+	UInterchangeShaderPortsAPI::ConnectDefaultOuputToInput(ParentShaderNode, InputChannelName, NodeToConnect->GetUniqueID());
+}
+
+void UInterchangeMaterialXTranslator::ConnectCellNoise3DInputToOutput(MaterialX::NodePtr UpstreamNode, UInterchangeShaderNode* ParentShaderNode, const FString& InputChannelName, TMap<FString, UInterchangeShaderNode*>& NamesToShaderNodes, UInterchangeBaseNodeContainer& NodeContainer) const
+{
+	using namespace UE::Interchange::Materials::Standard::Nodes;
+
+	// Let's use a vector noise for this one, the only one that is close to MaterialX implementation
+	UInterchangeShaderNode* NoiseNode = CreateShaderNode<UInterchangeShaderNode>(UpstreamNode->getName().c_str(), VectorNoise::Name.ToString(), NamesToShaderNodes, NodeContainer);
+	NoiseNode->AddInt32Attribute(VectorNoise::Attributes::Function, EVectorNoiseFunction::VNF_CellnoiseALU);
+
+	UInterchangeShaderPortsAPI::ConnectDefaultOuputToInput(ParentShaderNode, InputChannelName, NoiseNode->GetUniqueID());
+}
+
+void UInterchangeMaterialXTranslator::ConnectWorleyNoise3DInputToOutput(MaterialX::NodePtr UpstreamNode, UInterchangeShaderNode* ParentShaderNode, const FString& InputChannelName, TMap<FString, UInterchangeShaderNode*>& NamesToShaderNodes, UInterchangeBaseNodeContainer& NodeContainer) const
+{
+	using namespace UE::Interchange::Materials::Standard::Nodes;
+
+	//Also called Voronoi, the implementation is a bit different in UE, especially we don't have access to the jitter
+	UInterchangeShaderNode* NoiseNode = CreateShaderNode<UInterchangeShaderNode>(UpstreamNode->getName().c_str(), Noise::Name.ToString(), NamesToShaderNodes, NodeContainer);
+	NoiseNode->AddInt32Attribute(Noise::Attributes::Function, ENoiseFunction::NOISEFUNCTION_VoronoiALU);
+
+	UInterchangeShaderPortsAPI::ConnectDefaultOuputToInput(ParentShaderNode, InputChannelName, NoiseNode->GetUniqueID());
+}
+
 UInterchangeShaderNode* UInterchangeMaterialXTranslator::CreateMaskShaderNode(uint8 RGBA, const FString& NodeName, TMap<FString, UInterchangeShaderNode*>& NamesToShaderNodes, UInterchangeBaseNodeContainer& NodeContainer) const
 {
 	bool bR = (0b1000 & RGBA) >> 3;
@@ -2145,6 +2232,26 @@ void UInterchangeMaterialXTranslator::AddTexCoordToTiledImageNodes(MaterialX::No
 	}
 }
 
+FString UInterchangeMaterialXTranslator::GetAttributeParentName(MaterialX::NodePtr Node) const
+{
+	FString ParentName;
+	const mx::StringVec & Attributes =  Node->getAttributeNames();
+
+	// For consistency the parent attribute has an index attach to it to ensure unicity
+	// Attributes are set in order, we only need to take the first one (it will be remove after that)
+	for(const std::string& Attrib : Attributes)
+	{
+		if(Attrib.find(mx::Attributes::ParentName) != std::string::npos)
+		{
+			ParentName = Node->getAttribute(Attrib).c_str();
+			Node->removeAttribute(Attrib);
+			break;
+		}
+	}
+
+	return ParentName.IsEmpty() ? Node->getName().c_str() : ParentName;
+}
+
 #endif //WITH_EDITOR
 
 #undef LOCTEXT_NAMESPACE
diff --git a/Engine/Plugins/Interchange/Runtime/Source/Import/Public/MaterialX/InterchangeMaterialXTranslator.h b/Engine/Plugins/Interchange/Runtime/Source/Import/Public/MaterialX/InterchangeMaterialXTranslator.h
--- a/Engine/Plugins/Interchange/Runtime/Source/Import/Public/MaterialX/InterchangeMaterialXTranslator.h
+++ b/Engine/Plugins/Interchange/Runtime/Source/Import/Public/MaterialX/InterchangeMaterialXTranslator.h
@@ -352,6 +352,9 @@ protected:
 	void ConnectTangentInputToOutput(MaterialX::NodePtr UpstreamNode, UInterchangeShaderNode* ParentShaderNode, const FString& InputChannelName, TMap<FString, UInterchangeShaderNode*>& NamesToShaderNodes, UInterchangeBaseNodeContainer& NodeContainer) const;
 	void ConnectBitangentInputToOutput(MaterialX::NodePtr UpstreamNode, UInterchangeShaderNode* ParentShaderNode, const FString& InputChannelName, TMap<FString, UInterchangeShaderNode*>& NamesToShaderNodes, UInterchangeBaseNodeContainer& NodeContainer) const;
 	void ConnectTimeInputToOutput(MaterialX::NodePtr UpstreamNode, UInterchangeShaderNode* ParentShaderNode, const FString& InputChannelName, TMap<FString, UInterchangeShaderNode*>& NamesToShaderNodes, UInterchangeBaseNodeContainer& NodeContainer) const;
+	void ConnectNoise3DInputToOutput(MaterialX::NodePtr UpstreamNode, UInterchangeShaderNode* ParentShaderNode, const FString& InputChannelName, TMap<FString, UInterchangeShaderNode*>& NamesToShaderNodes, UInterchangeBaseNodeContainer& NodeContainer) const;
+	void ConnectCellNoise3DInputToOutput(MaterialX::NodePtr UpstreamNode, UInterchangeShaderNode* ParentShaderNode, const FString& InputChannelName, TMap<FString, UInterchangeShaderNode*>& NamesToShaderNodes, UInterchangeBaseNodeContainer& NodeContainer) const;
+	void ConnectWorleyNoise3DInputToOutput(MaterialX::NodePtr UpstreamNode, UInterchangeShaderNode* ParentShaderNode, const FString& InputChannelName, TMap<FString, UInterchangeShaderNode*>& NamesToShaderNodes, UInterchangeBaseNodeContainer& NodeContainer) const;
 
 	UInterchangeShaderNode* CreateMaskShaderNode(uint8 RGBA, const FString & NodeName, TMap<FString, UInterchangeShaderNode*>& NamesToShaderNodes, UInterchangeBaseNodeContainer& NodeContainer) const;
 
@@ -418,6 +421,14 @@ protected:
 	 */
 	void AddTexCoordToTiledImageNodes(MaterialX::NodeGraphPtr NodeGraph) const;
 
+	/**
+	 * Retrieve the interchange parent name of a MaterialX node, useful when a node is a combination of several nodes connected to different inputs (look for example Noise3D)
+	 *
+	 * @param Node - The node to retrieve the parent name 
+	 * @return The node of the parent
+	 */
+	FString GetAttributeParentName(MaterialX::NodePtr Node) const;
+
 private:
 
 	TMap<TPair<FString, FString>, FString> InputNamesMaterialX2UE; //given a MaterialX node (category - input), return the UE/Interchange input name.
diff --git a/Engine/Plugins/Interchange/Runtime/Source/Nodes/Public/InterchangeMaterialDefinitions.h b/Engine/Plugins/Interchange/Runtime/Source/Nodes/Public/InterchangeMaterialDefinitions.h
--- a/Engine/Plugins/Interchange/Runtime/Source/Nodes/Public/InterchangeMaterialDefinitions.h
+++ b/Engine/Plugins/Interchange/Runtime/Source/Nodes/Public/InterchangeMaterialDefinitions.h
@@ -86,6 +86,31 @@ namespace Interchange
 					}
 				}
 
+				namespace Noise
+				{
+					const FName Name = TEXT("Noise");
+
+					namespace Inputs
+					{
+						const FName Position = TEXT("Position");
+						const FName FilterWidth = TEXT("FilterWidth");
+					}
+
+					namespace Attributes
+					{
+						const FName Scale = TEXT("Scale");
+						const FName Quality = TEXT("Quality");
+						const FName Function = TEXT("NoiseFunction");
+						const FName Turbulence = TEXT("bTurbulence");
+						const FName Levels = TEXT("Levels");
+						const FName OutputMin = TEXT("OutputMin");
+						const FName OutputMax = TEXT("OutputMax");
+						const FName LevelScale = TEXT("LevelScale");
+						const FName Tiling = TEXT("bTiling");
+						const FName RepeatSize = TEXT("RepeatSize");
+					}
+				}
+
 				namespace OneMinus
 				{
 					const FName Name = TEXT("OneMinus");
@@ -176,6 +201,25 @@ namespace Interchange
 						const FName Input = TEXT("Input");
 					}
 				}
+
+				namespace VectorNoise
+				{
+
+					const FName Name = TEXT("VectorNoise");
+
+					namespace Attributes
+					{
+						const FName Function = TEXT("NoiseFunction");
+						const FName Quality = TEXT("Quality");
+						const FName Tiling = TEXT("bTiling");
+						const FName TileSize = TEXT("TileSize");
+					}
+
+					namespace Inputs
+					{
+						const FName Position = TEXT("Position");
+					}
+				}
 			}
 		}
 
diff --git a/Engine/Plugins/Interchange/Runtime/Source/Pipelines/Private/InterchangeGenericMaterialPipeline.cpp b/Engine/Plugins/Interchange/Runtime/Source/Pipelines/Private/InterchangeGenericMaterialPipeline.cpp
--- a/Engine/Plugins/Interchange/Runtime/Source/Pipelines/Private/InterchangeGenericMaterialPipeline.cpp
+++ b/Engine/Plugins/Interchange/Runtime/Source/Pipelines/Private/InterchangeGenericMaterialPipeline.cpp
@@ -26,6 +26,7 @@
 #include "Materials/MaterialExpressionLinearInterpolate.h"
 #include "Materials/MaterialExpressionMaterialFunctionCall.h"
 #include "Materials/MaterialExpressionMultiply.h"
+#include "Materials/MaterialExpressionNoise.h"
 #include "Materials/MaterialExpressionOneMinus.h"
 #include "Materials/MaterialExpressionScalarParameter.h"
 #include "Materials/MaterialExpressionTextureCoordinate.h"
@@ -36,6 +37,7 @@
 #include "Materials/MaterialExpressionTime.h"
 #include "Materials/MaterialExpressionTransformPosition.h"
 #include "Materials/MaterialExpressionTransform.h"
+#include "Materials/MaterialExpressionVectorNoise.h"
 #include "Materials/MaterialExpressionVectorParameter.h"
 #include "Materials/MaterialInstanceConstant.h"
 #include "Materials/MaterialInstanceDynamic.h"
@@ -71,6 +73,58 @@ FString LexToString(UInterchangeGenericMaterialPipeline::EMaterialInputType Valu
 	}
 }
 
+FMaterialXPipelineSettings::FMaterialXPipelineSettings()
+{
+	SurfaceShader.FindOrAdd(EInterchangeMaterialXShaders::StandardSurface, FSoftObjectPath{ TEXT("MaterialFunction'/Interchange/Functions/MX_StandardSurface.MX_StandardSurface'") });
+	SurfaceShader.FindOrAdd(EInterchangeMaterialXShaders::StandardSurfaceTransmission, FSoftObjectPath{ TEXT("MaterialFunction'/Interchange/Functions/MX_TransmissionSurface.MX_TransmissionSurface'") });
+}
+
+bool FMaterialXPipelineSettings::AreRequiredPackagesLoaded()
+{
+	auto ArePackagesLoaded = [this](const TMap<EInterchangeMaterialXShaders, FSoftObjectPath>& ObjectPaths)
+	{
+		bool bAllLoaded = true;
+
+		for(const TPair<EInterchangeMaterialXShaders, FSoftObjectPath>& Pair : ObjectPaths)
+		{
+			const FSoftObjectPath& ObjectPath = Pair.Get<1>();
+
+			if(!ObjectPath.ResolveObject())
+			{
+				FString PackagePath = ObjectPath.GetLongPackageName();
+				if(FPackageName::DoesPackageExist(PackagePath))
+				{
+					if(!ObjectPath.TryLoad())
+					{
+						UE_LOG(LogInterchangePipeline, Warning, TEXT("Couldn't load %s"), *PackagePath);
+						bAllLoaded = false;
+					}
+				}
+				else
+				{
+					UE_LOG(LogInterchangePipeline, Warning, TEXT("Couldn't find %s"), *PackagePath);
+					bAllLoaded = false;
+				}
+			}
+		}
+
+		return bAllLoaded;
+	};
+
+	return ArePackagesLoaded(SurfaceShader);
+}
+
+FString FMaterialXPipelineSettings::GetAssetPathString(EInterchangeMaterialXShaders ShaderType) const
+{
+	FString AssetPath;
+	if(const FSoftObjectPath* ObjectPath = SurfaceShader.Find(ShaderType))
+	{
+		AssetPath = ObjectPath->GetAssetPathString();
+	}
+
+	return AssetPath;
+}
+
 namespace UE::Interchange::InterchangeGenericMaterialPipeline::Private
 {
 	bool AreRequiredPackagesLoaded()
@@ -106,8 +160,6 @@ namespace UE::Interchange::InterchangeGenericMaterialPipeline::Private
 		};
 
 		TArray<FString> RequiredPackages = {
-			TEXT("MaterialFunction'/Interchange/Functions/MX_StandardSurface.MX_StandardSurface'"),
-			TEXT("MaterialFunction'/Interchange/Functions/MX_TransmissionSurface.MX_TransmissionSurface'"),
 			TEXT("MaterialFunction'/Engine/Functions/Engine_MaterialFunctions01/Shading/ConvertFromDiffSpec.ConvertFromDiffSpec'"),
 			TEXT("MaterialFunction'/Engine/Functions/Engine_MaterialFunctions01/Texturing/FlattenNormal.FlattenNormal'"),
 			TEXT("MaterialFunction'/Engine/Functions/Engine_MaterialFunctions02/Utility/MakeFloat3.MakeFloat3'"),
@@ -180,6 +232,11 @@ void UInterchangeGenericMaterialPipeline::AdjustSettingsForContext(EInterchangeP
 	{
 		UE_LOG(LogInterchangePipeline, Warning, TEXT("UInterchangeGenericMaterialPipeline: Some required packages are missing. Material import might be wrong"));
 	}
+
+	if(!MaterialXSettings.AreRequiredPackagesLoaded())
+	{
+		UE_LOG(LogInterchangePipeline, Warning, TEXT("UInterchangeGenericMaterialPipeline: Some required packages are missing. Material import might be wrong"));
+	}
 }
 
 void UInterchangeGenericMaterialPipeline::ExecutePreImportPipeline(UInterchangeBaseNodeContainer* InBaseNodeContainer, const TArray<UInterchangeSourceData*>& InSourceDatas)
@@ -624,11 +681,11 @@ bool UInterchangeGenericMaterialPipeline::HandleStandardSurfaceModel(const UInte
 		const FName MaterialFunctionMemberName = GET_MEMBER_NAME_CHECKED(UMaterialExpressionMaterialFunctionCall, MaterialFunction);
 
 		using namespace UE::Interchange::Materials;
-		FString StandardSurfaceOrTransmissionPath{ TEXT("MaterialFunction'/Interchange/Functions/MX_StandardSurface.MX_StandardSurface'") };
+		FString StandardSurfaceOrTransmissionPath = MaterialXSettings.GetAssetPathString(EInterchangeMaterialXShaders::StandardSurface);
 
 		if(UInterchangeShaderPortsAPI::HasInput(ShaderGraphNode, StandardSurface::Parameters::Transmission))
 		{
-			StandardSurfaceOrTransmissionPath = TEXT("MaterialFunction'/Interchange/Functions/MX_TransmissionSurface.MX_TransmissionSurface'");
+			StandardSurfaceOrTransmissionPath = MaterialXSettings.GetAssetPathString(EInterchangeMaterialXShaders::StandardSurfaceTransmission);
 		}
 
 		FunctionCallExpression->AddStringAttribute(
@@ -1583,6 +1640,168 @@ void UInterchangeGenericMaterialPipeline::HandleTransformVectorNode(const UInter
 	}
 }
 
+void UInterchangeGenericMaterialPipeline::HandleNoiseNode(const UInterchangeShaderNode* ShaderNode, UInterchangeBaseMaterialFactoryNode* MaterialFactoryNode, UInterchangeMaterialExpressionFactoryNode* NoiseFactoryNode)
+{
+	using namespace UE::Interchange::Materials::Standard::Nodes;
+
+	NoiseFactoryNode->SetCustomExpressionClassName(UMaterialExpressionNoise::StaticClass()->GetName());
+
+	// Position
+	{
+		TTuple<UInterchangeMaterialExpressionFactoryNode*, FString> PositionExpression =
+			CreateMaterialExpressionForInput(MaterialFactoryNode, ShaderNode, Noise::Inputs::Position.ToString(), NoiseFactoryNode->GetUniqueID());
+
+		if(PositionExpression.Get<0>())
+		{
+			UInterchangeShaderPortsAPI::ConnectOuputToInput(NoiseFactoryNode, GET_MEMBER_NAME_CHECKED(UMaterialExpressionNoise, Position).ToString(),
+															PositionExpression.Get<0>()->GetUniqueID(), PositionExpression.Get<1>());
+		}
+	}
+
+	// FilterWidth
+	{
+		TTuple<UInterchangeMaterialExpressionFactoryNode*, FString> FilterWidthExpression =
+			CreateMaterialExpressionForInput(MaterialFactoryNode, ShaderNode, Noise::Inputs::FilterWidth.ToString(), NoiseFactoryNode->GetUniqueID());
+
+		if(FilterWidthExpression.Get<0>())
+		{
+			UInterchangeShaderPortsAPI::ConnectOuputToInput(NoiseFactoryNode, GET_MEMBER_NAME_CHECKED(UMaterialExpressionNoise, FilterWidth).ToString(),
+															FilterWidthExpression.Get<0>()->GetUniqueID(), FilterWidthExpression.Get<1>());
+		}
+	}
+
+	// Scale
+	if(float Scale; ShaderNode->GetFloatAttribute(Noise::Attributes::Scale, Scale))
+	{
+		const FName ScaleMemberName = GET_MEMBER_NAME_CHECKED(UMaterialExpressionNoise, Scale);
+		NoiseFactoryNode->AddFloatAttribute(ScaleMemberName, Scale);
+		NoiseFactoryNode->AddApplyAndFillDelegates<float>(ScaleMemberName, UMaterialExpressionNoise::StaticClass(), ScaleMemberName);
+	}
+
+	// Quality
+	if(int32 Quality; ShaderNode->GetInt32Attribute(Noise::Attributes::Quality, Quality))
+	{
+		const FName QualityMemberName = GET_MEMBER_NAME_CHECKED(UMaterialExpressionNoise, Quality);
+		NoiseFactoryNode->AddInt32Attribute(QualityMemberName, Quality);
+		NoiseFactoryNode->AddApplyAndFillDelegates<int32>(QualityMemberName, UMaterialExpressionNoise::StaticClass(), QualityMemberName);
+	}
+
+	// Noise Function
+	if(int32 NoiseFunction; ShaderNode->GetInt32Attribute(Noise::Attributes::Function, NoiseFunction))
+	{
+		const FName NoiseFunctionMemberName = GET_MEMBER_NAME_CHECKED(UMaterialExpressionNoise, NoiseFunction);
+		NoiseFactoryNode->AddInt32Attribute(NoiseFunctionMemberName, NoiseFunction);
+		NoiseFactoryNode->AddApplyAndFillDelegates<int32>(NoiseFunctionMemberName, UMaterialExpressionNoise::StaticClass(), NoiseFunctionMemberName);
+	}
+
+	// Turbulence
+	if(bool bTurbulence; ShaderNode->GetBooleanAttribute(Noise::Attributes::Turbulence, bTurbulence))
+	{
+		const FName TurbulenceMemberName = GET_MEMBER_NAME_CHECKED(UMaterialExpressionNoise, bTurbulence);
+		NoiseFactoryNode->AddBooleanAttribute(TurbulenceMemberName, bTurbulence);
+		NoiseFactoryNode->AddApplyAndFillDelegates<bool>(TurbulenceMemberName, UMaterialExpressionNoise::StaticClass(), TurbulenceMemberName);
+	}
+
+	// Levels
+	if(int32 Levels; ShaderNode->GetInt32Attribute(Noise::Attributes::Levels, Levels))
+	{
+		const FName LevelsMemberName = GET_MEMBER_NAME_CHECKED(UMaterialExpressionNoise, Levels);
+		NoiseFactoryNode->AddInt32Attribute(LevelsMemberName, Levels);
+		NoiseFactoryNode->AddApplyAndFillDelegates<int32>(LevelsMemberName, UMaterialExpressionNoise::StaticClass(), LevelsMemberName);
+	}
+
+	// Output Min
+	if(float OutputMin; ShaderNode->GetFloatAttribute(Noise::Attributes::OutputMin, OutputMin))
+	{
+		const FName OutputMinMemberName = GET_MEMBER_NAME_CHECKED(UMaterialExpressionNoise, OutputMin);
+		NoiseFactoryNode->AddFloatAttribute(OutputMinMemberName, OutputMin);
+		NoiseFactoryNode->AddApplyAndFillDelegates<float>(OutputMinMemberName, UMaterialExpressionNoise::StaticClass(), OutputMinMemberName);
+	}
+
+	// Output Max
+	if(float OutputMax; ShaderNode->GetFloatAttribute(Noise::Attributes::OutputMax, OutputMax))
+	{
+		const FName OutputMaxMemberName = GET_MEMBER_NAME_CHECKED(UMaterialExpressionNoise, OutputMax);
+		NoiseFactoryNode->AddFloatAttribute(OutputMaxMemberName, OutputMax);
+		NoiseFactoryNode->AddApplyAndFillDelegates<float>(OutputMaxMemberName, UMaterialExpressionNoise::StaticClass(), OutputMaxMemberName);
+	}
+
+	// Level Scale
+	if(float LevelScale; ShaderNode->GetFloatAttribute(Noise::Attributes::LevelScale, LevelScale))
+	{
+		const FName LevelScaleMemberName = GET_MEMBER_NAME_CHECKED(UMaterialExpressionNoise, LevelScale);
+		NoiseFactoryNode->AddFloatAttribute(LevelScaleMemberName, LevelScale);
+		NoiseFactoryNode->AddApplyAndFillDelegates<float>(LevelScaleMemberName, UMaterialExpressionNoise::StaticClass(), LevelScaleMemberName);
+	}
+
+	// Tiling
+	if(bool bTiling; ShaderNode->GetBooleanAttribute(Noise::Attributes::Tiling, bTiling))
+	{
+		const FName TilingMemberName = GET_MEMBER_NAME_CHECKED(UMaterialExpressionNoise, bTiling);
+		NoiseFactoryNode->AddBooleanAttribute(TilingMemberName, bTiling);
+		NoiseFactoryNode->AddApplyAndFillDelegates<bool>(TilingMemberName, UMaterialExpressionNoise::StaticClass(), TilingMemberName);
+	}
+
+	// Levels
+	if(int32 RepeatSize; ShaderNode->GetInt32Attribute(Noise::Attributes::RepeatSize, RepeatSize))
+	{
+		const FName RepeatSizeMemberName = GET_MEMBER_NAME_CHECKED(UMaterialExpressionNoise, RepeatSize);
+		NoiseFactoryNode->AddInt32Attribute(RepeatSizeMemberName, RepeatSize);
+		NoiseFactoryNode->AddApplyAndFillDelegates<int32>(RepeatSizeMemberName, UMaterialExpressionNoise::StaticClass(), RepeatSizeMemberName);
+	}
+}
+
+void UInterchangeGenericMaterialPipeline::HandleVectorNoiseNode(const UInterchangeShaderNode* ShaderNode, UInterchangeBaseMaterialFactoryNode* MaterialFactoryNode, UInterchangeMaterialExpressionFactoryNode* NoiseFactoryNode)
+{
+	using namespace UE::Interchange::Materials::Standard::Nodes;
+
+	NoiseFactoryNode->SetCustomExpressionClassName(UMaterialExpressionVectorNoise::StaticClass()->GetName());
+
+	// Position
+	{
+		TTuple<UInterchangeMaterialExpressionFactoryNode*, FString> PositionExpression =
+			CreateMaterialExpressionForInput(MaterialFactoryNode, ShaderNode, VectorNoise::Inputs::Position.ToString(), NoiseFactoryNode->GetUniqueID());
+
+		if(PositionExpression.Get<0>())
+		{
+			UInterchangeShaderPortsAPI::ConnectOuputToInput(NoiseFactoryNode, GET_MEMBER_NAME_CHECKED(UMaterialExpressionVectorNoise, Position).ToString(),
+															PositionExpression.Get<0>()->GetUniqueID(), PositionExpression.Get<1>());
+		}
+	}
+
+	// Noise Function
+	if(int32 NoiseFunction; ShaderNode->GetInt32Attribute(VectorNoise::Attributes::Function, NoiseFunction))
+	{
+		const FName NoiseFunctionMemberName = GET_MEMBER_NAME_CHECKED(UMaterialExpressionVectorNoise, NoiseFunction);
+		NoiseFactoryNode->AddInt32Attribute(NoiseFunctionMemberName, NoiseFunction);
+		NoiseFactoryNode->AddApplyAndFillDelegates<int32>(NoiseFunctionMemberName, UMaterialExpressionVectorNoise::StaticClass(), NoiseFunctionMemberName);
+	}
+
+	// Quality
+	if(int32 Quality; ShaderNode->GetInt32Attribute(VectorNoise::Attributes::Quality, Quality))
+	{
+		const FName QualityMemberName = GET_MEMBER_NAME_CHECKED(UMaterialExpressionVectorNoise, Quality);
+		NoiseFactoryNode->AddInt32Attribute(QualityMemberName, Quality);
+		NoiseFactoryNode->AddApplyAndFillDelegates<int32>(QualityMemberName, UMaterialExpressionVectorNoise::StaticClass(), QualityMemberName);
+	}
+
+	// Tiling
+	if(bool bTiling; ShaderNode->GetBooleanAttribute(VectorNoise::Attributes::Tiling, bTiling))
+	{
+		const FName TilingMemberName = GET_MEMBER_NAME_CHECKED(UMaterialExpressionVectorNoise, bTiling);
+		NoiseFactoryNode->AddBooleanAttribute(TilingMemberName, bTiling);
+		NoiseFactoryNode->AddApplyAndFillDelegates<bool>(TilingMemberName, UMaterialExpressionVectorNoise::StaticClass(), TilingMemberName);
+	}
+
+	// Tile Size
+	if(int32 TileSize; ShaderNode->GetInt32Attribute(VectorNoise::Attributes::Function, TileSize))
+	{
+		const FName TileSizeMemberName = GET_MEMBER_NAME_CHECKED(UMaterialExpressionVectorNoise, TileSize);
+		NoiseFactoryNode->AddInt32Attribute(TileSizeMemberName, TileSize);
+		NoiseFactoryNode->AddApplyAndFillDelegates<int32>(TileSizeMemberName, UMaterialExpressionVectorNoise::StaticClass(), TileSizeMemberName);
+	}
+}
+
 UInterchangeMaterialExpressionFactoryNode* UInterchangeGenericMaterialPipeline::CreateMaterialExpressionForShaderNode(UInterchangeBaseMaterialFactoryNode* MaterialFactoryNode,
 	const UInterchangeShaderNode* ShaderNode, const FString& ParentUid)
 {
@@ -1643,6 +1862,10 @@ UInterchangeMaterialExpressionFactoryNode* UInterchangeGenericMaterialPipeline::
 	{
 		HandleMaskNode(ShaderNode, MaterialFactoryNode, MaterialExpression);
 	}
+	else if(*ShaderType == Nodes::Noise::Name)
+	{
+		HandleNoiseNode(ShaderNode, MaterialFactoryNode, MaterialExpression);
+	}
 	else if (*ShaderType == Nodes::TextureCoordinate::Name)
 	{
 		HandleTextureCoordinateNode(ShaderNode, MaterialFactoryNode, MaterialExpression);
@@ -1663,6 +1886,10 @@ UInterchangeMaterialExpressionFactoryNode* UInterchangeGenericMaterialPipeline::
 	{
 		HandleTransformVectorNode(ShaderNode, MaterialFactoryNode, MaterialExpression);
 	}
+	else if(*ShaderType == Nodes::TransformVector::Name)
+	{
+		HandleVectorNoiseNode(ShaderNode, MaterialFactoryNode, MaterialExpression);
+	}
 	else
 	{
 		const FString ExpressionClassName = TEXT("MaterialExpression") + ShaderType;
diff --git a/Engine/Plugins/Interchange/Runtime/Source/Pipelines/Public/InterchangeGenericMaterialPipeline.h b/Engine/Plugins/Interchange/Runtime/Source/Pipelines/Public/InterchangeGenericMaterialPipeline.h
--- a/Engine/Plugins/Interchange/Runtime/Source/Pipelines/Public/InterchangeGenericMaterialPipeline.h
+++ b/Engine/Plugins/Interchange/Runtime/Source/Pipelines/Public/InterchangeGenericMaterialPipeline.h
@@ -29,6 +29,33 @@ enum class EInterchangeMaterialImportOption : uint8
 	ImportAsMaterialInstances,
 };
 
+
+UENUM(BlueprintType)
+enum class EInterchangeMaterialXShaders : uint8
+{
+	/** Default settings for Autodesk's Standard Surface shader	*/
+	StandardSurface,
+
+	/** Standard Surface shader use for translucency	*/
+	StandardSurfaceTransmission
+};
+
+USTRUCT(BlueprintType)
+struct FMaterialXPipelineSettings
+{
+	GENERATED_BODY()
+
+public:
+	FMaterialXPipelineSettings();
+
+	UPROPERTY(EditAnywhere, BlueprintReadWrite, Category = "MaterialX")
+	TMap<EInterchangeMaterialXShaders, FSoftObjectPath> SurfaceShader;
+
+	bool AreRequiredPackagesLoaded();
+
+	FString GetAssetPathString(EInterchangeMaterialXShaders ShaderType) const;
+};
+
 UCLASS(BlueprintType, editinlinenew)
 class INTERCHANGEPIPELINES_API UInterchangeGenericMaterialPipeline : public UInterchangePipelineBase
 {
@@ -51,6 +78,9 @@ public:
 	UPROPERTY(VisibleAnywhere, Instanced, Category = "Textures")
 	TObjectPtr<UInterchangeGenericTexturePipeline> TexturePipeline;
 
+	UPROPERTY(EditAnywhere, BlueprintReadWrite, Category = "MaterialX", DisplayName = "Settings")
+	FMaterialXPipelineSettings MaterialXSettings;
+
 	/** BEGIN UInterchangePipelineBase overrides */
 	virtual void PreDialogCleanup(const FName PipelineStackName) override;
 	virtual bool IsSettingsAreValid(TOptional<FText>& OutInvalidReason) const override;
@@ -127,7 +157,8 @@ private:
 	void HandleTimeNode(const UInterchangeShaderNode* ShaderNode, UInterchangeBaseMaterialFactoryNode* MaterialFactoryNode, UInterchangeMaterialExpressionFactoryNode* TimeFactoryNode);
 	void HandleTransformPositionNode(const UInterchangeShaderNode* ShaderNode, UInterchangeBaseMaterialFactoryNode* MaterialFactoryNode, UInterchangeMaterialExpressionFactoryNode* TransformPositionFactoryNode);
 	void HandleTransformVectorNode(const UInterchangeShaderNode* ShaderNode, UInterchangeBaseMaterialFactoryNode* MaterialFactoryNode, UInterchangeMaterialExpressionFactoryNode* TransformVectorFactoryNode);
-
+	void HandleNoiseNode(const UInterchangeShaderNode* ShaderNode, UInterchangeBaseMaterialFactoryNode* MaterialFactoryNode, UInterchangeMaterialExpressionFactoryNode* NoiseFactoryNode);
+	void HandleVectorNoiseNode(const UInterchangeShaderNode* ShaderNode, UInterchangeBaseMaterialFactoryNode* MaterialFactoryNode, UInterchangeMaterialExpressionFactoryNode* NoiseFactoryNode);
 
 	UInterchangeMaterialExpressionFactoryNode* CreateMaterialExpressionForShaderNode(UInterchangeBaseMaterialFactoryNode* MaterialFactoryNode, const UInterchangeShaderNode* ShaderNode, const FString& ParentUid);
 	TTuple<UInterchangeMaterialExpressionFactoryNode*, FString> CreateMaterialExpressionForInput(UInterchangeBaseMaterialFactoryNode* MaterialFactoryNode, const UInterchangeShaderNode* ShaderNode, const FString& InputName, const FString& ParentUid);
@@ -177,5 +208,3 @@ private:
 
 	TArray<FMaterialExpressionCreationContext> MaterialExpressionCreationContextStack;
 };
-
-
