MaterialX:
- Changing back Subsurface Profile by Subsurface
- Connecting Refraction
- Exposing the MX_TransmissionSurface to the library

diff --git a/Engine/Plugins/Interchange/Runtime/Source/Nodes/Public/InterchangeMaterialDefinitions.h b/Engine/Plugins/Interchange/Runtime/Source/Nodes/Public/InterchangeMaterialDefinitions.h
--- a/Engine/Plugins/Interchange/Runtime/Source/Nodes/Public/InterchangeMaterialDefinitions.h
+++ b/Engine/Plugins/Interchange/Runtime/Source/Nodes/Public/InterchangeMaterialDefinitions.h
@@ -234,6 +234,7 @@ namespace Interchange
 				const FName Occlusion = TEXT("Occlusion"); // Type: float
 				const FName IndexOfRefraction = TEXT("IOR"); // Type: float
 				const FName BxDF = TEXT("BxDF"); // input/output of BSDF or BRDF or BXDF or BTDF data
+				const FName Refraction = TEXT("Refraction"); // input/output of BSDF or BRDF or BXDF or BTDF data
 			}
 		}
 
diff --git a/Engine/Plugins/Interchange/Runtime/Source/Pipelines/Private/InterchangeGenericMaterialPipeline.cpp b/Engine/Plugins/Interchange/Runtime/Source/Pipelines/Private/InterchangeGenericMaterialPipeline.cpp
--- a/Engine/Plugins/Interchange/Runtime/Source/Pipelines/Private/InterchangeGenericMaterialPipeline.cpp
+++ b/Engine/Plugins/Interchange/Runtime/Source/Pipelines/Private/InterchangeGenericMaterialPipeline.cpp
@@ -780,11 +780,12 @@ bool UInterchangeGenericMaterialPipeline::HandleStandardSurfaceModel(const UInte
 		else if(UInterchangeShaderPortsAPI::HasInput(ShaderGraphNode, StandardSurface::Parameters::Subsurface))
 		{
 			MaterialFactoryNode->ConnectOutputToSubsurface(FunctionCallExpressionUid, Subsurface::Parameters::SubsurfaceColor.ToString());
-			MaterialFactoryNode->SetCustomShadingModel(MSM_SubsurfaceProfile); //MaterialX subsurface fits more with subsurface profile than subsurface, see: standard_surface_chess_set (King_B/W, Queen_B/W)
+			MaterialFactoryNode->SetCustomShadingModel(MSM_Subsurface);
 		}
 		else if (UInterchangeShaderPortsAPI::HasInput(ShaderGraphNode, StandardSurface::Parameters::Transmission))
 		{
 			MaterialFactoryNode->ConnectOutputToTransmissionColor(FunctionCallExpressionUid, ThinTranslucent::Parameters::TransmissionColor.ToString());
+			MaterialFactoryNode->ConnectOutputToRefraction(FunctionCallExpressionUid, Common::Parameters::Refraction.ToString());
 
 			MaterialFactoryNode->SetCustomBlendMode(EBlendMode::BLEND_Translucent);
 			MaterialFactoryNode->SetCustomShadingModel(EMaterialShadingModel::MSM_ThinTranslucent);
