Fixing duplicated material expressions when reimporting a Material from a file

diff --git a/Engine/Plugins/Interchange/Runtime/Source/Import/Private/Material/InterchangeMaterialFactory.cpp b/Engine/Plugins/Interchange/Runtime/Source/Import/Private/Material/InterchangeMaterialFactory.cpp
--- a/Engine/Plugins/Interchange/Runtime/Source/Import/Private/Material/InterchangeMaterialFactory.cpp
+++ b/Engine/Plugins/Interchange/Runtime/Source/Import/Private/Material/InterchangeMaterialFactory.cpp
@@ -183,6 +183,9 @@ namespace UE
 						, Arguments(InArguments)
 					{
 						check((Material || MaterialFunction) && !(Material && MaterialFunction));
+
+						//We need to put in place a better mechanism for the reimport, for the moment let's delete all material expressions to avoid duplicates	
+						DeleteAllMaterialExpressions();
 					}
 
 					UMaterialExpression* CreateExpressionsForNode(const UInterchangeMaterialExpressionFactoryNode& ExpressionNode)
@@ -297,6 +300,28 @@ namespace UE
 						return MaterialExpression;
 					}
 
+					void DeleteAllMaterialExpressions()
+					{
+						if(Material)
+						{
+							// Copy the Material Expressions in a TArray, otherwise working directly on the TArrayView messes up with the expressions
+							// and doesn't make a full clean up, especially when we're not in the Game Thread
+							if(TArray<UMaterialExpression*> MaterialExpressions(Material->GetExpressions()); !MaterialExpressions.IsEmpty())
+							{
+								Material->Modify();
+								for(UMaterialExpression* MaterialExpression : MaterialExpressions)
+								{
+									MaterialExpression->Modify();
+									Material->GetExpressionCollection().RemoveExpression(MaterialExpression);
+									Material->RemoveExpressionParameter(MaterialExpression);
+									// Make sure the deleted expression is caught by gc
+									MaterialExpression->MarkAsGarbage();
+								}
+
+								Material->MarkPackageDirty();
+							}
+						}
+					}
 
 					UMaterial* Material = nullptr;
 					UMaterialFunction* MaterialFunction = nullptr;
