[interchange] [gltf] transparency issues on models fix
(Allowing Interchange Unlit material to be transparent)

diff --git a/Engine/Plugins/Interchange/Runtime/Source/Import/Private/Gltf/InterchangeGltfTranslator.cpp b/Engine/Plugins/Interchange/Runtime/Source/Import/Private/Gltf/InterchangeGltfTranslator.cpp
--- a/Engine/Plugins/Interchange/Runtime/Source/Import/Private/Gltf/InterchangeGltfTranslator.cpp
+++ b/Engine/Plugins/Interchange/Runtime/Source/Import/Private/Gltf/InterchangeGltfTranslator.cpp
@@ -391,6 +391,16 @@ void UInterchangeGltfTranslator::HandleGltfMaterial( UInterchangeBaseNodeContain
 				BaseColorFactor, Standard::Nodes::TextureSample::Outputs::RGB.ToString());
 		}
 
+		// Opacity (use the base color alpha channel)
+		if (GltfMaterial.AlphaMode != GLTF::FMaterial::EAlphaMode::Opaque)
+		{
+			TVariant< FLinearColor, float > OpacityFactor;
+			OpacityFactor.Set< float >(GltfMaterial.BaseColorFactor.W);
+
+			HandleGltfMaterialParameter(NodeContainer, GltfMaterial.BaseColor, ShaderGraphNode, PBR::Parameters::Opacity.ToString(),
+				OpacityFactor, Standard::Nodes::TextureSample::Outputs::A.ToString());
+		}
+
 		return;
 	}
 	
diff --git a/Engine/Plugins/Interchange/Runtime/Source/Pipelines/Private/InterchangeGenericMaterialPipeline.cpp b/Engine/Plugins/Interchange/Runtime/Source/Pipelines/Private/InterchangeGenericMaterialPipeline.cpp
--- a/Engine/Plugins/Interchange/Runtime/Source/Pipelines/Private/InterchangeGenericMaterialPipeline.cpp
+++ b/Engine/Plugins/Interchange/Runtime/Source/Pipelines/Private/InterchangeGenericMaterialPipeline.cpp
@@ -2520,6 +2520,24 @@ bool UInterchangeGenericMaterialPipeline::HandleUnlitModel(const UInterchangeSha
 
 	bool bShadingModelHandled = false;
 
+	{
+		//gltf allows unlit color to be also translucent:
+		const bool bHasInput = UInterchangeShaderPortsAPI::HasInput(ShaderGraphNode, Common::Parameters::Opacity);
+
+		if (bHasInput)
+		{
+			TTuple<UInterchangeMaterialExpressionFactoryNode*, FString> OpacityExpressionFactoryNode =
+				CreateMaterialExpressionForInput(MaterialFactoryNode, ShaderGraphNode, Common::Parameters::Opacity.ToString(), MaterialFactoryNode->GetUniqueID());
+
+			if (OpacityExpressionFactoryNode.Get<0>())
+			{
+				MaterialFactoryNode->ConnectOutputToOpacity(OpacityExpressionFactoryNode.Get<0>()->GetUniqueID(), OpacityExpressionFactoryNode.Get<1>());
+			}
+
+			MaterialFactoryNode->SetCustomBlendMode(EBlendMode::BLEND_Translucent);
+		}
+	}
+
 	// Unlit Color
 	{
 		const bool bHasInput = UInterchangeShaderPortsAPI::HasInput(ShaderGraphNode, Unlit::Parameters::UnlitColor);
