[interchange] [gltf] geometry missing due to naming clashes fix

diff --git a/Engine/Plugins/Editor/GLTFImporter/Source/GLTFCore/Private/GLTF/GLTFAsset.cpp b/Engine/Plugins/Editor/GLTFImporter/Source/GLTFCore/Private/GLTF/GLTFAsset.cpp
--- a/Engine/Plugins/Editor/GLTFImporter/Source/GLTFCore/Private/GLTF/GLTFAsset.cpp
+++ b/Engine/Plugins/Editor/GLTFImporter/Source/GLTFCore/Private/GLTF/GLTFAsset.cpp
@@ -15,9 +15,11 @@ namespace GLTF
 			{
 				T& Obj = Objects[Index];
 
+				Obj.UniqueId = Prefix + FString::FromInt(Index);
+
 				if (Obj.Name.IsEmpty())
 				{
-					Obj.Name = Prefix + FString::FromInt(Index);
+					Obj.Name = Obj.UniqueId;
 				}
 			}
 		}
@@ -71,6 +73,22 @@ namespace GLTF
 	{
 		check(!Prefix.IsEmpty());
 
+		{
+			const FString ScenePrefix = Prefix + TEXT("_scene_");
+
+			for (size_t SceneIndex = 0; SceneIndex < Scenes.Num(); SceneIndex++)
+			{
+				FScene& Scene = Scenes[SceneIndex];
+
+				Scene.UniqueId = ScenePrefix + FString::FromInt(SceneIndex);
+
+				if (Scene.Name.IsEmpty())
+				{
+					Scene.Name = Scene.UniqueId;
+				}
+			}
+		}
+
 		{
 			const FString NodePrefix = Prefix + TEXT("_node_");
 			const FString JointPrefix = Prefix + TEXT("_joint_");
@@ -78,13 +96,14 @@ namespace GLTF
 			{
 				FNode& Node = Nodes[NodeIndex];
 
+				bool bIsJoint = Node.Type == FNode::EType::Joint;
+				Node.UniqueId = (bIsJoint ? JointPrefix : NodePrefix);
+
+				Node.UniqueId = Node.UniqueId + FString::FromInt(NodeIndex);
+
 				if (Node.Name.IsEmpty())
 				{
-					bool bIsJoint = Node.Type == FNode::EType::Joint;
-					Node.Name     = (bIsJoint ? JointPrefix : NodePrefix);
-
-					// Make sure node names are unique, in case Name is not empty it is presumed to be unique / handled by the processing pipelines.
-					Node.Name = Node.Name + FString::FromInt(NodeIndex);
+					Node.Name = Node.UniqueId;
 				}
 			}
 		}
@@ -95,6 +114,8 @@ namespace GLTF
 			{
 				FTexture& Tex = Textures[TextureIndex];
 
+				Tex.UniqueId = TexPrefix + FString::FromInt(TextureIndex);
+
 				if (!Tex.Name.IsEmpty())
 				{
 					// keep base part of the name as it was in gltf
@@ -105,14 +126,19 @@ namespace GLTF
 				}
 				else if (!Tex.Source.URI.IsEmpty())
 				{
-					Tex.Name = FPaths::GetBaseFilename(Tex.Source.URI);
+					//in case texture is embedded:
+					if (Tex.Source.URI.StartsWith(TEXT("data:")))
+					{
+						Tex.Name = Tex.UniqueId;
+					}
+					else
+					{
+						Tex.Name = FPaths::GetBaseFilename(Tex.Source.URI);
+					}
 				}
 				else
 				{
-					// GLTF texture name has decorative purpose, not guaranteed to be unique
-					// only its index is unique. Same with glTF image or its source file's basename
-					// So always include texture index into texture's name to increase probability that names are unique
-					Tex.Name = TexPrefix + FString::FromInt(TextureIndex);
+					Tex.Name = Tex.UniqueId;
 				}
 
 			}
@@ -124,6 +150,8 @@ namespace GLTF
 			{
 				FCamera& Camera = Cameras[CameraIndex];
 
+				Camera.UniqueId = CameraPrefix + FString::FromInt(CameraIndex);
+
 				if (Camera.Name.IsEmpty())
 				{
 					if (!Camera.Node.Name.IsEmpty())
@@ -132,7 +160,7 @@ namespace GLTF
 					}
 					else
 					{
-						Camera.Name = CameraPrefix + FString::FromInt(CameraIndex);
+						Camera.Name = Camera.UniqueId;
 					}
 				}
 			}
@@ -144,6 +172,8 @@ namespace GLTF
 			{
 				FLight& Light = Lights[LightIndex];
 
+				Light.UniqueId = LightPrefix + FString::FromInt(LightIndex);
+
 				if (Light.Name.IsEmpty())
 				{
 					if (Light.Node && !Light.Node->Name.IsEmpty())
@@ -152,18 +182,17 @@ namespace GLTF
 					}
 					else
 					{
-						Light.Name = LightPrefix + FString::FromInt(LightIndex);
+						Light.Name = Light.UniqueId;
 					}
 				}
 			}
 		}
 
-		// General pattern here is that we'll have a Prefix that will be a filename (like MyFile). If the
-		// GLTF object has an actual name, it will end up like <index>_<objectname>, like "0_shoe" or "3_spotlight".
-		// If it doesn't have a name, it will end up with a name like <index>_<filename>_<objecttype>, e.g.
-		// "0_MyFile_material" or "3_myfile_mesh".
-		// We want to do images after textures as the textures may want to reuse source image names, but we don't
-		// want to use those if we generated them ourselves.
+		// General pattern here is that we'll have a Prefix that will be a filename (like MyFile). 
+		// The name (display label) and the UniqueId of an object has been separated.
+		// UniqueId will consist of: [Prefil]+[_ElementType_]+[ElementIndex]
+		// Name will be the provided GLTF object.name, if present,
+		// if Name is not present the UniqueId will be used as the Name.
 		GLTF::GenerateNames(Prefix + TEXT("_material_"), Materials);
 		GLTF::GenerateNames(Prefix + TEXT("_skin_"), Skins);
 		GLTF::GenerateNames(Prefix + TEXT("_animation_"), Animations);
diff --git a/Engine/Plugins/Editor/GLTFImporter/Source/GLTFCore/Public/GLTFAnimation.h b/Engine/Plugins/Editor/GLTFImporter/Source/GLTFCore/Public/GLTFAnimation.h
--- a/Engine/Plugins/Editor/GLTFImporter/Source/GLTFCore/Public/GLTFAnimation.h
+++ b/Engine/Plugins/Editor/GLTFImporter/Source/GLTFCore/Public/GLTFAnimation.h
@@ -69,6 +69,8 @@ namespace GLTF
 		FString          Name;
 		TArray<FSampler> Samplers;
 		TArray<FChannel> Channels;
+
+		FString          UniqueId; //will be generated in FAsset::GenerateNames
 	};
 
 }  // namespace GLTF
diff --git a/Engine/Plugins/Editor/GLTFImporter/Source/GLTFCore/Public/GLTFAsset.h b/Engine/Plugins/Editor/GLTFImporter/Source/GLTFCore/Public/GLTFAsset.h
--- a/Engine/Plugins/Editor/GLTFImporter/Source/GLTFCore/Public/GLTFAsset.h
+++ b/Engine/Plugins/Editor/GLTFImporter/Source/GLTFCore/Public/GLTFAsset.h
@@ -37,6 +37,7 @@ namespace GLTF
 	{
 		FString       Name;
 		TArray<int32> Nodes;
+		FString	      UniqueId; //will be generated in FAsset::GenerateNames
 	};
 
 	struct GLTFCORE_API FMetadata
diff --git a/Engine/Plugins/Editor/GLTFImporter/Source/GLTFCore/Public/GLTFMaterial.h b/Engine/Plugins/Editor/GLTFImporter/Source/GLTFCore/Public/GLTFMaterial.h
--- a/Engine/Plugins/Editor/GLTFImporter/Source/GLTFCore/Public/GLTFMaterial.h
+++ b/Engine/Plugins/Editor/GLTFImporter/Source/GLTFCore/Public/GLTFMaterial.h
@@ -164,6 +164,8 @@ namespace GLTF
 		bool     bHasIOR;
 		bool     bHasSpecular;
 
+		FString UniqueId; //will be generated in FAsset::GenerateNames
+
 		FMaterial(const FString& Name)
 		    : Name(Name)
 		    , BaseColorFactor {1.0f, 1.0f, 1.0f, 1.0f}
diff --git a/Engine/Plugins/Editor/GLTFImporter/Source/GLTFCore/Public/GLTFMesh.h b/Engine/Plugins/Editor/GLTFImporter/Source/GLTFCore/Public/GLTFMesh.h
--- a/Engine/Plugins/Editor/GLTFImporter/Source/GLTFCore/Public/GLTFMesh.h
+++ b/Engine/Plugins/Editor/GLTFImporter/Source/GLTFCore/Public/GLTFMesh.h
@@ -107,6 +107,8 @@ namespace GLTF
 	{
 		FString				Name;
 		TArray<FPrimitive>	Primitives;
+
+		FString				UniqueId; //will be generated in FAsset::GenerateNames
 	
 		bool HasNormals() const;
 		bool HasTangents() const;
diff --git a/Engine/Plugins/Editor/GLTFImporter/Source/GLTFCore/Public/GLTFNode.h b/Engine/Plugins/Editor/GLTFImporter/Source/GLTFCore/Public/GLTFNode.h
--- a/Engine/Plugins/Editor/GLTFImporter/Source/GLTFCore/Public/GLTFNode.h
+++ b/Engine/Plugins/Editor/GLTFImporter/Source/GLTFCore/Public/GLTFNode.h
@@ -40,6 +40,8 @@ namespace GLTF
 		int32 ParentIndex;
 		int32 RootJointIndex; //only valid if node is of Joint Type
 
+		FString UniqueId; //will be generated in FAsset::GenerateNames
+
 		FNode()
 		    : Type(EType::None)
 		    , MeshIndex(INDEX_NONE)
@@ -78,6 +80,7 @@ namespace GLTF
 		float ZNear;
 		float ZFar;
 		bool  bIsPerspective;
+		FString UniqueId; //will be generated in FAsset::GenerateNames
 
 		FCamera(const FNode& Node)
 		    : Node(Node)
@@ -119,6 +122,8 @@ namespace GLTF
 		float Range;
 		FSpot Spot;
 
+		FString UniqueId; //will be generated in FAsset::GenerateNames
+
 		FLight(const FNode* Node)
 		    : Node(Node)
 		    , Color(1.f)
@@ -135,6 +140,8 @@ namespace GLTF
 		TArray<int32>    Joints;    // each is an index into FAsset::Nodes
 		int32            Skeleton;  // root node, index into FAsset::Nodes
 
+		FString          UniqueId; //will be generated in FAsset::GenerateNames
+
 		FSkinInfo(const FAccessor& InverseBindMatrices)
 		    : InverseBindMatrices(InverseBindMatrices)
 		    , Skeleton(INDEX_NONE)
diff --git a/Engine/Plugins/Editor/GLTFImporter/Source/GLTFCore/Public/GLTFTexture.h b/Engine/Plugins/Editor/GLTFImporter/Source/GLTFCore/Public/GLTFTexture.h
--- a/Engine/Plugins/Editor/GLTFImporter/Source/GLTFCore/Public/GLTFTexture.h
+++ b/Engine/Plugins/Editor/GLTFImporter/Source/GLTFCore/Public/GLTFTexture.h
@@ -24,6 +24,8 @@ namespace GLTF
 		uint32       DataByteLength;
 		const uint8* Data;
 
+		FString      UniqueId; //will be generated in FAsset::GenerateNames
+
 		FImage()
 		    : DataByteLength(0)
 		    , Data(nullptr)
@@ -76,6 +78,8 @@ namespace GLTF
 		const FSampler& Sampler;
 		FString         Name;
 
+		FString			UniqueId; //will be generated in FAsset::GenerateNames
+
 		FTexture(const FString& InName, const FImage& InSource, const FSampler& InSampler)
 		    : Source(InSource)
 		    , Sampler(InSampler)
diff --git a/Engine/Plugins/Interchange/Runtime/Source/Import/Private/Gltf/InterchangeGltfTranslator.cpp b/Engine/Plugins/Interchange/Runtime/Source/Import/Private/Gltf/InterchangeGltfTranslator.cpp
--- a/Engine/Plugins/Interchange/Runtime/Source/Import/Private/Gltf/InterchangeGltfTranslator.cpp
+++ b/Engine/Plugins/Interchange/Runtime/Source/Import/Private/Gltf/InterchangeGltfTranslator.cpp
@@ -134,7 +134,7 @@ void UInterchangeGltfTranslator::HandleGltfNode( UInterchangeBaseNodeContainer&
 {
 	using namespace UE::Interchange::Gltf::Private;
 
-	const FString NodeUid = ParentNodeUid + TEXT("\\") + GltfNode.Name;
+	const FString NodeUid = ParentNodeUid + TEXT("\\") + GltfNode.UniqueId;
 
 	const UInterchangeSceneNode* ParentSceneNode = Cast< UInterchangeSceneNode >( NodeContainer.GetNode( ParentNodeUid ) );
 
@@ -173,7 +173,7 @@ void UInterchangeGltfTranslator::HandleGltfNode( UInterchangeBaseNodeContainer&
 			{
 				HandleGltfMesh(NodeContainer, GltfAsset.Meshes[GltfNode.MeshIndex], GltfNode.MeshIndex, UnusedMeshIndices);
 
-				const FString MeshNodeUid = TEXT("\\Mesh\\") + GltfAsset.Meshes[ GltfNode.MeshIndex ].Name;
+				const FString MeshNodeUid = TEXT("\\Mesh\\") + GltfAsset.Meshes[ GltfNode.MeshIndex ].UniqueId;
 				InterchangeSceneNode->SetCustomAssetInstanceUid( MeshNodeUid );
 
 				if (!bHasVariants && GltfAsset.Variants.Num() > 0)
@@ -190,7 +190,7 @@ void UInterchangeGltfTranslator::HandleGltfNode( UInterchangeBaseNodeContainer&
 
 			if ( GltfAsset.Cameras.IsValidIndex( GltfNode.CameraIndex ) )
 			{
-				const FString CameraNodeUid = TEXT("\\Camera\\") + GltfAsset.Cameras[ GltfNode.CameraIndex ].Name;
+				const FString CameraNodeUid = TEXT("\\Camera\\") + GltfAsset.Cameras[ GltfNode.CameraIndex ].UniqueId;
 				InterchangeSceneNode->SetCustomAssetInstanceUid( CameraNodeUid );
 			}
 			break;
@@ -202,7 +202,7 @@ void UInterchangeGltfTranslator::HandleGltfNode( UInterchangeBaseNodeContainer&
 
 			if ( GltfAsset.Lights.IsValidIndex( GltfNode.LightIndex ) )
 			{
-				const FString LightNodeUid = TEXT("\\Light\\") + GltfAsset.Lights[ GltfNode.LightIndex ].Name;
+				const FString LightNodeUid = TEXT("\\Light\\") + GltfAsset.Lights[ GltfNode.LightIndex ].UniqueId;
 				InterchangeSceneNode->SetCustomAssetInstanceUid( LightNodeUid );
 			}
 		}
@@ -271,14 +271,13 @@ void UInterchangeGltfTranslator::HandleGltfMaterialParameter( UInterchangeBaseNo
 		UInterchangeShaderNode* ColorNode = UInterchangeShaderNode::Create( &NodeContainer, ColorNodeName, ShaderNode.GetUniqueID() );
 		ColorNode->SetCustomShaderType( Standard::Nodes::TextureSample::Name.ToString() );
 
-		const FString TextureName = GltfAsset.Textures[ TextureMap.TextureIndex ].Name;
-		const FString TextureUid = UInterchangeTextureNode::MakeNodeUid( TextureName );
+		const FString TextureUid = UInterchangeTextureNode::MakeNodeUid(GltfAsset.Textures[TextureMap.TextureIndex].UniqueId);
 
 		ColorNode->AddStringAttribute( UInterchangeShaderPortsAPI::MakeInputValueKey( Standard::Nodes::TextureSample::Inputs::Texture.ToString() ), TextureUid );
 
 		if (TextureMap.TexCoord > 0 || TextureMap.bHasTextureTransform)
 		{
-			UInterchangeShaderNode* TexCoordNode = UInterchangeShaderNode::Create(&NodeContainer, TEXT("TexCoord"), ShaderNode.GetUniqueID());
+			UInterchangeShaderNode* TexCoordNode = UInterchangeShaderNode::Create(&NodeContainer, MapName + TEXT("\\TexCoord"), ShaderNode.GetUniqueID());
 			TexCoordNode->SetCustomShaderType(Standard::Nodes::TextureCoordinate::Name.ToString());
 
 			TexCoordNode->AddInt32Attribute(UInterchangeShaderPortsAPI::MakeInputValueKey(Standard::Nodes::TextureCoordinate::Inputs::Index.ToString()), TextureMap.TexCoord);
@@ -889,7 +888,8 @@ bool UInterchangeGltfTranslator::Translate( UInterchangeBaseNodeContainer& NodeC
 		int32 TextureIndex = 0;
 		for ( const GLTF::FTexture& GltfTexture : GltfAsset.Textures )
 		{
-			UInterchangeTexture2DNode* TextureNode = UInterchangeTexture2DNode::Create( &NodeContainer, GltfTexture.Name);
+			UInterchangeTexture2DNode* TextureNode = UInterchangeTexture2DNode::Create(&NodeContainer, GltfTexture.UniqueId, GltfTexture.Name);
+
 			TextureNode->SetCustomFilter(ConvertFilter(GltfTexture.Sampler.MinFilter));
 			TextureNode->SetPayLoadKey( LexToString( TextureIndex++ ) );
 
@@ -903,7 +903,7 @@ bool UInterchangeGltfTranslator::Translate( UInterchangeBaseNodeContainer& NodeC
 		int32 MaterialIndex = 0;
 		for ( const GLTF::FMaterial& GltfMaterial : GltfAsset.Materials )
 		{
-			UInterchangeShaderGraphNode* ShaderGraphNode = UInterchangeShaderGraphNode::Create( &NodeContainer, GltfMaterial.Name );
+			UInterchangeShaderGraphNode* ShaderGraphNode = UInterchangeShaderGraphNode::Create(&NodeContainer, GltfMaterial.UniqueId, GltfMaterial.Name);
 
 			HandleGltfMaterial( NodeContainer, GltfMaterial, *ShaderGraphNode );
 			++MaterialIndex;
@@ -926,7 +926,7 @@ bool UInterchangeGltfTranslator::Translate( UInterchangeBaseNodeContainer& NodeC
 		for ( const GLTF::FCamera& GltfCamera : GltfAsset.Cameras )
 		{
 			UInterchangeCameraNode* CameraNode = NewObject< UInterchangeCameraNode >( &NodeContainer );
-			FString CameraNodeUid = TEXT("\\Camera\\") + GltfCamera.Name;
+			FString CameraNodeUid = TEXT("\\Camera\\") + GltfCamera.UniqueId;
 			CameraNode->InitializeNode( CameraNodeUid, GltfCamera.Name, EInterchangeNodeContainerType::TranslatedAsset );
 
 			float       AspectRatio;
@@ -959,7 +959,7 @@ bool UInterchangeGltfTranslator::Translate( UInterchangeBaseNodeContainer& NodeC
 		int32 LightIndex = 0;
 		for ( const GLTF::FLight& GltfLight : GltfAsset.Lights )
 		{
-			FString LightNodeUid = TEXT("\\Light\\") + GltfLight.Name;
+			FString LightNodeUid = TEXT("\\Light\\") + GltfLight.UniqueId;
 
 			switch (GltfLight.Type)
 			{
@@ -1023,12 +1023,8 @@ bool UInterchangeGltfTranslator::Translate( UInterchangeBaseNodeContainer& NodeC
 			UInterchangeSceneNode* SceneNode = NewObject< UInterchangeSceneNode >( &NodeContainer );
 
 			FString SceneName = GltfScene.Name;
-			if (SceneName.IsEmpty())
-			{
-				SceneName = FileName + TEXT("_scene_") + LexToString(SceneIndex++);
-			}
 
-			FString SceneNodeUid = TEXT("\\Scene\\") + SceneName;
+			FString SceneNodeUid = TEXT("\\Scene\\") + GltfScene.UniqueId;
 			SceneNode->InitializeNode( SceneNodeUid, SceneName, EInterchangeNodeContainerType::TranslatedScene );
 			NodeContainer.AddNode( SceneNode );
 
@@ -1366,9 +1362,8 @@ void UInterchangeGltfTranslator::HandleGltfAnimation(UInterchangeBaseNodeContain
 
 	UInterchangeAnimationTrackSetNode* TrackSetNode = NewObject< UInterchangeAnimationTrackSetNode >(&NodeContainer);
 
-	const FString AnimTrackSetName = GltfAnimation.Name;
-	const FString AnimTrackSetNodeUid = TEXT("\\Animation\\") + AnimTrackSetName;
-	TrackSetNode->InitializeNode(AnimTrackSetNodeUid, AnimTrackSetName, EInterchangeNodeContainerType::TranslatedAsset);
+	const FString AnimTrackSetNodeUid = TEXT("\\Animation\\") + GltfAnimation.UniqueId;
+	TrackSetNode->InitializeNode(AnimTrackSetNodeUid, GltfAnimation.Name, EInterchangeNodeContainerType::TranslatedAsset);
 
 	for (const TTuple<const GLTF::FNode*, TArray<int32>>& NodeChannelsEntry : NodeChannelsMap)
 	{
@@ -1444,8 +1439,7 @@ void UInterchangeGltfTranslator::SetTextureSRGB(UInterchangeBaseNodeContainer& N
 {
 	if (GltfAsset.Textures.IsValidIndex(TextureMap.TextureIndex))
 	{
-		const FString TextureName = GltfAsset.Textures[TextureMap.TextureIndex].Name;
-		const FString TextureUid = UInterchangeTextureNode::MakeNodeUid(TextureName);
+		const FString TextureUid = UInterchangeTextureNode::MakeNodeUid(GltfAsset.Textures[TextureMap.TextureIndex].UniqueId);
 		if (UInterchangeTextureNode* TextureNode = const_cast<UInterchangeTextureNode*>(Cast<UInterchangeTextureNode>(NodeContainer.GetNode(TextureUid))))
 		{
 			TextureNode->SetCustomSRGB(true);
@@ -1456,8 +1450,7 @@ void UInterchangeGltfTranslator::SetTextureFlipGreenChannel(UInterchangeBaseNode
 {
 	if (GltfAsset.Textures.IsValidIndex(TextureMap.TextureIndex))
 	{
-		const FString TextureName = GltfAsset.Textures[TextureMap.TextureIndex].Name;
-		const FString TextureUid = UInterchangeTextureNode::MakeNodeUid(TextureName);
+		const FString TextureUid = UInterchangeTextureNode::MakeNodeUid(GltfAsset.Textures[TextureMap.TextureIndex].UniqueId);
 		if (UInterchangeTextureNode* TextureNode = const_cast<UInterchangeTextureNode*>(Cast<UInterchangeTextureNode>(NodeContainer.GetNode(TextureUid))))
 		{
 			TextureNode->SetCustombFlipGreenChannel(true);
@@ -1558,7 +1551,7 @@ void UInterchangeGltfTranslator::HandleGltfVariants(UInterchangeBaseNodeContaine
 								}
 
 								const GLTF::FMaterial& GltfMaterial = Materials[VariantMapping.MaterialIndex];
-								const FString MaterialUid = UInterchangeShaderGraphNode::MakeNodeUid(GltfMaterial.Name);
+								const FString MaterialUid = UInterchangeShaderGraphNode::MakeNodeUid(GltfMaterial.UniqueId);
 
 								VariantSetNode->AddCustomDependencyUid(MaterialUid);
 
@@ -1639,7 +1632,7 @@ bool UInterchangeGltfTranslator::GetVariantSetPayloadData(UE::Interchange::FVari
 						}
 
 						const GLTF::FMaterial& GltfMaterial = Materials[VariantMapping.MaterialIndex];
-						const FString MaterialUid = UInterchangeShaderGraphNode::MakeNodeUid(GltfMaterial.Name);
+						const FString MaterialUid = UInterchangeShaderGraphNode::MakeNodeUid(GltfMaterial.UniqueId);
 
 						for (int32 VariantIndex : VariantMapping.VariantIndices)
 						{
@@ -1752,8 +1745,9 @@ void UInterchangeGltfTranslator::HandleGltfSkeletons(UInterchangeBaseNodeContain
 
 			//Skeletal Mesh's naming policy: (Mesh.Name)_(RootJointNode.Name) naming policy:
 			FString SkeletalName = GltfAsset.Meshes[MeshIndex].Name + TEXT("_") + GltfAsset.Nodes[RootJointIndex].Name;
+			FString SkeletalId = GltfAsset.Meshes[MeshIndex].UniqueId + TEXT("_") + GltfAsset.Nodes[RootJointIndex].UniqueId;
 
-			UInterchangeMeshNode* SkeletalMeshNode = HandleGltfMesh(NodeContainer, GltfAsset.Meshes[MeshIndex], MeshIndex, UnusedMeshIndices, SkeletalName);
+			UInterchangeMeshNode* SkeletalMeshNode = HandleGltfMesh(NodeContainer, GltfAsset.Meshes[MeshIndex], MeshIndex, UnusedMeshIndices, SkeletalName, SkeletalId);
 
 			SkeletalMeshNode->SetSkinnedMesh(true);
 
@@ -1804,7 +1798,8 @@ void UInterchangeGltfTranslator::HandleGltfSkeletons(UInterchangeBaseNodeContain
 UInterchangeMeshNode* UInterchangeGltfTranslator::HandleGltfMesh(UInterchangeBaseNodeContainer& NodeContainer, 
 	const GLTF::FMesh& GltfMesh, int MeshIndex, 
 	TSet<int>& UnusedMeshIndices, 
-	const FString& SkeletalName/*If set it creates the mesh even if it was already created (for Skeletals)*/) const
+	const FString& SkeletalName/*If set it creates the mesh even if it was already created (for Skeletals)*/,
+	const FString& SkeletalId) const
 {
 	if (!UnusedMeshIndices.Contains(MeshIndex) && SkeletalName.Len() == 0)
 	{
@@ -1814,7 +1809,7 @@ UInterchangeMeshNode* UInterchangeGltfTranslator::HandleGltfMesh(UInterchangeBas
 	UnusedMeshIndices.Remove(MeshIndex);
 
 	UInterchangeMeshNode* MeshNode = NewObject< UInterchangeMeshNode >(&NodeContainer);
-	FString MeshNodeUid = TEXT("\\Mesh\\") + MeshName;
+	FString MeshNodeUid = TEXT("\\Mesh\\") + (SkeletalId.Len() ? SkeletalId : GltfMesh.UniqueId);
 
 	MeshNode->InitializeNode(MeshNodeUid, MeshName, EInterchangeNodeContainerType::TranslatedAsset);
 	MeshNode->SetPayLoadKey(LexToString(MeshIndex));
@@ -1829,7 +1824,7 @@ UInterchangeMeshNode* UInterchangeGltfTranslator::HandleGltfMesh(UInterchangeBas
 		if (GltfAsset.Materials.IsValidIndex(Primitive.MaterialIndex))
 		{
 			const FString MaterialName = GltfAsset.Materials[Primitive.MaterialIndex].Name;
-			const FString ShaderGraphNodeUid = UInterchangeShaderGraphNode::MakeNodeUid(MaterialName);
+			const FString ShaderGraphNodeUid = UInterchangeShaderGraphNode::MakeNodeUid(GltfAsset.Materials[Primitive.MaterialIndex].UniqueId);
 			MeshNode->SetSlotMaterialDependencyUid(MaterialName, ShaderGraphNodeUid);
 		}
 	}
diff --git a/Engine/Plugins/Interchange/Runtime/Source/Import/Public/Gltf/InterchangeGltfTranslator.h b/Engine/Plugins/Interchange/Runtime/Source/Import/Public/Gltf/InterchangeGltfTranslator.h
--- a/Engine/Plugins/Interchange/Runtime/Source/Import/Public/Gltf/InterchangeGltfTranslator.h
+++ b/Engine/Plugins/Interchange/Runtime/Source/Import/Public/Gltf/InterchangeGltfTranslator.h
@@ -80,7 +80,7 @@ protected:
 	void HandleGltfAnimation(UInterchangeBaseNodeContainer& NodeContainer, int32 AnimationIndex) const;
 	void HandleGltfVariants(UInterchangeBaseNodeContainer& NodeContainer, const FString& FileName) const;
 	UInterchangeMeshNode* HandleGltfMesh(UInterchangeBaseNodeContainer& NodeContainer, const GLTF::FMesh& GltfMesh, int MeshIndex,
-		TSet<int>& UnusedMeshIndices, const FString& SkeletalName = "" /*If set it creates the mesh even if it was already created (for Skeletals)*/) const;
+		TSet<int>& UnusedMeshIndices, const FString& SkeletalName = "" /*If set it creates the mesh even if it was already created (for Skeletals)*/, const FString& SkeletalId = "") const;
 
 	/** Support for KHR_materials_clearcoat */
 	void HandleGltfClearCoat( UInterchangeBaseNodeContainer& NodeContainer, const GLTF::FMaterial& GltfMaterial, UInterchangeShaderGraphNode& ShaderGraphNode, const bool bSwapNormalAndClearCoatNormal ) const;
diff --git a/Engine/Plugins/Interchange/Runtime/Source/Nodes/Private/InterchangeShaderGraphNode.cpp b/Engine/Plugins/Interchange/Runtime/Source/Nodes/Private/InterchangeShaderGraphNode.cpp
--- a/Engine/Plugins/Interchange/Runtime/Source/Nodes/Private/InterchangeShaderGraphNode.cpp
+++ b/Engine/Plugins/Interchange/Runtime/Source/Nodes/Private/InterchangeShaderGraphNode.cpp
@@ -167,14 +167,15 @@ FString UInterchangeShaderGraphNode::MakeNodeUid(const FStringView NodeName)
 	return FString(UInterchangeBaseNode::HierarchySeparator) + TEXT("Shaders") + FString(UInterchangeBaseNode::HierarchySeparator) + NodeName;
 }
 
-UInterchangeShaderGraphNode* UInterchangeShaderGraphNode::Create(UInterchangeBaseNodeContainer* NodeContainer, const FStringView NodeName)
+UInterchangeShaderGraphNode* UInterchangeShaderGraphNode::Create(UInterchangeBaseNodeContainer* NodeContainer, const FStringView ShaderGraphNodeName /*Or ShaderGraphNodeId*/, const FStringView ShaderGraphNodeDisplayLabel)
 {
 	check(NodeContainer);
-
-	const FString ShaderGraphNodeUid = MakeNodeUid(NodeName);
+	
+	const FString NodeUid = MakeNodeUid(ShaderGraphNodeName);
+	const FStringView NodeDisplayLabel = ShaderGraphNodeDisplayLabel.IsEmpty() ? ShaderGraphNodeName : ShaderGraphNodeDisplayLabel;
 
 	UInterchangeShaderGraphNode* ShaderGraphNode = NewObject< UInterchangeShaderGraphNode >(NodeContainer);
-	ShaderGraphNode->InitializeNode(ShaderGraphNodeUid, FString(NodeName), EInterchangeNodeContainerType::TranslatedAsset);
+	ShaderGraphNode->InitializeNode(NodeUid, FString(NodeDisplayLabel), EInterchangeNodeContainerType::TranslatedAsset);
 
 	NodeContainer->AddNode(ShaderGraphNode);
 
diff --git a/Engine/Plugins/Interchange/Runtime/Source/Nodes/Private/InterchangeTexture2DNode.cpp b/Engine/Plugins/Interchange/Runtime/Source/Nodes/Private/InterchangeTexture2DNode.cpp
--- a/Engine/Plugins/Interchange/Runtime/Source/Nodes/Private/InterchangeTexture2DNode.cpp
+++ b/Engine/Plugins/Interchange/Runtime/Source/Nodes/Private/InterchangeTexture2DNode.cpp
@@ -10,12 +10,14 @@ FString UInterchangeTexture2DNode::MakeNodeUid(const FStringView NodeName)
 	return UInterchangeTextureNode::MakeNodeUid(NodeName);
 }
 
- UInterchangeTexture2DNode* UInterchangeTexture2DNode::Create(UInterchangeBaseNodeContainer* NodeContainer, const FStringView TextureNodeName)
+ UInterchangeTexture2DNode* UInterchangeTexture2DNode::Create(UInterchangeBaseNodeContainer* NodeContainer, const FStringView TextureNodeName /*Or TextureNodeId*/, const FStringView TextureNodeDisplayLabel)
 {
 	UInterchangeTexture2DNode* TextureNode = NewObject< UInterchangeTexture2DNode >(NodeContainer);
-	const FString TextureNodeUid = UInterchangeTextureNode::MakeNodeUid(TextureNodeName);
 
-	TextureNode->InitializeNode(TextureNodeUid, FString(TextureNodeName), EInterchangeNodeContainerType::TranslatedAsset);
+	const FString NodeUid = MakeNodeUid(TextureNodeName);
+	const FStringView NodeDisplayLabel = TextureNodeDisplayLabel.IsEmpty() ? TextureNodeName : TextureNodeDisplayLabel;
+
+	TextureNode->InitializeNode(NodeUid, FString(NodeDisplayLabel), EInterchangeNodeContainerType::TranslatedAsset);
 	NodeContainer->AddNode(TextureNode);
 
 	return TextureNode;
diff --git a/Engine/Plugins/Interchange/Runtime/Source/Nodes/Public/InterchangeShaderGraphNode.h b/Engine/Plugins/Interchange/Runtime/Source/Nodes/Public/InterchangeShaderGraphNode.h
--- a/Engine/Plugins/Interchange/Runtime/Source/Nodes/Public/InterchangeShaderGraphNode.h
+++ b/Engine/Plugins/Interchange/Runtime/Source/Nodes/Public/InterchangeShaderGraphNode.h
@@ -169,8 +169,9 @@ public:
 
 	/**
 	 * Creates a new UInterchangeShaderGraphNode and adds it to NodeContainer as a translated node.
+	 * In case the ShaderGraphNodeDisplayLabel is not empty, ShaderGraphNodeName supposed to hold the ShaderGraphNodeId
 	 */
-	static UInterchangeShaderGraphNode* Create(UInterchangeBaseNodeContainer* NodeContainer, const FStringView NodeName);
+	static UInterchangeShaderGraphNode* Create(UInterchangeBaseNodeContainer* NodeContainer, const FStringView ShaderGraphNodeName /*Or ShaderGraphNodeId*/, const FStringView ShaderGraphNodeDisplayLabel = FStringView());
 
 	/**
 	 * Return the node type name of the class, we use this when reporting error
diff --git a/Engine/Plugins/Interchange/Runtime/Source/Nodes/Public/InterchangeTexture2DNode.h b/Engine/Plugins/Interchange/Runtime/Source/Nodes/Public/InterchangeTexture2DNode.h
--- a/Engine/Plugins/Interchange/Runtime/Source/Nodes/Public/InterchangeTexture2DNode.h
+++ b/Engine/Plugins/Interchange/Runtime/Source/Nodes/Public/InterchangeTexture2DNode.h
@@ -36,8 +36,9 @@ public:
 
 	/**
 	 * Creates a new UInterchangeTexture2DNode and adds it to NodeContainer as a translated node.
+	 * In case the TextureNodeDisplayLabel is not empty, TextureNodeName supposed to hold the TextureNodeId
 	 */
-	static UInterchangeTexture2DNode* Create(UInterchangeBaseNodeContainer* NodeContainer, const FStringView TextureNodeName);
+	static UInterchangeTexture2DNode* Create(UInterchangeBaseNodeContainer* NodeContainer, const FStringView TextureNodeName /*Or TextureNodeId*/, const FStringView TextureNodeDisplayLabel = FStringView());
 
 	virtual void PostInitProperties()
 	{
diff --git a/Engine/Source/Runtime/Interchange/Engine/Private/Tasks/InterchangeTaskParsing.cpp b/Engine/Source/Runtime/Interchange/Engine/Private/Tasks/InterchangeTaskParsing.cpp
--- a/Engine/Source/Runtime/Interchange/Engine/Private/Tasks/InterchangeTaskParsing.cpp
+++ b/Engine/Source/Runtime/Interchange/Engine/Private/Tasks/InterchangeTaskParsing.cpp
@@ -277,6 +277,11 @@ void UE::Interchange::FTaskParsing::DoTask(ENamedThreads::Type CurrentThread, co
 					RenameInfo.NewName = NewName;
 					RenameInfo.SourceIndex = SourceIndex;
 					FactoryNode->SetDisplayLabel(FactoryNode->GetDisplayLabel() + NameIndexString);
+					if (FactoryNode->HasAttribute(UE::Interchange::FBaseNodeStaticData::AssetNameKey()))
+					{
+						//TextureFactorNodes automatically set up the AssetNames which then are used for asset creation:
+						FactoryNode->SetAssetName(FactoryNode->GetDisplayLabel());
+					}
 				}
 				else
 				{
