Undo //UE5/Main/Engine/... changelist 23252118

diff --git a/Engine/Plugins/Editor/GLTFImporter/Source/GLTFCore/Private/GLTF/GLTFAsset.cpp b/Engine/Plugins/Editor/GLTFImporter/Source/GLTFCore/Private/GLTF/GLTFAsset.cpp
--- a/Engine/Plugins/Editor/GLTFImporter/Source/GLTFCore/Private/GLTF/GLTFAsset.cpp
+++ b/Engine/Plugins/Editor/GLTFImporter/Source/GLTFCore/Private/GLTF/GLTFAsset.cpp
@@ -15,11 +15,9 @@ namespace GLTF
 			{
 				T& Obj = Objects[Index];
 
-				Obj.UniqueId = Prefix + FString::FromInt(Index);
-
 				if (Obj.Name.IsEmpty())
 				{
-					Obj.Name = Obj.UniqueId;
+					Obj.Name = Prefix + FString::FromInt(Index);
 				}
 			}
 		}
@@ -73,22 +71,6 @@ namespace GLTF
 	{
 		check(!Prefix.IsEmpty());
 
-		{
-			const FString ScenePrefix = Prefix + TEXT("_scene_");
-
-			for (size_t SceneIndex = 0; SceneIndex < Scenes.Num(); SceneIndex++)
-			{
-				FScene& Scene = Scenes[SceneIndex];
-
-				Scene.UniqueId = ScenePrefix + FString::FromInt(SceneIndex);
-
-				if (Scene.Name.IsEmpty())
-				{
-					Scene.Name = Scene.UniqueId;
-				}
-			}
-		}
-
 		{
 			const FString NodePrefix = Prefix + TEXT("_node_");
 			const FString JointPrefix = Prefix + TEXT("_joint_");
@@ -96,14 +78,13 @@ namespace GLTF
 			{
 				FNode& Node = Nodes[NodeIndex];
 
-				bool bIsJoint = Node.Type == FNode::EType::Joint;
-				Node.UniqueId = (bIsJoint ? JointPrefix : NodePrefix);
-
-				Node.UniqueId = Node.UniqueId + FString::FromInt(NodeIndex);
-
 				if (Node.Name.IsEmpty())
 				{
-					Node.Name = Node.UniqueId;
+					bool bIsJoint = Node.Type == FNode::EType::Joint;
+					Node.Name     = (bIsJoint ? JointPrefix : NodePrefix);
+
+					// Make sure node names are unique, in case Name is not empty it is presumed to be unique / handled by the processing pipelines.
+					Node.Name = Node.Name + FString::FromInt(NodeIndex);
 				}
 			}
 		}
@@ -114,8 +95,6 @@ namespace GLTF
 			{
 				FTexture& Tex = Textures[TextureIndex];
 
-				Tex.UniqueId = TexPrefix + FString::FromInt(TextureIndex);
-
 				if (!Tex.Name.IsEmpty())
 				{
 					// keep base part of the name as it was in gltf
@@ -126,19 +105,14 @@ namespace GLTF
 				}
 				else if (!Tex.Source.URI.IsEmpty())
 				{
-					//in case texture is embedded:
-					if (Tex.Source.URI.StartsWith(TEXT("data:")))
-					{
-						Tex.Name = Tex.UniqueId;
-					}
-					else
-					{
-						Tex.Name = FPaths::GetBaseFilename(Tex.Source.URI);
-					}
+					Tex.Name = FPaths::GetBaseFilename(Tex.Source.URI);
 				}
 				else
 				{
-					Tex.Name = Tex.UniqueId;
+					// GLTF texture name has decorative purpose, not guaranteed to be unique
+					// only its index is unique. Same with glTF image or its source file's basename
+					// So always include texture index into texture's name to increase probability that names are unique
+					Tex.Name = TexPrefix + FString::FromInt(TextureIndex);
 				}
 
 			}
@@ -150,8 +124,6 @@ namespace GLTF
 			{
 				FCamera& Camera = Cameras[CameraIndex];
 
-				Camera.UniqueId = CameraPrefix + FString::FromInt(CameraIndex);
-
 				if (Camera.Name.IsEmpty())
 				{
 					if (!Camera.Node.Name.IsEmpty())
@@ -160,7 +132,7 @@ namespace GLTF
 					}
 					else
 					{
-						Camera.Name = Camera.UniqueId;
+						Camera.Name = CameraPrefix + FString::FromInt(CameraIndex);
 					}
 				}
 			}
@@ -172,8 +144,6 @@ namespace GLTF
 			{
 				FLight& Light = Lights[LightIndex];
 
-				Light.UniqueId = LightPrefix + FString::FromInt(LightIndex);
-
 				if (Light.Name.IsEmpty())
 				{
 					if (Light.Node && !Light.Node->Name.IsEmpty())
@@ -182,17 +152,18 @@ namespace GLTF
 					}
 					else
 					{
-						Light.Name = Light.UniqueId;
+						Light.Name = LightPrefix + FString::FromInt(LightIndex);
 					}
 				}
 			}
 		}
 
-		// General pattern here is that we'll have a Prefix that will be a filename (like MyFile). 
-		// The name (display label) and the UniqueId of an object has been separated.
-		// UniqueId will consist of: [Prefil]+[_ElementType_]+[ElementIndex]
-		// Name will be the provided GLTF object.name, if present,
-		// if Name is not present the UniqueId will be used as the Name.
+		// General pattern here is that we'll have a Prefix that will be a filename (like MyFile). If the
+		// GLTF object has an actual name, it will end up like <index>_<objectname>, like "0_shoe" or "3_spotlight".
+		// If it doesn't have a name, it will end up with a name like <index>_<filename>_<objecttype>, e.g.
+		// "0_MyFile_material" or "3_myfile_mesh".
+		// We want to do images after textures as the textures may want to reuse source image names, but we don't
+		// want to use those if we generated them ourselves.
 		GLTF::GenerateNames(Prefix + TEXT("_material_"), Materials);
 		GLTF::GenerateNames(Prefix + TEXT("_skin_"), Skins);
 		GLTF::GenerateNames(Prefix + TEXT("_animation_"), Animations);
diff --git a/Engine/Plugins/Editor/GLTFImporter/Source/GLTFCore/Public/GLTFAnimation.h b/Engine/Plugins/Editor/GLTFImporter/Source/GLTFCore/Public/GLTFAnimation.h
--- a/Engine/Plugins/Editor/GLTFImporter/Source/GLTFCore/Public/GLTFAnimation.h
+++ b/Engine/Plugins/Editor/GLTFImporter/Source/GLTFCore/Public/GLTFAnimation.h
@@ -69,8 +69,6 @@ namespace GLTF
 		FString          Name;
 		TArray<FSampler> Samplers;
 		TArray<FChannel> Channels;
-
-		FString          UniqueId; //will be generated in FAsset::GenerateNames
 	};
 
 }  // namespace GLTF
diff --git a/Engine/Plugins/Editor/GLTFImporter/Source/GLTFCore/Public/GLTFAsset.h b/Engine/Plugins/Editor/GLTFImporter/Source/GLTFCore/Public/GLTFAsset.h
--- a/Engine/Plugins/Editor/GLTFImporter/Source/GLTFCore/Public/GLTFAsset.h
+++ b/Engine/Plugins/Editor/GLTFImporter/Source/GLTFCore/Public/GLTFAsset.h
@@ -37,7 +37,6 @@ namespace GLTF
 	{
 		FString       Name;
 		TArray<int32> Nodes;
-		FString	      UniqueId; //will be generated in FAsset::GenerateNames
 	};
 
 	struct GLTFCORE_API FMetadata
diff --git a/Engine/Plugins/Editor/GLTFImporter/Source/GLTFCore/Public/GLTFMaterial.h b/Engine/Plugins/Editor/GLTFImporter/Source/GLTFCore/Public/GLTFMaterial.h
--- a/Engine/Plugins/Editor/GLTFImporter/Source/GLTFCore/Public/GLTFMaterial.h
+++ b/Engine/Plugins/Editor/GLTFImporter/Source/GLTFCore/Public/GLTFMaterial.h
@@ -164,8 +164,6 @@ namespace GLTF
 		bool     bHasIOR;
 		bool     bHasSpecular;
 
-		FString UniqueId; //will be generated in FAsset::GenerateNames
-
 		FMaterial(const FString& Name)
 		    : Name(Name)
 		    , BaseColorFactor {1.0f, 1.0f, 1.0f, 1.0f}
diff --git a/Engine/Plugins/Editor/GLTFImporter/Source/GLTFCore/Public/GLTFMesh.h b/Engine/Plugins/Editor/GLTFImporter/Source/GLTFCore/Public/GLTFMesh.h
--- a/Engine/Plugins/Editor/GLTFImporter/Source/GLTFCore/Public/GLTFMesh.h
+++ b/Engine/Plugins/Editor/GLTFImporter/Source/GLTFCore/Public/GLTFMesh.h
@@ -107,8 +107,6 @@ namespace GLTF
 	{
 		FString				Name;
 		TArray<FPrimitive>	Primitives;
-
-		FString				UniqueId; //will be generated in FAsset::GenerateNames
 	
 		bool HasNormals() const;
 		bool HasTangents() const;
diff --git a/Engine/Plugins/Editor/GLTFImporter/Source/GLTFCore/Public/GLTFNode.h b/Engine/Plugins/Editor/GLTFImporter/Source/GLTFCore/Public/GLTFNode.h
--- a/Engine/Plugins/Editor/GLTFImporter/Source/GLTFCore/Public/GLTFNode.h
+++ b/Engine/Plugins/Editor/GLTFImporter/Source/GLTFCore/Public/GLTFNode.h
@@ -40,8 +40,6 @@ namespace GLTF
 		int32 ParentIndex;
 		int32 RootJointIndex; //only valid if node is of Joint Type
 
-		FString UniqueId; //will be generated in FAsset::GenerateNames
-
 		FNode()
 		    : Type(EType::None)
 		    , MeshIndex(INDEX_NONE)
@@ -80,7 +78,6 @@ namespace GLTF
 		float ZNear;
 		float ZFar;
 		bool  bIsPerspective;
-		FString UniqueId; //will be generated in FAsset::GenerateNames
 
 		FCamera(const FNode& Node)
 		    : Node(Node)
@@ -122,8 +119,6 @@ namespace GLTF
 		float Range;
 		FSpot Spot;
 
-		FString UniqueId; //will be generated in FAsset::GenerateNames
-
 		FLight(const FNode* Node)
 		    : Node(Node)
 		    , Color(1.f)
@@ -140,8 +135,6 @@ namespace GLTF
 		TArray<int32>    Joints;    // each is an index into FAsset::Nodes
 		int32            Skeleton;  // root node, index into FAsset::Nodes
 
-		FString          UniqueId; //will be generated in FAsset::GenerateNames
-
 		FSkinInfo(const FAccessor& InverseBindMatrices)
 		    : InverseBindMatrices(InverseBindMatrices)
 		    , Skeleton(INDEX_NONE)
diff --git a/Engine/Plugins/Editor/GLTFImporter/Source/GLTFCore/Public/GLTFTexture.h b/Engine/Plugins/Editor/GLTFImporter/Source/GLTFCore/Public/GLTFTexture.h
--- a/Engine/Plugins/Editor/GLTFImporter/Source/GLTFCore/Public/GLTFTexture.h
+++ b/Engine/Plugins/Editor/GLTFImporter/Source/GLTFCore/Public/GLTFTexture.h
@@ -24,8 +24,6 @@ namespace GLTF
 		uint32       DataByteLength;
 		const uint8* Data;
 
-		FString      UniqueId; //will be generated in FAsset::GenerateNames
-
 		FImage()
 		    : DataByteLength(0)
 		    , Data(nullptr)
@@ -78,8 +76,6 @@ namespace GLTF
 		const FSampler& Sampler;
 		FString         Name;
 
-		FString			UniqueId; //will be generated in FAsset::GenerateNames
-
 		FTexture(const FString& InName, const FImage& InSource, const FSampler& InSampler)
 		    : Source(InSource)
 		    , Sampler(InSampler)
diff --git a/Engine/Plugins/Interchange/Runtime/Source/Import/Private/Gltf/InterchangeGltfTranslator.cpp b/Engine/Plugins/Interchange/Runtime/Source/Import/Private/Gltf/InterchangeGltfTranslator.cpp
--- a/Engine/Plugins/Interchange/Runtime/Source/Import/Private/Gltf/InterchangeGltfTranslator.cpp
+++ b/Engine/Plugins/Interchange/Runtime/Source/Import/Private/Gltf/InterchangeGltfTranslator.cpp
@@ -134,7 +134,7 @@ void UInterchangeGltfTranslator::HandleGltfNode( UInterchangeBaseNodeContainer&
 {
 	using namespace UE::Interchange::Gltf::Private;
 
-	const FString NodeUid = ParentNodeUid + TEXT("\\") + GltfNode.UniqueId;
+	const FString NodeUid = ParentNodeUid + TEXT("\\") + GltfNode.Name;
 
 	const UInterchangeSceneNode* ParentSceneNode = Cast< UInterchangeSceneNode >( NodeContainer.GetNode( ParentNodeUid ) );
 
@@ -173,7 +173,7 @@ void UInterchangeGltfTranslator::HandleGltfNode( UInterchangeBaseNodeContainer&
 			{
 				HandleGltfMesh(NodeContainer, GltfAsset.Meshes[GltfNode.MeshIndex], GltfNode.MeshIndex, UnusedMeshIndices);
 
-				const FString MeshNodeUid = TEXT("\\Mesh\\") + GltfAsset.Meshes[ GltfNode.MeshIndex ].UniqueId;
+				const FString MeshNodeUid = TEXT("\\Mesh\\") + GltfAsset.Meshes[ GltfNode.MeshIndex ].Name;
 				InterchangeSceneNode->SetCustomAssetInstanceUid( MeshNodeUid );
 
 				if (!bHasVariants && GltfAsset.Variants.Num() > 0)
@@ -190,7 +190,7 @@ void UInterchangeGltfTranslator::HandleGltfNode( UInterchangeBaseNodeContainer&
 
 			if ( GltfAsset.Cameras.IsValidIndex( GltfNode.CameraIndex ) )
 			{
-				const FString CameraNodeUid = TEXT("\\Camera\\") + GltfAsset.Cameras[ GltfNode.CameraIndex ].UniqueId;
+				const FString CameraNodeUid = TEXT("\\Camera\\") + GltfAsset.Cameras[ GltfNode.CameraIndex ].Name;
 				InterchangeSceneNode->SetCustomAssetInstanceUid( CameraNodeUid );
 			}
 			break;
@@ -202,7 +202,7 @@ void UInterchangeGltfTranslator::HandleGltfNode( UInterchangeBaseNodeContainer&
 
 			if ( GltfAsset.Lights.IsValidIndex( GltfNode.LightIndex ) )
 			{
-				const FString LightNodeUid = TEXT("\\Light\\") + GltfAsset.Lights[ GltfNode.LightIndex ].UniqueId;
+				const FString LightNodeUid = TEXT("\\Light\\") + GltfAsset.Lights[ GltfNode.LightIndex ].Name;
 				InterchangeSceneNode->SetCustomAssetInstanceUid( LightNodeUid );
 			}
 		}
@@ -271,13 +271,14 @@ void UInterchangeGltfTranslator::HandleGltfMaterialParameter( UInterchangeBaseNo
 		UInterchangeShaderNode* ColorNode = UInterchangeShaderNode::Create( &NodeContainer, ColorNodeName, ShaderNode.GetUniqueID() );
 		ColorNode->SetCustomShaderType( Standard::Nodes::TextureSample::Name.ToString() );
 
-		const FString TextureUid = UInterchangeTextureNode::MakeNodeUid(GltfAsset.Textures[TextureMap.TextureIndex].UniqueId);
+		const FString TextureName = GltfAsset.Textures[ TextureMap.TextureIndex ].Name;
+		const FString TextureUid = UInterchangeTextureNode::MakeNodeUid( TextureName );
 
 		ColorNode->AddStringAttribute( UInterchangeShaderPortsAPI::MakeInputValueKey( Standard::Nodes::TextureSample::Inputs::Texture.ToString() ), TextureUid );
 
 		if (TextureMap.TexCoord > 0 || TextureMap.bHasTextureTransform)
 		{
-			UInterchangeShaderNode* TexCoordNode = UInterchangeShaderNode::Create(&NodeContainer, MapName + TEXT("\\TexCoord"), ShaderNode.GetUniqueID());
+			UInterchangeShaderNode* TexCoordNode = UInterchangeShaderNode::Create(&NodeContainer, TEXT("TexCoord"), ShaderNode.GetUniqueID());
 			TexCoordNode->SetCustomShaderType(Standard::Nodes::TextureCoordinate::Name.ToString());
 
 			TexCoordNode->AddInt32Attribute(UInterchangeShaderPortsAPI::MakeInputValueKey(Standard::Nodes::TextureCoordinate::Inputs::Index.ToString()), TextureMap.TexCoord);
@@ -888,8 +889,7 @@ bool UInterchangeGltfTranslator::Translate( UInterchangeBaseNodeContainer& NodeC
 		int32 TextureIndex = 0;
 		for ( const GLTF::FTexture& GltfTexture : GltfAsset.Textures )
 		{
-			UInterchangeTexture2DNode* TextureNode = UInterchangeTexture2DNode::Create(&NodeContainer, GltfTexture.UniqueId, GltfTexture.Name);
-
+			UInterchangeTexture2DNode* TextureNode = UInterchangeTexture2DNode::Create( &NodeContainer, GltfTexture.Name);
 			TextureNode->SetCustomFilter(ConvertFilter(GltfTexture.Sampler.MinFilter));
 			TextureNode->SetPayLoadKey( LexToString( TextureIndex++ ) );
 
@@ -903,7 +903,7 @@ bool UInterchangeGltfTranslator::Translate( UInterchangeBaseNodeContainer& NodeC
 		int32 MaterialIndex = 0;
 		for ( const GLTF::FMaterial& GltfMaterial : GltfAsset.Materials )
 		{
-			UInterchangeShaderGraphNode* ShaderGraphNode = UInterchangeShaderGraphNode::Create(&NodeContainer, GltfMaterial.UniqueId, GltfMaterial.Name);
+			UInterchangeShaderGraphNode* ShaderGraphNode = UInterchangeShaderGraphNode::Create( &NodeContainer, GltfMaterial.Name );
 
 			HandleGltfMaterial( NodeContainer, GltfMaterial, *ShaderGraphNode );
 			++MaterialIndex;
@@ -926,7 +926,7 @@ bool UInterchangeGltfTranslator::Translate( UInterchangeBaseNodeContainer& NodeC
 		for ( const GLTF::FCamera& GltfCamera : GltfAsset.Cameras )
 		{
 			UInterchangeCameraNode* CameraNode = NewObject< UInterchangeCameraNode >( &NodeContainer );
-			FString CameraNodeUid = TEXT("\\Camera\\") + GltfCamera.UniqueId;
+			FString CameraNodeUid = TEXT("\\Camera\\") + GltfCamera.Name;
 			CameraNode->InitializeNode( CameraNodeUid, GltfCamera.Name, EInterchangeNodeContainerType::TranslatedAsset );
 
 			float       AspectRatio;
@@ -959,7 +959,7 @@ bool UInterchangeGltfTranslator::Translate( UInterchangeBaseNodeContainer& NodeC
 		int32 LightIndex = 0;
 		for ( const GLTF::FLight& GltfLight : GltfAsset.Lights )
 		{
-			FString LightNodeUid = TEXT("\\Light\\") + GltfLight.UniqueId;
+			FString LightNodeUid = TEXT("\\Light\\") + GltfLight.Name;
 
 			switch (GltfLight.Type)
 			{
@@ -1023,8 +1023,12 @@ bool UInterchangeGltfTranslator::Translate( UInterchangeBaseNodeContainer& NodeC
 			UInterchangeSceneNode* SceneNode = NewObject< UInterchangeSceneNode >( &NodeContainer );
 
 			FString SceneName = GltfScene.Name;
+			if (SceneName.IsEmpty())
+			{
+				SceneName = FileName + TEXT("_scene_") + LexToString(SceneIndex++);
+			}
 
-			FString SceneNodeUid = TEXT("\\Scene\\") + GltfScene.UniqueId;
+			FString SceneNodeUid = TEXT("\\Scene\\") + SceneName;
 			SceneNode->InitializeNode( SceneNodeUid, SceneName, EInterchangeNodeContainerType::TranslatedScene );
 			NodeContainer.AddNode( SceneNode );
 
@@ -1362,8 +1366,9 @@ void UInterchangeGltfTranslator::HandleGltfAnimation(UInterchangeBaseNodeContain
 
 	UInterchangeAnimationTrackSetNode* TrackSetNode = NewObject< UInterchangeAnimationTrackSetNode >(&NodeContainer);
 
-	const FString AnimTrackSetNodeUid = TEXT("\\Animation\\") + GltfAnimation.UniqueId;
-	TrackSetNode->InitializeNode(AnimTrackSetNodeUid, GltfAnimation.Name, EInterchangeNodeContainerType::TranslatedAsset);
+	const FString AnimTrackSetName = GltfAnimation.Name;
+	const FString AnimTrackSetNodeUid = TEXT("\\Animation\\") + AnimTrackSetName;
+	TrackSetNode->InitializeNode(AnimTrackSetNodeUid, AnimTrackSetName, EInterchangeNodeContainerType::TranslatedAsset);
 
 	for (const TTuple<const GLTF::FNode*, TArray<int32>>& NodeChannelsEntry : NodeChannelsMap)
 	{
@@ -1439,7 +1444,8 @@ void UInterchangeGltfTranslator::SetTextureSRGB(UInterchangeBaseNodeContainer& N
 {
 	if (GltfAsset.Textures.IsValidIndex(TextureMap.TextureIndex))
 	{
-		const FString TextureUid = UInterchangeTextureNode::MakeNodeUid(GltfAsset.Textures[TextureMap.TextureIndex].UniqueId);
+		const FString TextureName = GltfAsset.Textures[TextureMap.TextureIndex].Name;
+		const FString TextureUid = UInterchangeTextureNode::MakeNodeUid(TextureName);
 		if (UInterchangeTextureNode* TextureNode = const_cast<UInterchangeTextureNode*>(Cast<UInterchangeTextureNode>(NodeContainer.GetNode(TextureUid))))
 		{
 			TextureNode->SetCustomSRGB(true);
@@ -1450,7 +1456,8 @@ void UInterchangeGltfTranslator::SetTextureFlipGreenChannel(UInterchangeBaseNode
 {
 	if (GltfAsset.Textures.IsValidIndex(TextureMap.TextureIndex))
 	{
-		const FString TextureUid = UInterchangeTextureNode::MakeNodeUid(GltfAsset.Textures[TextureMap.TextureIndex].UniqueId);
+		const FString TextureName = GltfAsset.Textures[TextureMap.TextureIndex].Name;
+		const FString TextureUid = UInterchangeTextureNode::MakeNodeUid(TextureName);
 		if (UInterchangeTextureNode* TextureNode = const_cast<UInterchangeTextureNode*>(Cast<UInterchangeTextureNode>(NodeContainer.GetNode(TextureUid))))
 		{
 			TextureNode->SetCustombFlipGreenChannel(true);
@@ -1551,7 +1558,7 @@ void UInterchangeGltfTranslator::HandleGltfVariants(UInterchangeBaseNodeContaine
 								}
 
 								const GLTF::FMaterial& GltfMaterial = Materials[VariantMapping.MaterialIndex];
-								const FString MaterialUid = UInterchangeShaderGraphNode::MakeNodeUid(GltfMaterial.UniqueId);
+								const FString MaterialUid = UInterchangeShaderGraphNode::MakeNodeUid(GltfMaterial.Name);
 
 								VariantSetNode->AddCustomDependencyUid(MaterialUid);
 
@@ -1632,7 +1639,7 @@ bool UInterchangeGltfTranslator::GetVariantSetPayloadData(UE::Interchange::FVari
 						}
 
 						const GLTF::FMaterial& GltfMaterial = Materials[VariantMapping.MaterialIndex];
-						const FString MaterialUid = UInterchangeShaderGraphNode::MakeNodeUid(GltfMaterial.UniqueId);
+						const FString MaterialUid = UInterchangeShaderGraphNode::MakeNodeUid(GltfMaterial.Name);
 
 						for (int32 VariantIndex : VariantMapping.VariantIndices)
 						{
@@ -1745,9 +1752,8 @@ void UInterchangeGltfTranslator::HandleGltfSkeletons(UInterchangeBaseNodeContain
 
 			//Skeletal Mesh's naming policy: (Mesh.Name)_(RootJointNode.Name) naming policy:
 			FString SkeletalName = GltfAsset.Meshes[MeshIndex].Name + TEXT("_") + GltfAsset.Nodes[RootJointIndex].Name;
-			FString SkeletalId = GltfAsset.Meshes[MeshIndex].UniqueId + TEXT("_") + GltfAsset.Nodes[RootJointIndex].UniqueId;
 
-			UInterchangeMeshNode* SkeletalMeshNode = HandleGltfMesh(NodeContainer, GltfAsset.Meshes[MeshIndex], MeshIndex, UnusedMeshIndices, SkeletalName, SkeletalId);
+			UInterchangeMeshNode* SkeletalMeshNode = HandleGltfMesh(NodeContainer, GltfAsset.Meshes[MeshIndex], MeshIndex, UnusedMeshIndices, SkeletalName);
 
 			SkeletalMeshNode->SetSkinnedMesh(true);
 
@@ -1798,8 +1804,7 @@ void UInterchangeGltfTranslator::HandleGltfSkeletons(UInterchangeBaseNodeContain
 UInterchangeMeshNode* UInterchangeGltfTranslator::HandleGltfMesh(UInterchangeBaseNodeContainer& NodeContainer, 
 	const GLTF::FMesh& GltfMesh, int MeshIndex, 
 	TSet<int>& UnusedMeshIndices, 
-	const FString& SkeletalName/*If set it creates the mesh even if it was already created (for Skeletals)*/,
-	const FString& SkeletalId) const
+	const FString& SkeletalName/*If set it creates the mesh even if it was already created (for Skeletals)*/) const
 {
 	if (!UnusedMeshIndices.Contains(MeshIndex) && SkeletalName.Len() == 0)
 	{
@@ -1809,7 +1814,7 @@ UInterchangeMeshNode* UInterchangeGltfTranslator::HandleGltfMesh(UInterchangeBas
 	UnusedMeshIndices.Remove(MeshIndex);
 
 	UInterchangeMeshNode* MeshNode = NewObject< UInterchangeMeshNode >(&NodeContainer);
-	FString MeshNodeUid = TEXT("\\Mesh\\") + (SkeletalId.Len() ? SkeletalId : GltfMesh.UniqueId);
+	FString MeshNodeUid = TEXT("\\Mesh\\") + MeshName;
 
 	MeshNode->InitializeNode(MeshNodeUid, MeshName, EInterchangeNodeContainerType::TranslatedAsset);
 	MeshNode->SetPayLoadKey(LexToString(MeshIndex));
@@ -1824,7 +1829,7 @@ UInterchangeMeshNode* UInterchangeGltfTranslator::HandleGltfMesh(UInterchangeBas
 		if (GltfAsset.Materials.IsValidIndex(Primitive.MaterialIndex))
 		{
 			const FString MaterialName = GltfAsset.Materials[Primitive.MaterialIndex].Name;
-			const FString ShaderGraphNodeUid = UInterchangeShaderGraphNode::MakeNodeUid(GltfAsset.Materials[Primitive.MaterialIndex].UniqueId);
+			const FString ShaderGraphNodeUid = UInterchangeShaderGraphNode::MakeNodeUid(MaterialName);
 			MeshNode->SetSlotMaterialDependencyUid(MaterialName, ShaderGraphNodeUid);
 		}
 	}
diff --git a/Engine/Plugins/Interchange/Runtime/Source/Import/Public/Gltf/InterchangeGltfTranslator.h b/Engine/Plugins/Interchange/Runtime/Source/Import/Public/Gltf/InterchangeGltfTranslator.h
--- a/Engine/Plugins/Interchange/Runtime/Source/Import/Public/Gltf/InterchangeGltfTranslator.h
+++ b/Engine/Plugins/Interchange/Runtime/Source/Import/Public/Gltf/InterchangeGltfTranslator.h
@@ -80,7 +80,7 @@ protected:
 	void HandleGltfAnimation(UInterchangeBaseNodeContainer& NodeContainer, int32 AnimationIndex) const;
 	void HandleGltfVariants(UInterchangeBaseNodeContainer& NodeContainer, const FString& FileName) const;
 	UInterchangeMeshNode* HandleGltfMesh(UInterchangeBaseNodeContainer& NodeContainer, const GLTF::FMesh& GltfMesh, int MeshIndex,
-		TSet<int>& UnusedMeshIndices, const FString& SkeletalName = "" /*If set it creates the mesh even if it was already created (for Skeletals)*/, const FString& SkeletalId = "") const;
+		TSet<int>& UnusedMeshIndices, const FString& SkeletalName = "" /*If set it creates the mesh even if it was already created (for Skeletals)*/) const;
 
 	/** Support for KHR_materials_clearcoat */
 	void HandleGltfClearCoat( UInterchangeBaseNodeContainer& NodeContainer, const GLTF::FMaterial& GltfMaterial, UInterchangeShaderGraphNode& ShaderGraphNode, const bool bSwapNormalAndClearCoatNormal ) const;
diff --git a/Engine/Plugins/Interchange/Runtime/Source/Nodes/Private/InterchangeShaderGraphNode.cpp b/Engine/Plugins/Interchange/Runtime/Source/Nodes/Private/InterchangeShaderGraphNode.cpp
--- a/Engine/Plugins/Interchange/Runtime/Source/Nodes/Private/InterchangeShaderGraphNode.cpp
+++ b/Engine/Plugins/Interchange/Runtime/Source/Nodes/Private/InterchangeShaderGraphNode.cpp
@@ -167,15 +167,14 @@ FString UInterchangeShaderGraphNode::MakeNodeUid(const FStringView NodeName)
 	return FString(UInterchangeBaseNode::HierarchySeparator) + TEXT("Shaders") + FString(UInterchangeBaseNode::HierarchySeparator) + NodeName;
 }
 
-UInterchangeShaderGraphNode* UInterchangeShaderGraphNode::Create(UInterchangeBaseNodeContainer* NodeContainer, const FStringView ShaderGraphNodeName /*Or ShaderGraphNodeId*/, const FStringView ShaderGraphNodeDisplayLabel)
+UInterchangeShaderGraphNode* UInterchangeShaderGraphNode::Create(UInterchangeBaseNodeContainer* NodeContainer, const FStringView NodeName)
 {
 	check(NodeContainer);
-	
-	const FString NodeUid = MakeNodeUid(ShaderGraphNodeName);
-	const FStringView NodeDisplayLabel = ShaderGraphNodeDisplayLabel.IsEmpty() ? ShaderGraphNodeName : ShaderGraphNodeDisplayLabel;
+
+	const FString ShaderGraphNodeUid = MakeNodeUid(NodeName);
 
 	UInterchangeShaderGraphNode* ShaderGraphNode = NewObject< UInterchangeShaderGraphNode >(NodeContainer);
-	ShaderGraphNode->InitializeNode(NodeUid, FString(NodeDisplayLabel), EInterchangeNodeContainerType::TranslatedAsset);
+	ShaderGraphNode->InitializeNode(ShaderGraphNodeUid, FString(NodeName), EInterchangeNodeContainerType::TranslatedAsset);
 
 	NodeContainer->AddNode(ShaderGraphNode);
 
diff --git a/Engine/Plugins/Interchange/Runtime/Source/Nodes/Private/InterchangeTexture2DNode.cpp b/Engine/Plugins/Interchange/Runtime/Source/Nodes/Private/InterchangeTexture2DNode.cpp
--- a/Engine/Plugins/Interchange/Runtime/Source/Nodes/Private/InterchangeTexture2DNode.cpp
+++ b/Engine/Plugins/Interchange/Runtime/Source/Nodes/Private/InterchangeTexture2DNode.cpp
@@ -10,14 +10,12 @@ FString UInterchangeTexture2DNode::MakeNodeUid(const FStringView NodeName)
 	return UInterchangeTextureNode::MakeNodeUid(NodeName);
 }
 
- UInterchangeTexture2DNode* UInterchangeTexture2DNode::Create(UInterchangeBaseNodeContainer* NodeContainer, const FStringView TextureNodeName /*Or TextureNodeId*/, const FStringView TextureNodeDisplayLabel)
+ UInterchangeTexture2DNode* UInterchangeTexture2DNode::Create(UInterchangeBaseNodeContainer* NodeContainer, const FStringView TextureNodeName)
 {
 	UInterchangeTexture2DNode* TextureNode = NewObject< UInterchangeTexture2DNode >(NodeContainer);
+	const FString TextureNodeUid = UInterchangeTextureNode::MakeNodeUid(TextureNodeName);
 
-	const FString NodeUid = MakeNodeUid(TextureNodeName);
-	const FStringView NodeDisplayLabel = TextureNodeDisplayLabel.IsEmpty() ? TextureNodeName : TextureNodeDisplayLabel;
-
-	TextureNode->InitializeNode(NodeUid, FString(NodeDisplayLabel), EInterchangeNodeContainerType::TranslatedAsset);
+	TextureNode->InitializeNode(TextureNodeUid, FString(TextureNodeName), EInterchangeNodeContainerType::TranslatedAsset);
 	NodeContainer->AddNode(TextureNode);
 
 	return TextureNode;
diff --git a/Engine/Plugins/Interchange/Runtime/Source/Nodes/Public/InterchangeShaderGraphNode.h b/Engine/Plugins/Interchange/Runtime/Source/Nodes/Public/InterchangeShaderGraphNode.h
--- a/Engine/Plugins/Interchange/Runtime/Source/Nodes/Public/InterchangeShaderGraphNode.h
+++ b/Engine/Plugins/Interchange/Runtime/Source/Nodes/Public/InterchangeShaderGraphNode.h
@@ -169,9 +169,8 @@ public:
 
 	/**
 	 * Creates a new UInterchangeShaderGraphNode and adds it to NodeContainer as a translated node.
-	 * In case the ShaderGraphNodeDisplayLabel is not empty, ShaderGraphNodeName supposed to hold the ShaderGraphNodeId
 	 */
-	static UInterchangeShaderGraphNode* Create(UInterchangeBaseNodeContainer* NodeContainer, const FStringView ShaderGraphNodeName /*Or ShaderGraphNodeId*/, const FStringView ShaderGraphNodeDisplayLabel = FStringView());
+	static UInterchangeShaderGraphNode* Create(UInterchangeBaseNodeContainer* NodeContainer, const FStringView NodeName);
 
 	/**
 	 * Return the node type name of the class, we use this when reporting error
diff --git a/Engine/Plugins/Interchange/Runtime/Source/Nodes/Public/InterchangeTexture2DNode.h b/Engine/Plugins/Interchange/Runtime/Source/Nodes/Public/InterchangeTexture2DNode.h
--- a/Engine/Plugins/Interchange/Runtime/Source/Nodes/Public/InterchangeTexture2DNode.h
+++ b/Engine/Plugins/Interchange/Runtime/Source/Nodes/Public/InterchangeTexture2DNode.h
@@ -36,9 +36,8 @@ public:
 
 	/**
 	 * Creates a new UInterchangeTexture2DNode and adds it to NodeContainer as a translated node.
-	 * In case the TextureNodeDisplayLabel is not empty, TextureNodeName supposed to hold the TextureNodeId
 	 */
-	static UInterchangeTexture2DNode* Create(UInterchangeBaseNodeContainer* NodeContainer, const FStringView TextureNodeName /*Or TextureNodeId*/, const FStringView TextureNodeDisplayLabel = FStringView());
+	static UInterchangeTexture2DNode* Create(UInterchangeBaseNodeContainer* NodeContainer, const FStringView TextureNodeName);
 
 	virtual void PostInitProperties()
 	{
diff --git a/Engine/Source/Runtime/Interchange/Engine/Private/Tasks/InterchangeTaskParsing.cpp b/Engine/Source/Runtime/Interchange/Engine/Private/Tasks/InterchangeTaskParsing.cpp
--- a/Engine/Source/Runtime/Interchange/Engine/Private/Tasks/InterchangeTaskParsing.cpp
+++ b/Engine/Source/Runtime/Interchange/Engine/Private/Tasks/InterchangeTaskParsing.cpp
@@ -277,11 +277,6 @@ void UE::Interchange::FTaskParsing::DoTask(ENamedThreads::Type CurrentThread, co
 					RenameInfo.NewName = NewName;
 					RenameInfo.SourceIndex = SourceIndex;
 					FactoryNode->SetDisplayLabel(FactoryNode->GetDisplayLabel() + NameIndexString);
-					if (FactoryNode->HasAttribute(UE::Interchange::FBaseNodeStaticData::AssetNameKey()))
-					{
-						//TextureFactorNodes automatically set up the AssetNames which then are used for asset creation:
-						FactoryNode->SetAssetName(FactoryNode->GetDisplayLabel());
-					}
 				}
 				else
 				{
