Interchange: fix obj reimport adding material when the material do not have any name specified in the obj.

#rnx

diff --git a/Engine/Plugins/Interchange/Runtime/Source/Import/Private/Mesh/InterchangeOBJTranslator.cpp b/Engine/Plugins/Interchange/Runtime/Source/Import/Private/Mesh/InterchangeOBJTranslator.cpp
--- a/Engine/Plugins/Interchange/Runtime/Source/Import/Private/Mesh/InterchangeOBJTranslator.cpp
+++ b/Engine/Plugins/Interchange/Runtime/Source/Import/Private/Mesh/InterchangeOBJTranslator.cpp
@@ -13,6 +13,7 @@
 #include "InterchangeShaderGraphNode.h"
 #include "InterchangeTexture2DNode.h"
 #include "InterchangeTextureNode.h"
+#include "Materials/Material.h"
 #include "Misc/ConfigCacheIni.h"
 #include "Misc/FileHelper.h"
 #include "Misc/Paths.h"
@@ -297,7 +298,9 @@ FMeshDescription FObjData::MakeMeshDescriptionForGroup(const FString& GroupName)
 	// Create polygon group
 
 	FPolygonGroupID PolygonGroupIndex = MeshDescription.CreatePolygonGroup();
-	Attributes.GetPolygonGroupMaterialSlotNames()[PolygonGroupIndex] = FName(GroupData.MaterialName);
+	const FString MaterialName = GroupData.MaterialName.IsEmpty() ? UMaterial::GetDefaultMaterial(MD_Surface)->GetName() : GroupData.MaterialName;
+	ensure(!MaterialName.IsEmpty());
+	Attributes.GetPolygonGroupMaterialSlotNames()[PolygonGroupIndex] = FName(MaterialName);
 
 	// Create faces.
 	// n-gons are preserved
