[interchange] [gltf] geometry missing due to naming clashes fix (specifically for 5.1)

diff --git a/Engine/Plugins/Interchange/Runtime/Source/Import/Private/Gltf/InterchangeGltfTranslator.cpp b/Engine/Plugins/Interchange/Runtime/Source/Import/Private/Gltf/InterchangeGltfTranslator.cpp
--- a/Engine/Plugins/Interchange/Runtime/Source/Import/Private/Gltf/InterchangeGltfTranslator.cpp
+++ b/Engine/Plugins/Interchange/Runtime/Source/Import/Private/Gltf/InterchangeGltfTranslator.cpp
@@ -127,14 +127,81 @@ namespace UE::Interchange::Gltf::Private
 			Node.Transform.SetTranslation(Node.Transform.GetTranslation() * Scale);
 		}
 	}
+
+	//Temporary Solution specifically for 5.1.1
+	// As we cannot modify public headers (because of binary compatibility ) all necessary changes for gltf object uniqueness have been moved into here.
+	// Appropriate/fully fledged solution will be addressed for 5.2
+	struct GLTFUniqueIdGenerator
+	{
+		FString Prefix;
+		
+		// Node To UniqueId maps
+		TMap<const GLTF::FAnimation*, FString> AnimationToUniqueId;
+		TMap<const GLTF::FScene*, FString> SceneToUniqueId;
+		TMap<const GLTF::FMaterial*, FString> MaterialToUniqueId;
+		TMap<const GLTF::FMesh*, FString> MeshToUniqueId;
+		TMap<const GLTF::FNode*, FString> NodeToUniqueId;
+		TMap<const GLTF::FCamera*, FString> CameraToUniqueId;
+		TMap<const GLTF::FLight*, FString> LightToUniqueId;
+		TMap<const GLTF::FSkinInfo*, FString> SkinToUniqueId;
+		TMap<const GLTF::FTexture*, FString> TextureToUniqueId;
+		TMap<const GLTF::FImage*, FString> ImageToUniqueId;
+
+		void Clear()
+		{
+			Prefix = "";
+
+			AnimationToUniqueId.Empty();
+			SceneToUniqueId.Empty();
+			MaterialToUniqueId.Empty();
+			MeshToUniqueId.Empty();
+			NodeToUniqueId.Empty();
+			CameraToUniqueId.Empty();
+			LightToUniqueId.Empty();
+			SkinToUniqueId.Empty();
+			TextureToUniqueId.Empty();
+			ImageToUniqueId.Empty();
+		}
+
+		template <typename T>
+		FString GetUniqueId(T Object, TMap<T, FString>& ObjectToUniqueId, const FString& ObjectPrefix)
+		{
+			if (ObjectToUniqueId.Contains(Object))
+			{
+				return ObjectToUniqueId[Object];
+			}
+			else
+			{
+				FString UniqueIdCandidate = Prefix + ObjectPrefix + FString::FromInt(ObjectToUniqueId.Num());
+				ObjectToUniqueId.Add(Object, UniqueIdCandidate);
+				return UniqueIdCandidate;
+			}
+		}
+
+		FString GetUniqueId(const GLTF::FAnimation* Object) { return GetUniqueId<const GLTF::FAnimation*>(Object, AnimationToUniqueId, "_animation_"); }
+		FString GetUniqueId(const GLTF::FScene* Object) { return GetUniqueId<const GLTF::FScene*>(Object, SceneToUniqueId, "_scene_"); }
+		FString GetUniqueId(const GLTF::FMaterial* Object) { return GetUniqueId<const GLTF::FMaterial*>(Object, MaterialToUniqueId, "_material_"); }
+		FString GetUniqueId(const GLTF::FMesh* Object) { return GetUniqueId<const GLTF::FMesh*>(Object, MeshToUniqueId, "_mesh_"); }
+		FString GetUniqueId(const GLTF::FNode* Object) { return GetUniqueId<const GLTF::FNode*>(Object, NodeToUniqueId, "_node_"); }
+		FString GetUniqueId(const GLTF::FCamera* Object) { return GetUniqueId<const GLTF::FCamera*>(Object, CameraToUniqueId, "_camera_"); }
+		FString GetUniqueId(const GLTF::FLight* Object) { return GetUniqueId<const GLTF::FLight*>(Object, LightToUniqueId, "_light_"); }
+		FString GetUniqueId(const GLTF::FSkinInfo* Object) { return GetUniqueId<const GLTF::FSkinInfo*>(Object, SkinToUniqueId, "_skin_"); }
+		FString GetUniqueId(const GLTF::FTexture* Object) { return GetUniqueId<const GLTF::FTexture*>(Object, TextureToUniqueId, "_texture_"); }
+		FString GetUniqueId(const GLTF::FImage* Object) { return GetUniqueId<const GLTF::FImage*>(Object, ImageToUniqueId, "_image_"); }
+	};
+
+	static TMap<const UInterchangeSourceData*, GLTFUniqueIdGenerator> GltfUniqueIdGenerators;
+	static TMap<const UInterchangeSourceData*, int> GltfUniqueIdGeneratorsRefCounter;
 }
 
 void UInterchangeGltfTranslator::HandleGltfNode( UInterchangeBaseNodeContainer& NodeContainer, const GLTF::FNode& GltfNode, const FString& ParentNodeUid, const int32 NodeIndex, 
 	bool &bHasVariants, TArray<int32>& SkinnedMeshNodes, TSet<int>& UnusedMeshIndices ) const
 {
+	
+
 	using namespace UE::Interchange::Gltf::Private;
 
-	const FString NodeUid = ParentNodeUid + TEXT("\\") + GltfNode.Name;
+	const FString NodeUid = ParentNodeUid + TEXT("\\") + GltfUniqueIdGenerators[GetSourceData()].GetUniqueId(&GltfNode);
 
 	const UInterchangeSceneNode* ParentSceneNode = Cast< UInterchangeSceneNode >( NodeContainer.GetNode( ParentNodeUid ) );
 
@@ -173,7 +240,7 @@ void UInterchangeGltfTranslator::HandleGltfNode( UInterchangeBaseNodeContainer&
 			{
 				HandleGltfMesh(NodeContainer, GltfAsset.Meshes[GltfNode.MeshIndex], GltfNode.MeshIndex, UnusedMeshIndices);
 
-				const FString MeshNodeUid = TEXT("\\Mesh\\") + GltfAsset.Meshes[ GltfNode.MeshIndex ].Name;
+				const FString MeshNodeUid = TEXT("\\Mesh\\") + GltfUniqueIdGenerators[GetSourceData()].GetUniqueId(&GltfAsset.Meshes[ GltfNode.MeshIndex ]);
 				InterchangeSceneNode->SetCustomAssetInstanceUid( MeshNodeUid );
 
 				if (!bHasVariants && GltfAsset.Variants.Num() > 0)
@@ -190,7 +257,7 @@ void UInterchangeGltfTranslator::HandleGltfNode( UInterchangeBaseNodeContainer&
 
 			if ( GltfAsset.Cameras.IsValidIndex( GltfNode.CameraIndex ) )
 			{
-				const FString CameraNodeUid = TEXT("\\Camera\\") + GltfAsset.Cameras[ GltfNode.CameraIndex ].Name;
+				const FString CameraNodeUid = TEXT("\\Camera\\") + GltfUniqueIdGenerators[GetSourceData()].GetUniqueId(&GltfAsset.Cameras[ GltfNode.CameraIndex ]);
 				InterchangeSceneNode->SetCustomAssetInstanceUid( CameraNodeUid );
 			}
 			break;
@@ -202,7 +269,7 @@ void UInterchangeGltfTranslator::HandleGltfNode( UInterchangeBaseNodeContainer&
 
 			if ( GltfAsset.Lights.IsValidIndex( GltfNode.LightIndex ) )
 			{
-				const FString LightNodeUid = TEXT("\\Light\\") + GltfAsset.Lights[ GltfNode.LightIndex ].Name;
+				const FString LightNodeUid = TEXT("\\Light\\") + GltfUniqueIdGenerators[GetSourceData()].GetUniqueId(&GltfAsset.Lights[ GltfNode.LightIndex ]);
 				InterchangeSceneNode->SetCustomAssetInstanceUid( LightNodeUid );
 			}
 		}
@@ -271,14 +338,13 @@ void UInterchangeGltfTranslator::HandleGltfMaterialParameter( UInterchangeBaseNo
 		UInterchangeShaderNode* ColorNode = UInterchangeShaderNode::Create( &NodeContainer, ColorNodeName, ShaderNode.GetUniqueID() );
 		ColorNode->SetCustomShaderType( Standard::Nodes::TextureSample::Name.ToString() );
 
-		const FString TextureName = GltfAsset.Textures[ TextureMap.TextureIndex ].Name;
-		const FString TextureUid = UInterchangeTextureNode::MakeNodeUid( TextureName );
+		const FString TextureUid = UInterchangeTextureNode::MakeNodeUid(GltfUniqueIdGenerators[GetSourceData()].GetUniqueId(&GltfAsset.Textures[TextureMap.TextureIndex]));
 
 		ColorNode->AddStringAttribute( UInterchangeShaderPortsAPI::MakeInputValueKey( Standard::Nodes::TextureSample::Inputs::Texture.ToString() ), TextureUid );
 
 		if (TextureMap.TexCoord > 0 || TextureMap.bHasTextureTransform)
 		{
-			UInterchangeShaderNode* TexCoordNode = UInterchangeShaderNode::Create(&NodeContainer, TEXT("TexCoord"), ShaderNode.GetUniqueID());
+			UInterchangeShaderNode* TexCoordNode = UInterchangeShaderNode::Create(&NodeContainer, MapName + TEXT("\\TexCoord"), ShaderNode.GetUniqueID());
 			TexCoordNode->SetCustomShaderType(Standard::Nodes::TextureCoordinate::Name.ToString());
 
 			TexCoordNode->AddInt32Attribute(UInterchangeShaderPortsAPI::MakeInputValueKey(Standard::Nodes::TextureCoordinate::Inputs::Index.ToString()), TextureMap.TexCoord);
@@ -882,14 +948,19 @@ bool UInterchangeGltfTranslator::Translate( UInterchangeBaseNodeContainer& NodeC
 		}
 	}
 
-	const_cast< UInterchangeGltfTranslator* >( this )->GltfAsset.GenerateNames(FileName);
+	const_cast< UInterchangeGltfTranslator* >( this )->GltfAsset.GenerateNames(FileName); //uniqueIds are generated on demand
+	GltfUniqueIdGeneratorsRefCounter.FindOrAdd(GetSourceData())++;
+	GltfUniqueIdGenerators.FindOrAdd(GetSourceData());
+	GltfUniqueIdGenerators[GetSourceData()].Prefix = FileName;
 
 	// Textures
 	{
 		int32 TextureIndex = 0;
 		for ( const GLTF::FTexture& GltfTexture : GltfAsset.Textures )
 		{
-			UInterchangeTexture2DNode* TextureNode = UInterchangeTexture2DNode::Create( &NodeContainer, GltfTexture.Name);
+			UInterchangeTexture2DNode* TextureNode = UInterchangeTexture2DNode::Create(&NodeContainer, GltfUniqueIdGenerators[GetSourceData()].GetUniqueId(&GltfTexture));
+			TextureNode->SetDisplayLabel(GltfTexture.Name);
+
 			TextureNode->SetCustomFilter(ConvertFilter(GltfTexture.Sampler.MinFilter));
 			TextureNode->SetPayLoadKey( LexToString( TextureIndex++ ) );
 
@@ -903,7 +974,8 @@ bool UInterchangeGltfTranslator::Translate( UInterchangeBaseNodeContainer& NodeC
 		int32 MaterialIndex = 0;
 		for ( const GLTF::FMaterial& GltfMaterial : GltfAsset.Materials )
 		{
-			UInterchangeShaderGraphNode* ShaderGraphNode = UInterchangeShaderGraphNode::Create( &NodeContainer, GltfMaterial.Name );
+			UInterchangeShaderGraphNode* ShaderGraphNode = UInterchangeShaderGraphNode::Create(&NodeContainer, GltfUniqueIdGenerators[GetSourceData()].GetUniqueId(&GltfMaterial));
+			ShaderGraphNode->SetDisplayLabel(GltfMaterial.Name);
 
 			HandleGltfMaterial( NodeContainer, GltfMaterial, *ShaderGraphNode );
 			++MaterialIndex;
@@ -926,7 +998,7 @@ bool UInterchangeGltfTranslator::Translate( UInterchangeBaseNodeContainer& NodeC
 		for ( const GLTF::FCamera& GltfCamera : GltfAsset.Cameras )
 		{
 			UInterchangeCameraNode* CameraNode = NewObject< UInterchangeCameraNode >( &NodeContainer );
-			FString CameraNodeUid = TEXT("\\Camera\\") + GltfCamera.Name;
+			FString CameraNodeUid = TEXT("\\Camera\\") + GltfUniqueIdGenerators[GetSourceData()].GetUniqueId(&GltfCamera);
 			CameraNode->InitializeNode( CameraNodeUid, GltfCamera.Name, EInterchangeNodeContainerType::TranslatedAsset );
 
 			float       AspectRatio;
@@ -959,7 +1031,7 @@ bool UInterchangeGltfTranslator::Translate( UInterchangeBaseNodeContainer& NodeC
 		int32 LightIndex = 0;
 		for ( const GLTF::FLight& GltfLight : GltfAsset.Lights )
 		{
-			FString LightNodeUid = TEXT("\\Light\\") + GltfLight.Name;
+			FString LightNodeUid = TEXT("\\Light\\") + GltfUniqueIdGenerators[GetSourceData()].GetUniqueId(&GltfLight);
 
 			switch (GltfLight.Type)
 			{
@@ -1023,12 +1095,8 @@ bool UInterchangeGltfTranslator::Translate( UInterchangeBaseNodeContainer& NodeC
 			UInterchangeSceneNode* SceneNode = NewObject< UInterchangeSceneNode >( &NodeContainer );
 
 			FString SceneName = GltfScene.Name;
-			if (SceneName.IsEmpty())
-			{
-				SceneName = FileName + TEXT("_scene_") + LexToString(SceneIndex++);
-			}
 
-			FString SceneNodeUid = TEXT("\\Scene\\") + SceneName;
+			FString SceneNodeUid = TEXT("\\Scene\\") + GltfUniqueIdGenerators[GetSourceData()].GetUniqueId(&GltfScene);
 			SceneNode->InitializeNode( SceneNodeUid, SceneName, EInterchangeNodeContainerType::TranslatedScene );
 			NodeContainer.AddNode( SceneNode );
 
@@ -1107,6 +1175,14 @@ bool UInterchangeGltfTranslator::Translate( UInterchangeBaseNodeContainer& NodeC
 		UE_LOG(LogInterchangeImport, Warning, TEXT("GLTF Mesh Import Warning. Gltf Mesh Usage expectation is not met."));
 	}
 
+	GltfUniqueIdGeneratorsRefCounter[GetSourceData()]--;
+	if (GltfUniqueIdGeneratorsRefCounter[GetSourceData()] <= 0)
+	{
+		GltfUniqueIdGenerators[GetSourceData()].Clear();
+		GltfUniqueIdGenerators.Remove(GetSourceData());
+		GltfUniqueIdGeneratorsRefCounter.Remove(GetSourceData());
+	}
+
 	return true;
 }
 
@@ -1366,9 +1442,8 @@ void UInterchangeGltfTranslator::HandleGltfAnimation(UInterchangeBaseNodeContain
 
 	UInterchangeAnimationTrackSetNode* TrackSetNode = NewObject< UInterchangeAnimationTrackSetNode >(&NodeContainer);
 
-	const FString AnimTrackSetName = GltfAnimation.Name;
-	const FString AnimTrackSetNodeUid = TEXT("\\Animation\\") + AnimTrackSetName;
-	TrackSetNode->InitializeNode(AnimTrackSetNodeUid, AnimTrackSetName, EInterchangeNodeContainerType::TranslatedAsset);
+	const FString AnimTrackSetNodeUid = TEXT("\\Animation\\") + UE::Interchange::Gltf::Private::GltfUniqueIdGenerators[GetSourceData()].GetUniqueId(&GltfAnimation);
+	TrackSetNode->InitializeNode(AnimTrackSetNodeUid, GltfAnimation.Name, EInterchangeNodeContainerType::TranslatedAsset);
 
 	for (const TTuple<const GLTF::FNode*, TArray<int32>>& NodeChannelsEntry : NodeChannelsMap)
 	{
@@ -1444,8 +1519,7 @@ void UInterchangeGltfTranslator::SetTextureSRGB(UInterchangeBaseNodeContainer& N
 {
 	if (GltfAsset.Textures.IsValidIndex(TextureMap.TextureIndex))
 	{
-		const FString TextureName = GltfAsset.Textures[TextureMap.TextureIndex].Name;
-		const FString TextureUid = UInterchangeTextureNode::MakeNodeUid(TextureName);
+		const FString TextureUid = UInterchangeTextureNode::MakeNodeUid(UE::Interchange::Gltf::Private::GltfUniqueIdGenerators[GetSourceData()].GetUniqueId(&GltfAsset.Textures[TextureMap.TextureIndex]));
 		if (UInterchangeTextureNode* TextureNode = const_cast<UInterchangeTextureNode*>(Cast<UInterchangeTextureNode>(NodeContainer.GetNode(TextureUid))))
 		{
 			TextureNode->SetCustomSRGB(true);
@@ -1456,8 +1530,7 @@ void UInterchangeGltfTranslator::SetTextureFlipGreenChannel(UInterchangeBaseNode
 {
 	if (GltfAsset.Textures.IsValidIndex(TextureMap.TextureIndex))
 	{
-		const FString TextureName = GltfAsset.Textures[TextureMap.TextureIndex].Name;
-		const FString TextureUid = UInterchangeTextureNode::MakeNodeUid(TextureName);
+		const FString TextureUid = UInterchangeTextureNode::MakeNodeUid(UE::Interchange::Gltf::Private::GltfUniqueIdGenerators[GetSourceData()].GetUniqueId(&GltfAsset.Textures[TextureMap.TextureIndex]));
 		if (UInterchangeTextureNode* TextureNode = const_cast<UInterchangeTextureNode*>(Cast<UInterchangeTextureNode>(NodeContainer.GetNode(TextureUid))))
 		{
 			TextureNode->SetCustombFlipGreenChannel(true);
@@ -1558,7 +1631,7 @@ void UInterchangeGltfTranslator::HandleGltfVariants(UInterchangeBaseNodeContaine
 								}
 
 								const GLTF::FMaterial& GltfMaterial = Materials[VariantMapping.MaterialIndex];
-								const FString MaterialUid = UInterchangeShaderGraphNode::MakeNodeUid(GltfMaterial.Name);
+								const FString MaterialUid = UInterchangeShaderGraphNode::MakeNodeUid(UE::Interchange::Gltf::Private::GltfUniqueIdGenerators[GetSourceData()].GetUniqueId(&GltfMaterial));
 
 								VariantSetNode->AddCustomDependencyUid(MaterialUid);
 
@@ -1639,7 +1712,7 @@ bool UInterchangeGltfTranslator::GetVariantSetPayloadData(UE::Interchange::FVari
 						}
 
 						const GLTF::FMaterial& GltfMaterial = Materials[VariantMapping.MaterialIndex];
-						const FString MaterialUid = UInterchangeShaderGraphNode::MakeNodeUid(GltfMaterial.Name);
+						const FString MaterialUid = UInterchangeShaderGraphNode::MakeNodeUid(UE::Interchange::Gltf::Private::GltfUniqueIdGenerators[GetSourceData()].GetUniqueId(&GltfMaterial));
 
 						for (int32 VariantIndex : VariantMapping.VariantIndices)
 						{
@@ -1752,8 +1825,9 @@ void UInterchangeGltfTranslator::HandleGltfSkeletons(UInterchangeBaseNodeContain
 
 			//Skeletal Mesh's naming policy: (Mesh.Name)_(RootJointNode.Name) naming policy:
 			FString SkeletalName = GltfAsset.Meshes[MeshIndex].Name + TEXT("_") + GltfAsset.Nodes[RootJointIndex].Name;
+			FString SkeletalId = UE::Interchange::Gltf::Private::GltfUniqueIdGenerators[GetSourceData()].GetUniqueId(&GltfAsset.Meshes[MeshIndex]) + TEXT("_") + UE::Interchange::Gltf::Private::GltfUniqueIdGenerators[GetSourceData()].GetUniqueId(&GltfAsset.Nodes[RootJointIndex]);
 
-			UInterchangeMeshNode* SkeletalMeshNode = HandleGltfMesh(NodeContainer, GltfAsset.Meshes[MeshIndex], MeshIndex, UnusedMeshIndices, SkeletalName);
+			UInterchangeMeshNode* SkeletalMeshNode = HandleGltfMesh(NodeContainer, GltfAsset.Meshes[MeshIndex], MeshIndex, UnusedMeshIndices, SkeletalId);
 
 			SkeletalMeshNode->SetSkinnedMesh(true);
 
@@ -1814,7 +1888,7 @@ UInterchangeMeshNode* UInterchangeGltfTranslator::HandleGltfMesh(UInterchangeBas
 	UnusedMeshIndices.Remove(MeshIndex);
 
 	UInterchangeMeshNode* MeshNode = NewObject< UInterchangeMeshNode >(&NodeContainer);
-	FString MeshNodeUid = TEXT("\\Mesh\\") + MeshName;
+	FString MeshNodeUid = TEXT("\\Mesh\\") + (SkeletalName.Len() ? SkeletalName : UE::Interchange::Gltf::Private::GltfUniqueIdGenerators[GetSourceData()].GetUniqueId(&GltfMesh));
 
 	MeshNode->InitializeNode(MeshNodeUid, MeshName, EInterchangeNodeContainerType::TranslatedAsset);
 	MeshNode->SetPayLoadKey(LexToString(MeshIndex));
@@ -1829,7 +1903,7 @@ UInterchangeMeshNode* UInterchangeGltfTranslator::HandleGltfMesh(UInterchangeBas
 		if (GltfAsset.Materials.IsValidIndex(Primitive.MaterialIndex))
 		{
 			const FString MaterialName = GltfAsset.Materials[Primitive.MaterialIndex].Name;
-			const FString ShaderGraphNodeUid = UInterchangeShaderGraphNode::MakeNodeUid(MaterialName);
+			const FString ShaderGraphNodeUid = UInterchangeShaderGraphNode::MakeNodeUid(UE::Interchange::Gltf::Private::GltfUniqueIdGenerators[GetSourceData()].GetUniqueId(&GltfAsset.Materials[Primitive.MaterialIndex]));
 			MeshNode->SetSlotMaterialDependencyUid(MaterialName, ShaderGraphNodeUid);
 		}
 	}
diff --git a/Engine/Source/Runtime/Interchange/Engine/Private/Tasks/InterchangeTaskParsing.cpp b/Engine/Source/Runtime/Interchange/Engine/Private/Tasks/InterchangeTaskParsing.cpp
--- a/Engine/Source/Runtime/Interchange/Engine/Private/Tasks/InterchangeTaskParsing.cpp
+++ b/Engine/Source/Runtime/Interchange/Engine/Private/Tasks/InterchangeTaskParsing.cpp
@@ -277,6 +277,11 @@ void UE::Interchange::FTaskParsing::DoTask(ENamedThreads::Type CurrentThread, co
 					RenameInfo.NewName = NewName;
 					RenameInfo.SourceIndex = SourceIndex;
 					FactoryNode->SetDisplayLabel(FactoryNode->GetDisplayLabel() + NameIndexString);
+					if (FactoryNode->HasAttribute(UE::Interchange::FBaseNodeStaticData::AssetNameKey()))
+					{
+						//TextureFactorNodes automatically set up the AssetNames which then are used for asset creation:
+						FactoryNode->SetAssetName(FactoryNode->GetDisplayLabel());
+					}
 				}
 				else
 				{
