Interchange: The interchange generic asset pipeline now support "BoneInfluenceLimit" skeletalmesh build settings.

#rnx

diff --git a/Engine/Plugins/Interchange/Runtime/Source/FactoryNodes/Private/InterchangeSkeletalMeshFactoryNode.cpp b/Engine/Plugins/Interchange/Runtime/Source/FactoryNodes/Private/InterchangeSkeletalMeshFactoryNode.cpp
--- a/Engine/Plugins/Interchange/Runtime/Source/FactoryNodes/Private/InterchangeSkeletalMeshFactoryNode.cpp
+++ b/Engine/Plugins/Interchange/Runtime/Source/FactoryNodes/Private/InterchangeSkeletalMeshFactoryNode.cpp
@@ -196,6 +196,23 @@ bool UInterchangeSkeletalMeshFactoryNode::FillCustomMorphThresholdPositionFromAs
 	IMPLEMENT_SKELETALMESH_BUILD_ASSET_TO_VALUE(MorphThresholdPosition, MorphThresholdPosition);
 }
 
+bool UInterchangeSkeletalMeshFactoryNode::GetCustomBoneInfluenceLimit(int32& AttributeValue) const
+{
+	IMPLEMENT_NODE_ATTRIBUTE_GETTER(BoneInfluenceLimit, int32)
+}
+bool UInterchangeSkeletalMeshFactoryNode::SetCustomBoneInfluenceLimit(const int32& AttributeValue, bool bAddApplyDelegate)
+{
+	IMPLEMENT_NODE_ATTRIBUTE_SETTER_WITH_CUSTOM_DELEGATE(UInterchangeSkeletalMeshFactoryNode, BoneInfluenceLimit, int32);
+}
+bool UInterchangeSkeletalMeshFactoryNode::ApplyCustomBoneInfluenceLimitToAsset(UObject* Asset) const
+{
+	IMPLEMENT_SKELETAL_BUILD_VALUE_TO_ASSET(BoneInfluenceLimit, int32, BoneInfluenceLimit);
+}
+bool UInterchangeSkeletalMeshFactoryNode::FillCustomBoneInfluenceLimitFromAsset(UObject* Asset)
+{
+	IMPLEMENT_SKELETALMESH_BUILD_ASSET_TO_VALUE(BoneInfluenceLimit, BoneInfluenceLimit);
+}
+
 bool UInterchangeSkeletalMeshFactoryNode::GetCustomImportContentType(EInterchangeSkeletalMeshContentType& AttributeValue) const
 {
 	FString OperationName = GetTypeName() + TEXT(".GetImportContentType");
diff --git a/Engine/Plugins/Interchange/Runtime/Source/FactoryNodes/Public/InterchangeSkeletalMeshFactoryNode.h b/Engine/Plugins/Interchange/Runtime/Source/FactoryNodes/Public/InterchangeSkeletalMeshFactoryNode.h
--- a/Engine/Plugins/Interchange/Runtime/Source/FactoryNodes/Public/InterchangeSkeletalMeshFactoryNode.h
+++ b/Engine/Plugins/Interchange/Runtime/Source/FactoryNodes/Public/InterchangeSkeletalMeshFactoryNode.h
@@ -128,6 +128,22 @@ public:
 	UFUNCTION(BlueprintCallable, Category = "Interchange | Node | SkeletalMesh")
 	bool SetCustomMorphThresholdPosition(const float& AttributeValue, bool bAddApplyDelegate = true);
 
+	/**
+	 * Query The maximum number of bone influences to allow each vertex in this mesh to use.
+	 * If set higher than the limit determined by the project settings, it has no effect.
+	 * If set to 0, the value is taken from the DefaultBoneInfluenceLimit project setting.
+	 */
+	UFUNCTION(BlueprintCallable, Category = "Interchange | Node | SkeletalMesh")
+	bool GetCustomBoneInfluenceLimit(int32& AttributeValue) const;
+
+	/**
+	 * Set The maximum number of bone influences to allow each vertex in this mesh to use.
+	 * If set higher than the limit determined by the project settings, it has no effect.
+	 * If set to 0, the value is taken from the DefaultBoneInfluenceLimit project setting.
+	 */
+	UFUNCTION(BlueprintCallable, Category = "Interchange | Node | SkeletalMesh")
+	bool SetCustomBoneInfluenceLimit(const int32& AttributeValue, bool bAddApplyDelegate = true);
+
 	/**
 	 * The skeletal mesh thumbnail can have an overlay if the last re-import was geometry only. This thumbnail overlay feature use the metadata to find out if the last import was geometry only.
 	 */
@@ -148,6 +164,7 @@ private:
 	const UE::Interchange::FAttributeKey Macro_CustomThresholdTangentNormalKey = UE::Interchange::FAttributeKey(TEXT("ThresholdTangentNormal"));
 	const UE::Interchange::FAttributeKey Macro_CustomThresholdUVKey = UE::Interchange::FAttributeKey(TEXT("ThresholdUV"));
 	const UE::Interchange::FAttributeKey Macro_CustomMorphThresholdPositionKey = UE::Interchange::FAttributeKey(TEXT("MorphThresholdPosition"));
+	const UE::Interchange::FAttributeKey Macro_CustomBoneInfluenceLimitKey = UE::Interchange::FAttributeKey(TEXT("BoneInfluenceLimit"));
 
 	bool ApplyCustomUseHighPrecisionSkinWeightsToAsset(UObject* Asset) const;
 	bool FillCustomUseHighPrecisionSkinWeightsFromAsset(UObject* Asset);
@@ -159,6 +176,8 @@ private:
 	bool FillCustomThresholdUVFromAsset(UObject* Asset);
 	bool ApplyCustomMorphThresholdPositionToAsset(UObject* Asset) const;
 	bool FillCustomMorphThresholdPositionFromAsset(UObject* Asset);
+	bool ApplyCustomBoneInfluenceLimitToAsset(UObject* Asset) const;
+	bool FillCustomBoneInfluenceLimitFromAsset(UObject* Asset);
 protected:
 	TSubclassOf<USkeletalMesh> AssetClass = nullptr;
 };
diff --git a/Engine/Plugins/Interchange/Runtime/Source/Pipelines/Private/InterchangeGenericSkeletalMeshPipeline.cpp b/Engine/Plugins/Interchange/Runtime/Source/Pipelines/Private/InterchangeGenericSkeletalMeshPipeline.cpp
--- a/Engine/Plugins/Interchange/Runtime/Source/Pipelines/Private/InterchangeGenericSkeletalMeshPipeline.cpp
+++ b/Engine/Plugins/Interchange/Runtime/Source/Pipelines/Private/InterchangeGenericSkeletalMeshPipeline.cpp
@@ -503,6 +503,7 @@ UInterchangeSkeletalMeshFactoryNode* UInterchangeGenericMeshPipeline::CreateSkel
 	SkeletalMeshFactoryNode->SetCustomThresholdTangentNormal(ThresholdTangentNormal);
 	SkeletalMeshFactoryNode->SetCustomThresholdUV(ThresholdUV);
 	SkeletalMeshFactoryNode->SetCustomMorphThresholdPosition(MorphThresholdPosition);
+	SkeletalMeshFactoryNode->SetCustomBoneInfluenceLimit(BoneInfluenceLimit);
 
 	return SkeletalMeshFactoryNode;
 }
diff --git a/Engine/Plugins/Interchange/Runtime/Source/Pipelines/Public/InterchangeGenericMeshPipeline.h b/Engine/Plugins/Interchange/Runtime/Source/Pipelines/Public/InterchangeGenericMeshPipeline.h
--- a/Engine/Plugins/Interchange/Runtime/Source/Pipelines/Public/InterchangeGenericMeshPipeline.h
+++ b/Engine/Plugins/Interchange/Runtime/Source/Pipelines/Public/InterchangeGenericMeshPipeline.h
@@ -175,6 +175,14 @@ public:
 	UPROPERTY(EditAnywhere, BlueprintReadWrite, Category = "Skeletal Meshes", meta = (SubCategory = "Build"))
 	float MorphThresholdPosition = 0.015f;
 
+	/**
+	 * The maximum number of bone influences to allow each vertex in this mesh to use.
+	 * If set higher than the limit determined by the project settings, it has no effect.
+	 * If set to 0, the value is taken from the DefaultBoneInfluenceLimit project setting.
+	 */
+	UPROPERTY(EditAnywhere, BlueprintReadWrite, Category = "Skeletal Meshes", meta = (SubCategory = "Build"))
+	int32 BoneInfluenceLimit = 0;
+
 	virtual void AdjustSettingsForContext(EInterchangePipelineContext ImportType, TObjectPtr<UObject> ReimportAsset) override;
 
 	virtual void PreDialogCleanup(const FName PipelineStackName) override;
