Fixed RetainJPEG broken in Interchange Texture import
Solution: Taking in account value of TextureImporter.RetainJpegFormat from editor.ini before setting the RawData on the Payload.

diff --git a/Engine/Plugins/Interchange/Runtime/Source/Import/Private/Texture/InterchangeJPGTranslator.cpp b/Engine/Plugins/Interchange/Runtime/Source/Import/Private/Texture/InterchangeJPGTranslator.cpp
--- a/Engine/Plugins/Interchange/Runtime/Source/Import/Private/Texture/InterchangeJPGTranslator.cpp
+++ b/Engine/Plugins/Interchange/Runtime/Source/Import/Private/Texture/InterchangeJPGTranslator.cpp
@@ -10,6 +10,7 @@
 #include "InterchangeImportLog.h"
 #include "InterchangeTextureNode.h"
 #include "Memory/SharedBuffer.h"
+#include "Misc/ConfigCacheIni.h"
 #include "Misc/FileHelper.h"
 #include "Misc/Paths.h"
 #include "Misc/ScopedSlowTask.h"
@@ -176,6 +177,13 @@ TOptional<UE::Interchange::FImportImage> UInterchangeJPGTranslator::GetTexturePa
 		bShouldAllocateRawDataBuffer
 	);
 
+	// Honor setting from TextureImporter.RetainJpegFormat in Editor.ini if it exists
+	const bool bShouldImportRawCache = bShouldImportRaw;
+	if (GConfig->GetBool(TEXT("TextureImporter"), TEXT("RetainJpegFormat"), bShouldImportRaw, GEditorIni) && bShouldImportRawCache != bShouldImportRaw)
+	{
+		UE_LOG(LogInterchangeImport, Log, TEXT("JPEG file [%s]: Pipeline setting 'bPreferCompressedSourceData' has been overridden by Editor setting 'RetainJpegFormat'"), *Filename);
+	}
+
 	TArray64<uint8> RawData;
 	if (bShouldImportRaw)
 	{
