The texture normal map identification is done after the factory pass before the postedit change, this allow to not have to build the texture twice if we detect its a normal maps

#rnx

diff --git a/Engine/Plugins/Interchange/Runtime/Source/Pipelines/Private/InterchangeGenericAssetsPipeline.cpp b/Engine/Plugins/Interchange/Runtime/Source/Pipelines/Private/InterchangeGenericAssetsPipeline.cpp
--- a/Engine/Plugins/Interchange/Runtime/Source/Pipelines/Private/InterchangeGenericAssetsPipeline.cpp
+++ b/Engine/Plugins/Interchange/Runtime/Source/Pipelines/Private/InterchangeGenericAssetsPipeline.cpp
@@ -219,6 +219,22 @@ void UInterchangeGenericAssetsPipeline::ExecutePipeline(UInterchangeBaseNodeCont
 		});
 }
 
+void UInterchangeGenericAssetsPipeline::ExecutePostFactoryPipeline(const UInterchangeBaseNodeContainer* InBaseNodeContainer, const FString& NodeKey, UObject* CreatedAsset, bool bIsAReimport)
+{
+	if (MaterialPipeline)
+	{
+		MaterialPipeline->ScriptedExecutePostFactoryPipeline(InBaseNodeContainer, NodeKey, CreatedAsset, bIsAReimport);
+	}
+	if (MeshPipeline)
+	{
+		MeshPipeline->ScriptedExecutePostFactoryPipeline(InBaseNodeContainer, NodeKey, CreatedAsset, bIsAReimport);
+	}
+	if (AnimationPipeline)
+	{
+		AnimationPipeline->ScriptedExecutePostFactoryPipeline(InBaseNodeContainer, NodeKey, CreatedAsset, bIsAReimport);
+	}
+}
+
 void UInterchangeGenericAssetsPipeline::ExecutePostImportPipeline(const UInterchangeBaseNodeContainer* InBaseNodeContainer, const FString& NodeKey, UObject* CreatedAsset, bool bIsAReimport)
 {
 	if (MaterialPipeline)
diff --git a/Engine/Plugins/Interchange/Runtime/Source/Pipelines/Private/InterchangeGenericMaterialPipeline.cpp b/Engine/Plugins/Interchange/Runtime/Source/Pipelines/Private/InterchangeGenericMaterialPipeline.cpp
--- a/Engine/Plugins/Interchange/Runtime/Source/Pipelines/Private/InterchangeGenericMaterialPipeline.cpp
+++ b/Engine/Plugins/Interchange/Runtime/Source/Pipelines/Private/InterchangeGenericMaterialPipeline.cpp
@@ -341,6 +341,14 @@ void UInterchangeGenericMaterialPipeline::ExecutePipeline(UInterchangeBaseNodeCo
 	}
 }
 
+void UInterchangeGenericMaterialPipeline::ExecutePostFactoryPipeline(const UInterchangeBaseNodeContainer* InBaseNodeContainer, const FString& NodeKey, UObject* CreatedAsset, bool bIsAReimport)
+{
+	if (TexturePipeline)
+	{
+		TexturePipeline->ScriptedExecutePostFactoryPipeline(InBaseNodeContainer, NodeKey, CreatedAsset, bIsAReimport);
+	}
+}
+
 void UInterchangeGenericMaterialPipeline::ExecutePostImportPipeline(const UInterchangeBaseNodeContainer* InBaseNodeContainer, const FString& NodeKey, UObject* CreatedAsset, bool bIsAReimport)
 {
 	if (TexturePipeline)
diff --git a/Engine/Plugins/Interchange/Runtime/Source/Pipelines/Private/InterchangeGenericTexturePipeline.cpp b/Engine/Plugins/Interchange/Runtime/Source/Pipelines/Private/InterchangeGenericTexturePipeline.cpp
--- a/Engine/Plugins/Interchange/Runtime/Source/Pipelines/Private/InterchangeGenericTexturePipeline.cpp
+++ b/Engine/Plugins/Interchange/Runtime/Source/Pipelines/Private/InterchangeGenericTexturePipeline.cpp
@@ -73,7 +73,6 @@ namespace UE::Interchange::Private
 	{
 		if (Texture)
 		{
-			Texture->PreEditChange(nullptr);
 			if (UE::NormalMapIdentification::HandleAssetPostImport(Texture))
 			{
 				UE_LOG(LogInterchangePipeline, Display, TEXT("Auto-detected normal map"));
@@ -84,7 +83,6 @@ namespace UE::Interchange::Private
 				}
 			}
 			// this will rebuild the texture if it changed to normal map
-			Texture->PostEditChange(); 
 		}
 	}
 #endif
@@ -171,7 +169,7 @@ void UInterchangeGenericTexturePipeline::ExecutePipeline(UInterchangeBaseNodeCon
 	}
 }
 
-void UInterchangeGenericTexturePipeline::ExecutePostImportPipeline(const UInterchangeBaseNodeContainer* InBaseNodeContainer, const FString& NodeKey, UObject* CreatedAsset, bool bIsAReimport)
+void UInterchangeGenericTexturePipeline::ExecutePostFactoryPipeline(const UInterchangeBaseNodeContainer* InBaseNodeContainer, const FString& NodeKey, UObject* CreatedAsset, bool bIsAReimport)
 {
 	//We do not use the provided base container since ExecutePreImportPipeline cache it
 	//We just make sure the same one is pass in parameter
@@ -335,6 +333,7 @@ void UInterchangeGenericTexturePipeline::PostImportTextureAssetImport(UObject* C
 #if WITH_EDITOR
 
 	// this is run on main thread
+	check(IsInGameThread());
 	//	after texture may have started compiling
 
 	if (!bIsAReimport && bDetectNormalMapTexture)
@@ -344,42 +343,9 @@ void UInterchangeGenericTexturePipeline::PostImportTextureAssetImport(UObject* C
 			// if it's already a normal map, no need to run NormalMapIdentification
 			if (!Texture->IsNormalMap())
 			{
-				// @todo Interchange: This can create 2 build of the texture (we should revisit this at some point)
-
-				if (FTextureCompilingManager::Get().IsCompilingTexture(Texture))
-				{
-					// install a lambda to do AdjustTextureForNormalMap which will run after compile is done
-
-					TWeakObjectPtr<UTexture> WeakTexturePtr = Texture;
-					TSharedRef<FDelegateHandle> HandlePtr = MakeShared<FDelegateHandle>();
-					HandlePtr.Get() = FTextureCompilingManager::Get().OnTexturePostCompileEvent().AddLambda([this, WeakTexturePtr, HandlePtr](const TArrayView<UTexture* const>&)
-						{
-							if (UTexture* TextureToTest = WeakTexturePtr.Get())
-							{
-								if (FTextureCompilingManager::Get().IsCompilingTexture(TextureToTest))
-								{
-									// I'm still in compile queue, this lambda will get called again when I finish the next compile
-									//	OnTexturePostCompileEvent().Remove is not called, leave the lambda
-									// this can happen if multiple compiles of this texture got queued
-									return;
-								}
-
-								UE::Interchange::Private::AdjustTextureForNormalMap(TextureToTest, bFlipNormalMapGreenChannel);
-								// ?? does the asset registry need to be fixed?
-								//	asset import notifications were sent earlier
-							}
-
-							// remove myself:
-							FTextureCompilingManager::Get().OnTexturePostCompileEvent().Remove(HandlePtr.Get());
-						});
-				}
-				else
-				{
-					//Texture->Modify(); // Modify blocks on async build task being done
-
-					// AdjustTextureForNormalMap does a PostEditChange which triggers a rebuild
-					UE::Interchange::Private::AdjustTextureForNormalMap(Texture, bFlipNormalMapGreenChannel);
-				}
+ 				check(!FTextureCompilingManager::Get().IsCompilingTexture(Texture));
+				// AdjustTextureForNormalMap does a PostEditChange which triggers a rebuild
+				UE::Interchange::Private::AdjustTextureForNormalMap(Texture, bFlipNormalMapGreenChannel);
 			}
 		}
 	}
diff --git a/Engine/Plugins/Interchange/Runtime/Source/Pipelines/Public/InterchangeGenericAssetsPipeline.h b/Engine/Plugins/Interchange/Runtime/Source/Pipelines/Public/InterchangeGenericAssetsPipeline.h
--- a/Engine/Plugins/Interchange/Runtime/Source/Pipelines/Public/InterchangeGenericAssetsPipeline.h
+++ b/Engine/Plugins/Interchange/Runtime/Source/Pipelines/Public/InterchangeGenericAssetsPipeline.h
@@ -81,7 +81,7 @@ public:
 protected:
 
 	virtual void ExecutePipeline(UInterchangeBaseNodeContainer* InBaseNodeContainer, const TArray<UInterchangeSourceData*>& InSourceDatas) override;
-
+	virtual void ExecutePostFactoryPipeline(const UInterchangeBaseNodeContainer* BaseNodeContainer, const FString& NodeKey, UObject* CreatedAsset, bool bIsAReimport) override;
 	virtual void ExecutePostImportPipeline(const UInterchangeBaseNodeContainer* BaseNodeContainer, const FString& NodeKey, UObject* CreatedAsset, bool bIsAReimport) override;
 
 	virtual bool CanExecuteOnAnyThread(EInterchangePipelineTask PipelineTask) override
diff --git a/Engine/Plugins/Interchange/Runtime/Source/Pipelines/Public/InterchangeGenericMaterialPipeline.h b/Engine/Plugins/Interchange/Runtime/Source/Pipelines/Public/InterchangeGenericMaterialPipeline.h
--- a/Engine/Plugins/Interchange/Runtime/Source/Pipelines/Public/InterchangeGenericMaterialPipeline.h
+++ b/Engine/Plugins/Interchange/Runtime/Source/Pipelines/Public/InterchangeGenericMaterialPipeline.h
@@ -88,6 +88,7 @@ public:
 
 protected:
 	virtual void ExecutePipeline(UInterchangeBaseNodeContainer* InBaseNodeContainer, const TArray<UInterchangeSourceData*>& InSourceDatas) override;
+	virtual void ExecutePostFactoryPipeline(const UInterchangeBaseNodeContainer* BaseNodeContainer, const FString& NodeKey, UObject* CreatedAsset, bool bIsAReimport) override;
 	virtual void ExecutePostImportPipeline(const UInterchangeBaseNodeContainer* BaseNodeContainer, const FString& NodeKey, UObject* CreatedAsset, bool bIsAReimport) override;
 	virtual void SetReimportSourceIndex(UClass* ReimportObjectClass, const int32 SourceFileIndex) override;
 	/** END UInterchangePipelineBase overrides */
diff --git a/Engine/Plugins/Interchange/Runtime/Source/Pipelines/Public/InterchangeGenericTexturePipeline.h b/Engine/Plugins/Interchange/Runtime/Source/Pipelines/Public/InterchangeGenericTexturePipeline.h
--- a/Engine/Plugins/Interchange/Runtime/Source/Pipelines/Public/InterchangeGenericTexturePipeline.h
+++ b/Engine/Plugins/Interchange/Runtime/Source/Pipelines/Public/InterchangeGenericTexturePipeline.h
@@ -62,7 +62,7 @@ public:
 public:
 	virtual void AdjustSettingsForContext(EInterchangePipelineContext ImportType, TObjectPtr<UObject> ReimportAsset) override;
 	virtual void ExecutePipeline(UInterchangeBaseNodeContainer* InBaseNodeContainer, const TArray<UInterchangeSourceData*>& InSourceDatas) override;
-	virtual void ExecutePostImportPipeline(const UInterchangeBaseNodeContainer* InBaseNodeContainer, const FString& NodeKey, UObject* CreatedAsset, bool bIsAReimport) override;
+	virtual void ExecutePostFactoryPipeline(const UInterchangeBaseNodeContainer* BaseNodeContainer, const FString& NodeKey, UObject* CreatedAsset, bool bIsAReimport) override;
 
 protected:
 	UInterchangeTextureFactoryNode* HandleCreationOfTextureFactoryNode(const UInterchangeTextureNode* TextureNode);
diff --git a/Engine/Source/Runtime/TextureUtilitiesCommon/Private/NormalMapIdentification.cpp b/Engine/Source/Runtime/TextureUtilitiesCommon/Private/NormalMapIdentification.cpp
--- a/Engine/Source/Runtime/TextureUtilitiesCommon/Private/NormalMapIdentification.cpp
+++ b/Engine/Source/Runtime/TextureUtilitiesCommon/Private/NormalMapIdentification.cpp
@@ -6,6 +6,7 @@
 
 #include "Engine/Texture2D.h"
 #include "Framework/Notifications/NotificationManager.h"
+#include "TextureCompiler.h"
 #include "Widgets/Notifications/SNotificationList.h"
 
 #define NORMALMAP_IDENTIFICATION_TIMING	(0)
@@ -451,6 +452,14 @@ public:
 		UTexture2D* Texture2D = Texture.IsValid() ? Cast<UTexture2D>(Texture.Get()) : NULL;
 		if ( Texture2D )
 		{
+			if (FTextureCompilingManager::Get().IsCompilingTexture(Texture2D))
+			{
+				// Block until compile is done
+				TArray<UTexture*> TextureArray;
+				TextureArray.Add(Texture2D);
+				FTextureCompilingManager::Get().FinishCompilation(TextureArray);
+			}
+
 			if ( Texture2D->CompressionSettings == TC_Normalmap )
 			{
 				// Must wait until the texture is done with previous operations before changing settings and getting it to rebuild.
@@ -465,7 +474,6 @@ public:
 					Texture2D->SRGB = true;
 					Texture2D->LODGroup = TEXTUREGROUP_World;
 				}
-
 				Texture2D->PostEditChange();
 			}
 		}
@@ -492,7 +500,6 @@ bool UE::NormalMapIdentification::HandleAssetPostImport( UTexture* Texture )
 			// Set the compression settings and no gamma correction for a normal map
 			{
 				Texture->SetFlags(RF_Transactional);
-				Texture->Modify(); // Modify blocks on pending texture builds
 				Texture->CompressionSettings = TC_Normalmap;
 				Texture->SRGB = false;
 				Texture->LODGroup = TEXTUREGROUP_WorldNormalMap;
