AssetTypeActions_MaterialInstance:
- Allow material instances to be reimported from the content browser

InterchangeMaterialFactory:
- Adding the Interchange reimport strategy for the Material Instances

DatasmithImportFactory
- Fixing the window dialog not be prompted when a Material Instance  uses the Reimport With New File from the Content Browser

diff --git a/Engine/Plugins/Enterprise/DatasmithImporter/Source/DatasmithImporter/Private/DatasmithImportFactory.cpp b/Engine/Plugins/Enterprise/DatasmithImporter/Source/DatasmithImporter/Private/DatasmithImportFactory.cpp
--- a/Engine/Plugins/Enterprise/DatasmithImporter/Source/DatasmithImporter/Private/DatasmithImportFactory.cpp
+++ b/Engine/Plugins/Enterprise/DatasmithImporter/Source/DatasmithImporter/Private/DatasmithImportFactory.cpp
@@ -31,6 +31,7 @@
 #include "Logging/LogMacros.h"
 #include "Logging/MessageLog.h"
 #include "Materials/Material.h"
+#include "Materials/MaterialInstanceConstant.h"
 #include "Misc/PackageName.h"
 #include "Misc/ScopedSlowTask.h"
 #include "PackageTools.h"
@@ -442,6 +443,7 @@ FText UDatasmithImportFactory::GetDisplayName() const
 bool UDatasmithImportFactory::DoesSupportClass(UClass* InClass)
 {
 	return InClass == UMaterial::StaticClass()
+		|| InClass == UMaterialInstanceConstant::StaticClass()
 		|| InClass == UStaticMesh::StaticClass()
 		|| InClass == UDatasmithScene::StaticClass();
 }
diff --git a/Engine/Plugins/Interchange/Runtime/Source/Import/Private/Material/InterchangeMaterialFactory.cpp b/Engine/Plugins/Interchange/Runtime/Source/Import/Private/Material/InterchangeMaterialFactory.cpp
--- a/Engine/Plugins/Interchange/Runtime/Source/Import/Private/Material/InterchangeMaterialFactory.cpp
+++ b/Engine/Plugins/Interchange/Runtime/Source/Import/Private/Material/InterchangeMaterialFactory.cpp
@@ -1,13 +1,14 @@
 // Copyright Epic Games, Inc. All Rights Reserved. 
 #include "Material/InterchangeMaterialFactory.h"
 
+#include "InterchangeAssetImportData.h"
 #include "InterchangeImportCommon.h"
 #include "InterchangeImportLog.h"
 #include "InterchangeMaterialFactoryNode.h"
-#include "InterchangeTextureNode.h"
-#include "InterchangeTextureFactoryNode.h"
 #include "InterchangeShaderGraphNode.h"
 #include "InterchangeSourceData.h"
+#include "InterchangeTextureNode.h"
+#include "InterchangeTextureFactoryNode.h"
 #include "Materials/Material.h"
 #include "Materials/MaterialExpression.h"
 #include "Materials/MaterialExpressionClearCoatNormalCustomOutput.h"
@@ -423,7 +424,7 @@ UObject* UInterchangeMaterialFactory::CreateAsset(const FCreateAssetParams& Argu
 		return nullptr;
 	}
 
-	const UInterchangeBaseMaterialFactoryNode* MaterialFactoryNode = Cast<UInterchangeBaseMaterialFactoryNode>(Arguments.AssetNode);
+	UInterchangeBaseMaterialFactoryNode* MaterialFactoryNode = Cast<UInterchangeBaseMaterialFactoryNode>(Arguments.AssetNode);
 	if (MaterialFactoryNode == nullptr)
 	{
 		return nullptr;
@@ -474,7 +475,26 @@ UObject* UInterchangeMaterialFactory::CreateAsset(const FCreateAssetParams& Argu
 				MaterialFactoryNode->ApplyAllCustomAttributeToObject(MaterialObject);
 			}
 		}
-		
+#if WITH_EDITORONLY_DATA
+		else
+		{
+			if(UMaterialInstance* MaterialInstance = Cast<UMaterialInstance>(MaterialObject))
+			{
+				//Apply the re import strategy 
+				UInterchangeAssetImportData* InterchangeAssetImportData = Cast<UInterchangeAssetImportData>(MaterialInstance->AssetImportData);
+				UInterchangeFactoryBaseNode* PreviousNode = nullptr;
+				if(InterchangeAssetImportData)
+				{
+					PreviousNode = InterchangeAssetImportData->NodeContainer->GetFactoryNode(InterchangeAssetImportData->NodeUniqueID);
+				}
+				UInterchangeFactoryBaseNode* CurrentNode = NewObject<UInterchangeMaterialFactoryNode>(GetTransientPackage());
+				UInterchangeBaseNode::CopyStorage(MaterialFactoryNode, CurrentNode);
+				CurrentNode->FillAllCustomAttributeFromObject(MaterialInstance);
+				UE::Interchange::FFactoryCommon::ApplyReimportStrategyToAsset(MaterialInstance, PreviousNode, CurrentNode, MaterialFactoryNode);
+			}
+		}
+#endif // #if WITH_EDITORONLY_DATA
+
 		//Getting the file Hash will cache it into the source data
 		Arguments.SourceData->GetFileContentHash();
 
diff --git a/Engine/Source/Developer/AssetTools/Private/AssetTypeActions/AssetTypeActions_MaterialInstanceConstant.h b/Engine/Source/Developer/AssetTools/Private/AssetTypeActions/AssetTypeActions_MaterialInstanceConstant.h
--- a/Engine/Source/Developer/AssetTools/Private/AssetTypeActions/AssetTypeActions_MaterialInstanceConstant.h
+++ b/Engine/Source/Developer/AssetTools/Private/AssetTypeActions/AssetTypeActions_MaterialInstanceConstant.h
@@ -24,6 +24,7 @@ public:
 	virtual void OpenAssetEditor( const TArray<UObject*>& InObjects, TSharedPtr<class IToolkitHost> EditWithinLevelEditor = TSharedPtr<IToolkitHost>() ) override;
 	virtual bool CanFilter() override { return true; }
 	virtual uint32 GetCategories() override { return FAssetTypeActions_MaterialInterface::GetCategories() | AssetCategoryBit; }
+	virtual bool IsImportedAsset() const override {	return true; }
 
 private:
 	/** Handler for when FindParent is selected */
