Fixed issue with masked material

diff --git a/Engine/Plugins/Interchange/Runtime/Source/Pipelines/Private/InterchangeGenericMaterialPipeline.cpp b/Engine/Plugins/Interchange/Runtime/Source/Pipelines/Private/InterchangeGenericMaterialPipeline.cpp
--- a/Engine/Plugins/Interchange/Runtime/Source/Pipelines/Private/InterchangeGenericMaterialPipeline.cpp
+++ b/Engine/Plugins/Interchange/Runtime/Source/Pipelines/Private/InterchangeGenericMaterialPipeline.cpp
@@ -170,6 +170,37 @@ namespace UE::Interchange::InterchangeGenericMaterialPipeline::Private
 
 		return bRequiredPackagesLoaded;
 	}
+
+	void UpdateBlendModeBasedOnOpacityAttributes(const UInterchangeShaderGraphNode* ShaderGraphNode, UInterchangeMaterialFactoryNode* MaterialFactoryNode)
+	{
+		// Opacity Clip Value
+		bool bIsMasked = false;
+		{
+			float OpacityClipValue;
+			if (ShaderGraphNode->GetCustomOpacityMaskClipValue(OpacityClipValue))
+			{
+				MaterialFactoryNode->SetCustomOpacityMaskClipValue(OpacityClipValue);
+				bIsMasked = true;
+			}
+		}
+
+		// Don't change the blend mode if it was already set
+		TEnumAsByte<EBlendMode> BlendMode = bIsMasked ? EBlendMode::BLEND_Masked : EBlendMode::BLEND_Translucent;
+		if (!MaterialFactoryNode->GetCustomBlendMode(BlendMode))
+		{
+			MaterialFactoryNode->SetCustomBlendMode(BlendMode);
+		}
+
+		// If bland mode is masked or translucent, set lighting mode accordingly without changing it if it was already set
+		if(BlendMode == EBlendMode::BLEND_Masked || BlendMode == EBlendMode::BLEND_Translucent)
+		{
+			TEnumAsByte<ETranslucencyLightingMode> LightingMode = ETranslucencyLightingMode::TLM_Surface;
+			if (!MaterialFactoryNode->GetCustomTranslucencyLightingMode(LightingMode))
+			{
+				MaterialFactoryNode->SetCustomTranslucencyLightingMode(LightingMode);
+			}
+		}
+	}
 }
 
 UInterchangeGenericMaterialPipeline::UInterchangeGenericMaterialPipeline()
@@ -1078,29 +1109,9 @@ void UInterchangeGenericMaterialPipeline::HandleCommonParameters(const UIntercha
 					MaterialFactoryNode->ConnectOutputToOpacity(OpacityExpressionFactoryNode.Get<0>()->GetUniqueID(), OpacityExpressionFactoryNode.Get<1>());
 				}
 
-				// Opacity Clip Value
-				bool bIsMasked = false;
-				{
-					float OpacityClipValue;
-					if (ShaderGraphNode->GetCustomOpacityMaskClipValue(OpacityClipValue))
-					{
-						MaterialFactoryNode->SetCustomOpacityMaskClipValue(OpacityClipValue);
-						bIsMasked = true;
-					}
-				}
-
-				// Don't change the blend mode or the lighting mode if they were already set
-				TEnumAsByte<EBlendMode> BlendMode = bIsMasked ? EBlendMode::BLEND_Masked : EBlendMode::BLEND_Translucent;
-				if (!MaterialFactoryNode->GetCustomBlendMode(BlendMode))
-				{
-					MaterialFactoryNode->SetCustomBlendMode(BlendMode);
+				using namespace UE::Interchange::InterchangeGenericMaterialPipeline;
 
-					TEnumAsByte<ETranslucencyLightingMode> LightingMode = ETranslucencyLightingMode::TLM_Surface;
-					if (!MaterialFactoryNode->GetCustomTranslucencyLightingMode(LightingMode))
-					{
-						MaterialFactoryNode->SetCustomTranslucencyLightingMode(LightingMode);
-					}
-				}
+				Private::UpdateBlendModeBasedOnOpacityAttributes(ShaderGraphNode, MaterialFactoryNode);
 			}
 		}
 	}
@@ -2528,24 +2539,6 @@ bool UInterchangeGenericMaterialPipeline::HandleUnlitModel(const UInterchangeSha
 
 	bool bShadingModelHandled = false;
 
-	{
-		//gltf allows unlit color to be also translucent:
-		const bool bHasInput = UInterchangeShaderPortsAPI::HasInput(ShaderGraphNode, Common::Parameters::Opacity);
-
-		if (bHasInput)
-		{
-			TTuple<UInterchangeMaterialExpressionFactoryNode*, FString> OpacityExpressionFactoryNode =
-				CreateMaterialExpressionForInput(MaterialFactoryNode, ShaderGraphNode, Common::Parameters::Opacity.ToString(), MaterialFactoryNode->GetUniqueID());
-
-			if (OpacityExpressionFactoryNode.Get<0>())
-			{
-				MaterialFactoryNode->ConnectOutputToOpacity(OpacityExpressionFactoryNode.Get<0>()->GetUniqueID(), OpacityExpressionFactoryNode.Get<1>());
-			}
-
-			MaterialFactoryNode->SetCustomBlendMode(EBlendMode::BLEND_Translucent);
-		}
-	}
-
 	// Unlit Color
 	{
 		const bool bHasInput = UInterchangeShaderPortsAPI::HasInput(ShaderGraphNode, Unlit::Parameters::UnlitColor);
@@ -2560,6 +2553,26 @@ bool UInterchangeGenericMaterialPipeline::HandleUnlitModel(const UInterchangeSha
 				MaterialFactoryNode->ConnectOutputToEmissiveColor(ExpressionFactoryNode.Get<0>()->GetUniqueID(), ExpressionFactoryNode.Get<1>());
 			}
 
+			//gltf allows unlit color to be also translucent:
+			{
+				const bool bHasOpacityInput = UInterchangeShaderPortsAPI::HasInput(ShaderGraphNode, Common::Parameters::Opacity);
+
+				if (bHasOpacityInput)
+				{
+					TTuple<UInterchangeMaterialExpressionFactoryNode*, FString> OpacityExpressionFactoryNode =
+						CreateMaterialExpressionForInput(MaterialFactoryNode, ShaderGraphNode, Common::Parameters::Opacity.ToString(), MaterialFactoryNode->GetUniqueID());
+
+					if (OpacityExpressionFactoryNode.Get<0>())
+					{
+						MaterialFactoryNode->ConnectOutputToOpacity(OpacityExpressionFactoryNode.Get<0>()->GetUniqueID(), OpacityExpressionFactoryNode.Get<1>());
+					}
+
+					using namespace UE::Interchange::InterchangeGenericMaterialPipeline;
+
+					Private::UpdateBlendModeBasedOnOpacityAttributes(ShaderGraphNode, MaterialFactoryNode);
+				}
+			}
+
 			bShadingModelHandled = true;
 		}
 	}
