Interchange:
- Add Material SkinXX workflow in the generic asset pipeline
- Add LOD Naming workflow in the generic asset pipeline, the difference with the fbx legacy is we wont create generated LOD if there is hole in the specified lod
- Staticmesh collision was already implement
- UVChannel is implement in the fbx translator since the meshdescription can't pass the UVChannel names. Only fbx will support this feature to make sure interchange fbx is on par with legacy fbx.

#rnx

diff --git a/Engine/Plugins/Interchange/Runtime/Source/Dispatcher/Public/InterchangeCommands.h b/Engine/Plugins/Interchange/Runtime/Source/Dispatcher/Public/InterchangeCommands.h
--- a/Engine/Plugins/Interchange/Runtime/Source/Dispatcher/Public/InterchangeCommands.h
+++ b/Engine/Plugins/Interchange/Runtime/Source/Dispatcher/Public/InterchangeCommands.h
@@ -17,7 +17,7 @@ namespace UE
 		{
 		public:
 			//Major version should be updated when there is an existing API that has been change.
-			static int32 GetMajor() { return 11; }
+			static int32 GetMajor() { return 12; }
 			//Minor version must be updated when there is an addition to the API.
 			static int32 GetMinor() { return 0; }
 			//Patch version should be update if there is some bug fixes in the private code.
diff --git a/Engine/Plugins/Interchange/Runtime/Source/FactoryNodes/Private/InterchangeMeshFactoryNode.cpp b/Engine/Plugins/Interchange/Runtime/Source/FactoryNodes/Private/InterchangeMeshFactoryNode.cpp
--- a/Engine/Plugins/Interchange/Runtime/Source/FactoryNodes/Private/InterchangeMeshFactoryNode.cpp
+++ b/Engine/Plugins/Interchange/Runtime/Source/FactoryNodes/Private/InterchangeMeshFactoryNode.cpp
@@ -174,7 +174,12 @@ bool UInterchangeMeshFactoryNode::SetSlotMaterialDependencyUid(const FString& Sl
 
 bool UInterchangeMeshFactoryNode::RemoveSlotMaterialDependencyUid(const FString& SlotName)
 {
-	return SlotMaterialDependencies.RemoveKey(SlotName);
+	if (SlotMaterialDependencies.RemoveKey(SlotName))
+	{
+		SlotMaterialDependencies.RebuildCache();
+		return true;
+	}
+	return false;
 }
 
 bool UInterchangeMeshFactoryNode::GetCustomRecomputeNormals(bool& AttributeValue) const
diff --git a/Engine/Plugins/Interchange/Runtime/Source/Import/Private/Animation/InterchangeAnimSequenceFactory.cpp b/Engine/Plugins/Interchange/Runtime/Source/Import/Private/Animation/InterchangeAnimSequenceFactory.cpp
--- a/Engine/Plugins/Interchange/Runtime/Source/Import/Private/Animation/InterchangeAnimSequenceFactory.cpp
+++ b/Engine/Plugins/Interchange/Runtime/Source/Import/Private/Animation/InterchangeAnimSequenceFactory.cpp
@@ -298,7 +298,7 @@ namespace UE::Interchange::Private
 
 			const double BakeInterval = 1.0 / SampleRate;
 			const double SequenceLength = FMath::Max<double>(RangeEnd - RangeStart, MINIMUM_ANIMATION_LENGTH);
-			int32 BakeKeyCount = (SequenceLength / BakeInterval) + 1;
+			int32 BakeKeyCount = FMath::RoundToInt32(SequenceLength * SampleRate) + 1;
 
 			TMap<FString, FString> PayloadKeys;
 			AnimSequenceFactoryNode->GetSceneNodeAnimationPayloadKeys(PayloadKeys);
diff --git a/Engine/Plugins/Interchange/Runtime/Source/Import/Private/Mesh/InterchangeSkeletalMeshFactory.cpp b/Engine/Plugins/Interchange/Runtime/Source/Import/Private/Mesh/InterchangeSkeletalMeshFactory.cpp
--- a/Engine/Plugins/Interchange/Runtime/Source/Import/Private/Mesh/InterchangeSkeletalMeshFactory.cpp
+++ b/Engine/Plugins/Interchange/Runtime/Source/Import/Private/Mesh/InterchangeSkeletalMeshFactory.cpp
@@ -1062,19 +1062,24 @@ UObject* UInterchangeSkeletalMeshFactory::CreateAsset(const FCreateAssetParams&
 	// Update skeletal materials
 	TArray<FSkeletalMaterial>& Materials = SkeletalMesh->GetMaterials();
 
-	auto UpdateOrAddSkeletalMaterial = [&Materials](const FName& MaterialSlotName, UMaterialInterface* MaterialInterface)
+	auto UpdateOrAddSkeletalMaterial = [&Materials, bIsReImport](const FName& MaterialSlotName, UMaterialInterface* MaterialInterface)
 	{
+		UMaterialInterface* NewMaterial = MaterialInterface ? MaterialInterface : UMaterial::GetDefaultMaterial(MD_Surface);
+
 		FSkeletalMaterial* SkeletalMaterial = Materials.FindByPredicate([&MaterialSlotName](const FSkeletalMaterial& Material) { return Material.MaterialSlotName == MaterialSlotName; });
-		
 		if (SkeletalMaterial)
 		{
-			SkeletalMaterial->MaterialInterface = MaterialInterface;
+			//When we do a reimport we update the material interface only if the specified MaterialInterface is not null
+			if (!bIsReImport || MaterialInterface || !SkeletalMaterial->MaterialInterface)
+			{
+				SkeletalMaterial->MaterialInterface = NewMaterial;
+			}
 		}
 		else
 		{
 			const bool bEnableShadowCasting = true;
 			const bool bInRecomputeTangent = false;
-			Materials.Emplace(MaterialInterface, bEnableShadowCasting, bInRecomputeTangent, MaterialSlotName, MaterialSlotName);
+			Materials.Emplace(NewMaterial, bEnableShadowCasting, bInRecomputeTangent, MaterialSlotName, MaterialSlotName);
 		}
 	};
 
@@ -1089,7 +1094,7 @@ UObject* UInterchangeSkeletalMeshFactory::CreateAsset(const FCreateAssetParams&
 		const UInterchangeBaseMaterialFactoryNode* MaterialFactoryNode = Cast<UInterchangeBaseMaterialFactoryNode>(Arguments.NodeContainer->GetNode(SlotMaterialDependency.Value));
 		if (!MaterialFactoryNode || !MaterialFactoryNode->IsEnabled())
 		{
-			UpdateOrAddSkeletalMaterial(MaterialSlotName, UMaterial::GetDefaultMaterial(MD_Surface));
+			UpdateOrAddSkeletalMaterial(MaterialSlotName, nullptr);
 			continue;
 		}
 
@@ -1097,12 +1102,12 @@ UObject* UInterchangeSkeletalMeshFactory::CreateAsset(const FCreateAssetParams&
 		MaterialFactoryNode->GetCustomReferenceObject(MaterialFactoryNodeReferenceObject);
 		if (!MaterialFactoryNodeReferenceObject.IsValid())
 		{
-			UpdateOrAddSkeletalMaterial(MaterialSlotName, UMaterial::GetDefaultMaterial(MD_Surface));
+			UpdateOrAddSkeletalMaterial(MaterialSlotName, nullptr);
 			continue;
 		}
 
 		UMaterialInterface* MaterialInterface = Cast<UMaterialInterface>(MaterialFactoryNodeReferenceObject.ResolveObject());
-		UpdateOrAddSkeletalMaterial(MaterialSlotName, MaterialInterface ? MaterialInterface : UMaterial::GetDefaultMaterial(MD_Surface));
+		UpdateOrAddSkeletalMaterial(MaterialSlotName, MaterialInterface ? MaterialInterface : nullptr);
 	}
 
 	for (int32 LodIndex = 0; LodIndex < LodCount; ++LodIndex)
@@ -1413,9 +1418,22 @@ UObject* UInterchangeSkeletalMeshFactory::CreateAsset(const FCreateAssetParams&
 				if (LODMatIndex == INDEX_NONE)
 				{
 					LODMatIndex = Materials.Add(FSkeletalMaterial(ImportedMaterials[ImportedMaterialIndex].Material.Get(), true, false, ImportedMaterialName, ImportedMaterialName));
+					LodInfo->LODMaterialMap.Add(LODMatIndex);
+				}
+				else
+				{
+					if (CurrentLodIndex == 0)
+					{
+						//For base lod, the skeletalmesh build will reorder the FSkeletalMeshImportData face in the correct order
+						//see FLODUtilities::AdjustImportDataFaceMaterialIndex function, so we dont need to set 
+						LodInfo->LODMaterialMap.Add(ImportedMaterialIndex);
+					}
+					else
+					{
+						//For non base lod we want to use the LODMatIndex
+						LodInfo->LODMaterialMap.Add(LODMatIndex);
+					}
 				}
-
-				LodInfo->LODMaterialMap.Add(LODMatIndex);
 			}
 		}
 
diff --git a/Engine/Plugins/Interchange/Runtime/Source/Import/Private/Mesh/InterchangeStaticMeshFactory.cpp b/Engine/Plugins/Interchange/Runtime/Source/Import/Private/Mesh/InterchangeStaticMeshFactory.cpp
--- a/Engine/Plugins/Interchange/Runtime/Source/Import/Private/Mesh/InterchangeStaticMeshFactory.cpp
+++ b/Engine/Plugins/Interchange/Runtime/Source/Import/Private/Mesh/InterchangeStaticMeshFactory.cpp
@@ -126,6 +126,7 @@ UObject* UInterchangeStaticMeshFactory::CreateAsset(const FCreateAssetParams& Ar
 	// create an asset if it doesn't exist
 	UObject* ExistingAsset = StaticFindObject(nullptr, Arguments.Parent, *Arguments.AssetName);
 
+	const bool bReimport = Arguments.ReimportObject&& ExistingAsset;
 	UObject* StaticMeshObject = nullptr;
 
 	// create a new static mesh or overwrite existing asset, if possible
@@ -170,28 +171,33 @@ UObject* UInterchangeStaticMeshFactory::CreateAsset(const FCreateAssetParams& Ar
 
 	// If we are reimporting, cache the existing vertex colors so they can be optionally reapplied after reimport
 	TMap<FVector3f, FColor> ExisitingVertexColorData;
-	if (Arguments.ReimportObject && ExistingAsset)
+	if (bReimport)
 	{
 		StaticMesh->GetVertexColorData(ExisitingVertexColorData);
 	}
 
 	// Set material slots from imported materials
-	auto UpdateOrAddStaticMaterial = [&StaticMesh](const FName& MaterialSlotName, UMaterialInterface* MaterialInterface)
+	auto UpdateOrAddStaticMaterial = [&StaticMesh, bReimport](const FName& MaterialSlotName, UMaterialInterface* MaterialInterface)
 	{
-		int32 MaterialSlotIndex = StaticMesh->GetMaterialIndexFromImportedMaterialSlotName(MaterialSlotName);
+		UMaterialInterface* NewMaterial = MaterialInterface ? MaterialInterface : UMaterial::GetDefaultMaterial(MD_Surface);
 
-		if (MaterialSlotIndex == INDEX_NONE)
+		FStaticMaterial* StaticMaterial = StaticMesh->GetStaticMaterials().FindByPredicate([&MaterialSlotName](const FStaticMaterial& Material) { return Material.MaterialSlotName == MaterialSlotName; });
+		if (StaticMaterial)
 		{
-			MaterialSlotIndex = StaticMesh->GetStaticMaterials().Emplace(MaterialInterface, MaterialSlotName);
+			//When we do a re-import we update the material interface only if the specified MaterialInterface is not null
+			if (!bReimport || MaterialInterface || !StaticMaterial->MaterialInterface)
+			{
+				StaticMaterial->MaterialInterface = NewMaterial;
+			}
+		}
+		else
+		{
+			int32 MaterialSlotIndex = StaticMesh->GetStaticMaterials().Emplace(NewMaterial, MaterialSlotName);
 #if !WITH_EDITOR
 			// UV density is not supported to be generated at runtime for now. We fake that it has been initialized so that we don't trigger ensures.
 			StaticMesh->GetStaticMaterials()[MaterialSlotIndex].UVChannelData = FMeshUVChannelInfo(1.f);
 #endif
 		}
-		else
-		{
-			StaticMesh->GetStaticMaterials()[MaterialSlotIndex].MaterialInterface = MaterialInterface;
-		}
 	};
 
 	TMap<FString, FString> SlotMaterialDependencies;
@@ -203,19 +209,19 @@ UObject* UInterchangeStaticMeshFactory::CreateAsset(const FCreateAssetParams& Ar
 		const UInterchangeBaseMaterialFactoryNode* MaterialFactoryNode = Cast<UInterchangeBaseMaterialFactoryNode>(Arguments.NodeContainer->GetNode(SlotMaterialDependency.Value));
 		if (!MaterialFactoryNode || !MaterialFactoryNode->IsEnabled())
 		{
-			UpdateOrAddStaticMaterial(MaterialSlotName, UMaterial::GetDefaultMaterial(MD_Surface));
+			UpdateOrAddStaticMaterial(MaterialSlotName, nullptr);
 			continue;
 		}
 		FSoftObjectPath MaterialFactoryNodeReferenceObject;
 		MaterialFactoryNode->GetCustomReferenceObject(MaterialFactoryNodeReferenceObject);
 		if (!MaterialFactoryNodeReferenceObject.IsValid())
 		{
-			UpdateOrAddStaticMaterial(MaterialSlotName, UMaterial::GetDefaultMaterial(MD_Surface));
+			UpdateOrAddStaticMaterial(MaterialSlotName, nullptr);
 			continue;
 		}
 
 		UMaterialInterface* MaterialInterface = Cast<UMaterialInterface>(MaterialFactoryNodeReferenceObject.TryLoad());
-		UpdateOrAddStaticMaterial(MaterialSlotName, MaterialInterface ? MaterialInterface : UMaterial::GetDefaultMaterial(MD_Surface));
+		UpdateOrAddStaticMaterial(MaterialSlotName, MaterialInterface ? MaterialInterface : nullptr);
 	}
 
 	// Now import geometry for each LOD
diff --git a/Engine/Plugins/Interchange/Runtime/Source/Parsers/Fbx/Private/FbxAPI.cpp b/Engine/Plugins/Interchange/Runtime/Source/Parsers/Fbx/Private/FbxAPI.cpp
--- a/Engine/Plugins/Interchange/Runtime/Source/Parsers/Fbx/Private/FbxAPI.cpp
+++ b/Engine/Plugins/Interchange/Runtime/Source/Parsers/Fbx/Private/FbxAPI.cpp
@@ -25,6 +25,17 @@ namespace UE
 			FFbxParser::~FFbxParser()
 			{
 				PayloadContexts.Empty();
+				FbxHelper = nullptr;
+			}
+			
+			const TSharedPtr<FFbxHelper> FFbxParser::GetFbxHelper()
+			{
+				if (!FbxHelper.IsValid())
+				{
+					FbxHelper = MakeShared<FFbxHelper>();
+				}
+				check(FbxHelper.IsValid());
+				return FbxHelper;
 			}
 
 			bool FFbxParser::LoadFbxFile(const FString& Filename)
diff --git a/Engine/Plugins/Interchange/Runtime/Source/Parsers/Fbx/Private/FbxAPI.h b/Engine/Plugins/Interchange/Runtime/Source/Parsers/Fbx/Private/FbxAPI.h
--- a/Engine/Plugins/Interchange/Runtime/Source/Parsers/Fbx/Private/FbxAPI.h
+++ b/Engine/Plugins/Interchange/Runtime/Source/Parsers/Fbx/Private/FbxAPI.h
@@ -18,6 +18,7 @@ namespace UE
 		namespace Private
 		{
 			class FPayloadContextBase;
+			struct FFbxHelper;
 		}
 	}
 }
@@ -38,6 +39,9 @@ namespace UE
 
 				~FFbxParser();
 
+				//return the fbx helper for this parser
+				const TSharedPtr<FFbxHelper> GetFbxHelper();
+
 				/* Load an fbx file into the fbx sdk, return false if the file could not be load. */
 				bool LoadFbxFile(const FString& Filename);
 
@@ -86,6 +90,7 @@ namespace UE
 				FbxGeometryConverter* SDKGeometryConverter = nullptr;
 				FString SourceFilename;
 				TMap<FString, TSharedPtr<FPayloadContextBase>> PayloadContexts;
+				TSharedPtr<FFbxHelper> FbxHelper;
 			};
 		}//ns Private
 	}//ns Interchange
diff --git a/Engine/Plugins/Interchange/Runtime/Source/Parsers/Fbx/Private/FbxAnimation.cpp b/Engine/Plugins/Interchange/Runtime/Source/Parsers/Fbx/Private/FbxAnimation.cpp
--- a/Engine/Plugins/Interchange/Runtime/Source/Parsers/Fbx/Private/FbxAnimation.cpp
+++ b/Engine/Plugins/Interchange/Runtime/Source/Parsers/Fbx/Private/FbxAnimation.cpp
@@ -562,7 +562,7 @@ namespace UE::Interchange::Private
 		if (!ensure(FetchPayloadData.Node))
 		{
 			UInterchangeResultError_Generic* Message = Parser.AddMessage<UInterchangeResultError_Generic>();
-			Message->InterchangeKey = FFbxHelper::GetFbxNodeHierarchyName(FetchPayloadData.Node);
+			Message->InterchangeKey = Parser.GetFbxHelper()->GetFbxNodeHierarchyName(FetchPayloadData.Node);
 			Message->Text = LOCTEXT("FBXNodeNull", "Cannot fetch FBX animation transform payload because the FBX node is null.");
 			return false;
 		}
@@ -592,7 +592,7 @@ namespace UE::Interchange::Private
 		if (!ensure(FetchPayloadData.Node != nullptr))
 		{
 			UInterchangeResultError_Generic* Message = Parser.AddMessage<UInterchangeResultError_Generic>();
-			Message->InterchangeKey = FFbxHelper::GetFbxNodeHierarchyName(FetchPayloadData.Node);
+			Message->InterchangeKey = Parser.GetFbxHelper()->GetFbxNodeHierarchyName(FetchPayloadData.Node);
 			Message->Text = LOCTEXT("InternalFetchCurveNodePayloadToFile_FBXNodeNull", "Cannot fetch FBX animation curve payload because the FBX node is null.");
 			return false;
 		}
@@ -600,7 +600,7 @@ namespace UE::Interchange::Private
 		if (!ensure(FetchPayloadData.AnimCurves != nullptr))
 		{
 			UInterchangeResultError_Generic* Message = Parser.AddMessage<UInterchangeResultError_Generic>();
-			Message->InterchangeKey = FFbxHelper::GetFbxNodeHierarchyName(FetchPayloadData.Node);
+			Message->InterchangeKey = Parser.GetFbxHelper()->GetFbxNodeHierarchyName(FetchPayloadData.Node);
 			Message->Text = LOCTEXT("InternalFetchCurveNodePayloadToFile_FBXCurveNull", "Cannot fetch FBX user attribute animation curve payload because the FBX anim curve node is null.");
 			return false;
 		}
@@ -721,7 +721,7 @@ namespace UE::Interchange::Private
 			UInterchangeResultError_Generic* Message = Parser.AddMessage<UInterchangeResultError_Generic>();
 			if (FbxNode* Node = Geometry->GetNode())
 			{
-				Message->InterchangeKey = FFbxHelper::GetFbxNodeHierarchyName(Node);
+				Message->InterchangeKey = Parser.GetFbxHelper()->GetFbxNodeHierarchyName(Node);
 			}
 			Message->Text = LOCTEXT("InternalFetchMorphTargetCurvePayloadToFile_FBXCurveNull", "Cannot fetch FBX morph target animation curve payload because the FBX anim curve node is null.");
 			return false;
@@ -743,6 +743,7 @@ namespace UE::Interchange::Private
 	}
 
 	void FFbxAnimation::AddNodeTransformAnimation(FbxScene* SDKScene
+		, FFbxParser& Parser
 		, FbxNode* Node
 		, UInterchangeSceneNode* SceneNode
 		, TMap<FString, TSharedPtr<FPayloadContextBase>>& PayloadContexts)
@@ -758,7 +759,7 @@ namespace UE::Interchange::Private
 
 		if (Parameters.IsNodeAnimated)
 		{
-			FString PayLoadKey = FFbxHelper::GetFbxNodeHierarchyName(Node) + TEXT("_AnimationPayloadKey");
+			FString PayLoadKey = Parser.GetFbxHelper()->GetFbxNodeHierarchyName(Node) + TEXT("_AnimationPayloadKey");
 			if (ensure(!PayloadContexts.Contains(PayLoadKey)))
 			{
 				TSharedPtr<FAnimationPayloadContextTransform> AnimPayload = MakeShared<FAnimationPayloadContextTransform>();
@@ -775,7 +776,8 @@ namespace UE::Interchange::Private
 		}
 	}
 
-	void FFbxAnimation::AddNodeAttributeCurvesAnimation(FbxNode* Node
+	void FFbxAnimation::AddNodeAttributeCurvesAnimation(FFbxParser& Parser
+		, FbxNode* Node
 		, FbxProperty& Property
 		, FbxAnimCurveNode* AnimCurveNode
 		, UInterchangeSceneNode* SceneNode
@@ -783,8 +785,8 @@ namespace UE::Interchange::Private
 		, EFbxType PropertyType
 		, TOptional<FString>& OutPayloadKey)
 	{
-		const FString PropertyName = FFbxHelper::GetFbxPropertyName(Property);
-		const FString PayLoadKey = FFbxHelper::GetFbxNodeHierarchyName(Node) + PropertyName + TEXT("_AnimationPayloadKey");
+		const FString PropertyName = Parser.GetFbxHelper()->GetFbxPropertyName(Property);
+		const FString PayLoadKey = Parser.GetFbxHelper()->GetFbxNodeHierarchyName(Node) + PropertyName + TEXT("_AnimationPayloadKey");
 		if (ensure(!PayloadContexts.Contains(PayLoadKey)))
 		{
 			TSharedPtr<FAnimationPayloadContextTransform> AnimPayload = MakeShared<FAnimationPayloadContextTransform>();
diff --git a/Engine/Plugins/Interchange/Runtime/Source/Parsers/Fbx/Private/FbxAnimation.h b/Engine/Plugins/Interchange/Runtime/Source/Parsers/Fbx/Private/FbxAnimation.h
--- a/Engine/Plugins/Interchange/Runtime/Source/Parsers/Fbx/Private/FbxAnimation.h
+++ b/Engine/Plugins/Interchange/Runtime/Source/Parsers/Fbx/Private/FbxAnimation.h
@@ -61,12 +61,14 @@ namespace UE::Interchange::Private
 	public:
 		/** This function add the payload key if the scene node transform is animated. */
 		static void AddNodeTransformAnimation(FbxScene* SDKScene
+			, FFbxParser& Parser
 			, FbxNode* Node
 			, UInterchangeSceneNode* UnrealNode
 			, TMap<FString, TSharedPtr<FPayloadContextBase>>& PayloadContexts);
 
 		/** This function add the payload key for an animated node user attribute. */
-		static void AddNodeAttributeCurvesAnimation(FbxNode* Node
+		static void AddNodeAttributeCurvesAnimation(FFbxParser& Parser
+			, FbxNode* Node
 			, FbxProperty& Property
 			, FbxAnimCurveNode* AnimCurveNode
 			, UInterchangeSceneNode* SceneNode
diff --git a/Engine/Plugins/Interchange/Runtime/Source/Parsers/Fbx/Private/FbxCamera.cpp b/Engine/Plugins/Interchange/Runtime/Source/Parsers/Fbx/Private/FbxCamera.cpp
--- a/Engine/Plugins/Interchange/Runtime/Source/Parsers/Fbx/Private/FbxCamera.cpp
+++ b/Engine/Plugins/Interchange/Runtime/Source/Parsers/Fbx/Private/FbxCamera.cpp
@@ -66,8 +66,8 @@ namespace UE
 
 					if (NodeAttribute && NodeAttribute->GetAttributeType() == FbxNodeAttribute::eCamera)
 					{
-						FString NodeName = FFbxHelper::GetNodeAttributeName(NodeAttribute, UInterchangeCameraNode::StaticAssetTypeName());
-						FString NodeUid = FFbxHelper::GetNodeAttributeUniqueID(NodeAttribute, UInterchangeCameraNode::StaticAssetTypeName());
+						FString NodeName = Parser.GetFbxHelper()->GetNodeAttributeName(NodeAttribute, UInterchangeCameraNode::StaticAssetTypeName());
+						FString NodeUid = Parser.GetFbxHelper()->GetNodeAttributeUniqueID(NodeAttribute, UInterchangeCameraNode::StaticAssetTypeName());
 
 						const UInterchangeCameraNode* CameraNode = Cast<const UInterchangeCameraNode>(NodeContainer.GetNode(NodeUid));
 
diff --git a/Engine/Plugins/Interchange/Runtime/Source/Parsers/Fbx/Private/FbxHelper.cpp b/Engine/Plugins/Interchange/Runtime/Source/Parsers/Fbx/Private/FbxHelper.cpp
--- a/Engine/Plugins/Interchange/Runtime/Source/Parsers/Fbx/Private/FbxHelper.cpp
+++ b/Engine/Plugins/Interchange/Runtime/Source/Parsers/Fbx/Private/FbxHelper.cpp
@@ -14,7 +14,7 @@ namespace UE
 	{
 		namespace Private
 		{
-			FString FFbxHelper::GetMeshName(FbxGeometryBase* Mesh)
+			FString FFbxHelper::GetMeshName(FbxGeometryBase* Mesh) const
 			{
 				if (!Mesh)
 				{
@@ -34,7 +34,7 @@ namespace UE
 				return GetNodeAttributeName(Mesh, DefaultPrefix);
 			}
 
-			FString FFbxHelper::GetMeshUniqueID(FbxGeometryBase* Mesh)
+			FString FFbxHelper::GetMeshUniqueID(FbxGeometryBase* Mesh) const
 			{
 				if (!Mesh)
 				{
@@ -54,7 +54,7 @@ namespace UE
 				return GetNodeAttributeUniqueID(Mesh, MeshUniqueID);
 			}
 
-			FString FFbxHelper::GetNodeAttributeName(FbxNodeAttribute* NodeAttribute, const FStringView DefaultNamePrefix)
+			FString FFbxHelper::GetNodeAttributeName(FbxNodeAttribute* NodeAttribute, const FStringView DefaultNamePrefix) const
 			{
 				FString NodeAttributeName = FFbxHelper::GetFbxObjectName(NodeAttribute);
 				if (NodeAttributeName.IsEmpty())
@@ -72,7 +72,7 @@ namespace UE
 				return NodeAttributeName;
 			}
 
-			FString FFbxHelper::GetNodeAttributeUniqueID(FbxNodeAttribute* NodeAttribute, const FStringView Prefix)
+			FString FFbxHelper::GetNodeAttributeUniqueID(FbxNodeAttribute* NodeAttribute, const FStringView Prefix) const
 			{
 				if (!NodeAttribute)
 				{
@@ -99,7 +99,7 @@ namespace UE
 				return NodeAttributeUniqueID;
 			}
 
-			FString FFbxHelper::GetFbxPropertyName(const FbxProperty& Property)
+			FString FFbxHelper::GetFbxPropertyName(const FbxProperty& Property) const
 			{
 				FString PropertyName = FFbxConvert::MakeName(Property.GetName());
 				if (PropertyName.Equals(TEXT("none"), ESearchCase::IgnoreCase))
@@ -110,7 +110,7 @@ namespace UE
 				return PropertyName;
 			}
 
-			FString FFbxHelper::GetFbxObjectName(const FbxObject* Object)
+			FString FFbxHelper::GetFbxObjectName(const FbxObject* Object) const
 			{
 				if (!Object)
 				{
@@ -122,10 +122,35 @@ namespace UE
 					//Replace None by Null because None clash with NAME_None and the create asset will instead call the object ClassName_X
 					ObjName = TEXT("Null");
 				}
+				
+				//Material name clash have to be sorted here since the unique ID is only compose of the name.
+				//If we do not do it only one material node will be created.
+				if (Object->Is<FbxSurfaceMaterial>())
+				{
+					const FbxObject* SurfaceMaterialClash = MaterialNameClashMap.FindOrAdd(ObjName);
+					if (SurfaceMaterialClash != nullptr && SurfaceMaterialClash != Object)
+					{
+						int32 UniqueID = 1;
+						FString MaterialNameClash;
+						bool bBreak = false;
+						do
+						{
+							MaterialNameClash = ObjName + TEXT(NAMECLASH1_KEY) + FString::FromInt(UniqueID++);
+							SurfaceMaterialClash = MaterialNameClashMap.FindOrAdd(MaterialNameClash);
+							if (SurfaceMaterialClash == nullptr || SurfaceMaterialClash == Object)
+							{
+								bBreak = true;
+							}
+						} while (!bBreak);
+						ObjName = MaterialNameClash;
+					}
+					MaterialNameClashMap.FindChecked(ObjName) = Object;
+				}
+
 				return ObjName;
 			}
 
-			FString FFbxHelper::GetFbxNodeHierarchyName(const FbxNode* Node)
+			FString FFbxHelper::GetFbxNodeHierarchyName(const FbxNode* Node) const
 			{
 				if (!Node)
 				{
@@ -150,7 +175,7 @@ namespace UE
 				return UniqueID;
 			}
 
-			FString FFbxHelper::GetUniqueIDString(const uint64 UniqueID)
+			FString FFbxHelper::GetUniqueIDString(const uint64 UniqueID) const
 			{
 				FStringFormatNamedArguments FormatArguments;
 				FormatArguments.Add(TEXT("UniqueID"), UniqueID);
diff --git a/Engine/Plugins/Interchange/Runtime/Source/Parsers/Fbx/Private/FbxHelper.h b/Engine/Plugins/Interchange/Runtime/Source/Parsers/Fbx/Private/FbxHelper.h
--- a/Engine/Plugins/Interchange/Runtime/Source/Parsers/Fbx/Private/FbxHelper.h
+++ b/Engine/Plugins/Interchange/Runtime/Source/Parsers/Fbx/Private/FbxHelper.h
@@ -3,55 +3,53 @@
 #pragma once
 
 #include "CoreMinimal.h"
+#include "Containers/UnrealString.h"
 #include "FbxInclude.h"
 
-namespace UE
+namespace UE::Interchange::Private
 {
-	namespace Interchange
+	class FFbxParser;
+
+	class FPayloadContextBase
+	{
+	public:
+		virtual ~FPayloadContextBase() {}
+		virtual FString GetPayloadType() const { return FString(); }
+		virtual bool FetchPayloadToFile(FFbxParser& Parser, const FString& PayloadFilepath) { return false; }
+		virtual bool FetchAnimationBakeTransformPayloadToFile(FFbxParser& Parser, const double BakeFrequency, const double RangeStartTime, const double RangeEndTime, const FString& PayloadFilepath) { return false; }
+	};
+
+	struct FFbxHelper
 	{
-		namespace Private
-		{
-			class FFbxParser;
-
-			class FPayloadContextBase
-			{
-			public:
-				virtual ~FPayloadContextBase() {}
-				virtual FString GetPayloadType() const { return FString(); }
-				virtual bool FetchPayloadToFile(FFbxParser& Parser, const FString& PayloadFilepath) { return false; }
-				virtual bool FetchAnimationBakeTransformPayloadToFile(FFbxParser& Parser, const double BakeFrequency, const double RangeStartTime, const double RangeEndTime, const FString& PayloadFilepath) { return false; }
-			};
-
-			struct FFbxHelper
-			{
-			public:
-
-				static FString GetMeshName(FbxGeometryBase* Mesh);
-				static FString GetMeshUniqueID(FbxGeometryBase* Mesh);
-
-				static FString GetNodeAttributeName(FbxNodeAttribute* NodeAttribute, const FStringView DefaultNamePrefix);
-				static FString GetNodeAttributeUniqueID(FbxNodeAttribute* NodeAttribute, const FStringView Prefix);
-
-				/**
-				 * Return the name of an FbxProperty, return empty string if the property is null.
-				 */
-				static FString GetFbxPropertyName(const FbxProperty& Property);
-
-				/**
-				 * Return the name of an FbxObject, return empty string if the object is null.
-				 */
-				static FString GetFbxObjectName(const FbxObject* Object);
-
-				/**
-				 * Return a string with the name of all the parent in the hierarchy separate by a dot( . ) from the fbx root node to the specified node.
-				 * This is a way to have a valid unique ID for a fbx node that will be the same if the fbx change when we re-import.
-				 * Using the fbx sdk int32 uniqueID is not valid anymore if the fbx is re-exported.
-				 */
-				static FString GetFbxNodeHierarchyName(const FbxNode* Node);
-
-			protected:
-				static FString GetUniqueIDString(const uint64 UniqueID);
-			};
-		}//ns Private
-	}//ns Interchange
-}//ns UE
+	public:
+
+		FString GetMeshName(FbxGeometryBase* Mesh) const;
+		FString GetMeshUniqueID(FbxGeometryBase* Mesh) const;
+
+		FString GetNodeAttributeName(FbxNodeAttribute* NodeAttribute, const FStringView DefaultNamePrefix) const;
+		FString GetNodeAttributeUniqueID(FbxNodeAttribute* NodeAttribute, const FStringView Prefix) const;
+
+		/**
+			* Return the name of an FbxProperty, return empty string if the property is null.
+			*/
+		FString GetFbxPropertyName(const FbxProperty& Property) const;
+
+		/**
+			* Return the name of an FbxObject, return empty string if the object is null.
+			*/
+		FString GetFbxObjectName(const FbxObject* Object) const;
+
+		/**
+			* Return a string with the name of all the parent in the hierarchy separate by a dot( . ) from the fbx root node to the specified node.
+			* This is a way to have a valid unique ID for a fbx node that will be the same if the fbx change when we re-import.
+			* Using the fbx sdk int32 uniqueID is not valid anymore if the fbx is re-exported.
+			*/
+		FString GetFbxNodeHierarchyName(const FbxNode* Node) const;
+
+	protected:
+		FString GetUniqueIDString(const uint64 UniqueID) const;
+
+	private:
+		mutable TMap<FString, const FbxObject*> MaterialNameClashMap;
+	};
+}//ns UE::Interchange::Private
diff --git a/Engine/Plugins/Interchange/Runtime/Source/Parsers/Fbx/Private/FbxLight.cpp b/Engine/Plugins/Interchange/Runtime/Source/Parsers/Fbx/Private/FbxLight.cpp
--- a/Engine/Plugins/Interchange/Runtime/Source/Parsers/Fbx/Private/FbxLight.cpp
+++ b/Engine/Plugins/Interchange/Runtime/Source/Parsers/Fbx/Private/FbxLight.cpp
@@ -65,8 +65,8 @@ namespace UE
 
 					if (NodeAttribute && NodeAttribute->GetAttributeType() == FbxNodeAttribute::eLight)
 					{
-						FString NodeName = FFbxHelper::GetNodeAttributeName(NodeAttribute, UInterchangeBaseLightNode::StaticAssetTypeName());
-						FString NodeUid = FFbxHelper::GetNodeAttributeUniqueID(NodeAttribute, UInterchangeBaseLightNode::StaticAssetTypeName());
+						FString NodeName = Parser.GetFbxHelper()->GetNodeAttributeName(NodeAttribute, UInterchangeBaseLightNode::StaticAssetTypeName());
+						FString NodeUid = Parser.GetFbxHelper()->GetNodeAttributeUniqueID(NodeAttribute, UInterchangeBaseLightNode::StaticAssetTypeName());
 
 						const UInterchangeBaseLightNode* LightNode = Cast<const UInterchangeBaseLightNode>(NodeContainer.GetNode(NodeUid));
 
diff --git a/Engine/Plugins/Interchange/Runtime/Source/Parsers/Fbx/Private/FbxMaterial.cpp b/Engine/Plugins/Interchange/Runtime/Source/Parsers/Fbx/Private/FbxMaterial.cpp
--- a/Engine/Plugins/Interchange/Runtime/Source/Parsers/Fbx/Private/FbxMaterial.cpp
+++ b/Engine/Plugins/Interchange/Runtime/Source/Parsers/Fbx/Private/FbxMaterial.cpp
@@ -161,7 +161,7 @@ namespace UE
 				using namespace UE::Interchange::Materials;
 
 				//Create a material node
-				FString MaterialName = FFbxHelper::GetFbxObjectName(SurfaceMaterial);
+				FString MaterialName = Parser.GetFbxHelper()->GetFbxObjectName(SurfaceMaterial);
 				FString NodeUid = TEXT("\\Material\\") + MaterialName;
 				const UInterchangeShaderGraphNode* ExistingShaderGraphNode = Cast<const UInterchangeShaderGraphNode>(NodeContainer.GetNode(NodeUid));
 				if (ExistingShaderGraphNode)
@@ -288,7 +288,7 @@ namespace UE
 				{
 					FbxSurfaceMaterial* SurfaceMaterial = ParentFbxNode->GetMaterial(MaterialIndex);
 					const UInterchangeShaderGraphNode* ShaderGraphNode = AddShaderGraphNode(SurfaceMaterial, NodeContainer);
-					SceneNode->SetSlotMaterialDependencyUid(FFbxHelper::GetFbxObjectName(SurfaceMaterial), ShaderGraphNode->GetUniqueID());
+					SceneNode->SetSlotMaterialDependencyUid(Parser.GetFbxHelper()->GetFbxObjectName(SurfaceMaterial), ShaderGraphNode->GetUniqueID());
 				}
 			}
 
diff --git a/Engine/Plugins/Interchange/Runtime/Source/Parsers/Fbx/Private/FbxMesh.cpp b/Engine/Plugins/Interchange/Runtime/Source/Parsers/Fbx/Private/FbxMesh.cpp
--- a/Engine/Plugins/Interchange/Runtime/Source/Parsers/Fbx/Private/FbxMesh.cpp
+++ b/Engine/Plugins/Interchange/Runtime/Source/Parsers/Fbx/Private/FbxMesh.cpp
@@ -330,7 +330,7 @@ bool FMeshDescriptionImporter::FillMeshDescriptionFromFbxMesh(FbxMesh* Mesh, TAr
 		for (int32 MaterialIndex = 0; MaterialIndex < MaterialCount; ++MaterialIndex)
 		{
 			FbxSurfaceMaterial* FbxMaterial = MeshNode->GetMaterial(MaterialIndex);
-			MaterialNames.Add(*FFbxHelper::GetFbxObjectName(FbxMaterial));
+			MaterialNames.Add(*Parser.GetFbxHelper()->GetFbxObjectName(FbxMaterial));
 		}
 	}
 
@@ -1044,7 +1044,7 @@ bool FMeshDescriptionImporter::FillMeshDescriptionFromFbxMesh(FbxMesh* Mesh, TAr
 				if (!SortedJoints.Find(Link, BoneIndex))
 				{
 					BoneIndex = SortedJoints.Add(Link);
-					FString JointNodeUniqueID = FFbxHelper::GetFbxObjectName(Link);
+					FString JointNodeUniqueID = Parser.GetFbxHelper()->GetFbxObjectName(Link);
 					OutJointUniqueNames.Add(JointNodeUniqueID);
 				}
 
@@ -1134,7 +1134,7 @@ bool FMeshPayloadContext::FetchPayloadToFile(FFbxParser& Parser, const FString&
 	if (!ensure(SDKScene != nullptr))
 	{
 		UInterchangeResultError_Generic* Message = Parser.AddMessage<UInterchangeResultError_Generic>();
-		Message->InterchangeKey = FFbxHelper::GetMeshUniqueID(Mesh);
+		Message->InterchangeKey = Parser.GetFbxHelper()->GetMeshUniqueID(Mesh);
 		Message->Text = LOCTEXT("FBXSceneNull_Mesh", "Cannot fetch FBX mesh payload because the FBX scene is null.");
 		return false;
 	}
@@ -1142,7 +1142,7 @@ bool FMeshPayloadContext::FetchPayloadToFile(FFbxParser& Parser, const FString&
 	if (!ensure(Mesh != nullptr))
 	{
 		UInterchangeResultError_Generic* Message = Parser.AddMessage<UInterchangeResultError_Generic>();
-		Message->InterchangeKey = FFbxHelper::GetMeshUniqueID(Mesh);
+		Message->InterchangeKey = Parser.GetFbxHelper()->GetMeshUniqueID(Mesh);
 		Message->Text = LOCTEXT("FBXMeshNull", "Cannot fetch FBX mesh payload because the FBX mesh is null.");
 		return false;
 	}
@@ -1150,7 +1150,7 @@ bool FMeshPayloadContext::FetchPayloadToFile(FFbxParser& Parser, const FString&
 	if (!ensure(SDKGeometryConverter != nullptr))
 	{
 		UInterchangeResultError_Generic* Message = Parser.AddMessage<UInterchangeResultError_Generic>();
-		Message->InterchangeKey = FFbxHelper::GetMeshUniqueID(Mesh);
+		Message->InterchangeKey = Parser.GetFbxHelper()->GetMeshUniqueID(Mesh);
 		Message->Text = LOCTEXT("FBXConverterNull", "Cannot fetch FBX mesh payload because the FBX geometry converter is null.");
 		return false;
 	}
@@ -1165,7 +1165,7 @@ bool FMeshPayloadContext::FetchPayloadToFile(FFbxParser& Parser, const FString&
 		if (!MeshDescriptionImporter.FillSkinnedMeshDescriptionFromFbxMesh(Mesh, JointUniqueNames))
 		{
 			UInterchangeResultError_Generic* Message = Parser.AddMessage<UInterchangeResultError_Generic>();
-			Message->InterchangeKey = FFbxHelper::GetMeshUniqueID(Mesh);
+			Message->InterchangeKey = Parser.GetFbxHelper()->GetMeshUniqueID(Mesh);
 			Message->Text = LOCTEXT("SkinnedMeshDescriptionError", "Cannot fetch skinned mesh payload because there was an error when creating the MeshDescription.");
 			return false;
 		}
@@ -1178,7 +1178,7 @@ bool FMeshPayloadContext::FetchPayloadToFile(FFbxParser& Parser, const FString&
 		if (!MeshDescriptionImporter.FillStaticMeshDescriptionFromFbxMesh(Mesh))
 		{
 			UInterchangeResultError_Generic* Message = Parser.AddMessage<UInterchangeResultError_Generic>();
-			Message->InterchangeKey = FFbxHelper::GetMeshUniqueID(Mesh);
+			Message->InterchangeKey = Parser.GetFbxHelper()->GetMeshUniqueID(Mesh);
 			Message->Text = LOCTEXT("StaticMeshDescriptionError", "Cannot fetch static mesh payload because there was an error when creating the MeshDescription.");
 			return false;
 		}
@@ -1221,7 +1221,7 @@ bool FMorphTargetPayloadContext::FetchPayloadToFile(FFbxParser& Parser, const FS
 	if (!MeshDescriptionImporter.FillMeshDescriptionFromFbxShape(Shape))
 	{
 		UInterchangeResultError_Generic* Message = Parser.AddMessage<UInterchangeResultError_Generic>();
-		Message->InterchangeKey = FFbxHelper::GetMeshUniqueID(Shape);
+		Message->InterchangeKey = Parser.GetFbxHelper()->GetMeshUniqueID(Shape);
 		Message->Text = LOCTEXT("MeshDescriptionMorphTargetError", "Unable to create MeshDescription from FBX morph target.");
 
 		return false;
@@ -1286,8 +1286,8 @@ void FFbxMesh::AddAllMeshes(FbxScene* SDKScene, FbxGeometryConverter* SDKGeometr
 		{
 			continue;
 		}
-		FString MeshName = FFbxHelper::GetMeshName(Mesh);
-		FString MeshUniqueID = FFbxHelper::GetMeshUniqueID(Mesh);
+		FString MeshName = Parser.GetFbxHelper()->GetMeshName(Mesh);
+		FString MeshUniqueID = Parser.GetFbxHelper()->GetMeshUniqueID(Mesh);
 		const UInterchangeMeshNode* ExistingMeshNode = Cast<UInterchangeMeshNode>(NodeContainer.GetNode(MeshUniqueID));
 		UInterchangeMeshNode* MeshNode = nullptr;
 		if (ExistingMeshNode)
@@ -1365,7 +1365,7 @@ void FFbxMesh::AddAllMeshes(FbxScene* SDKScene, FbxGeometryConverter* SDKGeometr
 				for (int32 MaterialIndex = 0; MaterialIndex < MaterialCount; ++MaterialIndex)
 				{
 					FbxSurfaceMaterial* FbxMaterial = FbxMeshNode->GetMaterial(MaterialIndex);
-					FString MaterialName = FFbxHelper::GetFbxObjectName(FbxMaterial);
+					FString MaterialName = Parser.GetFbxHelper()->GetFbxObjectName(FbxMaterial);
 					FString MaterialUid = TEXT("\\Material\\") + MaterialName;
 					if (bAddAllNodeMaterials || MaterialIndexes.Contains(MaterialIndex))
 					{
@@ -1402,7 +1402,7 @@ void FFbxMesh::AddAllMeshes(FbxScene* SDKScene, FbxGeometryConverter* SDKGeometr
 		{
 			FbxBlendShape* MorphTarget = (FbxBlendShape*)Geometry->GetDeformer(MorphTargetIndex, FbxDeformer::eBlendShape);
 			const int32 MorphTargetChannelCount = MorphTarget->GetBlendShapeChannelCount();
-			FString MorphTargetName = FFbxHelper::GetFbxObjectName(MorphTarget);
+			FString MorphTargetName = Parser.GetFbxHelper()->GetFbxObjectName(MorphTarget);
 			// see below where this is used for explanation...
 			const bool bMightBeBadMAXFile = (MorphTargetName == FString("Morpher"));
 			for (int32 ChannelIndex = 0; ChannelIndex < MorphTargetChannelCount; ++ChannelIndex)
@@ -1414,7 +1414,7 @@ void FFbxMesh::AddAllMeshes(FbxScene* SDKScene, FbxGeometryConverter* SDKGeometr
 				}
 				//Find which morph target should we use according to the weight.
 				const int32 CurrentChannelMorphTargetCount = Channel->GetTargetShapeCount();
-				FString ChannelName = FFbxHelper::GetFbxObjectName(Channel);
+				FString ChannelName = Parser.GetFbxHelper()->GetFbxObjectName(Channel);
 				// Maya adds the name of the MorphTarget and an underscore to the front of the channel name, so remove it
 				if (ChannelName.StartsWith(MorphTargetName))
 				{
@@ -1426,13 +1426,13 @@ void FFbxMesh::AddAllMeshes(FbxScene* SDKScene, FbxGeometryConverter* SDKGeometr
 					FString ShapeName;
 					if (CurrentChannelMorphTargetCount > 1)
 					{
-						ShapeName = FFbxHelper::GetFbxObjectName(Shape);
+						ShapeName = Parser.GetFbxHelper()->GetFbxObjectName(Shape);
 					}
 					else
 					{
 						if (bMightBeBadMAXFile)
 						{
-							ShapeName = FFbxHelper::GetFbxObjectName(Shape);
+							ShapeName = Parser.GetFbxHelper()->GetFbxObjectName(Shape);
 						}
 						else
 						{
@@ -1442,8 +1442,8 @@ void FFbxMesh::AddAllMeshes(FbxScene* SDKScene, FbxGeometryConverter* SDKGeometr
 					}
 					ensure(!ShapeNameToFbxShape.Contains(ShapeName));
 					ShapeNameToFbxShape.Add(ShapeName, Shape);
-					FString MorphTargetAttributeName = FFbxHelper::GetMeshName(Shape);
-					FString MorphTargetUniqueID = FFbxHelper::GetMeshUniqueID(Shape);
+					FString MorphTargetAttributeName = Parser.GetFbxHelper()->GetMeshName(Shape);
+					FString MorphTargetUniqueID = Parser.GetFbxHelper()->GetMeshUniqueID(Shape);
 					const UInterchangeMeshNode* ExistingMorphTargetNode = Cast<const UInterchangeMeshNode>(NodeContainer.GetNode(MorphTargetUniqueID));
 					if (!ExistingMorphTargetNode)
 					{
@@ -1556,7 +1556,7 @@ bool FFbxMesh::ExtractSkinnedMeshNodeJoints(FbxScene* SDKScene, UInterchangeBase
 			{
 				bFoundValidJoint = true;
 
-				FString JointNodeUniqueID = FFbxHelper::GetFbxNodeHierarchyName(Link);
+				FString JointNodeUniqueID = Parser.GetFbxHelper()->GetFbxNodeHierarchyName(Link);
 
 				// find the bone index
 				if (!JointNodeUniqueIDs.Contains(JointNodeUniqueID))
diff --git a/Engine/Plugins/Interchange/Runtime/Source/Parsers/Fbx/Private/FbxMesh.h b/Engine/Plugins/Interchange/Runtime/Source/Parsers/Fbx/Private/FbxMesh.h
--- a/Engine/Plugins/Interchange/Runtime/Source/Parsers/Fbx/Private/FbxMesh.h
+++ b/Engine/Plugins/Interchange/Runtime/Source/Parsers/Fbx/Private/FbxMesh.h
@@ -47,8 +47,8 @@ namespace UE
 				T* AddMessage(FbxGeometryBase* FbxNode) const
 				{
 					T* Item = Parser.AddMessage<T>();
-					Item->MeshName = FFbxHelper::GetMeshName(FbxNode);
-					Item->InterchangeKey = FFbxHelper::GetMeshUniqueID(FbxNode);
+					Item->MeshName = Parser.GetFbxHelper()->GetMeshName(FbxNode);
+					Item->InterchangeKey = Parser.GetFbxHelper()->GetMeshUniqueID(FbxNode);
 					return Item;
 				}
 
diff --git a/Engine/Plugins/Interchange/Runtime/Source/Parsers/Fbx/Private/FbxScene.cpp b/Engine/Plugins/Interchange/Runtime/Source/Parsers/Fbx/Private/FbxScene.cpp
--- a/Engine/Plugins/Interchange/Runtime/Source/Parsers/Fbx/Private/FbxScene.cpp
+++ b/Engine/Plugins/Interchange/Runtime/Source/Parsers/Fbx/Private/FbxScene.cpp
@@ -36,7 +36,7 @@ namespace UE
 					FbxMesh* Mesh = static_cast<FbxMesh*>(NodeAttribute);
 					if (ensure(Mesh))
 					{
-						FString MeshRefString = FFbxHelper::GetMeshUniqueID(Mesh);
+						FString MeshRefString = Parser.GetFbxHelper()->GetMeshUniqueID(Mesh);
 						MeshNode = Cast<UInterchangeMeshNode>(NodeContainer.GetNode(MeshRefString));
 					}
 				}
@@ -60,9 +60,9 @@ namespace UE
 				}
 			}
 
-			void CreateAssetNodeReference(UInterchangeSceneNode* UnrealSceneNode, FbxNodeAttribute* NodeAttribute, UInterchangeBaseNodeContainer& NodeContainer, const FStringView TypeName)
+			void CreateAssetNodeReference(FFbxParser& Parser, UInterchangeSceneNode* UnrealSceneNode, FbxNodeAttribute* NodeAttribute, UInterchangeBaseNodeContainer& NodeContainer, const FStringView TypeName)
 			{
-				const FString AssetUniqueID = FFbxHelper::GetNodeAttributeUniqueID(NodeAttribute, TypeName);
+				const FString AssetUniqueID = Parser.GetFbxHelper()->GetNodeAttributeUniqueID(NodeAttribute, TypeName);
 
 				if (const UInterchangeBaseNode* AssetNode = NodeContainer.GetNode(AssetUniqueID))
 				{
@@ -72,12 +72,12 @@ namespace UE
 
 			void FFbxScene::CreateCameraNodeReference(UInterchangeSceneNode* UnrealSceneNode, FbxNodeAttribute* NodeAttribute, UInterchangeBaseNodeContainer& NodeContainer)
 			{
-				CreateAssetNodeReference(UnrealSceneNode, NodeAttribute, NodeContainer, UInterchangeCameraNode::StaticAssetTypeName());
+				CreateAssetNodeReference(Parser, UnrealSceneNode, NodeAttribute, NodeContainer, UInterchangeCameraNode::StaticAssetTypeName());
 			}
 
 			void FFbxScene::CreateLightNodeReference(UInterchangeSceneNode* UnrealSceneNode, FbxNodeAttribute* NodeAttribute, UInterchangeBaseNodeContainer& NodeContainer)
 			{
-				CreateAssetNodeReference(UnrealSceneNode, NodeAttribute, NodeContainer, UInterchangeLightNode::StaticAssetTypeName());
+				CreateAssetNodeReference(Parser, UnrealSceneNode, NodeAttribute, NodeContainer, UInterchangeLightNode::StaticAssetTypeName());
 			}
 
 			bool DoesTheParentHierarchyContainJoints(FbxNode* Node)
@@ -105,8 +105,8 @@ namespace UE
 				, TMap<FString, TSharedPtr<FPayloadContextBase, ESPMode::ThreadSafe>>& PayloadContexts)
 			{
 				constexpr bool bResetCache = false;
-				FString NodeName = FFbxHelper::GetFbxObjectName(Node);
-				FString NodeUniqueID = FFbxHelper::GetFbxNodeHierarchyName(Node);
+				FString NodeName = Parser.GetFbxHelper()->GetFbxObjectName(Node);
+				FString NodeUniqueID = Parser.GetFbxHelper()->GetFbxNodeHierarchyName(Node);
 
 				UInterchangeSceneNode* UnrealNode = CreateTransformNode(NodeContainer, NodeName, NodeUniqueID);
 				check(UnrealNode);
@@ -280,7 +280,7 @@ namespace UE
 					}
 				}
 				//Scene node transform can be animated, add the transform animation payload key.
-				FFbxAnimation::AddNodeTransformAnimation(SDKScene, Node, UnrealNode, PayloadContexts);
+				FFbxAnimation::AddNodeTransformAnimation(SDKScene, Parser, Node, UnrealNode, PayloadContexts);
 
 				//Add all custom Attributes for the node
 				FbxProperty Property = Node->GetFirstProperty();
@@ -289,14 +289,14 @@ namespace UE
 					EFbxType PropertyType =  Property.GetPropertyDataType().GetType();
 					if (Property.GetFlag(FbxPropertyFlags::eUserDefined) && FFbxAnimation::IsFbxPropertyTypeSupported(PropertyType))
 					{
-						FString PropertyName = FFbxHelper::GetFbxPropertyName(Property);
+						FString PropertyName = Parser.GetFbxHelper()->GetFbxPropertyName(Property);
 
 						FbxAnimCurveNode* CurveNode = Property.GetCurveNode();
 						TOptional<FString> PayloadKey;
 						if (CurveNode && CurveNode->IsAnimated())
 						{
 							//Attribute is animated, add the curves payload key that represent the attribute animation
-							FFbxAnimation::AddNodeAttributeCurvesAnimation(Node, Property, CurveNode, UnrealNode, PayloadContexts, PropertyType, PayloadKey);
+							FFbxAnimation::AddNodeAttributeCurvesAnimation(Parser, Node, Property, CurveNode, UnrealNode, PayloadContexts, PropertyType, PayloadKey);
 						}
 						switch (Property.GetPropertyDataType().GetType())
 						{
diff --git a/Engine/Plugins/Interchange/Runtime/Source/Pipelines/Private/InterchangeGenericSkeletalMeshPipeline.cpp b/Engine/Plugins/Interchange/Runtime/Source/Pipelines/Private/InterchangeGenericSkeletalMeshPipeline.cpp
--- a/Engine/Plugins/Interchange/Runtime/Source/Pipelines/Private/InterchangeGenericSkeletalMeshPipeline.cpp
+++ b/Engine/Plugins/Interchange/Runtime/Source/Pipelines/Private/InterchangeGenericSkeletalMeshPipeline.cpp
@@ -531,9 +531,9 @@ void UInterchangeGenericMeshPipeline::AddLodDataToSkeletalMesh(const UInterchang
 
 	const FString SkeletalMeshUid = SkeletalMeshFactoryNode->GetUniqueID();
 	const FString SkeletonUid = SkeletonFactoryNode->GetUniqueID();
-	for (const TPair<int32, TArray<FString>>& LodIndexAndNodeUids : NodeUidsPerLodIndex)
+	const int32 LodCount = NodeUidsPerLodIndex.Num();
+	for (int32 LodIndex = 0; LodIndex < LodCount; ++LodIndex)
 	{
-		const int32 LodIndex = LodIndexAndNodeUids.Key;
 		if (!CommonMeshesProperties->bImportLods && LodIndex > 0)
 		{
 			//If the pipeline should not import lods, skip any lod over base lod
@@ -541,7 +541,7 @@ void UInterchangeGenericMeshPipeline::AddLodDataToSkeletalMesh(const UInterchang
 		}
 
 		//Copy the nodes unique id because we need to remove nested mesh if the option is to not import them
-		TArray<FString> NodeUids = LodIndexAndNodeUids.Value;
+		TArray<FString> NodeUids = NodeUidsPerLodIndex.FindChecked(LodIndex);
 		if (!CommonSkeletalMeshesAndAnimationsProperties->bImportMeshesInBoneHierarchy)
 		{
 			UE::Interchange::SkeletalMeshGenericPipeline::RemoveNestedMeshNodes(BaseNodeContainer, SkeletonFactoryNode, NodeUids);
@@ -597,6 +597,7 @@ void UInterchangeGenericMeshPipeline::AddLodDataToSkeletalMesh(const UInterchang
 
 			LodDataNode->AddMeshUid(NodeUid);
 		}
+		UE::Interchange::MeshesUtilities::ReorderSlotMaterialDependencies(*SkeletalMeshFactoryNode, *BaseNodeContainer);
 	}
 }
 
diff --git a/Engine/Plugins/Interchange/Runtime/Source/Pipelines/Private/InterchangeGenericStaticMeshPipeline.cpp b/Engine/Plugins/Interchange/Runtime/Source/Pipelines/Private/InterchangeGenericStaticMeshPipeline.cpp
--- a/Engine/Plugins/Interchange/Runtime/Source/Pipelines/Private/InterchangeGenericStaticMeshPipeline.cpp
+++ b/Engine/Plugins/Interchange/Runtime/Source/Pipelines/Private/InterchangeGenericStaticMeshPipeline.cpp
@@ -474,17 +474,16 @@ void UInterchangeGenericMeshPipeline::AddLodDataToStaticMesh(UInterchangeStaticM
 	check(CommonMeshesProperties.IsValid());
 	const FString StaticMeshFactoryUid = StaticMeshFactoryNode->GetUniqueID();
 
-	for (const TPair<int32, TArray<FString>>& LodIndexAndNodeUids : NodeUidsPerLodIndex)
+	const int32 LodCount = NodeUidsPerLodIndex.Num();
+	for(int32 LodIndex = 0; LodIndex < LodCount; ++LodIndex)
 	{
-		const int32 LodIndex = LodIndexAndNodeUids.Key;
 		if (!CommonMeshesProperties->bImportLods && LodIndex > 0)
 		{
 			// If the pipeline should not import lods, skip any lod over base lod
 			continue;
 		}
 
-		const TArray<FString>& NodeUids = LodIndexAndNodeUids.Value;
-
+		const TArray<FString>& NodeUids = NodeUidsPerLodIndex.FindChecked(LodIndex);
 		// Create a lod data node with all the meshes for this LOD
 		const FString StaticMeshLodDataName = TEXT("LodData") + FString::FromInt(LodIndex);
 		const FString LODDataPrefix = TEXT("\\LodData") + (LodIndex > 0 ? FString::FromInt(LodIndex) : TEXT(""));
@@ -567,5 +566,6 @@ void UInterchangeGenericMeshPipeline::AddLodDataToStaticMesh(UInterchangeStaticM
 				LodDataNode->AddMeshUid(NodeUid);
 			}
 		}
+		UE::Interchange::MeshesUtilities::ReorderSlotMaterialDependencies(*StaticMeshFactoryNode, *BaseNodeContainer);
 	}
 }
diff --git a/Engine/Plugins/Interchange/Runtime/Source/Pipelines/Private/InterchangePipelineMeshesUtilities.cpp b/Engine/Plugins/Interchange/Runtime/Source/Pipelines/Private/InterchangePipelineMeshesUtilities.cpp
--- a/Engine/Plugins/Interchange/Runtime/Source/Pipelines/Private/InterchangePipelineMeshesUtilities.cpp
+++ b/Engine/Plugins/Interchange/Runtime/Source/Pipelines/Private/InterchangePipelineMeshesUtilities.cpp
@@ -13,6 +13,94 @@
 
 #include UE_INLINE_GENERATED_CPP_BY_NAME(InterchangePipelineMeshesUtilities)
 
+namespace UE::Private::InterchangeMeshPipeline
+{
+	void FindNamedLodGroup(UInterchangeBaseNodeContainer* BaseNodeContainer, TMap<FString, TArray<FString>>& SceneMeshNodeUidsPerLodParentUidMap)
+	{
+		const FString LodPrefix = TEXT("LOD");
+		BaseNodeContainer->IterateNodes(
+			[&BaseNodeContainer, &SceneMeshNodeUidsPerLodParentUidMap, &LodPrefix](const FString& NodeUid, UInterchangeBaseNode* Node)
+			{
+				if (Node->GetNodeContainerType() == EInterchangeNodeContainerType::TranslatedScene)
+				{
+					const UInterchangeSceneNode* SceneNode = Cast<UInterchangeSceneNode>(Node);
+					if (!SceneNode)
+					{
+						return;
+					}
+					FString MeshUid;
+					if (!SceneNode->GetCustomAssetInstanceUid(MeshUid))
+					{
+						return;
+					}
+					const UInterchangeBaseNode* MeshNode = BaseNodeContainer->GetNode(MeshUid);
+					if (!MeshNode || !MeshNode->IsA<UInterchangeMeshNode>())
+					{
+						return;
+					}
+					FString ParentUniqueID = SceneNode->GetParentUid();
+					const UInterchangeSceneNode* ParentSceneNode = Cast<UInterchangeSceneNode>(BaseNodeContainer->GetNode(ParentUniqueID));
+					if (!ParentSceneNode)
+					{
+						return;
+					}
+
+					//Skip this node if there is a parent in the hierarchy that is a specialize lod group
+					const UInterchangeSceneNode* ParentSceneNodeHierarchy = ParentSceneNode;
+					while (ParentSceneNodeHierarchy)
+					{
+						if (ParentSceneNodeHierarchy->IsSpecializedTypeContains(UE::Interchange::FSceneNodeStaticData::GetLodGroupSpecializeTypeString()))
+						{
+							return;
+						}
+						ParentSceneNodeHierarchy = Cast<const UInterchangeSceneNode>(BaseNodeContainer->GetNode(ParentSceneNodeHierarchy->GetParentUid()));
+					}
+
+					if (BaseNodeContainer->GetNodeChildrenCount(ParentUniqueID) <= 1)
+					{
+						//Do not create custom named LOD group that have only one lod
+						return;
+					}
+
+					FString SceneNodeName = SceneNode->GetDisplayLabel();
+					if (SceneNodeName.Len() > 5 && SceneNodeName.StartsWith(LodPrefix, ESearchCase::CaseSensitive) && SceneNodeName[4] == '_')
+					{
+						FString LODXNumber = SceneNodeName.RightChop(3).Left(1);
+						if (LODXNumber.IsNumeric())
+						{
+							int32 LodNumber = FPlatformString::Atoi(*LODXNumber);
+							FString MatchName = ParentUniqueID;
+							
+							TArray<FString>& LodChildUids = SceneMeshNodeUidsPerLodParentUidMap.FindOrAdd(MatchName);
+							//Add LOD at the correct index
+							if (LodNumber >= LodChildUids.Num())
+							{
+								int32 AddCount = LodNumber + 1 - LodChildUids.Num();
+								LodChildUids.AddDefaulted(AddCount);
+							}
+							LodChildUids[LodNumber] = SceneNode->GetUniqueID();
+						}
+					}
+				}
+			}
+		);
+
+		//Remove all empty entry, we use empty entry to set all lod in the correct order
+		for (TPair<FString, TArray<FString>>& LodPrefixNodePair : SceneMeshNodeUidsPerLodParentUidMap)
+		{
+			TArray<FString>& LodChildUids = LodPrefixNodePair.Value;
+			for (int32 ChildLodIndex = LodChildUids.Num() - 1; ChildLodIndex >= 0; ChildLodIndex--)
+			{
+				if (LodChildUids[ChildLodIndex].IsEmpty())
+				{
+					LodChildUids.RemoveAt(ChildLodIndex, 1, false);
+				}
+			}
+			//Shrink the array to the correct size
+			LodChildUids.Shrink();
+		}
+	}
+}
 
 static bool IsSceneNodeASocket(const UInterchangeSceneNode* SceneNode)
 {
@@ -46,10 +134,13 @@ UInterchangePipelineMeshesUtilities* UInterchangePipelineMeshesUtilities::Create
 		}
 	);
 
+	TMap<FString, TArray<FString>> SceneMeshNodeUidsPerLodParentUidMap;
+	UE::Private::InterchangeMeshPipeline::FindNamedLodGroup(BaseNodeContainer, SceneMeshNodeUidsPerLodParentUidMap);
+
 	//Find all translated scene node we need for this pipeline
 	bool bHasSockets = false;
 	BaseNodeContainer->IterateNodes(
-		[&PipelineMeshesUtilities, &BaseNodeContainer, &SkeletonRootNodeUids, &bHasSockets](const FString& NodeUid, const UInterchangeBaseNode* Node)
+		[&PipelineMeshesUtilities, &BaseNodeContainer, &SceneMeshNodeUidsPerLodParentUidMap, &SkeletonRootNodeUids, &bHasSockets](const FString& NodeUid, const UInterchangeBaseNode* Node)
 		{
 			if (Node->GetNodeContainerType() == EInterchangeNodeContainerType::TranslatedScene)
 			{
@@ -81,28 +172,44 @@ UInterchangePipelineMeshesUtilities* UInterchangePipelineMeshesUtilities::Create
 							{
 								const UInterchangeSceneNode* LodGroupNode = nullptr;
 								int32 LodIndex = 0;
-								FString LastChildUid = SceneNode->GetUniqueID();
-								do
+								if (SceneMeshNodeUidsPerLodParentUidMap.Contains(ParentMeshSceneNode->GetUniqueID()))
 								{
-									if (ParentMeshSceneNode->IsSpecializedTypeContains(UE::Interchange::FSceneNodeStaticData::GetLodGroupSpecializeTypeString()))
+									LodGroupNode = ParentMeshSceneNode;
+									const TArray<FString>& LodChildUids = SceneMeshNodeUidsPerLodParentUidMap.FindChecked(ParentMeshSceneNode->GetUniqueID());
+									for (int32 ChildLodIndex = 0; ChildLodIndex < LodChildUids.Num(); ++ChildLodIndex)
 									{
-										LodGroupNode = ParentMeshSceneNode;
-										TArray<FString> LodGroupChildrens = BaseNodeContainer->GetNodeChildrenUids(ParentMeshSceneNode->GetUniqueID());
-										for (int32 ChildLodIndex = 0; ChildLodIndex < LodGroupChildrens.Num(); ++ChildLodIndex)
+										const FString& ChildrenUid = LodChildUids[ChildLodIndex];
+										if (ChildrenUid.Equals(SceneNode->GetUniqueID()))
 										{
-											const FString& ChildrenUid = LodGroupChildrens[ChildLodIndex];
-											if (ChildrenUid.Equals(LastChildUid))
+											LodIndex = ChildLodIndex;
+											break;
+										}
+									}
+								}
+								else
+								{
+									FString LastChildUid = SceneNode->GetUniqueID();
+									do
+									{
+										if (ParentMeshSceneNode->IsSpecializedTypeContains(UE::Interchange::FSceneNodeStaticData::GetLodGroupSpecializeTypeString()))
+										{
+											LodGroupNode = ParentMeshSceneNode;
+											TArray<FString> LodGroupChildrens = BaseNodeContainer->GetNodeChildrenUids(ParentMeshSceneNode->GetUniqueID());
+											for (int32 ChildLodIndex = 0; ChildLodIndex < LodGroupChildrens.Num(); ++ChildLodIndex)
 											{
-												LodIndex = ChildLodIndex;
-												break;
+												const FString& ChildrenUid = LodGroupChildrens[ChildLodIndex];
+												if (ChildrenUid.Equals(LastChildUid))
+												{
+													LodIndex = ChildLodIndex;
+													break;
+												}
 											}
+											break;
 										}
-										break;
-									}
-									LastChildUid = ParentMeshSceneNode->GetUniqueID();
-									ParentMeshSceneNode = Cast<const UInterchangeSceneNode>(BaseNodeContainer->GetNode(ParentMeshSceneNode->GetParentUid()));
-								} while (ParentMeshSceneNode);
-
+										LastChildUid = ParentMeshSceneNode->GetUniqueID();
+										ParentMeshSceneNode = Cast<const UInterchangeSceneNode>(BaseNodeContainer->GetNode(ParentMeshSceneNode->GetParentUid()));
+									} while (ParentMeshSceneNode);
+								}
 								FInterchangeMeshGeometry& MeshGeometry = PipelineMeshesUtilities->MeshGeometriesPerMeshUid.FindChecked(MeshUid);
 								if (LodGroupNode)
 								{
diff --git a/Engine/Plugins/Interchange/Runtime/Source/Pipelines/Public/InterchangePipelineMeshesUtilities.h b/Engine/Plugins/Interchange/Runtime/Source/Pipelines/Public/InterchangePipelineMeshesUtilities.h
--- a/Engine/Plugins/Interchange/Runtime/Source/Pipelines/Public/InterchangePipelineMeshesUtilities.h
+++ b/Engine/Plugins/Interchange/Runtime/Source/Pipelines/Public/InterchangePipelineMeshesUtilities.h
@@ -289,12 +289,11 @@ namespace UE::Interchange::MeshesUtilities
 		for (const TPair<FString, FString>& SlotMaterialDependency : SlotMaterialDependencies)
 		{
 			const FString MaterialFactoryNodeUid = UInterchangeBaseMaterialFactoryNode::GetMaterialFactoryNodeUidFromMaterialNodeUid(SlotMaterialDependency.Value);
+			FactoryNode.SetSlotMaterialDependencyUid(SlotMaterialDependency.Key, MaterialFactoryNodeUid);
 			if (UInterchangeBaseMaterialFactoryNode* MaterialFactoryNode = Cast<UInterchangeBaseMaterialFactoryNode>(NodeContainer.GetFactoryNode(MaterialFactoryNodeUid)))
 			{
 				MaterialFactoryNode->SetEnabled(true);
 
-				FactoryNode.SetSlotMaterialDependencyUid(SlotMaterialDependency.Key, MaterialFactoryNodeUid);
-
 				// Create a factory dependency so Material asset are imported before the static mesh asset
 				TArray<FString> FactoryDependencies;
 				FactoryNode.GetFactoryDependencies(FactoryDependencies);
@@ -306,4 +305,89 @@ namespace UE::Interchange::MeshesUtilities
 		}
 
 	}
+
+	template<class T>
+	void ReorderSlotMaterialDependencies(T& FactoryNode, const UInterchangeBaseNodeContainer& NodeContainer)
+	{
+		TMap<FString, FString> SlotMaterialDependencies;
+		FactoryNode.GetSlotMaterialDependencies(SlotMaterialDependencies);
+
+		for (const TPair<FString, FString>& SlotMaterialDependency : SlotMaterialDependencies)
+		{
+			const FString& MaterialName = SlotMaterialDependency.Key;
+			FactoryNode.RemoveSlotMaterialDependencyUid(MaterialName);
+		}
+
+		TArray<FString> KeyReorder;
+		KeyReorder.Reserve(SlotMaterialDependencies.Num());
+		TArray<FString> MissingSuffixe;
+		MissingSuffixe.Reserve(SlotMaterialDependencies.Num());
+		//Reorder material using the skinXX workflow
+		for (const TPair<FString, FString>& SlotMaterialDependency : SlotMaterialDependencies)
+		{
+			const FString& MaterialName = SlotMaterialDependency.Key;
+			if (MaterialName.Len() > 6)
+			{
+				int32 Offset = MaterialName.Find(TEXT("_SKIN"), ESearchCase::IgnoreCase, ESearchDir::FromEnd);
+				if (Offset != INDEX_NONE)
+				{
+					// Chop off the material name so we are left with the number in _SKINXX
+					FString SkinXXNumber = MaterialName.Right(MaterialName.Len() - (Offset + 1)).RightChop(4);
+					if (SkinXXNumber.IsNumeric())
+					{
+						int32 TmpIndex = FPlatformString::Atoi(*SkinXXNumber);
+						if (TmpIndex >= 0)
+						{
+							while (KeyReorder.Num() <= TmpIndex)
+							{
+								KeyReorder.AddDefaulted();
+							}
+							check(KeyReorder.IsValidIndex(TmpIndex));
+							if (KeyReorder[TmpIndex].IsEmpty())
+							{
+								KeyReorder[TmpIndex] = MaterialName;
+								continue;
+							}
+						}
+					}
+				}
+				else
+				{
+					MissingSuffixe.Add(MaterialName);
+				}
+			}
+			else
+			{
+				MissingSuffixe.Add(MaterialName);
+			}
+		}
+
+		//The map is MaterialName, MaterialUid
+		TMap<FString, FString> ReorderedSlotMaterialDependencies;
+		ReorderedSlotMaterialDependencies.Reserve(SlotMaterialDependencies.Num());
+
+		//Start by adding the KeyReorder material
+		for (const FString& MaterialName : KeyReorder)
+		{
+			if (MaterialName.IsEmpty())
+			{
+				continue;
+			}
+			const FString& SlotMaterialUid = SlotMaterialDependencies.FindChecked(MaterialName);
+			ReorderedSlotMaterialDependencies.Add(MaterialName, SlotMaterialUid);
+		}
+		//Add the missing suffixe material
+		for (const FString& MaterialName : MissingSuffixe)
+		{
+			const FString& SlotMaterialUid = SlotMaterialDependencies.FindChecked(MaterialName);
+			ReorderedSlotMaterialDependencies.Add(MaterialName, SlotMaterialUid);
+		}
+
+		check(ReorderedSlotMaterialDependencies.Num() == SlotMaterialDependencies.Num());
+
+		for (const TPair<FString, FString>& SlotMaterialDependency : ReorderedSlotMaterialDependencies)
+		{
+			FactoryNode.SetSlotMaterialDependencyUid(SlotMaterialDependency.Key, SlotMaterialDependency.Value);
+		}
+	}
 }
\ No newline at end of file
