[Interchange] Send glTF-specific analytics on import

diff --git a/Engine/Plugins/Interchange/Runtime/Source/Import/Private/Gltf/InterchangeGltfTranslator.cpp b/Engine/Plugins/Interchange/Runtime/Source/Import/Private/Gltf/InterchangeGltfTranslator.cpp
--- a/Engine/Plugins/Interchange/Runtime/Source/Import/Private/Gltf/InterchangeGltfTranslator.cpp
+++ b/Engine/Plugins/Interchange/Runtime/Source/Import/Private/Gltf/InterchangeGltfTranslator.cpp
@@ -34,6 +34,8 @@
 
 #include "Gltf/InterchangeGltfPrivate.h"
 
+#include "EngineAnalytics.h"
+
 #include UE_INLINE_GENERATED_CPP_BY_NAME(InterchangeGltfTranslator)
 
 #define LOCTEXT_NAMESPACE "InterchangeGltfTranslator"
@@ -127,6 +129,65 @@ namespace UE::Interchange::Gltf::Private
 			Node.Transform.SetTranslation(Node.Transform.GetTranslation() * Scale);
 		}
 	}
+
+	enum TranslationResult : int32
+	{
+		SUCCESSFULL = 0,
+		INPUT_FILE_NOTFOUND,
+		GLTFREADER_FAILED,
+		NOTSUPPORTED_EXTENSION_FOUND
+	};
+	void SendAnalytics(const TranslationResult& TranslationResult,
+		const FString& NotSupportedExtensions = FString(),
+		const TSet<GLTF::EExtension>& ExtensionsUsed = TSet<GLTF::EExtension>(),
+		const TArray<FString>& RequiredExtensions = TArray<FString>(),
+		GLTF::FMetadata Metadata = GLTF::FMetadata(),
+		const FString& GLTFReaderLogMessage = "")
+	{
+		if (FEngineAnalytics::IsAvailable())
+		{
+			TArray<FString> ExtensionsUsedStringified;
+			for (GLTF::EExtension Extension : ExtensionsUsed)
+			{
+				ExtensionsUsedStringified.Add(FString(GLTF::ToString(Extension)));
+			}
+
+			TMap<FString, FString> MetadataExtras;
+			for (GLTF::FMetadata::FExtraData ExtraData : Metadata.Extras)
+			{
+				MetadataExtras.Add(ExtraData.Name, ExtraData.Value);
+			}
+
+			TArray<FAnalyticsEventAttribute> GLTFAnalytics;
+			if (NotSupportedExtensions.Len() > 0)			GLTFAnalytics.Add(FAnalyticsEventAttribute(TEXT("NotSupportedExtensions"), NotSupportedExtensions));
+			if (RequiredExtensions.Num() > 0)		GLTFAnalytics.Add(FAnalyticsEventAttribute(TEXT("RequiredExtensions"), RequiredExtensions));
+			if (ExtensionsUsedStringified.Num() > 0)		GLTFAnalytics.Add(FAnalyticsEventAttribute(TEXT("UsedExtensions"), ExtensionsUsedStringified));
+			if (Metadata.GeneratorName.Len() > 0)	GLTFAnalytics.Add(FAnalyticsEventAttribute(TEXT("MetaData.GeneratorName"), Metadata.GeneratorName));
+			if (MetadataExtras.Num() > 0)					GLTFAnalytics.Add(FAnalyticsEventAttribute(TEXT("MetaData.Extras"), MetadataExtras));
+			/*Version is always set at this point.*/		GLTFAnalytics.Add(FAnalyticsEventAttribute(TEXT("MetaData.Version"), Metadata.Version));
+
+			switch (TranslationResult)
+			{
+			case SUCCESSFULL:
+				GLTFAnalytics.Add(FAnalyticsEventAttribute(TEXT("TranslationResult"), "Successfull."));
+				break;
+			case INPUT_FILE_NOTFOUND:
+				GLTFAnalytics.Add(FAnalyticsEventAttribute(TEXT("TranslationResult"), "[Failed] Input File Not Found."));
+				break;
+			case GLTFREADER_FAILED:
+				GLTFAnalytics.Add(FAnalyticsEventAttribute(TEXT("TranslationResult"), "[Failed] Parsing error: " + GLTFReaderLogMessage));
+				break;
+			case NOTSUPPORTED_EXTENSION_FOUND:
+				GLTFAnalytics.Add(FAnalyticsEventAttribute(TEXT("TranslationResult"), "[Failed] Unsupported Extension Found."));
+				break;
+			default:
+				break;
+			}
+
+			//Send Analytics
+			FEngineAnalytics::GetProvider().RecordEvent(TEXT("Interchange.Usage.Import.GLTF"), GLTFAnalytics);
+		}
+	};
 }
 
 void UInterchangeGltfTranslator::HandleGltfNode( UInterchangeBaseNodeContainer& NodeContainer, const GLTF::FNode& GltfNode, const FString& ParentNodeUid, const int32 NodeIndex, 
@@ -842,8 +903,10 @@ bool UInterchangeGltfTranslator::Translate( UInterchangeBaseNodeContainer& NodeC
 	FString FilePath = GetSourceData()->GetFilename();
 	if ( !FPaths::FileExists( FilePath ) )
 	{
+		SendAnalytics(TranslationResult::INPUT_FILE_NOTFOUND);
 		return false;
 	}
+	const FString FileName = FPaths::GetBaseFilename(FilePath);
 
 	GLTF::FFileReader GltfFileReader;
 
@@ -851,36 +914,53 @@ bool UInterchangeGltfTranslator::Translate( UInterchangeBaseNodeContainer& NodeC
 	const bool bLoadMetaData = false;
 	GltfFileReader.ReadFile( FilePath, bLoadImageData, bLoadMetaData, const_cast< UInterchangeGltfTranslator* >( this )->GltfAsset );
 
-	ScaleNodeTranslations(const_cast<UInterchangeGltfTranslator*>(this)->GltfAsset.Nodes, GltfUnitConversionMultiplier);
-
-	const FString FileName = FPaths::GetBaseFilename(FilePath);
-
+	//Required Extension Check:
+	FString NotSupportedExtensions;
 	if (GltfAsset.RequiredExtensions.Num() != 0)
 	{
-		FString NotSupportedExtensions;
-		for ( const FString& RequiredExtension: GltfAsset.RequiredExtensions )
+		for (const FString& RequiredExtension : GltfAsset.RequiredExtensions)
 		{
-			if ( SupportedExtensions.Find( RequiredExtension ) == INDEX_NONE )
+			if (SupportedExtensions.Find(RequiredExtension) == INDEX_NONE)
 			{
-				if ( NotSupportedExtensions.Len() > 0 )
+				if (NotSupportedExtensions.Len() > 0)
 				{
 					NotSupportedExtensions += ", ";
 				}
 				NotSupportedExtensions += RequiredExtension;
 			}
 		}
-		if ( NotSupportedExtensions.Len() > 0 )
+	}
+
+	//Check if ReadFile failed:
+	TArray<GLTF::FLogMessage> GLTFReadFileLogMessages = GltfFileReader.GetLogMessages();
+	for (GLTF::FLogMessage& LogMessage : GLTFReadFileLogMessages)
+	{
+		if (LogMessage.Key == GLTF::EMessageSeverity::Error)
 		{
 			UInterchangeResultError_Generic* ErrorResult = AddMessage< UInterchangeResultError_Generic >();
 			ErrorResult->SourceAssetName = FileName;
-			ErrorResult->Text = FText::Format(
-				LOCTEXT( "UnsupportedRequiredExtensions", "Not All Required Extensions are supported. (Unsupported extensions: {0})"),
-				FText::FromString( NotSupportedExtensions ) );
+			ErrorResult->Text = FText::Format(LOCTEXT("GLTF::FFileReader::ReadFile Failed.", "LogMessage: {0}"), FText::FromString(LogMessage.Value));
 
+			SendAnalytics(TranslationResult::GLTFREADER_FAILED, NotSupportedExtensions, GltfAsset.ExtensionsUsed, GltfAsset.RequiredExtensions, GltfAsset.Metadata, LogMessage.Value);
 			return false;
 		}
 	}
 
+	//In case of non supported extensions fail out:
+	if (NotSupportedExtensions.Len() > 0)
+	{
+		UInterchangeResultError_Generic* ErrorResult = AddMessage< UInterchangeResultError_Generic >();
+		ErrorResult->SourceAssetName = FileName;
+		ErrorResult->Text = FText::Format(
+			LOCTEXT("UnsupportedRequiredExtensions", "Not All Required Extensions are supported. (Unsupported extensions: {0})"),
+			FText::FromString(NotSupportedExtensions));
+
+		SendAnalytics(TranslationResult::NOTSUPPORTED_EXTENSION_FOUND, NotSupportedExtensions, GltfAsset.ExtensionsUsed, GltfAsset.RequiredExtensions, GltfAsset.Metadata);
+		return false;
+	}
+
+	ScaleNodeTranslations(const_cast<UInterchangeGltfTranslator*>(this)->GltfAsset.Nodes, GltfUnitConversionMultiplier);
+
 	const_cast< UInterchangeGltfTranslator* >( this )->GltfAsset.GenerateNames(FileName);
 
 	// Textures
@@ -1105,6 +1185,7 @@ bool UInterchangeGltfTranslator::Translate( UInterchangeBaseNodeContainer& NodeC
 		UE_LOG(LogInterchangeImport, Warning, TEXT("GLTF Mesh Import Warning. Gltf Mesh Usage expectation is not met."));
 	}
 
+	SendAnalytics(TranslationResult::SUCCESSFULL, NotSupportedExtensions, GltfAsset.ExtensionsUsed, GltfAsset.RequiredExtensions, GltfAsset.Metadata);;
 	return true;
 }
 
