Fixed the import of UDIMs being broken when using jpeg files and interchange.

diff --git a/Engine/Plugins/Interchange/Runtime/Source/Import/Private/Texture/InterchangeJPGTranslator.cpp b/Engine/Plugins/Interchange/Runtime/Source/Import/Private/Texture/InterchangeJPGTranslator.cpp
--- a/Engine/Plugins/Interchange/Runtime/Source/Import/Private/Texture/InterchangeJPGTranslator.cpp
+++ b/Engine/Plugins/Interchange/Runtime/Source/Import/Private/Texture/InterchangeJPGTranslator.cpp
@@ -177,13 +177,6 @@ TOptional<UE::Interchange::FImportImage> UInterchangeJPGTranslator::GetTexturePa
 		bShouldAllocateRawDataBuffer
 	);
 
-	// Honor setting from TextureImporter.RetainJpegFormat in Editor.ini if it exists
-	const bool bShouldImportRawCache = bShouldImportRaw;
-	if (GConfig->GetBool(TEXT("TextureImporter"), TEXT("RetainJpegFormat"), bShouldImportRaw, GEditorIni) && bShouldImportRawCache != bShouldImportRaw)
-	{
-		UE_LOG(LogInterchangeImport, Log, TEXT("JPEG file [%s]: Pipeline setting 'bPreferCompressedSourceData' has been overridden by Editor setting 'RetainJpegFormat'"), *Filename);
-	}
-
 	TArray64<uint8> RawData;
 	if (bShouldImportRaw)
 	{
diff --git a/Engine/Plugins/Interchange/Runtime/Source/Import/Private/Texture/InterchangeTextureFactory.cpp b/Engine/Plugins/Interchange/Runtime/Source/Import/Private/Texture/InterchangeTextureFactory.cpp
--- a/Engine/Plugins/Interchange/Runtime/Source/Import/Private/Texture/InterchangeTextureFactory.cpp
+++ b/Engine/Plugins/Interchange/Runtime/Source/Import/Private/Texture/InterchangeTextureFactory.cpp
@@ -33,11 +33,13 @@
 #include "InterchangeTranslatorBase.h"
 #include "InterchangeVolumeTextureFactoryNode.h"
 #include "InterchangeVolumeTextureNode.h"
+#include "Misc/ConfigCacheIni.h"
 #include "Misc/CoreStats.h"
 #include "Nodes/InterchangeBaseNode.h"
 #include "ProfilingDebugging/CpuProfilerTrace.h"
 #include "Serialization/EditorBulkData.h"
 #include "Texture/InterchangeBlockedTexturePayloadInterface.h"
+#include "Texture/InterchangeJPGTranslator.h"
 #include "Texture/InterchangeSlicedTexturePayloadInterface.h"
 #include "Texture/InterchangeTextureLightProfilePayloadInterface.h"
 #include "Texture/InterchangeTexturePayloadInterface.h"
@@ -604,6 +606,16 @@ namespace UE::Interchange::Private::InterchangeTextureFactory
 			{
 				if (BlockAndSourceDataFiles.IsEmpty())
 				{
+					if (Translator->GetClass() == UInterchangeJPGTranslator::StaticClass())
+					{
+						// Honor setting from TextureImporter.RetainJpegFormat in Editor.ini if it exists (ideally we should deprecate this as it is confusing and probably not thread safe)
+						const bool bShouldImportRawCache = bShoudImportCompressedImage;
+						if (GConfig->GetBool(TEXT("TextureImporter"), TEXT("RetainJpegFormat"), bShoudImportCompressedImage, GEditorIni) && bShouldImportRawCache != bShoudImportCompressedImage)
+						{
+							UE_LOG(LogInterchangeImport, Log, TEXT("JPEG file [%s]: Pipeline setting 'bPreferCompressedSourceData' has been overridden by Editor setting 'RetainJpegFormat'"), *SourceData->GetFilename());
+						}
+					}
+
 					if (bShoudImportCompressedImage && TextureTranslator->SupportCompressedTexturePayloadData())
 					{
 						return FTexturePayloadVariant(TInPlaceType<TOptional<FImportImage>>(), TextureTranslator->GetCompressedTexturePayloadData(SourceData, PayloadKey));
