Fix a missing warning if there is no smooth group

#rnx

diff --git a/Engine/Plugins/Interchange/Runtime/Source/Parsers/Fbx/Private/FbxMesh.cpp b/Engine/Plugins/Interchange/Runtime/Source/Parsers/Fbx/Private/FbxMesh.cpp
--- a/Engine/Plugins/Interchange/Runtime/Source/Parsers/Fbx/Private/FbxMesh.cpp
+++ b/Engine/Plugins/Interchange/Runtime/Source/Parsers/Fbx/Private/FbxMesh.cpp
@@ -336,6 +336,12 @@ bool FMeshDescriptionImporter::FillMeshDescriptionFromFbxMesh(FbxMesh* Mesh, TAr
 
 	// Must do this before triangulating the mesh due to an FBX bug in TriangulateMeshAdvance
 	int32 LayerSmoothingCount = Mesh->GetLayerCount(FbxLayerElement::eSmoothing);
+	if (LayerSmoothingCount == 0)
+	{
+		UInterchangeResultMeshWarning_Generic* Message = AddMessage<UInterchangeResultMeshWarning_Generic>(Mesh);
+		Message->Text = LOCTEXT("MissingSmoothGroup", "No smoothing group information was found for this mesh '{MeshName}' in the FBX file. Please make sure to enable the 'Export Smoothing Groups' option in the FBX Exporter before exporting the file.");
+	}
+
 	for (int32 i = 0; i < LayerSmoothingCount; i++)
 	{
 		FbxLayerElementSmoothing const* SmoothingInfo = Mesh->GetLayer(0)->GetSmoothing();
