Interchange: Integrate interchange with the AssetImportTask workflow, by allowing user to specify an override pipeline stack in the AssetImportTask::Options parameter.

Also implement a pipeline asset name option so user can specify a name when they import only one asset.

#rnx

diff --git a/Engine/Plugins/Interchange/Runtime/Source/Pipelines/Private/InterchangeGenericAssetsPipeline.cpp b/Engine/Plugins/Interchange/Runtime/Source/Pipelines/Private/InterchangeGenericAssetsPipeline.cpp
--- a/Engine/Plugins/Interchange/Runtime/Source/Pipelines/Private/InterchangeGenericAssetsPipeline.cpp
+++ b/Engine/Plugins/Interchange/Runtime/Source/Pipelines/Private/InterchangeGenericAssetsPipeline.cpp
@@ -504,31 +504,31 @@ void UInterchangeGenericAssetsPipeline::SetReimportSourceIndex(UClass* ReimportO
 
 void UInterchangeGenericAssetsPipeline::ImplementUseSourceNameForAssetOption()
 {
-	if (bUseSourceNameForAsset)
-	{
-		const UClass* SkeletalMeshFactoryNodeClass = UInterchangeSkeletalMeshFactoryNode::StaticClass();
-		TArray<FString> SkeletalMeshNodeUids;
-		BaseNodeContainer->GetNodes(SkeletalMeshFactoryNodeClass, SkeletalMeshNodeUids);
+	const UClass* SkeletalMeshFactoryNodeClass = UInterchangeSkeletalMeshFactoryNode::StaticClass();
+	TArray<FString> SkeletalMeshNodeUids;
+	BaseNodeContainer->GetNodes(SkeletalMeshFactoryNodeClass, SkeletalMeshNodeUids);
 
-		const UClass* StaticMeshFactoryNodeClass = UInterchangeStaticMeshFactoryNode::StaticClass();
-		TArray<FString> StaticMeshNodeUids;
-		BaseNodeContainer->GetNodes(StaticMeshFactoryNodeClass, StaticMeshNodeUids);
+	const UClass* StaticMeshFactoryNodeClass = UInterchangeStaticMeshFactoryNode::StaticClass();
+	TArray<FString> StaticMeshNodeUids;
+	BaseNodeContainer->GetNodes(StaticMeshFactoryNodeClass, StaticMeshNodeUids);
 
-		const UClass* AnimSequenceFactoryNodeClass = UInterchangeAnimSequenceFactoryNode::StaticClass();
-		TArray<FString> AnimSequenceNodeUids;
-		BaseNodeContainer->GetNodes(AnimSequenceFactoryNodeClass, AnimSequenceNodeUids);
+	const UClass* AnimSequenceFactoryNodeClass = UInterchangeAnimSequenceFactoryNode::StaticClass();
+	TArray<FString> AnimSequenceNodeUids;
+	BaseNodeContainer->GetNodes(AnimSequenceFactoryNodeClass, AnimSequenceNodeUids);
 
-		//If we import only one mesh, we want to rename the mesh using the file name.
-		const int32 MeshesImportedNodeCount = SkeletalMeshNodeUids.Num() + StaticMeshNodeUids.Num();
+	//If we import only one mesh, we want to rename the mesh using the file name.
+	const int32 MeshesImportedNodeCount = SkeletalMeshNodeUids.Num() + StaticMeshNodeUids.Num();
 
-		//SkeletalMesh
-		MeshPipeline->ImplementUseSourceNameForAssetOptionSkeletalMesh(MeshesImportedNodeCount, bUseSourceNameForAsset);
+	//SkeletalMesh it must always be run even if there is no rename option, skeleton and physics asset will be rename properly
+	MeshPipeline->ImplementUseSourceNameForAssetOptionSkeletalMesh(MeshesImportedNodeCount, bUseSourceNameForAsset, AssetName);
 
+	if (bUseSourceNameForAsset || !AssetName.IsEmpty())
+	{
 		//StaticMesh
 		if (MeshesImportedNodeCount == 1 && StaticMeshNodeUids.Num() > 0)
 		{
 			UInterchangeStaticMeshFactoryNode* StaticMeshNode = Cast<UInterchangeStaticMeshFactoryNode>(BaseNodeContainer->GetFactoryNode(StaticMeshNodeUids[0]));
-			const FString DisplayLabelName = FPaths::GetBaseFilename(SourceDatas[0]->GetFilename());
+			const FString DisplayLabelName = AssetName.IsEmpty() ? FPaths::GetBaseFilename(SourceDatas[0]->GetFilename()) : AssetName;
 			StaticMeshNode->SetDisplayLabel(DisplayLabelName);
 		}
 
@@ -536,10 +536,9 @@ void UInterchangeGenericAssetsPipeline::ImplementUseSourceNameForAssetOption()
 		if (AnimSequenceNodeUids.Num() == 1)
 		{
 			UInterchangeAnimSequenceFactoryNode* AnimSequenceNode = Cast<UInterchangeAnimSequenceFactoryNode>(BaseNodeContainer->GetFactoryNode(AnimSequenceNodeUids[0]));
-			const FString DisplayLabelName = FPaths::GetBaseFilename(SourceDatas[0]->GetFilename()) + TEXT("_Anim");
+			const FString DisplayLabelName = (AssetName.IsEmpty() ? FPaths::GetBaseFilename(SourceDatas[0]->GetFilename()) : AssetName) + TEXT("_Anim");
 			AnimSequenceNode->SetDisplayLabel(DisplayLabelName);
 		}
 	}
-
 }
 
diff --git a/Engine/Plugins/Interchange/Runtime/Source/Pipelines/Private/InterchangeGenericSkeletalMeshPipeline.cpp b/Engine/Plugins/Interchange/Runtime/Source/Pipelines/Private/InterchangeGenericSkeletalMeshPipeline.cpp
--- a/Engine/Plugins/Interchange/Runtime/Source/Pipelines/Private/InterchangeGenericSkeletalMeshPipeline.cpp
+++ b/Engine/Plugins/Interchange/Runtime/Source/Pipelines/Private/InterchangeGenericSkeletalMeshPipeline.cpp
@@ -681,7 +681,7 @@ void UInterchangeGenericMeshPipeline::PostImportPhysicsAssetImport(UObject* Crea
 #endif //WITH_EDITOR
 }
 
-void UInterchangeGenericMeshPipeline::ImplementUseSourceNameForAssetOptionSkeletalMesh(const int32 MeshesImportedNodeCount, const bool bUseSourceNameForAsset)
+void UInterchangeGenericMeshPipeline::ImplementUseSourceNameForAssetOptionSkeletalMesh(const int32 MeshesImportedNodeCount, const bool bUseSourceNameForAsset, const FString& AssetName)
 {
 	check(CommonSkeletalMeshesAndAnimationsProperties.IsValid());
 
@@ -693,7 +693,7 @@ void UInterchangeGenericMeshPipeline::ImplementUseSourceNameForAssetOptionSkelet
 		return;
 	}
 	//If we import only one asset, and bUseSourceNameForAsset is true, we want to rename the asset using the file name.
-	const bool bShouldChangeAssetName = (bUseSourceNameForAsset && MeshesImportedNodeCount == 1);
+	const bool bShouldChangeAssetName = ((bUseSourceNameForAsset || !AssetName.IsEmpty()) && MeshesImportedNodeCount == 1);
 	const FString SkeletalMeshUid = SkeletalMeshNodeUids[0];
 	UInterchangeSkeletalMeshFactoryNode* SkeletalMeshNode = Cast<UInterchangeSkeletalMeshFactoryNode>(BaseNodeContainer->GetFactoryNode(SkeletalMeshUid));
 	if (!SkeletalMeshNode)
@@ -705,7 +705,7 @@ void UInterchangeGenericMeshPipeline::ImplementUseSourceNameForAssetOptionSkelet
 		
 	if (bShouldChangeAssetName)
 	{
-		DisplayLabelName = FPaths::GetBaseFilename(SourceDatas[0]->GetFilename());
+		DisplayLabelName = AssetName.IsEmpty() ? FPaths::GetBaseFilename(SourceDatas[0]->GetFilename()) : AssetName;
 		SkeletalMeshNode->SetDisplayLabel(DisplayLabelName);
 	}
 
diff --git a/Engine/Plugins/Interchange/Runtime/Source/Pipelines/Public/InterchangeGenericAssetsPipeline.h b/Engine/Plugins/Interchange/Runtime/Source/Pipelines/Public/InterchangeGenericAssetsPipeline.h
--- a/Engine/Plugins/Interchange/Runtime/Source/Pipelines/Public/InterchangeGenericAssetsPipeline.h
+++ b/Engine/Plugins/Interchange/Runtime/Source/Pipelines/Public/InterchangeGenericAssetsPipeline.h
@@ -36,10 +36,14 @@ public:
 	UPROPERTY(EditAnywhere, BlueprintReadWrite, Category = "Common", meta = (AdjustPipelineAndRefreshDetailOnChange = "True"))
 	EReimportStrategyFlags ReimportStrategy = EReimportStrategyFlags::ApplyNoProperties;
 
-	/** If enable and there is only one asset and one source data, we will name the asset like the source data name. */
+	/** If enable and AssetName is empty and there is only one asset and one source data, we will name the asset like the source data name. */
 	UPROPERTY(EditAnywhere, BlueprintReadWrite, Category = "Common")
 	bool bUseSourceNameForAsset = true;
 
+	/** If not empty, and there is only one asset and one source data, we will name the asset with this string. */
+	UPROPERTY(EditAnywhere, BlueprintReadWrite, Category = "Common")
+	FString AssetName;
+
 	/** Translation offset applied to meshes and animations. */
 	UPROPERTY(EditAnywhere, BlueprintReadWrite, Category = "Common", meta = (DisplayName = "Offset Translation"))
 	FVector ImportOffsetTranslation;
diff --git a/Engine/Plugins/Interchange/Runtime/Source/Pipelines/Public/InterchangeGenericMeshPipeline.h b/Engine/Plugins/Interchange/Runtime/Source/Pipelines/Public/InterchangeGenericMeshPipeline.h
--- a/Engine/Plugins/Interchange/Runtime/Source/Pipelines/Public/InterchangeGenericMeshPipeline.h
+++ b/Engine/Plugins/Interchange/Runtime/Source/Pipelines/Public/InterchangeGenericMeshPipeline.h
@@ -260,7 +260,7 @@ private:
 public:
 	
 	/** Specialize for skeletalmesh */
-	void ImplementUseSourceNameForAssetOptionSkeletalMesh(const int32 MeshesImportedNodeCount, const bool bUseSourceNameForAsset);
+	void ImplementUseSourceNameForAssetOptionSkeletalMesh(const int32 MeshesImportedNodeCount, const bool bUseSourceNameForAsset, const FString& AssetName);
 
 private:
 	/* Skeletal mesh API END                                                */
diff --git a/Engine/Source/Developer/AssetTools/Private/AssetTools.cpp b/Engine/Source/Developer/AssetTools/Private/AssetTools.cpp
--- a/Engine/Source/Developer/AssetTools/Private/AssetTools.cpp
+++ b/Engine/Source/Developer/AssetTools/Private/AssetTools.cpp
@@ -3258,6 +3258,17 @@ TArray<UObject*> UAssetToolsImpl::ImportAssetsInternal(const TArray<FString>& Fi
 				}
 				else
 				{
+					if (Params.AssetImportTask)
+					{
+						//Add the override pipelines if we have a task that specify valid interchange override pipelines
+						if (UInterchangePipelineStackOverride* PipelineStackOverride = Cast<UInterchangePipelineStackOverride>(Params.AssetImportTask->Options))
+						{
+							for (TObjectPtr<UInterchangePipelineBase> OverridePipeline : PipelineStackOverride->OverridePipelines)
+							{
+								ImportAssetParameters.OverridePipelines.Add(OverridePipeline);
+							}
+						}
+					}
 					UE::Interchange::FAssetImportResultRef InterchangeResult = (InterchangeManager.ImportAssetAsync(DestinationPath, ScopedSourceData.GetSourceData(), ImportAssetParameters));
 
 					// If we have an ImportTask, fill out the asynchronous results object here so the caller can see when the results are ready
diff --git a/Engine/Source/Editor/UnrealEd/Public/AssetImportTask.h b/Engine/Source/Editor/UnrealEd/Public/AssetImportTask.h
--- a/Engine/Source/Editor/UnrealEd/Public/AssetImportTask.h
+++ b/Engine/Source/Editor/UnrealEd/Public/AssetImportTask.h
@@ -39,7 +39,7 @@ public:
 	UPROPERTY(BlueprintReadWrite, Category = "Asset Import Task")
 	FString DestinationPath;
 
-	/** Optional custom name to import as */
+	/** Optional custom name to import as (if you are using interchange the name must be set in a pipeline and this field will be ignored)*/
 	UPROPERTY(BlueprintReadWrite, Category = "Asset Import Task")
 	FString DestinationName;
 
diff --git a/Engine/Source/Runtime/Interchange/Engine/Public/InterchangeManager.h b/Engine/Source/Runtime/Interchange/Engine/Public/InterchangeManager.h
--- a/Engine/Source/Runtime/Interchange/Engine/Public/InterchangeManager.h
+++ b/Engine/Source/Runtime/Interchange/Engine/Public/InterchangeManager.h
@@ -263,6 +263,19 @@ namespace UE
 	} //ns interchange
 } //ns UE
 
+/**
+ * This class is use to pass override pipelines in the ImportAssetTask Options member
+ */
+UCLASS(Transient, BlueprintType)
+class INTERCHANGEENGINE_API UInterchangePipelineStackOverride : public UObject
+{
+	GENERATED_BODY()
+public:
+
+	UPROPERTY(EditAnywhere, BlueprintReadWrite, Category = "Interchange|ImportAsset")
+		TArray<TObjectPtr<UInterchangePipelineBase>> OverridePipelines;
+};
+
 USTRUCT(BlueprintType)
 struct FImportAssetParameters
 {
