[Interchange] Adding MorphTargetCurveWeights support for SkeletalMeshActors

diff --git a/Engine/Plugins/Interchange/Runtime/Source/FactoryNodes/Private/InterchangeMeshActorFactoryNode.cpp b/Engine/Plugins/Interchange/Runtime/Source/FactoryNodes/Private/InterchangeMeshActorFactoryNode.cpp
--- a/Engine/Plugins/Interchange/Runtime/Source/FactoryNodes/Private/InterchangeMeshActorFactoryNode.cpp
+++ b/Engine/Plugins/Interchange/Runtime/Source/FactoryNodes/Private/InterchangeMeshActorFactoryNode.cpp
@@ -11,6 +11,7 @@
 UInterchangeMeshActorFactoryNode::UInterchangeMeshActorFactoryNode()
 {
 	SlotMaterialDependencies.Initialize(Attributes.ToSharedRef(), TEXT("__SlotMaterialDependencies__"));
+	MorphTargetCurveWeights.Initialize(Attributes.ToSharedRef(), TEXT("__MorphTargetCurves__Key"));
 }
 
 void UInterchangeMeshActorFactoryNode::GetSlotMaterialDependencies(TMap<FString, FString>& OutMaterialDependencies) const
@@ -33,3 +34,20 @@ bool UInterchangeMeshActorFactoryNode::RemoveSlotMaterialDependencyUid(const FSt
 	return SlotMaterialDependencies.RemoveKey(SlotName);
 }
 
+bool UInterchangeMeshActorFactoryNode::SetMorphTargetCurveWeight(const FString& MorphTargetName, const float& Weight)
+{
+	return MorphTargetCurveWeights.SetKeyValue(MorphTargetName, Weight);
+}
+
+void UInterchangeMeshActorFactoryNode::SetMorphTargetCurveWeights(const TMap<FString, float>& InMorphTargetCurveWeights)
+{
+	for (const TPair<FString, float>& MorphTargetCurve : InMorphTargetCurveWeights)
+	{
+		SetMorphTargetCurveWeight(MorphTargetCurve.Key, MorphTargetCurve.Value);
+	}
+}
+
+void UInterchangeMeshActorFactoryNode::GetMorphTargetCurveWeights(TMap<FString, float>& OutMorphTargetCurveWeights) const
+{
+	OutMorphTargetCurveWeights = MorphTargetCurveWeights.ToMap();
+}
\ No newline at end of file
diff --git a/Engine/Plugins/Interchange/Runtime/Source/FactoryNodes/Public/InterchangeMeshActorFactoryNode.h b/Engine/Plugins/Interchange/Runtime/Source/FactoryNodes/Public/InterchangeMeshActorFactoryNode.h
--- a/Engine/Plugins/Interchange/Runtime/Source/FactoryNodes/Public/InterchangeMeshActorFactoryNode.h
+++ b/Engine/Plugins/Interchange/Runtime/Source/FactoryNodes/Public/InterchangeMeshActorFactoryNode.h
@@ -51,6 +51,19 @@ public:
 	UFUNCTION(BlueprintCallable, Category = "Interchange | Node | StaticMesh")
 	bool RemoveSlotMaterialDependencyUid(const FString& SlotName);
 
+	/** Set MorphTarget with given weight. */
+	UFUNCTION(BlueprintCallable, Category = "Interchange | Node | SkeletalMesh")
+	bool SetMorphTargetCurveWeight(const FString& MorphTargetName, const float& Weight);
+
+	/** Set MorphTarget with given weight. */
+	UFUNCTION(BlueprintCallable, Category = "Interchange | Node | SkeletalMesh")
+	void SetMorphTargetCurveWeights(const TMap<FString, float>& InMorphTargetCurveWeights);
+
+	/** Get MorphTargets and their weights. */
+	UFUNCTION(BlueprintCallable, Category = "Interchange | Node | SkeletalMesh")
+	void GetMorphTargetCurveWeights(TMap<FString, float>& OutMorphTargetCurveWeights) const;
+
 private:
 	UE::Interchange::TMapAttributeHelper<FString, FString> SlotMaterialDependencies;
+	UE::Interchange::TMapAttributeHelper<FString, float> MorphTargetCurveWeights;
 };
\ No newline at end of file
diff --git a/Engine/Plugins/Interchange/Runtime/Source/Import/Private/Scene/InterchangeSkeletalMeshActorFactory.cpp b/Engine/Plugins/Interchange/Runtime/Source/Import/Private/Scene/InterchangeSkeletalMeshActorFactory.cpp
--- a/Engine/Plugins/Interchange/Runtime/Source/Import/Private/Scene/InterchangeSkeletalMeshActorFactory.cpp
+++ b/Engine/Plugins/Interchange/Runtime/Source/Import/Private/Scene/InterchangeSkeletalMeshActorFactory.cpp
@@ -46,33 +46,37 @@ void UInterchangeSkeletalMeshActorFactory::SetupSkeletalMeshActor(const UInterch
 	SkeletalMeshComponent->UnregisterComponent();
 }
 
-void UInterchangeSkeletalMeshActorFactory::FinalizeObject_GameThread(const FSetupObjectParams& Arguments)
+void UInterchangeSkeletalMeshActorFactory::SetupObject_GameThread(const FSetupObjectParams& Arguments)
 {
-	// Set the skeletal mesh on the component in the post import callback, once the skeletal mesh has been fully imported.
-
 	if (ASkeletalMeshActor* SkeletalMeshActor = Cast<ASkeletalMeshActor>(Arguments.ImportedObject))
 	{
-		if (USkeletalMeshComponent * SkeletalMeshComponent = SkeletalMeshActor->GetSkeletalMeshComponent())
+		if (USkeletalMeshComponent* SkeletalMeshComponent = SkeletalMeshActor->GetSkeletalMeshComponent())
 		{
+			TArray<FString> TargetNodeUids;
+			Arguments.FactoryNode->GetTargetNodeUids(TargetNodeUids);
+			const UInterchangeSkeletalMeshFactoryNode* SkeletalMeshFactoryNode = TargetNodeUids.IsEmpty() ? nullptr : Cast<UInterchangeSkeletalMeshFactoryNode>(Arguments.NodeContainer->GetFactoryNode(TargetNodeUids[0]));
+			if (SkeletalMeshFactoryNode)
 			{
-				TArray<FString> TargetNodeUids;
-				Arguments.FactoryNode->GetTargetNodeUids(TargetNodeUids);
-				const UInterchangeSkeletalMeshFactoryNode* SkeletalMeshFactoryNode = TargetNodeUids.IsEmpty() ? nullptr : Cast<UInterchangeSkeletalMeshFactoryNode>(Arguments.NodeContainer->GetFactoryNode(TargetNodeUids[0]));
-				if (SkeletalMeshFactoryNode)
+				FSoftObjectPath ReferenceObject;
+				SkeletalMeshFactoryNode->GetCustomReferenceObject(ReferenceObject);
+				if (USkeletalMesh* SkeletalMesh = Cast<USkeletalMesh>(ReferenceObject.TryLoad()))
 				{
-					FSoftObjectPath ReferenceObject;
-					SkeletalMeshFactoryNode->GetCustomReferenceObject(ReferenceObject);
-					if (USkeletalMesh* SkeletalMesh = Cast<USkeletalMesh>(ReferenceObject.TryLoad()))
+					SkeletalMeshComponent->SetSkeletalMesh(SkeletalMesh);
+
+					if (const UInterchangeMeshActorFactoryNode* MeshActorFactoryNode = Cast<UInterchangeMeshActorFactoryNode>(Arguments.FactoryNode))
 					{
-						SkeletalMeshComponent->SetSkeletalMesh(SkeletalMesh);
+						UE::Interchange::ActorHelper::ApplySlotMaterialDependencies(*Arguments.NodeContainer, *MeshActorFactoryNode, *SkeletalMeshComponent);
 
-						if (const UInterchangeMeshActorFactoryNode* MeshActorFactoryNode = Cast<UInterchangeMeshActorFactoryNode>(Arguments.FactoryNode))
+						//Set Morph Target Curves:
+						TMap<FString, float> MorphTargetCurveWeights;
+						MeshActorFactoryNode->GetMorphTargetCurveWeights(MorphTargetCurveWeights);
+						for (const TPair<FString, float>& MorphTargetCurve : MorphTargetCurveWeights)
 						{
-							UE::Interchange::ActorHelper::ApplySlotMaterialDependencies(*Arguments.NodeContainer, *MeshActorFactoryNode, *SkeletalMeshComponent);
+							SkeletalMeshComponent->SetMorphTarget(*MorphTargetCurve.Key, MorphTargetCurve.Value);
 						}
 					}
 				}
 			}
 		}
 	}
-}
+}
\ No newline at end of file
diff --git a/Engine/Plugins/Interchange/Runtime/Source/Import/Public/Scene/InterchangeSkeletalMeshActorFactory.h b/Engine/Plugins/Interchange/Runtime/Source/Import/Public/Scene/InterchangeSkeletalMeshActorFactory.h
--- a/Engine/Plugins/Interchange/Runtime/Source/Import/Public/Scene/InterchangeSkeletalMeshActorFactory.h
+++ b/Engine/Plugins/Interchange/Runtime/Source/Import/Public/Scene/InterchangeSkeletalMeshActorFactory.h
@@ -26,7 +26,7 @@ public:
 
 	virtual UObject* ImportSceneObject_GameThread(const UInterchangeFactoryBase::FImportSceneObjectsParams& CreateSceneObjectsParams) override;
 
-	virtual void FinalizeObject_GameThread(const FSetupObjectParams& Arguments) override;
+	virtual void SetupObject_GameThread(const FSetupObjectParams& Arguments) override;
 
 	// Interchange factory base interface end
 	//////////////////////////////////////////////////////////////////////////
diff --git a/Engine/Plugins/Interchange/Runtime/Source/Nodes/Private/InterchangeSceneNode.cpp b/Engine/Plugins/Interchange/Runtime/Source/Nodes/Private/InterchangeSceneNode.cpp
--- a/Engine/Plugins/Interchange/Runtime/Source/Nodes/Private/InterchangeSceneNode.cpp
+++ b/Engine/Plugins/Interchange/Runtime/Source/Nodes/Private/InterchangeSceneNode.cpp
@@ -46,6 +46,12 @@ namespace UE
 			static FString SlotMaterialDependenciesString(TEXT("__SlotMaterialDependencies__"));
 			return SlotMaterialDependenciesString;
 		}
+
+		const FString& FSceneNodeStaticData::GetMorphTargetCurveWeightsKey()
+		{
+			static FString MorphTargetCurvesKey(TEXT("__MorphTargetCurveWeights__Key"));
+			return MorphTargetCurvesKey;
+		}
 	}//ns Interchange
 }//ns UE
 
@@ -53,6 +59,7 @@ UInterchangeSceneNode::UInterchangeSceneNode()
 {
 	NodeSpecializeTypes.Initialize(Attributes, UE::Interchange::FSceneNodeStaticData::GetNodeSpecializeTypeBaseKey().ToString());
 	SlotMaterialDependencies.Initialize(Attributes.ToSharedRef(), UE::Interchange::FSceneNodeStaticData::GetSlotMaterialDependenciesString());
+	MorphTargetCurveWeights.Initialize(Attributes.ToSharedRef(), UE::Interchange::FSceneNodeStaticData::GetMorphTargetCurveWeightsKey());
 }
 
 /**
@@ -376,3 +383,13 @@ bool UInterchangeSceneNode::RemoveSlotMaterialDependencyUid(const FString& SlotN
 	return SlotMaterialDependencies.RemoveKey(SlotName);
 }
 
+
+bool UInterchangeSceneNode::SetMorphTargetCurveWeight(const FString& MorphTargetName, const float& Weight)
+{
+	return MorphTargetCurveWeights.SetKeyValue(MorphTargetName, Weight);
+}
+
+void UInterchangeSceneNode::GetMorphTargetCurveWeights(TMap<FString, float>& OutMorphTargetCurveWeights) const
+{
+	OutMorphTargetCurveWeights = MorphTargetCurveWeights.ToMap();
+}
\ No newline at end of file
diff --git a/Engine/Plugins/Interchange/Runtime/Source/Nodes/Public/InterchangeSceneNode.h b/Engine/Plugins/Interchange/Runtime/Source/Nodes/Public/InterchangeSceneNode.h
--- a/Engine/Plugins/Interchange/Runtime/Source/Nodes/Public/InterchangeSceneNode.h
+++ b/Engine/Plugins/Interchange/Runtime/Source/Nodes/Public/InterchangeSceneNode.h
@@ -23,6 +23,7 @@ namespace UE
 			static const FString& GetJointSpecializeTypeString();
 			static const FString& GetLodGroupSpecializeTypeString();
 			static const FString& GetSlotMaterialDependenciesString();
+			static const FString& GetMorphTargetCurveWeightsKey();
 		};
 
 	}//ns Interchange
@@ -196,6 +197,14 @@ public:
 	UFUNCTION(BlueprintCallable, Category = "Interchange | Node | Meshes")
 	bool RemoveSlotMaterialDependencyUid(const FString& SlotName);
 
+	/** Set MorphTarget with given weight. */
+	UFUNCTION(BlueprintCallable, Category = "Interchange | Node | SkeletalMesh")
+	bool SetMorphTargetCurveWeight(const FString& MorphTargetName, const float& Weight);
+
+	/** Get MorphTargets and their weights. */
+	UFUNCTION(BlueprintCallable, Category = "Interchange | Node | SkeletalMesh")
+	void GetMorphTargetCurveWeights(TMap<FString, float>& OutMorphTargetCurveWeights) const;
+
 private:
 
 	bool GetGlobalTransformInternal(const UE::Interchange::FAttributeKey LocalTransformKey
@@ -226,6 +235,9 @@ private:
 	//A scene node can have is own set of materials for the mesh it reference.
 	UE::Interchange::TMapAttributeHelper<FString, FString> SlotMaterialDependencies;
 
+	//A scene node can have different MorphTarget curve settings:
+	UE::Interchange::TMapAttributeHelper<FString, float> MorphTargetCurveWeights;
+
 	//mutable caches for global transforms
 	mutable TOptional<FTransform> CacheGlobalTransform;
 	mutable TOptional<FTransform> CacheBindPoseGlobalTransform;
diff --git a/Engine/Plugins/Interchange/Runtime/Source/Pipelines/Private/InterchangeGenericScenesPipeline.cpp b/Engine/Plugins/Interchange/Runtime/Source/Pipelines/Private/InterchangeGenericScenesPipeline.cpp
--- a/Engine/Plugins/Interchange/Runtime/Source/Pipelines/Private/InterchangeGenericScenesPipeline.cpp
+++ b/Engine/Plugins/Interchange/Runtime/Source/Pipelines/Private/InterchangeGenericScenesPipeline.cpp
@@ -287,6 +287,10 @@ void UInterchangeGenericLevelPipeline::SetUpFactoryNode(UInterchangeActorFactory
 
 			UE::Interchange::MeshesUtilities::ApplySlotMaterialDependencies(*MeshActorFactoryNode, SlotMaterialDependencies, *BaseNodeContainer);
 
+			TMap<FString, float> MorphTargetCurveWeights;
+			SceneNode->GetMorphTargetCurveWeights(MorphTargetCurveWeights);
+			MeshActorFactoryNode->SetMorphTargetCurveWeights(MorphTargetCurveWeights);
+
 			MeshActorFactoryNode->AddFactoryDependencyUid(UInterchangeFactoryBaseNode::BuildFactoryNodeUid(MeshNode->GetUniqueID()));
 		}
 	}
