Fixed crash when assigning a static mesh to a static mesh actor when importing a file twice in the same content folder
  - Solution: Get rid of the render data as a whole instead of just the resources

diff --git a/Engine/Plugins/Interchange/Runtime/Source/Import/Private/Mesh/InterchangeStaticMeshFactory.cpp b/Engine/Plugins/Interchange/Runtime/Source/Import/Private/Mesh/InterchangeStaticMeshFactory.cpp
--- a/Engine/Plugins/Interchange/Runtime/Source/Import/Private/Mesh/InterchangeStaticMeshFactory.cpp
+++ b/Engine/Plugins/Interchange/Runtime/Source/Import/Private/Mesh/InterchangeStaticMeshFactory.cpp
@@ -78,7 +78,7 @@ UObject* UInterchangeStaticMeshFactory::ImportAssetObject_GameThread(const FImpo
 		//This is a reimport, we are just re-updating the source data
 		StaticMesh = Cast<UStaticMesh>(ExistingAsset);
 
-		// Clear the render resources on the static mesh from the game thread so that we're ready to update it
+		// Clear the render data on the existing static mesh from the game thread so that we're ready to update it
 		if (StaticMesh && StaticMesh->AreRenderingResourcesInitialized())
 		{
 			const bool bInvalidateLighting = true;
@@ -86,6 +86,8 @@ UObject* UInterchangeStaticMeshFactory::ImportAssetObject_GameThread(const FImpo
 			FStaticMeshComponentRecreateRenderStateContext RecreateRenderStateContext(StaticMesh, bInvalidateLighting, bRefreshBounds);
 			StaticMesh->ReleaseResources();
 			StaticMesh->ReleaseResourcesFence.Wait();
+
+			StaticMesh->SetRenderData(nullptr);
 		}
 	}
 	
