[Interchange] MorphTarget Animation handling fix (for gltf)

diff --git a/Engine/Plugins/Interchange/Runtime/Source/Import/Private/Animation/InterchangeAnimSequenceFactory.cpp b/Engine/Plugins/Interchange/Runtime/Source/Import/Private/Animation/InterchangeAnimSequenceFactory.cpp
--- a/Engine/Plugins/Interchange/Runtime/Source/Import/Private/Animation/InterchangeAnimSequenceFactory.cpp
+++ b/Engine/Plugins/Interchange/Runtime/Source/Import/Private/Animation/InterchangeAnimSequenceFactory.cpp
@@ -515,23 +515,47 @@ namespace UE::Interchange::Private
 
 			//Import morph target curves
 			{
+				TMap<FString, FString> MorphTargetNodeAnimationPayloads;
+				AnimSequenceFactoryNode->GetMorphTargetNodeAnimationPayloadKeys(MorphTargetNodeAnimationPayloads);
+
 				TMap<FString, TFuture<TOptional<UE::Interchange::FAnimationCurvePayloadData>>> AnimationCurvesPayloads;
 				TMap<FString, FString> AnimationCurveMorphTargetNodeNames;
-				//Import morph target curves (FRichCurve is what anim api )
-				TArray<FString> AnimatedMorphTargetUids;
-				AnimSequenceFactoryNode->GetAnimatedMorphTargetDependencies(AnimatedMorphTargetUids);
-				for (const FString& MorphTargetUid : AnimatedMorphTargetUids)
+
+				if (MorphTargetNodeAnimationPayloads.Num() > 0)
+				{
+					for (const TPair<FString, FString>& MorphTargetNodeUidAnimationPayload : MorphTargetNodeAnimationPayloads)
+					{
+						FString PayloadKey = MorphTargetNodeUidAnimationPayload.Value;
+						if (PayloadKey.Len() == 0)
+						{
+							continue;
+						}
+						if (const UInterchangeMeshNode* MorphTargetNode = Cast<UInterchangeMeshNode>(NodeContainer->GetNode(MorphTargetNodeUidAnimationPayload.Key)))
+						{
+							AnimationCurvesPayloads.Add(PayloadKey, AnimSequenceTranslatorPayloadInterface->GetAnimationCurvePayloadData(PayloadKey));
+							AnimationCurveMorphTargetNodeNames.Add(PayloadKey, MorphTargetNode->GetDisplayLabel());
+						}
+					}
+				}
+				else
 				{
-					if (const UInterchangeMeshNode* MorphTargetNode = Cast<UInterchangeMeshNode>(NodeContainer->GetNode(MorphTargetUid)))
+					//Import morph target curves (FRichCurve is what anim api )
+					TArray<FString> AnimatedMorphTargetUids;
+					AnimSequenceFactoryNode->GetAnimatedMorphTargetDependencies(AnimatedMorphTargetUids);
+					for (const FString& MorphTargetUid : AnimatedMorphTargetUids)
 					{
-						TOptional<FString> PayloadKey = MorphTargetNode->GetAnimationCurvePayLoadKey();
-						if (PayloadKey.IsSet())
+						if (const UInterchangeMeshNode* MorphTargetNode = Cast<UInterchangeMeshNode>(NodeContainer->GetNode(MorphTargetUid)))
 						{
-							AnimationCurvesPayloads.Add(PayloadKey.GetValue(), AnimSequenceTranslatorPayloadInterface->GetAnimationCurvePayloadData(PayloadKey.GetValue()));
-							AnimationCurveMorphTargetNodeNames.Add(PayloadKey.GetValue(), MorphTargetNode->GetDisplayLabel());
+							TOptional<FString> PayloadKey = MorphTargetNode->GetAnimationCurvePayLoadKey();
+							if (PayloadKey.IsSet())
+							{
+								AnimationCurvesPayloads.Add(PayloadKey.GetValue(), AnimSequenceTranslatorPayloadInterface->GetAnimationCurvePayloadData(PayloadKey.GetValue()));
+								AnimationCurveMorphTargetNodeNames.Add(PayloadKey.GetValue(), MorphTargetNode->GetDisplayLabel());
+							}
 						}
 					}
 				}
+
 				for (TPair<FString, TFuture<TOptional<UE::Interchange::FAnimationCurvePayloadData>>>& CurveNameAndPayload : AnimationCurvesPayloads)
 				{
 					TOptional<UE::Interchange::FAnimationCurvePayloadData> AnimationCurvePayload = CurveNameAndPayload.Value.Get();
diff --git a/Engine/Plugins/Interchange/Runtime/Source/Nodes/Private/InterchangeAnimationTrackSetNode.cpp b/Engine/Plugins/Interchange/Runtime/Source/Nodes/Private/InterchangeAnimationTrackSetNode.cpp
--- a/Engine/Plugins/Interchange/Runtime/Source/Nodes/Private/InterchangeAnimationTrackSetNode.cpp
+++ b/Engine/Plugins/Interchange/Runtime/Source/Nodes/Private/InterchangeAnimationTrackSetNode.cpp
@@ -171,16 +171,6 @@ bool UInterchangeSkeletalAnimationTrackNode::SetCustomSkeletonNodeUid(const FStr
 	IMPLEMENT_NODE_ATTRIBUTE_SETTER_NODELEGATE(SkeletonNodeUid, FString);
 }
 
-bool UInterchangeSkeletalAnimationTrackNode::GetCustomSkeletalMeshNodeUid(FString& AttributeValue) const
-{
-	IMPLEMENT_NODE_ATTRIBUTE_GETTER(SkeletalMeshNodeUid, FString);
-}
-
-bool UInterchangeSkeletalAnimationTrackNode::SetCustomSkeletalMeshNodeUid(const FString& AttributeValue)
-{
-	IMPLEMENT_NODE_ATTRIBUTE_SETTER_NODELEGATE(SkeletalMeshNodeUid, FString);
-}
-
 bool UInterchangeSkeletalAnimationTrackNode::GetCustomAnimationSampleRate(double& AttributeValue) const
 {
 	IMPLEMENT_NODE_ATTRIBUTE_GETTER(AnimationSampleRate, double);
diff --git a/Engine/Plugins/Interchange/Runtime/Source/Nodes/Public/InterchangeAnimationTrackSetNode.h b/Engine/Plugins/Interchange/Runtime/Source/Nodes/Public/InterchangeAnimationTrackSetNode.h
--- a/Engine/Plugins/Interchange/Runtime/Source/Nodes/Public/InterchangeAnimationTrackSetNode.h
+++ b/Engine/Plugins/Interchange/Runtime/Source/Nodes/Public/InterchangeAnimationTrackSetNode.h
@@ -403,14 +403,6 @@ public:
 	UFUNCTION(BlueprintCallable, Category = "Interchange | Node | SkeletalAnimationTrack")
 	bool SetCustomSkeletonNodeUid(const FString& AttributeValue);
 
-	/** Get the skeletal mesh node unique id. Return false if the attribute is not set. */
-	UFUNCTION(BlueprintCallable, Category = "Interchange | Node | SkeletalAnimationTrack")
-	bool GetCustomSkeletalMeshNodeUid(FString& AttributeValue) const;
-
-	/** Set the skeletal mesh node unique id. Return false if the attribute cannot be set. */
-	UFUNCTION(BlueprintCallable, Category = "Interchange | Node | SkeletalAnimationTrack")
-	bool SetCustomSkeletalMeshNodeUid(const FString& AttributeValue);
-
 	/**
 	 * Set the animation sample rate. Return false if the attribute cannot be set.
 	 */
@@ -473,7 +465,6 @@ public:
 
 private:
 	const UE::Interchange::FAttributeKey Macro_CustomSkeletonNodeUidKey = UE::Interchange::FAttributeKey(TEXT("SkeletonNodeUid"));
-	const UE::Interchange::FAttributeKey Macro_CustomSkeletalMeshNodeUidKey = UE::Interchange::FAttributeKey(TEXT("SkeletalMeshNodeUid"));
 	const UE::Interchange::FAttributeKey Macro_CustomAnimationSampleRateKey = UE::Interchange::FAttributeKey(TEXT("AnimationSampleRate"));
 	const UE::Interchange::FAttributeKey Macro_CustomAnimationStartTimeKey = UE::Interchange::FAttributeKey(TEXT("AnimationStartTime"));
 	const UE::Interchange::FAttributeKey Macro_CustomAnimationStopTimeKey = UE::Interchange::FAttributeKey(TEXT("AnimationStopTime"));
diff --git a/Engine/Plugins/Interchange/Runtime/Source/Pipelines/Private/InterchangeGenericAnimationPipeline.cpp b/Engine/Plugins/Interchange/Runtime/Source/Pipelines/Private/InterchangeGenericAnimationPipeline.cpp
--- a/Engine/Plugins/Interchange/Runtime/Source/Pipelines/Private/InterchangeGenericAnimationPipeline.cpp
+++ b/Engine/Plugins/Interchange/Runtime/Source/Pipelines/Private/InterchangeGenericAnimationPipeline.cpp
@@ -644,11 +644,10 @@ void UInterchangeGenericAnimationPipeline::CreateAnimSequenceFactoryNode(UInterc
 		return;
 	}
 
-	FString SkeletalMeshNodeUid;
+	FString SkeletalMeshFactoryNodeUid;
 	const UInterchangeSkeletalMeshFactoryNode* SkeletalMeshFactoryNode = nullptr;
-	if (TrackNode.GetCustomSkeletalMeshNodeUid(SkeletalMeshNodeUid))
+	if (SkeletonFactoryNode->GetCustomSkeletalMeshFactoryNodeUid(SkeletalMeshFactoryNodeUid))
 	{
-		const FString SkeletalMeshFactoryNodeUid = UInterchangeFactoryBaseNode::BuildFactoryNodeUid(SkeletalMeshNodeUid);
 		SkeletalMeshFactoryNode = Cast<const UInterchangeSkeletalMeshFactoryNode>(BaseNodeContainer->GetFactoryNode(SkeletalMeshFactoryNodeUid));
 	}
 
