[Interchange - glTF] Add support for morph and morph animation

diff --git a/Engine/Plugins/Editor/GLTFImporter/Source/GLTFCore/Private/GLTF/GLTFAsset.cpp b/Engine/Plugins/Editor/GLTFImporter/Source/GLTFCore/Private/GLTF/GLTFAsset.cpp
--- a/Engine/Plugins/Editor/GLTFImporter/Source/GLTFCore/Private/GLTF/GLTFAsset.cpp
+++ b/Engine/Plugins/Editor/GLTFImporter/Source/GLTFCore/Private/GLTF/GLTFAsset.cpp
@@ -230,6 +230,67 @@ namespace GLTF
 		GLTF::ValidateFNameCriteria(TEXT("Texture_"), Textures);
 		GLTF::ValidateFNameCriteria(TEXT("Mesh_"), Meshes);
 		GLTF::ValidateFNameCriteria(TEXT("Animation_"), Animations);
+
+		//Generate MorphTarget Names:
+		//MorphTargetNames have to be unique across meshes as well:
+		TSet<FString> MorphTargetNamesUniquenessCheckAcrossMeshes;
+		int32 MorphTargetCounterAcrossMeshes = 0;
+		for (int32 MeshIndex = 0; MeshIndex < Meshes.Num(); MeshIndex++)
+		{
+			if (Meshes[MeshIndex].Primitives.Num() > 0)
+			{
+				TSet<FString> MorphTargetNamesUniquenessCheck; //per mesh
+
+				for (const FString& MorphTargetName : Meshes[MeshIndex].MorphTargetNames)
+				{
+					MorphTargetNamesUniquenessCheckAcrossMeshes.Add(MorphTargetName);
+					MorphTargetCounterAcrossMeshes++;
+
+					//validate uniqueness of MorphTargetNames
+					MorphTargetNamesUniquenessCheck.Add(MorphTargetName);
+				}
+
+				if (MorphTargetNamesUniquenessCheck.Num() != Meshes[MeshIndex].MorphTargetNames.Num())
+				{
+					Meshes[MeshIndex].MorphTargetNames.Empty();
+				}
+			}
+		}
+		bool bReBuildMorphTargetNames = MorphTargetNamesUniquenessCheckAcrossMeshes.Num() != MorphTargetCounterAcrossMeshes;
+
+		for (int32 MeshIndex = 0; MeshIndex < Meshes.Num(); MeshIndex++)
+		{
+			if (Meshes[MeshIndex].Primitives.Num() > 0)
+			{
+				//at this point the NumberOfMorphTargets across primitives have been validated (they have to be equal)
+				//Note: All primitives MUST have the same number of morph targets in the same order.
+				
+				if (bReBuildMorphTargetNames)
+				{
+					Meshes[MeshIndex].MorphTargetNames.Empty();
+				}
+
+				//check if number of morph targets and the number of morphtarget names:
+				//NumberOfMorphTargetsPerPrimitive will return the first Primitive's MorphTarget's count,
+				//		In case Primitives have varying MorphTarget counts then the ValidationCheck will report false.
+				int32 NumberOfMorphTargets = Meshes[MeshIndex].NumberOfMorphTargetsPerPrimitive();
+				if (NumberOfMorphTargets != Meshes[MeshIndex].MorphTargetNames.Num())
+				{
+					Meshes[MeshIndex].MorphTargetNames.Empty();
+				}
+
+				//if morph target names failed a criteria above, generate unique names:
+				if (Meshes[MeshIndex].MorphTargetNames.IsEmpty())
+				{
+					for (int32 MorphTargetIndex = 0; MorphTargetIndex < NumberOfMorphTargets; MorphTargetIndex++)
+					{
+						Meshes[MeshIndex].MorphTargetNames.Add(Meshes[MeshIndex].UniqueId + TEXT("_") + LexToString(MorphTargetIndex) + TEXT("_MorphTarget"));
+					}
+
+					Meshes[MeshIndex].GenerateIsValidCache(false);
+				}
+			}
+		}
 	}
 
 	void FAsset::GetRootNodes(TArray<int32>& NodeIndices)
@@ -266,10 +327,33 @@ namespace GLTF
 	{
 		int32 Res = EValidationCheck::Valid;
 		if (Meshes.FindByPredicate([](const FMesh& Mesh) { return !Mesh.IsValid(); }))
+		{
 			Res |= InvalidMeshPresent;
+		}
+		
+		if (Nodes.FindByPredicate([this](const FNode& Node)
+			{
+				if (Node.MeshIndex != INDEX_NONE)
+				{
+					//check Mesh index validity:
+					if (!Meshes.IsValidIndex(Node.MeshIndex))
+					{
+						return false;
+					}
+
+					//Check Morph Target Weights:
+					//	Number of MorphTargets across primitives are validated in Mesh.IsValid.
+					if (Node.MorphTargetWeights.Num() > 0 && (Node.MorphTargetWeights.Num() != Meshes[Node.MeshIndex].NumberOfMorphTargetsPerPrimitive()))
+					{
+						return false;
+					}
+				}
 
-		if (Nodes.FindByPredicate([](const FNode& Node) { return !Node.Transform.IsValid(); }))
-			Res |= InvalidNodeTransform;
+				return !Node.Transform.IsValid();
+			}))
+		{
+			Res |= InvalidNodePresent;
+		}
 
 		// TODO: lots more validation
 
diff --git a/Engine/Plugins/Editor/GLTFImporter/Source/GLTFCore/Private/GLTF/GLTFMesh.cpp b/Engine/Plugins/Editor/GLTFImporter/Source/GLTFCore/Private/GLTF/GLTFMesh.cpp
--- a/Engine/Plugins/Editor/GLTFImporter/Source/GLTFCore/Private/GLTF/GLTFMesh.cpp
+++ b/Engine/Plugins/Editor/GLTFImporter/Source/GLTFCore/Private/GLTF/GLTFMesh.cpp
@@ -13,11 +13,171 @@ namespace GLTF
 		bool IsConvertibleToNormalizedFloat(const FAccessor& Attrib)
 		{
 			return Attrib.ComponentType == FAccessor::EComponentType::F32 ||
-			       (Attrib.Normalized &&
-			        (Attrib.ComponentType == FAccessor::EComponentType::U8 || Attrib.ComponentType == FAccessor::EComponentType::U16));
+				(Attrib.Normalized &&
+					(Attrib.ComponentType == FAccessor::EComponentType::U8 || Attrib.ComponentType == FAccessor::EComponentType::U16));
 		}
 	}
 
+	FMorphTarget::FMorphTarget(const FAccessor& InPositionDisplacements, const FAccessor& InNormalDisplacements,
+		const FAccessor& InTangentDisplacements, const FAccessor& InTexCoord0Displacements, const FAccessor& InTexCoord1Displacements, const FAccessor& InColor0Deltas)
+		: PositionDisplacements(InPositionDisplacements)
+		, NormalDisplacements(InNormalDisplacements)
+		, TangentDisplacements(InTangentDisplacements)
+		, TexCoord0Displacements(InTexCoord0Displacements)
+		, TexCoord1Displacements(InTexCoord1Displacements)
+		, Color0Deltas(InColor0Deltas)
+	{
+	}
+
+	bool FMorphTarget::IsValid() const
+	{
+		return HasPositionDisplacements() || HasNormalDisplacements() || HasTangentDisplacements() || HasTexCoordDisplacements(0) || HasTexCoordDisplacements(1) || HasColorDeltas();
+	}
+	FMD5Hash FMorphTarget::GetHash() const
+	{
+		if (!IsValid())
+		{
+			return FMD5Hash();
+		}
+
+		auto HashAccessor = [](FMD5& MD5, const FAccessor& Accessor)
+		{
+			if (!Accessor.IsValid())
+			{
+				return;
+			}
+
+			FMD5Hash MD5Hash = Accessor.GetHash();
+			if (MD5Hash.IsValid())
+			{
+				MD5.Update(MD5Hash.GetBytes(), MD5Hash.GetSize());
+			}
+		};
+
+		FMD5 MD5;
+
+		HashAccessor(MD5, PositionDisplacements);
+		HashAccessor(MD5, NormalDisplacements);
+		HashAccessor(MD5, TangentDisplacements);
+		HashAccessor(MD5, TexCoord0Displacements);
+		HashAccessor(MD5, TexCoord1Displacements);
+		HashAccessor(MD5, Color0Deltas);
+
+		FMD5Hash Hash;
+		Hash.Set(MD5);
+
+		return Hash;
+	}
+
+	bool FMorphTarget::HasPositionDisplacements() const
+	{
+		return PositionDisplacements.IsValid();
+	}
+	void FMorphTarget::GetPositionDisplacements(TArray<FVector3f>& Buffer) const
+	{
+		PositionDisplacements.GetCoordArray(Buffer);
+	}
+	int32 FMorphTarget::GetNumberOfPositionDisplacements() const
+	{
+		return PositionDisplacements.Count;
+	}
+
+	bool FMorphTarget::HasNormalDisplacements() const
+	{
+		return NormalDisplacements.IsValid();
+	}
+	void FMorphTarget::GetNormalDisplacements(TArray<FVector3f>& Buffer) const
+	{
+		NormalDisplacements.GetCoordArray(Buffer);
+	}
+	int32 FMorphTarget::GetNumberOfNormalDisplacements() const
+	{
+		return NormalDisplacements.Count;
+	}
+
+	bool FMorphTarget::HasTangentDisplacements() const
+	{
+		return TangentDisplacements.IsValid();
+	}
+	void FMorphTarget::GetTangentDisplacements(TArray<FVector4f>& Buffer) const
+	{
+		const int32 N = TangentDisplacements.Count;
+		Buffer.Reserve(N);
+		for (int32 Index = 0; Index < N; ++Index)
+		{
+			const FVector TangentVector = TangentDisplacements.GetVec3(Index);
+			Buffer.Emplace((FVector3f)GLTF::ConvertVec3(TangentVector), 0.f);
+		}
+	}
+	int32 FMorphTarget::GetNumberOfTangentDisplacements() const
+	{
+		return TangentDisplacements.Count;
+	}
+
+	bool FMorphTarget::HasTexCoordDisplacements(uint32 Index) const
+	{
+		switch (Index)
+		{
+			case 0:
+				return TexCoord0Displacements.IsValid();
+			case 1:
+				return TexCoord1Displacements.IsValid();
+			default:
+				return false;
+		}
+	}
+	void FMorphTarget::GetTexCoordDisplacements(uint32 Index, TArray<FVector2f>& Buffer) const
+	{
+		switch (Index)
+		{
+		case 0:
+			return TexCoord0Displacements.GetVec2Array(Buffer);
+		case 1:
+			return TexCoord1Displacements.GetVec2Array(Buffer);
+		default:
+			break;
+		}
+	}
+	int32 FMorphTarget::GetNumberOfTexCoordDisplacements(uint32 Index) const
+	{
+		switch (Index)
+		{
+		case 0:
+			return TexCoord0Displacements.Count;
+		case 1:
+			return TexCoord1Displacements.Count;
+		default:
+			return 0;
+		}
+	}
+
+	bool FMorphTarget::HasColorDeltas() const
+	{
+		return Color0Deltas.IsValid() && (Color0Deltas.Type == FAccessor::EType::Vec3 || Color0Deltas.Type == FAccessor::EType::Vec4);
+	}
+	void FMorphTarget::GetColorDeltas(TArray<FVector4f>& Buffer) const
+	{
+		if (Color0Deltas.Type == FAccessor::EType::Vec4)
+		{
+			Color0Deltas.GetVec4Array(Buffer);
+		}
+		else if (Color0Deltas.Type == FAccessor::EType::Vec3)
+		{
+			const int32 N = Color0Deltas.Count;
+			Buffer.Reserve(N);
+			for (int32 Index = 0; Index < N; ++Index)
+			{
+				const FVector Vec = Color0Deltas.GetVec3(Index);
+				Buffer.Emplace((FVector3f)Vec, 0.f);
+			}
+		}
+	}
+	int32 FMorphTarget::GetNumberOfColorDeltas() const
+	{
+		return (Color0Deltas.Type == FAccessor::EType::Vec3 || Color0Deltas.Type == FAccessor::EType::Vec4) ? Color0Deltas.Count : 0;
+	}
+
+
 	FPrimitive::FPrimitive(EMode InMode, int32 InMaterial, const FAccessor& InIndices, const FAccessor& InPosition, const FAccessor& InNormal,
 	                       const FAccessor& InTangent, const FAccessor& InTexCoord0, const FAccessor& InTexCoord1, const FAccessor& InColor0,
 	                       const FAccessor& InJoints0, const FAccessor& InWeights0)
@@ -60,8 +220,6 @@ namespace GLTF
 				Buffer.Emplace((FVector3f)Vec, 1.f);
 			}
 		}
-		else
-			check(false);
 	}
 
 	void FPrimitive::GetJointInfluences(TArray<FJointInfluence>& Buffer) const
@@ -207,7 +365,24 @@ namespace GLTF
 		}
 	}
 
+	void FPrimitive::GenerateIsValidCache()
+	{
+		bIsValidCache = IsValidPrivate();
+	}
+
 	bool FPrimitive::IsValid() const
+	{
+		if (bIsValidCache.IsSet())
+		{
+			return bIsValidCache.GetValue();
+		}
+		else
+		{
+			return IsValidPrivate();
+		}
+	}
+
+	bool FPrimitive::IsValidPrivate() const
 	{
 		// make sure all semantic attributes meet the spec
 
@@ -249,7 +424,7 @@ namespace GLTF
 			}
 		}
 
-		const FAccessor* TexCoords[] = {&TexCoord0, &TexCoord1};
+		const FAccessor* TexCoords[] = { &TexCoord0, &TexCoord1 };
 		for (const FAccessor* TexCoord : TexCoords)
 		{
 			if (TexCoord->IsValid())
@@ -279,7 +454,70 @@ namespace GLTF
 			}
 		}
 
-		// TODO: validate ranges? index buffer values?
+		//Validate Ranges:
+		{
+			TArray<uint32> AttributesSizes;
+			const FAccessor* AttributeAccessors[] = { &Position, &Normal, &Tangent, &TexCoord0, &TexCoord1, &Color0, &Joints0, &Weights0 };
+			for (const FAccessor* AttributeAccessor : AttributeAccessors)
+			{
+				if (AttributeAccessor->IsValid() && (AttributeAccessor->Count > 0))
+				{
+					AttributesSizes.Add(AttributeAccessor->Count);
+				}
+			}
+
+			TArray<uint32> IndicesValues;
+			GetTriangleIndices(IndicesValues);
+
+			for (uint32 Index : IndicesValues)
+			{
+				for (const uint32& AttributeSize : AttributesSizes)
+				{
+					if (AttributeSize <= Index)
+					{
+						return false;
+					}
+				}
+			}
+		}
+
+		//Morph Targets:
+		for (const FMorphTarget& MorphTarget : MorphTargets)
+		{
+			if (!MorphTarget.IsValid())
+			{
+				return false;
+			}
+
+			if (MorphTarget.HasPositionDisplacements() && Position.Count != MorphTarget.GetNumberOfPositionDisplacements())
+			{
+				return false;
+			}
+
+			if (MorphTarget.HasNormalDisplacements() && Normal.Count != MorphTarget.GetNumberOfNormalDisplacements())
+			{
+				return false;
+			}
+
+			if (MorphTarget.HasTangentDisplacements() && Tangent.Count != MorphTarget.GetNumberOfTangentDisplacements())
+			{
+				return false;
+			}
+
+			if (MorphTarget.HasTexCoordDisplacements(0) && TexCoord0.Count != MorphTarget.GetNumberOfTexCoordDisplacements(0))
+			{
+				return false;
+			}
+			if (MorphTarget.HasTexCoordDisplacements(1) && TexCoord0.Count != MorphTarget.GetNumberOfTexCoordDisplacements(1))
+			{
+				return false;
+			}
+
+			if (MorphTarget.HasColorDeltas() && Color0.Count != MorphTarget.GetNumberOfColorDeltas())
+			{
+				return false;
+			}
+		}
 
 		return true;
 	}
@@ -345,6 +583,16 @@ namespace GLTF
 			HashAccessor(MD5, Weights0);
 		}
 
+		//Morph Targets:
+		for (const FMorphTarget& MorphTarget : MorphTargets)
+		{
+			FMD5Hash MorphTargetHash = MorphTarget.GetHash();
+			if (MorphTargetHash.IsValid())
+			{
+				MD5.Update(MorphTargetHash.GetBytes(), MorphTargetHash.GetSize());
+			}
+		}
+
 		FMD5Hash Hash;
 		Hash.Set(MD5);
 		return Hash;
@@ -372,4 +620,15 @@ namespace GLTF
 		Hash.Set(MD5);
 		return Hash;
 	}
+
+	//This is a helper functions for the validators.
+	//We are allowed to check only the first primitive's morph target count,
+	// because if there is a incosonsitency acorss primitive targets then the mesh IsValid will report false already.
+	int32 FMesh::NumberOfMorphTargetsPerPrimitive() const
+	{
+		int32 MorphTargetCounter = Primitives.Num() > 0 ? Primitives[0].MorphTargets.Num() : 0;
+
+		return MorphTargetCounter;
+	}
+
 }  // namespace GLTF
diff --git a/Engine/Plugins/Editor/GLTFImporter/Source/GLTFCore/Private/GLTF/GLTFReader.cpp b/Engine/Plugins/Editor/GLTFImporter/Source/GLTFCore/Private/GLTF/GLTFReader.cpp
--- a/Engine/Plugins/Editor/GLTFImporter/Source/GLTFCore/Private/GLTF/GLTFReader.cpp
+++ b/Engine/Plugins/Editor/GLTFImporter/Source/GLTFCore/Private/GLTF/GLTFReader.cpp
@@ -216,6 +216,19 @@ namespace GLTF
 		}
 	}
 
+	void FFileReader::SetupMorphTarget(const FJsonObject& Object, GLTF::FPrimitive& Primitive) const
+	{
+		const TArray<FValidAccessor>& A = Asset->Accessors;
+		const FAccessor& Position = AccessorAtIndex(A, GetIndex(Object, TEXT("POSITION")));
+		const FAccessor& Normal = AccessorAtIndex(A, GetIndex(Object, TEXT("NORMAL")));
+		const FAccessor& Tangent = AccessorAtIndex(A, GetIndex(Object, TEXT("TANGENT")));
+		const FAccessor& TexCoord0 = AccessorAtIndex(A, GetIndex(Object, TEXT("TEXCOORD_0")));
+		const FAccessor& TexCoord1 = AccessorAtIndex(A, GetIndex(Object, TEXT("TEXCOORD_1")));
+		const FAccessor& Color0 = AccessorAtIndex(A, GetIndex(Object, TEXT("COLOR_0")));
+
+		Primitive.MorphTargets.Emplace(Position, Normal, Tangent, TexCoord0, TexCoord1, Color0);
+	}
+
 	void FFileReader::SetupPrimitive(const FJsonObject& Object, FMesh& Mesh) const
 	{
 		const FPrimitive::EMode       Mode = PrimitiveModeFromNumber(GetUnsignedInt(Object, TEXT("mode"), (uint32)FPrimitive::EMode::Triangles));
@@ -236,9 +249,16 @@ namespace GLTF
 		const FAccessor&   Weights0   = AccessorAtIndex(A, GetIndex(Attributes, TEXT("WEIGHTS_0")));
 
 		Mesh.Primitives.Emplace(Mode, MaterialIndex, Indices, Position, Normal, Tangent, TexCoord0, TexCoord1, Color0, Joints0, Weights0);
-		if (!Mesh.Primitives.Last().IsValid())
+
+		//Morph Targets:
+		if (Object.HasField(TEXT("targets")))
 		{
-			Messages.Emplace(EMessageSeverity::Error, TEXT("Invalid primitive!"));
+			const TArray<TSharedPtr<FJsonValue> >& MorphTargets = Object.GetArrayField(TEXT("targets"));
+			for (TSharedPtr<FJsonValue> Value : MorphTargets)
+			{
+				const FJsonObject& MorphTargetObject = *Value->AsObject();
+				SetupMorphTarget(MorphTargetObject, Mesh.Primitives.Last());
+			}
 		}
 
 		ExtensionsHandler->SetupPrimitiveExtensions(Object, Mesh.Primitives.Last());
@@ -254,12 +274,49 @@ namespace GLTF
 		Mesh.Name = GetString(Object, TEXT("name"));
 		Mesh.Primitives.Reserve(PrimArray.Num());
 
+		int32 NumberOfMorphTargets = -1;
 		for (TSharedPtr<FJsonValue> Value : PrimArray)
 		{
 			const FJsonObject& PrimObject = *Value->AsObject();
 			SetupPrimitive(PrimObject, Mesh);
+
+			if (NumberOfMorphTargets == -1)
+			{
+				NumberOfMorphTargets = Mesh.Primitives.Last().MorphTargets.Num();
+			}
+			else
+			{
+				if (NumberOfMorphTargets != Mesh.Primitives.Last().MorphTargets.Num())
+				{
+					//All primitives MUST have the same number of morph targets in the same order.
+					Messages.Emplace(EMessageSeverity::Error, TEXT("Number of Primitive.Targets is not consistent across the Mesh."));
+				}
+			}
+		}
+
+		// Morph Target Weights:
+		const TArray<TSharedPtr<FJsonValue> >& MorphTargetWeightsArray = Object.GetArrayField(TEXT("weights"));
+		for (TSharedPtr<FJsonValue> Value : MorphTargetWeightsArray)
+		{
+			Mesh.MorphTargetWeights.Add(Value->AsNumber());
 		}
 
+		// Morph Target Names:
+		if (Object.HasField(TEXT("extras")))
+		{
+			const TSharedPtr<FJsonObject>& Extras = Object.GetObjectField(TEXT("extras"));
+			if (Extras->HasField(TEXT("targetNames")))
+			{
+				const TArray<TSharedPtr<FJsonValue>>& TargetNamesArray = Extras->GetArrayField(TEXT("targetNames"));
+				for (TSharedPtr<FJsonValue> Value : TargetNamesArray)
+				{
+					Mesh.MorphTargetNames.Add(Value->AsString());
+				}
+			}
+		}
+
+		Mesh.GenerateIsValidCache();
+
 		ExtensionsHandler->SetupMeshExtensions(Object, Mesh);
 	}
 
@@ -319,6 +376,16 @@ namespace GLTF
 		Node.Skindex     = GetIndex(Object, TEXT("skin"));
 		Node.CameraIndex = GetIndex(Object, TEXT("camera"));
 
+		if (Object.HasField(TEXT("weights")))
+		{
+			const TArray<TSharedPtr<FJsonValue> >& ChildArray = Object.GetArrayField(TEXT("weights"));
+			Node.MorphTargetWeights.Reserve(ChildArray.Num());
+			for (TSharedPtr<FJsonValue> Value : ChildArray)
+			{
+				Node.MorphTargetWeights.Add(Value->AsNumber());
+			}
+		}
+
 		ExtensionsHandler->SetupNodeExtensions(Object, Node);
 	}
 
@@ -698,11 +765,6 @@ namespace GLTF
 		const FString ResourcesPath = FPaths::GetPath(InFilePath);
 		ImportAsset(ResourcesPath, bInLoadImageData, OutAsset);
 
-		if (OutAsset.ValidationCheck() != FAsset::Valid)
-		{
-			Messages.Emplace(EMessageSeverity::Warning, TEXT("GLTF Asset imported is not valid."));
-		}
-
 		// generate asset name
 		{
 			OutAsset.Name = FPaths::GetBaseFilename(InFilePath);
@@ -719,6 +781,12 @@ namespace GLTF
 				}
 			}
 		}
+		OutAsset.GenerateNames(OutAsset.Name);
+
+		if (OutAsset.ValidationCheck() != FAsset::Valid)
+		{
+			Messages.Emplace(EMessageSeverity::Warning, FString::Printf(TEXT("For GLTF Asset [%s] not all imported objects are valid."), *OutAsset.Name));
+		}
 
 		JsonRoot.Reset();
 	}
diff --git a/Engine/Plugins/Editor/GLTFImporter/Source/GLTFCore/Private/GLTFMeshFactory.cpp b/Engine/Plugins/Editor/GLTFImporter/Source/GLTFCore/Private/GLTFMeshFactory.cpp
--- a/Engine/Plugins/Editor/GLTFImporter/Source/GLTFCore/Private/GLTFMeshFactory.cpp
+++ b/Engine/Plugins/Editor/GLTFImporter/Source/GLTFCore/Private/GLTFMeshFactory.cpp
@@ -13,7 +13,7 @@ namespace GLTF
 	public:
 		FMeshFactoryImpl();
 
-		void FillMeshDescription(const GLTF::FMesh &Mesh, FMeshDescription* MeshDescription);
+		void FillMeshDescription(const GLTF::FMesh &Mesh, FMeshDescription* MeshDescription, const TArray<float>& MorphTargetWeights);
 
 		void CleanUp();
 
@@ -31,7 +31,7 @@ namespace GLTF
 			const TVertexInstanceAttributesRef<FVector4f>&  VertexInstanceColors,
 			const TEdgeAttributesRef<bool>&                EdgeHardnesses,
 			FMeshDescription* MeshDescription,
-			bool bSkipTangents);
+			const TArray<float>& MorphTargetWeights);
 
 		inline TArray<FVector4f>& GetVector4dBuffer(int32 Index)
 		{
@@ -65,17 +65,37 @@ namespace GLTF
 		}
 
 	private:
+
+		enum Vector3fMorphTargetAttributes
+		{
+			POSITIONS,
+			NORMALS
+		};
+		void UpdateVector3fFromMorphTargets(TArray<FVector3f>& AttrbiuteValues, const FPrimitive& Primitive, const TArray<float>& MorphTargetWeights, const Vector3fMorphTargetAttributes& AttributeToProcess);
+
+		enum Vector4dMorphTargetAttributes
+		{
+			TANGENTS,
+			COLORS
+		};
+		void UpdateVector4dFromMorphTargets(TArray<FVector4f>& AttrbiuteValues, const FPrimitive& Primitive, const TArray<float>& MorphTargetWeights, const Vector4dMorphTargetAttributes& AttributeToProcess);
+
+		void UpdateTexCoordsFromMorphTargets(TArray<FVector2f>& AttrbiuteValues, const FPrimitive& Primitive, const TArray<float>& MorphTargetWeights, int32 Index);
+
 		enum
 		{
 			NormalBufferIndex = 0,
 			PositionBufferIndex = 1,
 			ReindexBufferIndex = 2,
-			VectorBufferCount = 3,
+			MorphTargetBufferIndex = 3,
+			VectorBufferCount = 4,
 			UvReindexBufferIndex = MAX_MESH_TEXTURE_COORDS_MD,
+			MorphTargetUvBufferIndex = UvReindexBufferIndex+1,
 			ColorBufferIndex = 0,
 			TangentBufferIndex = 1,
 			Reindex4dBufferIndex = 2,
-			Vector4dBufferCount = 3,
+			MorphTarget4dBufferIndex = 3,
+			Vector4dBufferCount = 4,
 		};
 
 		float                ImportUniformScale;
@@ -84,7 +104,7 @@ namespace GLTF
 		TMap<int32, FPolygonGroupID> MaterialIndexToPolygonGroupID;
 		TArray<FMeshFactory::FIndexVertexIdMap>    PositionIndexToVertexIdPerPrim;
 
-		TArray<FVector2f>                       Vector2dBuffers[MAX_MESH_TEXTURE_COORDS_MD + 1];
+		TArray<FVector2f>                       Vector2dBuffers[MAX_MESH_TEXTURE_COORDS_MD + 2];
 		TArray<FVector3f>                       VectorBuffers[VectorBufferCount];
 		TArray<FVector4f>                       Vector4dBuffers[Vector4dBufferCount];
 		TArray<uint32>                          IntBuffer;
@@ -154,10 +174,133 @@ namespace GLTF
 		CornerVertexInstanceIDs.SetNum(3);
 	}
 
-	void FMeshFactoryImpl::FillMeshDescription(const FMesh &Mesh, FMeshDescription* MeshDescription)
+	void FMeshFactoryImpl::UpdateVector3fFromMorphTargets(TArray<FVector3f>& AttributeValues, const FPrimitive& Primitive, const TArray<float>& MorphTargetWeights, const Vector3fMorphTargetAttributes& AttributeToProcess)
+	{
+		if (MorphTargetWeights.Num() != Primitive.MorphTargets.Num())
+		{
+			return;
+		}
+
+		for (int32 MorphTargetIndex = 0; MorphTargetIndex < Primitive.MorphTargets.Num(); MorphTargetIndex++)
+		{
+			float MorphTargetWeight = MorphTargetWeights.Num() > MorphTargetIndex ? MorphTargetWeights[MorphTargetIndex] : 0;
+
+			if (MorphTargetWeight > 0 && MorphTargetWeight <= 1.0)
+			{
+				TArray<FVector3f>& Displacements = GetVectorBuffer(MorphTargetBufferIndex);
+
+				switch (AttributeToProcess)
+				{
+				case POSITIONS:
+					if (!Primitive.MorphTargets[MorphTargetIndex].HasPositionDisplacements()) continue;
+					Primitive.MorphTargets[MorphTargetIndex].GetPositionDisplacements(Displacements);
+					break;
+				case NORMALS:
+					if (!Primitive.MorphTargets[MorphTargetIndex].HasNormalDisplacements()) continue;
+					Primitive.MorphTargets[MorphTargetIndex].GetNormalDisplacements(Displacements);
+					break;
+				default:
+					return;
+					break;
+				}
+
+				if (AttributeValues.Num() != Displacements.Num())
+				{
+					continue;
+				}
+				
+				for (int32 Index = 0; Index < Displacements.Num(); Index++)
+				{
+					AttributeValues[Index] += MorphTargetWeight * Displacements[Index];
+				}
+			}
+		}
+	}
+	void FMeshFactoryImpl::UpdateVector4dFromMorphTargets(TArray<FVector4f>& AttributeValues, const FPrimitive& Primitive, const TArray<float>& MorphTargetWeights, const Vector4dMorphTargetAttributes& AttributeToProcess)
 	{
-		const bool bSkipTangents = !Mesh.HasNormals(); // Per the GLTF spec, tangents should be ignored if no normals are provided
+		if (MorphTargetWeights.Num() != Primitive.MorphTargets.Num())
+		{
+			return;
+		}
+
+		for (int32 MorphTargetIndex = 0; MorphTargetIndex < Primitive.MorphTargets.Num(); MorphTargetIndex++)
+		{
+			float MorphTargetWeight = MorphTargetWeights.Num() > MorphTargetIndex ? MorphTargetWeights[MorphTargetIndex] : 0;
+
+			if (MorphTargetWeight > 0 && MorphTargetWeight <= 1.0)
+			{
+				TArray<FVector4f>& Displacements = GetVector4dBuffer(MorphTargetBufferIndex);
 
+				switch (AttributeToProcess)
+				{
+				case TANGENTS:
+					if (!Primitive.MorphTargets[MorphTargetIndex].HasTangentDisplacements()) continue;
+					Primitive.MorphTargets[MorphTargetIndex].GetTangentDisplacements(Displacements);
+					break;
+				case COLORS:
+					if (!Primitive.MorphTargets[MorphTargetIndex].HasColorDeltas()) continue;
+					Primitive.MorphTargets[MorphTargetIndex].GetColorDeltas(Displacements);
+					break;
+				default:
+					return;
+					break;
+				}
+
+				if (AttributeValues.Num() != Displacements.Num())
+				{
+					continue;
+				}
+
+				for (int32 Index = 0; Index < Displacements.Num(); Index++)
+				{
+					AttributeValues[Index] += MorphTargetWeight * Displacements[Index];
+
+					if (AttributeToProcess == Vector4dMorphTargetAttributes::COLORS)
+					{
+						//After applying color deltas, all components of each COLOR_0 morphed accessor element MUST be clamped to [0.0, 1.0] range.
+						AttributeValues[Index].X = FMath::Clamp(AttributeValues[Index].X, .0f, 1.f);
+						AttributeValues[Index].Y = FMath::Clamp(AttributeValues[Index].Y, .0f, 1.f);
+						AttributeValues[Index].Z = FMath::Clamp(AttributeValues[Index].W, .0f, 1.f);
+						AttributeValues[Index].W = FMath::Clamp(AttributeValues[Index].Z, .0f, 1.f);
+					}
+				}
+			}
+		}
+	}
+	void FMeshFactoryImpl::UpdateTexCoordsFromMorphTargets(TArray<FVector2f>& AttributeValues, const FPrimitive& Primitive, const TArray<float>& MorphTargetWeights, int32 UvIndex)
+	{
+		if (MorphTargetWeights.Num() != Primitive.MorphTargets.Num())
+		{
+			return;
+		}
+
+		for (int32 MorphTargetIndex = 0; MorphTargetIndex < Primitive.MorphTargets.Num(); MorphTargetIndex++)
+		{
+			float MorphTargetWeight = MorphTargetWeights.Num() > MorphTargetIndex ? MorphTargetWeights[MorphTargetIndex] : 0;
+
+			if (MorphTargetWeight > 0 && MorphTargetWeight <= 1.0)
+			{
+				TArray<FVector2f>& Displacements = GetVector2dBuffer(MorphTargetUvBufferIndex);
+
+				if (!Primitive.MorphTargets[MorphTargetIndex].HasTexCoordDisplacements(UvIndex)) continue;
+
+				Primitive.MorphTargets[MorphTargetIndex].GetTexCoordDisplacements(UvIndex, Displacements);
+
+				if (AttributeValues.Num() != Displacements.Num())
+				{
+					continue;
+				}
+
+				for (int32 Index = 0; Index < Displacements.Num(); Index++)
+				{
+					AttributeValues[Index] += MorphTargetWeight * Displacements[Index];
+				}
+			}
+		}
+	}
+
+	void FMeshFactoryImpl::FillMeshDescription(const FMesh &Mesh, FMeshDescription* MeshDescription, const TArray<float>& MorphTargetWeights)
+	{
 		const int32 NumUVs = FMath::Max(1, GetNumUVs(Mesh));
 
 		FStaticMeshAttributes StaticMeshAttributes(*MeshDescription);
@@ -192,6 +335,8 @@ namespace GLTF
 			TArray<FVector3f>& Positions = GetVectorBuffer(PositionBufferIndex);
 			Primitive.GetPositions(Positions);
 
+			UpdateVector3fFromMorphTargets(Positions, Primitive, MorphTargetWeights, Vector3fMorphTargetAttributes::POSITIONS);
+
 			FMeshFactory::FIndexVertexIdMap& PositionIndexToVertexId = PositionIndexToVertexIdPerPrim[Index];
 			PositionIndexToVertexId.Empty(Positions.Num());
 			for (int32 PositionIndex = 0; PositionIndex < Positions.Num(); ++PositionIndex)
@@ -229,7 +374,8 @@ namespace GLTF
 				ImportPrimitive(Primitive, Index, NumUVs, Mesh.HasTangents(), Mesh.HasColors(),  //
 					VertexInstanceNormals, VertexInstanceTangents, VertexInstanceBinormalSigns, VertexInstanceUVs,
 					VertexInstanceColors,  //
-					EdgeHardnesses, MeshDescription, bSkipTangents);
+					EdgeHardnesses, MeshDescription,
+					MorphTargetWeights);
 
 			bMeshUsesEmptyMaterial |= Primitive.MaterialIndex == INDEX_NONE;
 			for (int32 UVIndex = 0; UVIndex < NumUVs; ++UVIndex)
@@ -265,30 +411,17 @@ namespace GLTF
 		const TVertexInstanceAttributesRef<FVector4f>&  VertexInstanceColors,
 		const TEdgeAttributesRef<bool>&                EdgeHardnesses,
 		FMeshDescription* MeshDescription,
-		bool bSkipTangents)
+		const TArray<float>& MorphTargetWeights)
 	{
-
 		const FPolygonGroupID CurrentPolygonGroupID = MaterialIndexToPolygonGroupID[Primitive.MaterialIndex];
 		const uint32          TriCount = Primitive.TriangleCount();
 
 		TArray<uint32>& Indices = GetIntBuffer();
 		Primitive.GetTriangleIndices(Indices);
 
-		// Validate Indices against Positions:
-		//  In case of corrupted Indices return with bHasDegenerateTriangles.
-		{
-			TArray<FVector3f>& Positions = GetVectorBuffer(PositionBufferIndex);
-			Primitive.GetPositions(Positions);
-			uint32 PositionsSize = Positions.Num();
-
-			for (uint32 Index : Indices)
-			{
-				if (PositionsSize <= Index)
-				{
-					return true;
-				}
-			}
-		}
+		//Note from gltf 2.0 specification:
+		//	When the base mesh primitive does not specify normals, client implementations MUST calculate flat normals for each morph target; the provided tangents and their displacements (if present) MUST be ignored.
+		bool bIgnoreTangents = false;
 
 		TArray<FVector3f>& Normals = GetVectorBuffer(NormalBufferIndex);
 		// glTF does not guarantee each primitive within a mesh has the same attributes.
@@ -300,6 +433,9 @@ namespace GLTF
 		{
 			TArray<FVector3f>& ReindexBuffer = GetVectorBuffer(ReindexBufferIndex);
 			Primitive.GetNormals(Normals);
+
+			UpdateVector3fFromMorphTargets(Normals, Primitive, MorphTargetWeights, Vector3fMorphTargetAttributes::NORMALS);
+
 			ReIndex(Normals, Indices, ReindexBuffer);
 			Swap(Normals, ReindexBuffer);
 		}
@@ -307,14 +443,22 @@ namespace GLTF
 		{
 			TArray<FVector3f>& Positions = GetVectorBuffer(PositionBufferIndex);
 			Primitive.GetPositions(Positions);
+
+			//update positions with morph targets before generating flat normals:
+			UpdateVector3fFromMorphTargets(Positions, Primitive, MorphTargetWeights, Vector3fMorphTargetAttributes::POSITIONS);
+
 			GenerateFlatNormals(Positions, Indices, Normals);
+			bIgnoreTangents = true;
 		}
 
 		TArray<FVector4f>& Tangents = GetVector4dBuffer(TangentBufferIndex);
-		if (Primitive.HasTangents())
+		if (!bIgnoreTangents && Primitive.HasTangents())
 		{
 			TArray<FVector4f>& ReindexBuffer = GetVector4dBuffer(Reindex4dBufferIndex);
 			Primitive.GetTangents(Tangents);
+
+			UpdateVector4dFromMorphTargets(Tangents, Primitive, MorphTargetWeights, Vector4dMorphTargetAttributes::TANGENTS);
+
 			ReIndex(Tangents, Indices, ReindexBuffer);
 			Swap(Tangents, ReindexBuffer);
 		}
@@ -329,6 +473,9 @@ namespace GLTF
 		{
 			TArray<FVector4f>& ReindexBuffer = GetVector4dBuffer(Reindex4dBufferIndex);
 			Primitive.GetColors(Colors);
+
+			UpdateVector4dFromMorphTargets(Colors, Primitive, MorphTargetWeights, Vector4dMorphTargetAttributes::COLORS);
+
 			ReIndex(Colors, Indices, ReindexBuffer);
 			Swap(Colors, ReindexBuffer);
 		}
@@ -347,6 +494,9 @@ namespace GLTF
 			{
 				TArray<FVector2f>& ReindexBuffer = GetVector2dBuffer(UvReindexBufferIndex);
 				Primitive.GetTexCoords(UVIndex, *UVs[UVIndex]);
+
+				UpdateTexCoordsFromMorphTargets(*UVs[UVIndex], Primitive, MorphTargetWeights, UVIndex);
+
 				ReIndex(*UVs[UVIndex], Indices, ReindexBuffer);
 				Swap(*UVs[UVIndex], ReindexBuffer);
 			}
@@ -391,7 +541,7 @@ namespace GLTF
 
 				VertexInstanceNormals[VertexInstanceID] = Normals[IndiceIndex];
 
-				if (!bSkipTangents && Tangents.Num() > 0)
+				if (!bIgnoreTangents && Tangents.Num() > 0)
 				{
 					VertexInstanceTangents[VertexInstanceID] = Tangents[IndiceIndex];
 					VertexInstanceBinormalSigns[VertexInstanceID] = Tangents[IndiceIndex].W;
@@ -454,9 +604,9 @@ namespace GLTF
 
 	FMeshFactory::~FMeshFactory() {}
 
-	void FMeshFactory::FillMeshDescription(const GLTF::FMesh &Mesh, FMeshDescription* MeshDescription)
+	void FMeshFactory::FillMeshDescription(const GLTF::FMesh &Mesh, FMeshDescription* MeshDescription, const TArray<float>& MorphTargetWeights)
 	{
-		Impl->FillMeshDescription(Mesh, MeshDescription);
+		Impl->FillMeshDescription(Mesh, MeshDescription, MorphTargetWeights);
 	}
 
 	const TArray<FLogMessage>& FMeshFactory::GetLogMessages() const
diff --git a/Engine/Plugins/Editor/GLTFImporter/Source/GLTFCore/Public/GLTFAsset.h b/Engine/Plugins/Editor/GLTFImporter/Source/GLTFCore/Public/GLTFAsset.h
--- a/Engine/Plugins/Editor/GLTFImporter/Source/GLTFCore/Public/GLTFAsset.h
+++ b/Engine/Plugins/Editor/GLTFImporter/Source/GLTFCore/Public/GLTFAsset.h
@@ -71,7 +71,7 @@ namespace GLTF
 		{
 			Valid,
 			InvalidMeshPresent   = 0x1,
-			InvalidNodeTransform = 0x2,
+			InvalidNodePresent = 0x2,
 		};
 
 		FString Name;
@@ -110,6 +110,7 @@ namespace GLTF
 
 		/**
 		 * Will generate names for any entities(nodes, meshes, etc.) that have the name field missing.
+		 * Is called from GLTF::FFileReader::ReadFile
 		 *
 		 * @param Prefix - prefix to add to the entities name.
 		 */
diff --git a/Engine/Plugins/Editor/GLTFImporter/Source/GLTFCore/Public/GLTFMesh.h b/Engine/Plugins/Editor/GLTFImporter/Source/GLTFCore/Public/GLTFMesh.h
--- a/Engine/Plugins/Editor/GLTFImporter/Source/GLTFCore/Public/GLTFMesh.h
+++ b/Engine/Plugins/Editor/GLTFImporter/Source/GLTFCore/Public/GLTFMesh.h
@@ -40,6 +40,43 @@ namespace GLTF
 		TArray<int32> VariantIndices;
 	};
 
+	struct GLTFCORE_API FMorphTarget
+	{
+		FMorphTarget(const FAccessor& InPositionDisplacements, const FAccessor& InNormalDisplacements,
+			const FAccessor& InTangentDisplacements, const FAccessor& InTexCoord0Displacements, const FAccessor& InTexCoord1Displacements, const FAccessor& InColor0Deltas);
+
+		bool IsValid() const;
+		FMD5Hash GetHash() const;
+
+		bool HasPositionDisplacements() const;
+		void GetPositionDisplacements(TArray<FVector3f>& Buffer) const;
+		int32 GetNumberOfPositionDisplacements() const;
+
+		bool HasNormalDisplacements() const;
+		void GetNormalDisplacements(TArray<FVector3f>& Buffer) const;
+		int32 GetNumberOfNormalDisplacements() const;
+
+		bool HasTangentDisplacements() const;
+		void GetTangentDisplacements(TArray<FVector4f>& Buffer) const;
+		int32 GetNumberOfTangentDisplacements() const;
+
+		bool HasTexCoordDisplacements(uint32 Index) const;
+		void GetTexCoordDisplacements(uint32 Index, TArray<FVector2f>& Buffer) const;
+		int32 GetNumberOfTexCoordDisplacements(uint32 Index) const;
+
+		bool HasColorDeltas() const;
+		void GetColorDeltas(TArray<FVector4f>& Buffer) const;
+		int32 GetNumberOfColorDeltas() const;
+
+	private:
+		const FAccessor& PositionDisplacements;
+		const FAccessor& NormalDisplacements;
+		const FAccessor& TangentDisplacements;
+		const FAccessor& TexCoord0Displacements;
+		const FAccessor& TexCoord1Displacements;
+		const FAccessor& Color0Deltas;
+	};
+
 	struct GLTFCORE_API FPrimitive
 	{
 		enum class EMode
@@ -64,6 +101,7 @@ namespace GLTF
 		           const FAccessor& InTangent, const FAccessor& InTexCoord0, const FAccessor& InTexCoord1, const FAccessor& InColor0,
 		           const FAccessor& InJoints0, const FAccessor& InWeights0);
 
+		void GenerateIsValidCache();
 		bool IsValid() const;
 		FMD5Hash GetHash() const;
 
@@ -85,7 +123,12 @@ namespace GLTF
 		uint32 VertexCount() const;
 		uint32 TriangleCount() const;
 
+		//
+		TArray<FMorphTarget> MorphTargets;
+
 	private:
+		bool IsValidPrivate() const;
+
 		// index buffer
 		const FAccessor& Indices;
 
@@ -99,6 +142,9 @@ namespace GLTF
 		// skeletal mesh attributes
 		const FAccessor& Joints0;
 		const FAccessor& Weights0;
+
+		//Validity cache:
+		TOptional<bool> bIsValidCache;
 	};
 
 
@@ -107,6 +153,9 @@ namespace GLTF
 		FString				Name;
 		TArray<FPrimitive>	Primitives;
 
+		TArray<float>		MorphTargetWeights;
+		TArray<FString>		MorphTargetNames;
+
 		FString				UniqueId; //will be generated in FAsset::GenerateNames
 	
 		bool HasNormals() const;
@@ -115,8 +164,16 @@ namespace GLTF
 		bool HasColors() const;
 		bool HasJointWeights() const;
 
+		void GenerateIsValidCache(bool GenerateIsValidCacheForPrimitives = true);
 		bool IsValid() const;
 		FMD5Hash GetHash() const;
+
+		int32 NumberOfMorphTargetsPerPrimitive() const;
+
+	private:
+		bool IsValidPrivate() const;
+		//Validity cache:
+		TOptional<bool> bIsValidCache;
 	};
 
 	//
@@ -146,7 +203,7 @@ namespace GLTF
 
 	inline bool FPrimitive::HasColors() const
 	{
-		return Color0.IsValid();
+		return Color0.IsValid() && (Color0.Type == FAccessor::EType::Vec3 || Color0.Type == FAccessor::EType::Vec4);
 	}
 
 	inline void FPrimitive::GetPositions(TArray<FVector3f>& Buffer) const
@@ -218,9 +275,60 @@ namespace GLTF
 		return Result;
 	}
 
+	inline void FMesh::GenerateIsValidCache(bool GenerateIsValidCacheForPrimitives)
+	{
+		if (GenerateIsValidCacheForPrimitives)
+		{
+			for (FPrimitive& Primitive : Primitives)
+			{
+				Primitive.GenerateIsValidCache();
+			}
+		}
+
+		bIsValidCache = IsValidPrivate();
+	}
 	inline bool FMesh::IsValid() const
 	{
-		return Primitives.FindByPredicate([](const FPrimitive& Prim) { return !Prim.IsValid(); }) == nullptr;
+		if (bIsValidCache.IsSet())
+		{
+			return bIsValidCache.GetValue();
+		}
+		else
+		{
+			return IsValidPrivate();
+		}
+	}
+
+	inline bool FMesh::IsValidPrivate() const
+	{
+		//Validate Primitives:
+		bool bIsValid = Primitives.FindByPredicate([](const FPrimitive& Prim) { return !Prim.IsValid(); }) == nullptr;
+
+		//if MorphTargetNames are not set, but mesh has morph targets (which is likely), this will generate isValid false;
+		//IsValidCache will have to be (re-)genereated post FAsset::GenerateNames (which is called at the end of GLTF::FReader) to overcome this issue.
+
+		//Validate Morph Target Names and Morph Target Weights:
+		if ((MorphTargetNames.Num() > 0) && (MorphTargetWeights.Num() > 0))
+		{
+			if (MorphTargetNames.Num() != MorphTargetWeights.Num())
+			{
+				bIsValid = false;
+			}
+		}
+
+		//Validate Morph Target (and Morph Target Names) Counts:
+		int32 MorphTargetCounter = MorphTargetNames.Num();
+
+		for (const FPrimitive& Primitive : Primitives)
+		{
+			if (MorphTargetCounter != Primitive.MorphTargets.Num())
+			{
+				bIsValid = false;
+				break;
+			}
+		}
+
+		return bIsValid;
 	}
 
 }  // namespace GLTF
diff --git a/Engine/Plugins/Editor/GLTFImporter/Source/GLTFCore/Public/GLTFMeshFactory.h b/Engine/Plugins/Editor/GLTFImporter/Source/GLTFCore/Public/GLTFMeshFactory.h
--- a/Engine/Plugins/Editor/GLTFImporter/Source/GLTFCore/Public/GLTFMeshFactory.h
+++ b/Engine/Plugins/Editor/GLTFImporter/Source/GLTFCore/Public/GLTFMeshFactory.h
@@ -22,7 +22,7 @@ namespace GLTF
 
 		using FIndexVertexIdMap = TMap<int32, FVertexID>;
 
-		void FillMeshDescription(const GLTF::FMesh &Mesh, FMeshDescription* MeshDescription);
+		void FillMeshDescription(const GLTF::FMesh &Mesh, FMeshDescription* MeshDescription, const TArray<float>& MorphTargetWeights = TArray<float>());
 
 		float GetUniformScale() const;
 		void  SetUniformScale(float Scale);
diff --git a/Engine/Plugins/Editor/GLTFImporter/Source/GLTFCore/Public/GLTFNode.h b/Engine/Plugins/Editor/GLTFImporter/Source/GLTFCore/Public/GLTFNode.h
--- a/Engine/Plugins/Editor/GLTFImporter/Source/GLTFCore/Public/GLTFNode.h
+++ b/Engine/Plugins/Editor/GLTFImporter/Source/GLTFCore/Public/GLTFNode.h
@@ -38,6 +38,8 @@ namespace GLTF
 		int32 ParentIndex;
 		int32 RootJointIndex; //only valid if node is of Joint Type
 
+		TArray<float> MorphTargetWeights; //for the instantiated mesh with morph targets.
+
 		FString UniqueId; //will be generated in FAsset::GenerateNames
 
 		FNode()
diff --git a/Engine/Plugins/Editor/GLTFImporter/Source/GLTFCore/Public/GLTFReader.h b/Engine/Plugins/Editor/GLTFImporter/Source/GLTFCore/Public/GLTFReader.h
--- a/Engine/Plugins/Editor/GLTFImporter/Source/GLTFCore/Public/GLTFReader.h
+++ b/Engine/Plugins/Editor/GLTFImporter/Source/GLTFCore/Public/GLTFReader.h
@@ -12,6 +12,7 @@ namespace GLTF
 {
 	struct FAsset;
 	struct FMesh;
+	struct FPrimitive;
 	struct FTextureMap;
 	struct FMaterial;
 	class FBinaryFileReader;
@@ -42,6 +43,7 @@ namespace GLTF
 		void SetupBuffer(const FJsonObject& Object, const FString& Path);
 		void SetupBufferView(const FJsonObject& Object) const;
 		void SetupAccessor(const FJsonObject& Object) const;
+		void SetupMorphTarget(const FJsonObject& Object, GLTF::FPrimitive& Primitive) const;
 		void SetupPrimitive(const FJsonObject& Object, GLTF::FMesh& Mesh) const;
 		void SetupMesh(const FJsonObject& Object) const;
 
diff --git a/Engine/Plugins/Interchange/Runtime/Source/Import/Private/Gltf/InterchangeGltfAnimation.cpp b/Engine/Plugins/Interchange/Runtime/Source/Import/Private/Gltf/InterchangeGltfAnimation.cpp
--- a/Engine/Plugins/Interchange/Runtime/Source/Import/Private/Gltf/InterchangeGltfAnimation.cpp
+++ b/Engine/Plugins/Interchange/Runtime/Source/Import/Private/Gltf/InterchangeGltfAnimation.cpp
@@ -7,6 +7,7 @@
 #include "GLTFAnimation.h"
 #include "GLTFAsset.h"
 #include "GLTFNode.h"
+#include "GLTFMesh.h"
 #include "InterchangeImportLog.h"
 #include "Animation/AnimTypes.h"
 
@@ -200,7 +201,7 @@ namespace UE::Interchange::Gltf::Private
 		return ChannelIndices.Num() != 0;
 	}
 
-	bool GetAnimationTransformPayloadData(const FString& PayLoadKey, const GLTF::FAsset& GltfAsset, FAnimationCurvePayloadData& OutPayloadData)
+	bool GetTransformAnimationPayloadData(const FString& PayLoadKey, const GLTF::FAsset& GltfAsset, FAnimationCurvePayloadData& OutPayloadData)
 	{
 		//translation, rotation, scale time 3 channels:
 		OutPayloadData.Curves.SetNum(9);
@@ -213,7 +214,7 @@ namespace UE::Interchange::Gltf::Private
 			return false;
 		}
 		const GLTF::FAnimation& Animation = GltfAsset.Animations[AnimationIndex];
-		
+
 		for (const int32 ChannelIndex : ChannelIndices)
 		{
 			const GLTF::FAnimation::FChannel& Channel = Animation.Channels[ChannelIndex];
@@ -339,7 +340,7 @@ namespace UE::Interchange::Gltf::Private
 						}
 
 						Helper.Heading = Rotator;
-						
+
 						return Rotator.Euler();
 					}, 3, 4, true);
 				break;
@@ -356,7 +357,7 @@ namespace UE::Interchange::Gltf::Private
 			}
 			case GLTF::FAnimation::EPath::Weights:
 			{
-				UE_LOG(LogInterchangeImport, Warning, TEXT("Morph Animation type not supported yet."));
+				UE_LOG(LogInterchangeImport, Warning, TEXT("Animation[%d]:Channel[%d] : Morph Animation (Weight path) type not supported for Transform like Animations. (Morph Animations should be handled via \"GetMorphTargetAnimationPayloadData\" function."), AnimationIndex, ChannelIndex);
 				break;
 			}
 			default:
@@ -367,6 +368,102 @@ namespace UE::Interchange::Gltf::Private
 		return true;
 	}
 
+	bool GetMorphTargetAnimationPayloadData(const FString& PayLoadKey, const GLTF::FAsset& GltfAsset, FAnimationCurvePayloadData& OutPayloadData)
+	{
+		OutPayloadData.Curves.SetNum(1);
+
+		TArray<FString> PayLoadKeys;
+		PayLoadKey.ParseIntoArray(PayLoadKeys, TEXT(":"));
+
+		if (PayLoadKeys.Num() != 4)
+		{
+			return false;
+		}
+
+		int32 AnimationIndex;
+		int32 ChannelIndex;
+		int32 MeshIndex;
+		int32 MorphTargetIndex;
+
+		LexFromString(AnimationIndex, *PayLoadKeys[0]);
+		LexFromString(ChannelIndex, *PayLoadKeys[1]);
+		LexFromString(MeshIndex, *PayLoadKeys[2]);
+		LexFromString(MorphTargetIndex, *PayLoadKeys[3]);
+
+		if (GltfAsset.Animations.Num() <= AnimationIndex
+			|| GltfAsset.Animations[AnimationIndex].Channels.Num() <= ChannelIndex
+			|| GltfAsset.Meshes.Num() <= MeshIndex
+			|| GltfAsset.Meshes[MeshIndex].MorphTargetNames.Num() <= MorphTargetIndex)
+		{
+			return false;
+		}
+
+		const GLTF::FAnimation& Animation = GltfAsset.Animations[AnimationIndex];
+		const GLTF::FAnimation::FChannel& Channel = Animation.Channels[ChannelIndex];
+		const GLTF::FAnimation::FSampler& Sampler = Animation.Samplers[Channel.Sampler];
+		const GLTF::FMesh& Mesh = GltfAsset.Meshes[MeshIndex];
+
+		TArray<float> FrameTimeBuffer;
+		TArray<float> FrameDataBuffer;
+
+		Sampler.Input.GetFloatArray(FrameTimeBuffer);
+
+		ERichCurveInterpMode InterpolationMode = static_cast<ERichCurveInterpMode>(Sampler.Interpolation);
+
+		if (uint32(MorphTargetIndex * FrameDataBuffer.Num()) > Sampler.Output.Count)
+		{
+			return false;
+		}
+
+		//framebuffer is :
+		//morphTarget(n)_frame0 morphTarget(n+1)_frame0 morphTarget(n+2)_frame0 .... morphTarget(n)_frame1 morphTarget(n+1)_frame1 morphTarget(n+2)_frame1 .....
+		// CUBIC:
+		//morphTarget(n)_frame0_in morphTarget(n)_frame0 morphTarget(n)_frame0_out morphTarget(n+1)_frame0_in morphTarget(n+1)_frame0 morphTarget(n+1)_frame0_out morphTarget(n+2)_frame0_in morphTarget(n+2)_frame0 morphTarget(n+2)_frame0_out
+		//in/out values are currently ignored (in Interchange pipelines)
+		FrameDataBuffer.SetNumUninitialized(Sampler.Output.Count * (InterpolationMode == ERichCurveInterpMode::RCIM_Cubic ? 3 : 1));
+		Sampler.Output.GetFloatArray(reinterpret_cast<float*>(FrameDataBuffer.GetData()));
+
+		FRichCurve& Curve = OutPayloadData.Curves[0];
+
+		int32 RowMultiplier = Mesh.MorphTargetNames.Num() * (InterpolationMode == ERichCurveInterpMode::RCIM_Cubic ? 3 : 1);
+
+		for (int32 SampleIndex = 0; SampleIndex < FrameTimeBuffer.Num(); SampleIndex++)
+		{
+			if (InterpolationMode == ERichCurveInterpMode::RCIM_Cubic)
+			{
+				FKeyHandle Key = Curve.AddKey(FrameTimeBuffer[SampleIndex], FrameDataBuffer[MorphTargetIndex + (SampleIndex * RowMultiplier) + 1]);
+				FRichCurveKey& CurveKey = Curve.GetKey(Key);
+				CurveKey.ArriveTangent = FrameDataBuffer[MorphTargetIndex + (SampleIndex * RowMultiplier) + 0];
+				CurveKey.LeaveTangent = FrameDataBuffer[MorphTargetIndex + (SampleIndex * RowMultiplier) + 2];
+
+				Curve.SetKeyInterpMode(Key, InterpolationMode);
+			}
+			else
+			{
+				FKeyHandle Key = Curve.AddKey(FrameTimeBuffer[SampleIndex], FrameDataBuffer[MorphTargetIndex + (SampleIndex * RowMultiplier)]);
+
+				Curve.SetKeyInterpMode(Key, InterpolationMode);
+			}
+		}
+
+		return true;
+	}
+
+	bool GetAnimationPayloadData(const FString& PayLoadKey, const GLTF::FAsset& GltfAsset, FAnimationCurvePayloadData& OutPayloadData)
+	{
+		TArray<FString> MorphTargetCheckPayLoadKeys;
+		PayLoadKey.ParseIntoArray(MorphTargetCheckPayLoadKeys, TEXT("~"));
+
+		if (MorphTargetCheckPayLoadKeys.Num() == 2)
+		{
+			return GetMorphTargetAnimationPayloadData(MorphTargetCheckPayLoadKeys[1], GltfAsset, OutPayloadData);
+		}
+		else
+		{
+			return GetTransformAnimationPayloadData(PayLoadKey, GltfAsset, OutPayloadData);
+		}
+	}
+
 	bool GetBakedAnimationTransformPayloadData(const FString& PayLoadKey, const GLTF::FAsset& GltfAsset, FAnimationBakeTransformPayloadData& PayloadData)
 	{
 		TArray<int32> ChannelIndices;
diff --git a/Engine/Plugins/Interchange/Runtime/Source/Import/Private/Gltf/InterchangeGltfMesh.cpp b/Engine/Plugins/Interchange/Runtime/Source/Import/Private/Gltf/InterchangeGltfMesh.cpp
--- a/Engine/Plugins/Interchange/Runtime/Source/Import/Private/Gltf/InterchangeGltfMesh.cpp
+++ b/Engine/Plugins/Interchange/Runtime/Source/Import/Private/Gltf/InterchangeGltfMesh.cpp
@@ -73,8 +73,16 @@ namespace UE::Interchange::Gltf::Private
 
 	bool GetSkeletalMeshDescriptionForPayLoadKey(const GLTF::FAsset& GltfAsset, const FString& PayLoadKey, FMeshDescription& MeshDescription, TArray<FString>* OutJointUniqueNames)
 	{
+		TArray<FString> MorphTargetCheckPayLoadKeys;
+		PayLoadKey.ParseIntoArray(MorphTargetCheckPayLoadKeys, TEXT("~"));
+
+		if (MorphTargetCheckPayLoadKeys.Num() == 2)
+		{
+			return GetStaticMeshPayloadDataForPayLoadKey(GltfAsset, MorphTargetCheckPayLoadKeys[1], MeshDescription);
+		}
+		
 		TArray<FString> PayLoadKeys;
-		PayLoadKey.ParseIntoArray(PayLoadKeys, TEXT(":"));
+		MorphTargetCheckPayLoadKeys[0].ParseIntoArray(PayLoadKeys, TEXT(":"));
 		TMap<int32, TArray<int32>> MeshIndexToSkinIndicesMap;
 
 		int32 MeshIndex = 0;
@@ -267,10 +275,18 @@ namespace UE::Interchange::Gltf::Private
 		return true;
 	}
 
-	bool GetStaticMeshPayloadDataForPayLoadKey(const GLTF::FAsset& GltfAsset, const FString& PayLoadKey, FStaticMeshPayloadData& StaticMeshPayloadData)
+	bool GetStaticMeshPayloadDataForPayLoadKey(const GLTF::FAsset& GltfAsset, const FString& PayLoadKey, FMeshDescription& MeshDescription)
 	{
+		TArray<FString> PayLoadKeys;
+		PayLoadKey.ParseIntoArray(PayLoadKeys, TEXT(":"));
+
+		if (PayLoadKeys.Num() == 0)
+		{
+			return false;
+		}
+
 		int32 MeshIndex = 0;
-		LexFromString(MeshIndex, *PayLoadKey);
+		LexFromString(MeshIndex, *PayLoadKeys[0]);
 
 		if (!GltfAsset.Meshes.IsValidIndex(MeshIndex))
 		{
@@ -278,11 +294,28 @@ namespace UE::Interchange::Gltf::Private
 		}
 
 		const GLTF::FMesh& GltfMesh = GltfAsset.Meshes[MeshIndex];
+
+		TArray<float> MorphTargetWeights;
+		if (PayLoadKeys.Num() == 2)
+		{
+			int32 MorphTargetIndex;
+			LexFromString(MorphTargetIndex, *PayLoadKeys[1]);
+
+			MorphTargetWeights.SetNumZeroed(GltfMesh.MorphTargetNames.Num());
+
+			if (MorphTargetIndex >= MorphTargetWeights.Num())
+			{
+				return false;
+			}
+
+			MorphTargetWeights[MorphTargetIndex] = 1.0f;
+		}
+
 		GLTF::FMeshFactory MeshFactory;
-		MeshFactory.SetUniformScale(100.f); // GLTF is in meters while UE is in centimeters
-		MeshFactory.FillMeshDescription(GltfMesh, &StaticMeshPayloadData.MeshDescription);
+		MeshFactory.SetUniformScale(GltfUnitConversionMultiplier); // GLTF is in meters while UE is in centimeters
+		MeshFactory.FillMeshDescription(GltfMesh, &MeshDescription, MorphTargetWeights);
 
-		PatchPolygonGroups(StaticMeshPayloadData.MeshDescription, GltfAsset);
+		PatchPolygonGroups(MeshDescription, GltfAsset);
 
 		return true;
 	}
diff --git a/Engine/Plugins/Interchange/Runtime/Source/Import/Private/Gltf/InterchangeGltfPrivate.h b/Engine/Plugins/Interchange/Runtime/Source/Import/Private/Gltf/InterchangeGltfPrivate.h
--- a/Engine/Plugins/Interchange/Runtime/Source/Import/Private/Gltf/InterchangeGltfPrivate.h
+++ b/Engine/Plugins/Interchange/Runtime/Source/Import/Private/Gltf/InterchangeGltfPrivate.h
@@ -21,7 +21,7 @@ namespace UE::Interchange
 		static float GltfUnitConversionMultiplier = 100.f;
 
 		// Animation related functions
-		bool GetAnimationTransformPayloadData(const FString& PayLoadKey, const GLTF::FAsset& GltfAsset, FAnimationCurvePayloadData& OutPayloadData);
+		bool GetAnimationPayloadData(const FString& PayLoadKey, const GLTF::FAsset& GltfAsset, FAnimationCurvePayloadData& OutPayloadData);
 		bool GetBakedAnimationTransformPayloadData(const FString& PayLoadKey, const GLTF::FAsset& GltfAsset, FAnimationBakeTransformPayloadData& PayloadData);
 		// 
 
@@ -31,7 +31,7 @@ namespace UE::Interchange
 		bool GetSkeletalMeshDescriptionForPayLoadKey(const GLTF::FAsset& GltfAsset, const FString& PayLoadKey,
 			FMeshDescription& MeshDescription, TArray<FString>* OutJointUniqueNames);
 
-		bool GetStaticMeshPayloadDataForPayLoadKey(const GLTF::FAsset& GltfAsset, const FString& PayLoadKey, FStaticMeshPayloadData& StaticMeshPayloadData);
+		bool GetStaticMeshPayloadDataForPayLoadKey(const GLTF::FAsset& GltfAsset, const FString& PayLoadKey, FMeshDescription& MeshDescription);
 		//
 	}
 }
\ No newline at end of file
diff --git a/Engine/Plugins/Interchange/Runtime/Source/Import/Private/Gltf/InterchangeGltfTranslator.cpp b/Engine/Plugins/Interchange/Runtime/Source/Import/Private/Gltf/InterchangeGltfTranslator.cpp
--- a/Engine/Plugins/Interchange/Runtime/Source/Import/Private/Gltf/InterchangeGltfTranslator.cpp
+++ b/Engine/Plugins/Interchange/Runtime/Source/Import/Private/Gltf/InterchangeGltfTranslator.cpp
@@ -224,6 +224,25 @@ void UInterchangeGltfTranslator::HandleGltfNode( UInterchangeBaseNodeContainer&
 			{
 				bHasVariants |= CheckForVariants(GltfAsset.Meshes[GltfNode.MeshIndex], GltfAsset.Variants.Num(), GltfAsset.Materials.Num());
 			}
+
+			{//Set Morph Target Curve Weights
+				const TArray<FString>& MorphTargetNames = GltfAsset.Meshes[GltfNode.MeshIndex].MorphTargetNames;
+				int32 MorphTargetNamesCount = MorphTargetNames.Num();
+				const TArray<float>& MorphTargetWeights = (GltfNode.MorphTargetWeights.Num() > 0) ? GltfNode.MorphTargetWeights : GltfAsset.Meshes[GltfNode.MeshIndex].MorphTargetWeights;
+
+				if (MorphTargetWeights.Num() == MorphTargetNamesCount)
+				{
+					for (int32 MorphTargetIndex = 0; MorphTargetIndex < MorphTargetNamesCount; MorphTargetIndex++)
+					{
+						InterchangeSceneNode->SetMorphTargetCurveWeight(MorphTargetNames[MorphTargetIndex], MorphTargetWeights[MorphTargetIndex]);
+					}
+				}
+				else
+				{
+					UE_LOG(LogInterchangeImport, Warning, TEXT("GLTF Node [%] Import Warning. Gltf Node's MorphTargetNames count is missmatched against MorphTargetWeights count."), *GltfNode.UniqueId);
+				}
+			}
+
 			break;
 		}
 
@@ -237,10 +256,34 @@ void UInterchangeGltfTranslator::HandleGltfNode( UInterchangeBaseNodeContainer&
 		{
 			if ( GltfAsset.Meshes.IsValidIndex( GltfNode.MeshIndex ) )
 			{
-				HandleGltfMesh(NodeContainer, GltfAsset.Meshes[GltfNode.MeshIndex], GltfNode.MeshIndex, UnusedMeshIndices);
+				UInterchangeMeshNode* MeshNode = HandleGltfMesh(NodeContainer, GltfAsset.Meshes[GltfNode.MeshIndex], GltfNode.MeshIndex, UnusedMeshIndices);
 
-				const FString MeshNodeUid = TEXT("\\Mesh\\") + GltfAsset.Meshes[ GltfNode.MeshIndex ].UniqueId;
-				InterchangeSceneNode->SetCustomAssetInstanceUid( MeshNodeUid );
+				InterchangeSceneNode->SetCustomAssetInstanceUid( MeshNode->GetUniqueID() );
+				if (MeshNode->IsSkinnedMesh())
+				{
+					const TArray<FString>& MorphTargetNames = GltfAsset.Meshes[GltfNode.MeshIndex].MorphTargetNames;
+					int32 MorphTargetNamesCount = MorphTargetNames.Num();
+					const TArray<float>& MorphTargetWeights = (GltfNode.MorphTargetWeights.Num() > 0) ? GltfNode.MorphTargetWeights : GltfAsset.Meshes[GltfNode.MeshIndex].MorphTargetWeights;
+					
+					if (MorphTargetWeights.Num() == MorphTargetNamesCount)
+					{
+						for (int32 MorphTargetIndex = 0; MorphTargetIndex < MorphTargetNamesCount; MorphTargetIndex++)
+						{
+							InterchangeSceneNode->SetMorphTargetCurveWeight(MorphTargetNames[MorphTargetIndex], MorphTargetWeights[MorphTargetIndex]);
+						}
+					}
+					else
+					{
+						UE_LOG(LogInterchangeImport, Warning, TEXT("GLTF Node [%] Import Warning. Gltf Node's MorphTargetNames count is missmatched against MorphTargetWeights count."), *GltfNode.UniqueId);
+					}
+					
+					//Interchange/UE handles Morph Targets in skeletalMeshes:
+					InterchangeSceneNode->AddSpecializedType(UE::Interchange::FSceneNodeStaticData::GetJointSpecializeTypeString());
+					if (MeshNode->GetSkeletonDependeciesCount() == 0)
+					{
+						MeshNode->SetSkeletonDependencyUid(InterchangeSceneNode->GetUniqueID());
+					}
+				}
 
 				if (!bHasVariants && GltfAsset.Variants.Num() > 0)
 				{
@@ -911,7 +954,6 @@ bool UInterchangeGltfTranslator::Translate( UInterchangeBaseNodeContainer& NodeC
 		SendAnalytics(TranslationResult::INPUT_FILE_NOTFOUND);
 		return false;
 	}
-	const FString FileName = FPaths::GetBaseFilename(FilePath);
 
 	GLTF::FFileReader GltfFileReader;
 
@@ -919,6 +961,8 @@ bool UInterchangeGltfTranslator::Translate( UInterchangeBaseNodeContainer& NodeC
 	const bool bLoadMetaData = false;
 	GltfFileReader.ReadFile( FilePath, bLoadImageData, bLoadMetaData, const_cast< UInterchangeGltfTranslator* >( this )->GltfAsset );
 
+	const FString FileName = GltfAsset.Name;
+
 	//Required Extension Check:
 	FString NotSupportedExtensions;
 	if (GltfAsset.RequiredExtensions.Num() != 0)
@@ -966,8 +1010,6 @@ bool UInterchangeGltfTranslator::Translate( UInterchangeBaseNodeContainer& NodeC
 
 	ScaleNodeTranslations(const_cast<UInterchangeGltfTranslator*>(this)->GltfAsset.Nodes, GltfUnitConversionMultiplier);
 
-	const_cast< UInterchangeGltfTranslator* >( this )->GltfAsset.GenerateNames(FileName);
-
 	// Textures
 	{
 		int32 TextureIndex = 0;
@@ -1199,7 +1241,7 @@ TFuture< TOptional< UE::Interchange::FStaticMeshPayloadData > > UInterchangeGltf
 	return Async(EAsyncExecution::TaskGraph, [this, PayLoadKey]
 		{
 			UE::Interchange::FStaticMeshPayloadData StaticMeshPayloadData;
-			UE::Interchange::Gltf::Private::GetStaticMeshPayloadDataForPayLoadKey(GltfAsset, PayLoadKey, StaticMeshPayloadData);
+			UE::Interchange::Gltf::Private::GetStaticMeshPayloadDataForPayLoadKey(GltfAsset, PayLoadKey, StaticMeshPayloadData.MeshDescription);
 
 			TOptional<UE::Interchange::FStaticMeshPayloadData> Result;
 			Result.Emplace(StaticMeshPayloadData);
@@ -1259,7 +1301,7 @@ TFuture<TOptional<UE::Interchange::FAnimationCurvePayloadData>> UInterchangeGltf
 			TOptional<UE::Interchange::FAnimationCurvePayloadData> Result;
 			UE::Interchange::FAnimationCurvePayloadData AnimationCurvePayloadData;
 
-			if (UE::Interchange::Gltf::Private::GetAnimationTransformPayloadData(PayLoadKey, GltfAsset, AnimationCurvePayloadData))
+			if (UE::Interchange::Gltf::Private::GetAnimationPayloadData(PayLoadKey, GltfAsset, AnimationCurvePayloadData))
 			{
 				Result.Emplace(AnimationCurvePayloadData);
 			}
@@ -1305,7 +1347,7 @@ void UInterchangeGltfTranslator::HandleGltfAnimation(UInterchangeBaseNodeContain
 
 	TMap< const GLTF::FNode*, TArray<int32> > NodeChannelsMap;
 
-	TMap<int32, UInterchangeSkeletalAnimationTrackNode*> RootJointIndexToTrackNodeMap;
+	TMap<FString, UInterchangeSkeletalAnimationTrackNode*> RootJointIndexToTrackNodeMap;
 	TMap<UInterchangeSkeletalAnimationTrackNode*, TMap<FString, TArray<int32>>> TrackNodeToJointUidWithChannelsUsedMap;
 
 	for (int32 ChannelIndex = 0; ChannelIndex < GltfAnimation.Channels.Num(); ++ChannelIndex)
@@ -1313,34 +1355,102 @@ void UInterchangeGltfTranslator::HandleGltfAnimation(UInterchangeBaseNodeContain
 		const GLTF::FAnimation::FChannel& Channel = GltfAnimation.Channels[ChannelIndex];
 		const GLTF::FNode* AnimatedNode = &Channel.Target.Node;
 
-		if (AnimatedNode->Type == GLTF::FNode::EType::Joint && GltfAsset.Nodes.IsValidIndex(AnimatedNode->RootJointIndex))
+		const FString* AnimatedNodeUidPtr = NodeUidMap.Find(AnimatedNode);
+		if (!ensure(AnimatedNodeUidPtr))
 		{
-			const FString* JointNodeUid = NodeUidMap.Find(AnimatedNode);
-			if (ensure(JointNodeUid) && GltfAsset.Nodes.IsValidIndex(AnimatedNode->RootJointIndex))
+			continue;
+		}
+
+		FString AnimatedNodeUid = *AnimatedNodeUidPtr;
+
+		auto CreateSkeletalAnimationTrackNode = [&NodeContainer, &RootJointIndexToTrackNodeMap, &GltfAnimation, &AnimationIndex, &TrackNodeToJointUidWithChannelsUsedMap, &ChannelIndex, &AnimatedNodeUid](const FString& SkeletonNodeUid, const TMap<FString, FString>& AnimationPayloadKeyForMorphTargetNodeUids)
+		{
+			if (SkeletonNodeUid.Len() > 0)
 			{
 				UInterchangeSkeletalAnimationTrackNode* TrackNode = nullptr;
-				if (RootJointIndexToTrackNodeMap.Contains(AnimatedNode->RootJointIndex))
+				if (RootJointIndexToTrackNodeMap.Contains(SkeletonNodeUid))
 				{
-					TrackNode = RootJointIndexToTrackNodeMap[AnimatedNode->RootJointIndex];
+					TrackNode = RootJointIndexToTrackNodeMap[SkeletonNodeUid];
 				}
 				else
 				{
 					TrackNode = NewObject< UInterchangeSkeletalAnimationTrackNode >(&NodeContainer);
-					FString TrackNodeUid = "\\SkeletalAnimation\\" + LexToString(AnimatedNode->RootJointIndex) + "_" + LexToString(AnimationIndex);
+					FString TrackNodeUid = "\\SkeletalAnimation\\" + SkeletonNodeUid + "_" + LexToString(AnimationIndex);
 					TrackNode->InitializeNode(TrackNodeUid, GltfAnimation.Name, EInterchangeNodeContainerType::TranslatedAsset);
-					const GLTF::FNode& RootJointNode = GltfAsset.Nodes[AnimatedNode->RootJointIndex];
-					const FString* SkeletonNodeUid = NodeUidMap.Find(&RootJointNode);
-					TrackNode->SetCustomSkeletonNodeUid(*SkeletonNodeUid);
+					TrackNode->SetCustomSkeletonNodeUid(SkeletonNodeUid);
 
 					NodeContainer.AddNode(TrackNode);
 
-					RootJointIndexToTrackNodeMap.Add(AnimatedNode->RootJointIndex, TrackNode);
+					RootJointIndexToTrackNodeMap.Add(SkeletonNodeUid, TrackNode);
+				}
+
+				for (const TPair<FString, FString>& AnimationPayloadKeyForMorphTargetNodeUid : AnimationPayloadKeyForMorphTargetNodeUids)
+				{
+					TrackNode->SetAnimationPayloadKeyForMorphTargetNodeUid(AnimationPayloadKeyForMorphTargetNodeUid.Key, AnimationPayloadKeyForMorphTargetNodeUid.Value);
 				}
+
 				TMap<FString, TArray<int32>>& JointUidWithChannelsUsedMap = TrackNodeToJointUidWithChannelsUsedMap.FindOrAdd(TrackNode);
-				TArray<int32>& ChannelsUsed = JointUidWithChannelsUsedMap.FindOrAdd(*JointNodeUid);
+				TArray<int32>& ChannelsUsed = JointUidWithChannelsUsedMap.FindOrAdd(AnimatedNodeUid);
 				ChannelsUsed.Add(ChannelIndex);
 			}
+		};
+
+		
+		bool bAnimationChannelProcessed = false;
+
+		bool bSkeletalAnimation = AnimatedNode->Type == GLTF::FNode::EType::Joint && GltfAsset.Nodes.IsValidIndex(AnimatedNode->RootJointIndex);
+		if (bSkeletalAnimation)
+		{
+			const FString* SkeletonUidPtr = NodeUidMap.Find(&GltfAsset.Nodes[AnimatedNode->RootJointIndex]);
+			FString SkeletonNodeUid = *SkeletonUidPtr;
+
+			CreateSkeletalAnimationTrackNode(SkeletonNodeUid, TMap<FString, FString>());
+
+			bAnimationChannelProcessed = true;
+		}
+
+		bool bMorphTargetAnimation = Channel.Target.Path == GLTF::FAnimation::EPath::Weights;
+		if (bMorphTargetAnimation)
+		{
+			TMap<FString, FString> AnimationPayloadKeyForMorphTargetNodeUids;
+			//Find SceneNode that references the MeshNode:
+			if (const UInterchangeSceneNode* ConstSceneMeshActorNode = Cast< UInterchangeSceneNode >(NodeContainer.GetNode(*AnimatedNodeUid)))
+			{
+				FString SkeletalMeshUid;
+				if (ConstSceneMeshActorNode->GetCustomAssetInstanceUid(SkeletalMeshUid))
+				{
+					if (const UInterchangeMeshNode* MeshNode = Cast< UInterchangeMeshNode >(NodeContainer.GetNode(SkeletalMeshUid)))
+					{
+						TArray<FString> MorphTargetDependencies;
+						MeshNode->GetMorphTargetDependencies(MorphTargetDependencies);
+						for (const FString& MorphTargetDependencyUid : MorphTargetDependencies)
+						{
+							if (const UInterchangeMeshNode* MorphTargetNodeConst = Cast< UInterchangeMeshNode >(NodeContainer.GetNode(MorphTargetDependencyUid)))
+							{
+								if (MorphTargetNodeConst->GetPayLoadKey().IsSet())
+								{
+									FString PayLoadKey = MorphTargetNodeConst->GetPayLoadKey().GetValue(); // meshindex : MorphTargetIndex
+									PayLoadKey = TEXT("MorphTargetAnimation~") + LexToString(AnimationIndex) + TEXT(":") + LexToString(ChannelIndex) + TEXT(":") + PayLoadKey;
+
+									AnimationPayloadKeyForMorphTargetNodeUids.Add(MorphTargetDependencyUid, PayLoadKey);
+								}
+							}
+						}
+					}
+				}
+			}
+
+			if (AnimationPayloadKeyForMorphTargetNodeUids.Num() > 0)
+			{
+				FString SkeletonNodeUid = AnimatedNodeUid;
+				CreateSkeletalAnimationTrackNode(AnimatedNodeUid, AnimationPayloadKeyForMorphTargetNodeUids);
+			}
 			
+			bAnimationChannelProcessed = true;
+		}
+
+		if (bAnimationChannelProcessed)
+		{
 			continue;
 		}
 
@@ -1365,6 +1475,7 @@ void UInterchangeGltfTranslator::HandleGltfAnimation(UInterchangeBaseNodeContain
 
 			for (const TTuple<FString, TArray<int32>>& JointNodeUidAndChannelsUsedPair : TrackNodeAndJointNodeChannels.Value)
 			{
+				bool bHasNonWeightAnimationChannel = false;
 				double PreviousStopTime = StopTime;
 				//check channel length and build payload
 				FString Payload = LexToString(AnimationIndex);
@@ -1387,7 +1498,7 @@ void UInterchangeGltfTranslator::HandleGltfAnimation(UInterchangeBaseNodeContain
 					}
 					else
 					{
-						if (Sampler.Input.Count != Sampler.Output.Count)
+						if (Channel.Target.Path != GLTF::FAnimation::EPath::Weights && (Sampler.Input.Count != Sampler.Output.Count))
 						{
 							// if any of the channels are corrupt the joint will not receive any of the  animation data
 							Payload = "";
@@ -1418,13 +1529,22 @@ void UInterchangeGltfTranslator::HandleGltfAnimation(UInterchangeBaseNodeContain
 						FrameNumber = CurrentFrameNumber;
 						StopTime = FrameNumber * SingleFrameDuration;
 					}
-					Payload += ":" + LexToString(ChannelIndex);
+
+					if (Channel.Target.Path != GLTF::FAnimation::EPath::Weights)
+					{
+						bHasNonWeightAnimationChannel = true;
+						Payload += ":" + LexToString(ChannelIndex);
+					}
 				}
 
 				//set payload:
 				if (Payload.Len() > 0)
 				{
-					TrackNodeAndJointNodeChannels.Key->SetAnimationPayloadKeyForSceneNodeUid(JointNodeUidAndChannelsUsedPair.Key, Payload);
+					if (bHasNonWeightAnimationChannel)
+					{
+						TrackNodeAndJointNodeChannels.Key->SetAnimationPayloadKeyForSceneNodeUid(JointNodeUidAndChannelsUsedPair.Key, Payload);
+					}
+					
 					bHasAnimationPayloadSet = true;
 				}
 				else
@@ -1507,7 +1627,6 @@ void UInterchangeGltfTranslator::HandleGltfAnimation(UInterchangeBaseNodeContain
 				{
 					UsedChannels |= ScaleChannel;
 				} break;
-
 				default: break;
 			}
 		}
@@ -1782,9 +1901,17 @@ TFuture<TOptional<UE::Interchange::FSkeletalMeshLodPayloadData>> UInterchangeGlt
 
 TFuture<TOptional<UE::Interchange::FSkeletalMeshMorphTargetPayloadData>> UInterchangeGltfTranslator::GetSkeletalMeshMorphTargetPayloadData(const FString& PayLoadKey) const
 {
-	TSharedPtr<TPromise<TOptional<UE::Interchange::FSkeletalMeshMorphTargetPayloadData>>> Promise = MakeShared<TPromise<TOptional<UE::Interchange::FSkeletalMeshMorphTargetPayloadData>>>();
-	Promise->SetValue(TOptional<UE::Interchange::FSkeletalMeshMorphTargetPayloadData>{});
-	return Promise->GetFuture();
+	return Async(EAsyncExecution::TaskGraph, [this, PayLoadKey]
+		{
+			UE::Interchange::FSkeletalMeshMorphTargetPayloadData StaticMeshPayloadData;
+
+			UE::Interchange::Gltf::Private::GetStaticMeshPayloadDataForPayLoadKey(GltfAsset, PayLoadKey, StaticMeshPayloadData.LodMeshDescription);
+
+			TOptional<UE::Interchange::FSkeletalMeshMorphTargetPayloadData> Result;
+			Result.Emplace(StaticMeshPayloadData);
+
+			return Result;
+		});
 }
 
 void UInterchangeGltfTranslator::HandleGltfSkeletons(UInterchangeBaseNodeContainer& NodeContainer, const FString& SceneNodeUid, const TArray<int32>& SkinnedMeshNodes, TSet<int>& UnusedMeshIndices) const
@@ -1883,27 +2010,49 @@ void UInterchangeGltfTranslator::HandleGltfSkeletons(UInterchangeBaseNodeContain
 	}
 }
 
-UInterchangeMeshNode* UInterchangeGltfTranslator::HandleGltfMesh(UInterchangeBaseNodeContainer& NodeContainer, 
-	const GLTF::FMesh& GltfMesh, int MeshIndex, 
-	TSet<int>& UnusedMeshIndices, 
+UInterchangeMeshNode* UInterchangeGltfTranslator::HandleGltfMesh(UInterchangeBaseNodeContainer& NodeContainer,
+	const GLTF::FMesh& GltfMesh, int MeshIndex,
+	TSet<int>& UnusedMeshIndices,
 	const FString& SkeletalName/*If set it creates the mesh even if it was already created (for Skeletals)*/,
 	const FString& SkeletalId) const
 {
-	if (!UnusedMeshIndices.Contains(MeshIndex) && SkeletalName.Len() == 0)
+	FString MeshName = SkeletalName.Len() ? SkeletalName : GltfMesh.Name;
+	FString MeshNodeUid = TEXT("\\Mesh\\") + (SkeletalId.Len() ? SkeletalId : GltfMesh.UniqueId);
+
+	//check if Node already exist with MeshNodeUid:
+	if (const UInterchangeMeshNode* Node = Cast< UInterchangeMeshNode >(NodeContainer.GetNode(MeshNodeUid)))
 	{
-		return nullptr;
+		UInterchangeMeshNode* MeshNode = const_cast<UInterchangeMeshNode*>(Node);
+		if (ensure(MeshNode))
+		{
+			return MeshNode;
+		}
 	}
-	FString MeshName = SkeletalName.Len() ? SkeletalName : GltfMesh.Name;
+
+	//to track which meshes we have to generate a mesh node for at the end of Translate:
 	UnusedMeshIndices.Remove(MeshIndex);
 
+	//Create Mesh Node:
 	UInterchangeMeshNode* MeshNode = NewObject< UInterchangeMeshNode >(&NodeContainer);
-	FString MeshNodeUid = TEXT("\\Mesh\\") + (SkeletalId.Len() ? SkeletalId : GltfMesh.UniqueId);
-
 	MeshNode->InitializeNode(MeshNodeUid, MeshName, EInterchangeNodeContainerType::TranslatedAsset);
-	MeshNode->SetPayLoadKey(LexToString(MeshIndex));
+
+	//Generate Mesh Payload:
+	FString PayloadKey;
+	//if it has MorphTargets but is not a Skeletal (aka not called from HandleGltfSkeletons, aka SkeletalName.len = 0
+	//then the Payload processing should be calling the StaticMesh acquisition from the SkeletalMesh acquisition
+	// SkeletalMesh acqusition is called because UE handles MorphTargets in SkeletalMeshes 
+	// (hence Meshes with MorphTargets are marked to be skinned and the node that references it, is set to be Joint)
+	if (GltfMesh.MorphTargetNames.Num() > 0 && SkeletalName.Len() == 0)
+	{
+		PayloadKey = "StaticMeshWithMorphTarget~";
+		//if it is called from HandleGltfSkeletons then the payloadkey will be overwritten. (in HandleGltfSkeletons)
+	}
+	PayloadKey += LexToString(MeshIndex);
+	MeshNode->SetPayLoadKey(PayloadKey);
 
 	NodeContainer.AddNode(MeshNode);
 
+	//Set Slot Material Dependencies:
 	for (int32 PrimitiveCounter = 0; PrimitiveCounter < GltfMesh.Primitives.Num(); PrimitiveCounter++)
 	{
 		const GLTF::FPrimitive& Primitive = GltfMesh.Primitives[PrimitiveCounter];
@@ -1917,6 +2066,55 @@ UInterchangeMeshNode* UInterchangeGltfTranslator::HandleGltfMesh(UInterchangeBas
 		}
 	}
 
+	//Generate Morph Target Meshes:
+	if (GltfMesh.MorphTargetNames.Num() > 0)
+	{
+		MeshNode->SetSkinnedMesh(true);
+
+		for (int32 MorphTargetIndex = 0; MorphTargetIndex < GltfMesh.MorphTargetNames.Num(); MorphTargetIndex++)
+		{
+			//check if MorphTarget mesh was already created or not:
+			FString MorphTargetName = GltfMesh.MorphTargetNames[MorphTargetIndex]; //Morph Target Names are validated to be unique (GLTFAsset::GenerateNames)
+
+			//Add the MorphTargetName as a dependency to original mesh:
+			MeshNode->SetMorphTargetDependencyUid(MorphTargetName);
+
+			//check if Node already exist with MorphTargetName(uid):
+			if (const UInterchangeMeshNode* Node = Cast< UInterchangeMeshNode >(NodeContainer.GetNode(MorphTargetName)))
+			{
+				continue;
+			}
+
+			//create MorphTargetMeshNode:
+			UInterchangeMeshNode* MorphTargetMeshNode = NewObject< UInterchangeMeshNode >(&NodeContainer);
+			MorphTargetMeshNode->InitializeNode(MorphTargetName, MorphTargetName, EInterchangeNodeContainerType::TranslatedAsset);
+
+			//Generate Payload:
+			FString MorphTargetPayLoadKey = LexToString(MeshIndex) + TEXT(":") + LexToString(MorphTargetIndex);
+			MorphTargetMeshNode->SetPayLoadKey(MorphTargetPayLoadKey);
+
+			//set mesh as a morph target:
+			MorphTargetMeshNode->SetMorphTarget(true);
+			MorphTargetMeshNode->SetMorphTargetName(MorphTargetName);
+
+			NodeContainer.AddNode(MorphTargetMeshNode);
+
+			//Set Slot Material Dependencies:
+			for (int32 PrimitiveCounter = 0; PrimitiveCounter < GltfMesh.Primitives.Num(); PrimitiveCounter++)
+			{
+				const GLTF::FPrimitive& Primitive = GltfMesh.Primitives[PrimitiveCounter];
+
+				// Assign materials
+				if (GltfAsset.Materials.IsValidIndex(Primitive.MaterialIndex))
+				{
+					const FString MaterialName = GltfAsset.Materials[Primitive.MaterialIndex].Name;
+					const FString ShaderGraphNodeUid = UInterchangeShaderGraphNode::MakeNodeUid(GltfAsset.Materials[Primitive.MaterialIndex].UniqueId);
+					MorphTargetMeshNode->SetSlotMaterialDependencyUid(MaterialName, ShaderGraphNodeUid);
+				}
+			}
+		}
+	}
+
 	return MeshNode;
 }
 
