Interchange: Remove pipeline ensure for handling all mesh build properties.

#rnx

diff --git a/Engine/Plugins/Interchange/Runtime/Source/Pipelines/Private/InterchangeGenericMeshPipeline.cpp b/Engine/Plugins/Interchange/Runtime/Source/Pipelines/Private/InterchangeGenericMeshPipeline.cpp
--- a/Engine/Plugins/Interchange/Runtime/Source/Pipelines/Private/InterchangeGenericMeshPipeline.cpp
+++ b/Engine/Plugins/Interchange/Runtime/Source/Pipelines/Private/InterchangeGenericMeshPipeline.cpp
@@ -191,6 +191,9 @@ void UInterchangeGenericMeshPipeline::SetReimportSourceIndex(UClass* ReimportObj
 #if WITH_EDITOR
 bool UInterchangeGenericMeshPipeline::DoClassesIncludeAllEditableStructProperties(const TArray<const UClass*>& Classes, const UStruct* Struct)
 {
+	check(IsInGameThread());
+
+	bool bResult = true;
 	const FName CategoryKey("Category");
 	for (const FProperty* Property = Struct->PropertyLink; Property; Property = Property->PropertyLinkNext)
 	{
@@ -224,10 +227,11 @@ bool UInterchangeGenericMeshPipeline::DoClassesIncludeAllEditableStructPropertie
 			//Ensure to notify
 			if (!bFindProperty)
 			{
-				return false;
+				UE_LOG(LogInterchangePipeline, Log, TEXT("Interchange mesh pipeline do not include build property %s"), *PropertyName.ToString());
+				bResult = false;
 			}
 		}
 	}
-	return true;
+	return bResult;
 }
 #endif
\ No newline at end of file
diff --git a/Engine/Plugins/Interchange/Runtime/Source/Pipelines/Private/InterchangeGenericSkeletalMeshPipeline.cpp b/Engine/Plugins/Interchange/Runtime/Source/Pipelines/Private/InterchangeGenericSkeletalMeshPipeline.cpp
--- a/Engine/Plugins/Interchange/Runtime/Source/Pipelines/Private/InterchangeGenericSkeletalMeshPipeline.cpp
+++ b/Engine/Plugins/Interchange/Runtime/Source/Pipelines/Private/InterchangeGenericSkeletalMeshPipeline.cpp
@@ -87,15 +87,21 @@ void UInterchangeGenericMeshPipeline::ExecutePreImportPipelineSkeletalMesh()
 
 #if WITH_EDITOR
 	//Make sure the generic pipeline we cover all skeletalmesh build settings by asserting when we import
-	{
-		TArray<const UClass*> Classes;
-		Classes.Add(UInterchangeGenericCommonMeshesProperties::StaticClass());
-		Classes.Add(UInterchangeGenericMeshPipeline::StaticClass());
-		if (!ensure(DoClassesIncludeAllEditableStructProperties(Classes, FSkeletalMeshBuildSettings::StaticStruct())))
+	Async(EAsyncExecution::TaskGraphMainThread, []()
 		{
-			UE_LOG(LogInterchangePipeline, Log, TEXT("UInterchangeGenericMeshPipeline: The generic pipeline does not cover all skeletal mesh build options."));
-		}
-	}
+			static bool bVerifyBuildProperties = false;
+			if (!bVerifyBuildProperties)
+			{
+				bVerifyBuildProperties = true;
+				TArray<const UClass*> Classes;
+				Classes.Add(UInterchangeGenericCommonMeshesProperties::StaticClass());
+				Classes.Add(UInterchangeGenericMeshPipeline::StaticClass());
+				if (!DoClassesIncludeAllEditableStructProperties(Classes, FSkeletalMeshBuildSettings::StaticStruct()))
+				{
+					UE_LOG(LogInterchangePipeline, Log, TEXT("UInterchangeGenericMeshPipeline: The generic pipeline does not cover all skeletal mesh build options."));
+				}
+			}
+		});
 #endif
 
 	const bool bConvertStaticMeshToSkeletalMesh = (CommonMeshesProperties->ForceAllMeshAsType == EInterchangeForceMeshType::IFMT_SkeletalMesh);
diff --git a/Engine/Plugins/Interchange/Runtime/Source/Pipelines/Private/InterchangeGenericStaticMeshPipeline.cpp b/Engine/Plugins/Interchange/Runtime/Source/Pipelines/Private/InterchangeGenericStaticMeshPipeline.cpp
--- a/Engine/Plugins/Interchange/Runtime/Source/Pipelines/Private/InterchangeGenericStaticMeshPipeline.cpp
+++ b/Engine/Plugins/Interchange/Runtime/Source/Pipelines/Private/InterchangeGenericStaticMeshPipeline.cpp
@@ -140,15 +140,21 @@ void UInterchangeGenericMeshPipeline::ExecutePreImportPipelineStaticMesh()
 
 #if WITH_EDITOR
 	//Make sure the generic pipeline will cover all staticmesh build settings when we import
-	{
-		TArray<const UClass*> Classes;
-		Classes.Add(UInterchangeGenericCommonMeshesProperties::StaticClass());
-		Classes.Add(UInterchangeGenericMeshPipeline::StaticClass());
-		if (!ensure(DoClassesIncludeAllEditableStructProperties(Classes, FMeshBuildSettings::StaticStruct())))
+	Async(EAsyncExecution::TaskGraphMainThread, []()
 		{
-			UE_LOG(LogInterchangePipeline, Log, TEXT("UInterchangeGenericMeshPipeline: The generic pipeline does not cover all static mesh build options."));
-		}
-	}
+			static bool bVerifyBuildProperties = false;
+			if (!bVerifyBuildProperties)
+			{
+				bVerifyBuildProperties = true;
+				TArray<const UClass*> Classes;
+				Classes.Add(UInterchangeGenericCommonMeshesProperties::StaticClass());
+				Classes.Add(UInterchangeGenericMeshPipeline::StaticClass());
+				if (!DoClassesIncludeAllEditableStructProperties(Classes, FMeshBuildSettings::StaticStruct()))
+				{
+					UE_LOG(LogInterchangePipeline, Log, TEXT("UInterchangeGenericMeshPipeline: The generic pipeline does not cover all static mesh build options."));
+				}
+			}
+		});
 #endif
 
 	if (bImportStaticMeshes && (CommonMeshesProperties->ForceAllMeshAsType == EInterchangeForceMeshType::IFMT_None || CommonMeshesProperties->ForceAllMeshAsType == EInterchangeForceMeshType::IFMT_StaticMesh))
