Fix race condition crash between interchange worker dispatcher and handler.

diff --git a/Engine/Plugins/Interchange/Runtime/Source/Dispatcher/Private/InterchangeWorkerHandler.cpp b/Engine/Plugins/Interchange/Runtime/Source/Dispatcher/Private/InterchangeWorkerHandler.cpp
--- a/Engine/Plugins/Interchange/Runtime/Source/Dispatcher/Private/InterchangeWorkerHandler.cpp
+++ b/Engine/Plugins/Interchange/Runtime/Source/Dispatcher/Private/InterchangeWorkerHandler.cpp
@@ -260,6 +260,14 @@ namespace UE
 				return bSucceed;
 			};
 
+			auto SetTerminatedState = [this]()
+			{
+				//Always broadcast before setting the "terminated" state. API "IsAlive" is use by the dispatcher to reset the worker handler when worker handler state is "terminated"
+				//Notify the dispatcher the worker handler is done so it can terminate queue tasks
+				OnWorkerHandlerExitLoop.Broadcast();
+				WorkerState = EWorkerState::Terminated;
+			};
+
 			while (IsAlive())
 			{
 				switch (WorkerState)
@@ -272,7 +280,7 @@ namespace UE
 
 						if (ErrorState != EWorkerErrorState::Ok)
 						{
-							WorkerState = EWorkerState::Terminated;
+							SetTerminatedState();
 							break;
 						}
 
@@ -369,7 +377,7 @@ namespace UE
 						//Make sure all running task are completed with error
 						KillAllCurrentTasks();
 
-						WorkerState = EWorkerState::Terminated;
+						SetTerminatedState();
 						break;
 					}
 
@@ -379,8 +387,6 @@ namespace UE
 					}
 				}
 			}
-			//Notify the dispatcher the worker handler is done so it can terminate queue tasks
-			OnWorkerHandlerExitLoop.Broadcast();
 		}
 
 		void FInterchangeWorkerHandler::Stop()
