InterchangeMaterialFactory:
- Removing the TSet in the UpdateFromFunctionResource methods, causing the different tests to be unstable, for example it may happens that a material doesn't have any of its inputs connected in the material expression graph

MaterialExpressionRgbToHsv/HsvToRgb:
- returning the custom expression compile method instead of the direct method of the compiler, because on some rare case scenario, it can cause a crash, because of some inputs not properly connected in a material (linked to the UpdateFromFunctionResource in the Material Factory)

Python tests:
- Interchange: changing the camera name, by a camera location and rotation in the utils
- MaterialX: fixing some warnings when we load an asset asking for a PackageName instead of an ObjectPath (e.g: "Path/Asset.Asset" -> "Path/Asset")
- Changing the ground truth of the gltf test, using a camera instead of the viewport, too sensitive to the viewport changes

diff --git a/Engine/Plugins/Interchange/Runtime/Source/Import/Private/InterchangeImportModule.cpp b/Engine/Plugins/Interchange/Runtime/Source/Import/Private/InterchangeImportModule.cpp
--- a/Engine/Plugins/Interchange/Runtime/Source/Import/Private/InterchangeImportModule.cpp
+++ b/Engine/Plugins/Interchange/Runtime/Source/Import/Private/InterchangeImportModule.cpp
@@ -101,10 +101,7 @@ void FInterchangeImportModule::StartupModule()
 
 
 void FInterchangeImportModule::ShutdownModule()
-{
-	FInterchangeImportMaterialAsyncHelper& InterchangeMaterialAsyncHelper = FInterchangeImportMaterialAsyncHelper::GetInstance();
-	InterchangeMaterialAsyncHelper.CleanUp();
-}
+{}
 
 
 
diff --git a/Engine/Plugins/Interchange/Runtime/Source/Import/Private/Material/InterchangeMaterialFactory.cpp b/Engine/Plugins/Interchange/Runtime/Source/Import/Private/Material/InterchangeMaterialFactory.cpp
--- a/Engine/Plugins/Interchange/Runtime/Source/Import/Private/Material/InterchangeMaterialFactory.cpp
+++ b/Engine/Plugins/Interchange/Runtime/Source/Import/Private/Material/InterchangeMaterialFactory.cpp
@@ -1348,9 +1348,7 @@ void UInterchangeMaterialFunctionFactory::SetupMaterial(UMaterialFunction* Mater
 void FInterchangeImportMaterialAsyncHelper::UpdateFromFunctionResource(UMaterialExpressionMaterialFunctionCall* MaterialFunctionCall)
 {
 	FScopeLock Lock(&UpdatedMaterialFunctionCallsLock);
-	if(UpdatedMaterialFunctionCalls.Find(MaterialFunctionCall) == nullptr)
 	{
-		UpdatedMaterialFunctionCalls.Add(MaterialFunctionCall);
 		MaterialFunctionCall->UpdateFromFunctionResource();
 	}
 }
@@ -1358,9 +1356,7 @@ void FInterchangeImportMaterialAsyncHelper::UpdateFromFunctionResource(UMaterial
 void FInterchangeImportMaterialAsyncHelper::UpdateFromFunctionResource(UMaterialFunctionInterface* MaterialFunction)
 {
 	FScopeLock Lock(&UpdatedMaterialFunctionsLock);
-	if(UpdatedMaterialFunctions.Find(MaterialFunction) == nullptr)
 	{
-		UpdatedMaterialFunctions.Add(MaterialFunction);
 		MaterialFunction->UpdateFromFunctionResource();
 	}
 }
@@ -1372,11 +1368,5 @@ FInterchangeImportMaterialAsyncHelper& FInterchangeImportMaterialAsyncHelper::Ge
 	return Instance;
 }
 
-void FInterchangeImportMaterialAsyncHelper::CleanUp()
-{
-	UpdatedMaterialFunctionCalls.Empty();
-	UpdatedMaterialFunctions.Empty();
-}
-
 #undef LOCTEXT_NAMESPACE
 
diff --git a/Engine/Plugins/Interchange/Runtime/Source/Import/Private/MaterialX/MaterialExpressionHsvToRgb.cpp b/Engine/Plugins/Interchange/Runtime/Source/Import/Private/MaterialX/MaterialExpressionHsvToRgb.cpp
--- a/Engine/Plugins/Interchange/Runtime/Source/Import/Private/MaterialX/MaterialExpressionHsvToRgb.cpp
+++ b/Engine/Plugins/Interchange/Runtime/Source/Import/Private/MaterialX/MaterialExpressionHsvToRgb.cpp
@@ -41,8 +41,7 @@ int32 UMaterialExpressionHsvToRgb::Compile(FMaterialCompiler* Compiler, int32 Ou
     float3 p = abs(frac(c.xxx + K.xyz) * 6.0 - K.www);
     return c.z * lerp(K.xxx, saturate(p - K.xxx), c.y);)");
 
-	TArray<int32> Inputs{ Input.Compile(Compiler) };
-	return Compiler->CustomExpression(MaterialExpressionCustom, 0, Inputs);
+	return MaterialExpressionCustom->Compile(Compiler, OutputIndex);
 }
 
 void UMaterialExpressionHsvToRgb::GetCaption(TArray<FString>& OutCaptions) const
diff --git a/Engine/Plugins/Interchange/Runtime/Source/Import/Private/MaterialX/MaterialExpressionRgbToHsv.cpp b/Engine/Plugins/Interchange/Runtime/Source/Import/Private/MaterialX/MaterialExpressionRgbToHsv.cpp
--- a/Engine/Plugins/Interchange/Runtime/Source/Import/Private/MaterialX/MaterialExpressionRgbToHsv.cpp
+++ b/Engine/Plugins/Interchange/Runtime/Source/Import/Private/MaterialX/MaterialExpressionRgbToHsv.cpp
@@ -45,8 +45,7 @@ int32 UMaterialExpressionRgbToHsv::Compile(FMaterialCompiler* Compiler, int32 Ou
     float e = 1.0e-10;
     return float3(abs(q.z + (q.w - q.y) / (6.0 * d + e)), d / (q.x + e), q.x);)");
 		
-	TArray<int32> Inputs{ Input.Compile(Compiler) };
-	return Compiler->CustomExpression(MaterialExpressionCustom, 0, Inputs);
+	return MaterialExpressionCustom->Compile(Compiler, OutputIndex);
 }
 
 void UMaterialExpressionRgbToHsv::GetCaption(TArray<FString>& OutCaptions) const
diff --git a/Engine/Plugins/Interchange/Runtime/Source/Import/Public/Material/InterchangeMaterialFactory.h b/Engine/Plugins/Interchange/Runtime/Source/Import/Public/Material/InterchangeMaterialFactory.h
--- a/Engine/Plugins/Interchange/Runtime/Source/Import/Public/Material/InterchangeMaterialFactory.h
+++ b/Engine/Plugins/Interchange/Runtime/Source/Import/Public/Material/InterchangeMaterialFactory.h
@@ -84,8 +84,6 @@ public:
 
 	static FInterchangeImportMaterialAsyncHelper& GetInstance();
 
-	void CleanUp();
-
 #if WITH_EDITOR
 	void UpdateFromFunctionResource(UMaterialExpressionMaterialFunctionCall* MaterialFunctionCall);
 
@@ -95,8 +93,5 @@ public:
 private:
 
 	FCriticalSection UpdatedMaterialFunctionCallsLock;
-	TSet<UMaterialExpressionMaterialFunctionCall*> UpdatedMaterialFunctionCalls;
-
 	FCriticalSection UpdatedMaterialFunctionsLock;
-	TSet<UMaterialFunctionInterface*> UpdatedMaterialFunctions;
 };
\ No newline at end of file
