MaterialX: Fixing crash when importing rotate3d nodes without the input "amount"

diff --git a/Engine/Plugins/Interchange/Runtime/Source/Import/Private/MaterialX/InterchangeMaterialXTranslator.cpp b/Engine/Plugins/Interchange/Runtime/Source/Import/Private/MaterialX/InterchangeMaterialXTranslator.cpp
--- a/Engine/Plugins/Interchange/Runtime/Source/Import/Private/MaterialX/InterchangeMaterialXTranslator.cpp
+++ b/Engine/Plugins/Interchange/Runtime/Source/Import/Private/MaterialX/InterchangeMaterialXTranslator.cpp
@@ -1491,17 +1491,20 @@ void UInterchangeMaterialXTranslator::ConnectRotate3DInputToOutput(mx::NodePtr U
 	AddAttributeFromValueOrInterface(UpstreamNode->getInput("axis"), TEXT("NormalizedRotationAxis"), Rotate3DNode);
 
 	// we create a Divide node to convert MaterialX angle in degrees to UE's angle which is a value between [0,1]
-	if(mx::InputPtr Input = UpstreamNode->getInput("amount"); !AddAttributeFromValueOrInterface(Input, TEXT("RotationAngle"), Rotate3DNode))
+	if(mx::InputPtr Input = UpstreamNode->getInput("amount"))
 	{
-		//copy "in" input into "in1" and create "amount" input as "in2" with the value 360 because UE's angle is a value between [0,1]
-		mx::NodePtr NewDivideNode = CreateNode(UpstreamNode->getParent()->asA<mx::NodeGraph>(),
-											   UpstreamNode->getName().c_str(),
-											   mx::Category::Divide,
-											   { {"in1", Input} }, // rename input to match <divide> input "in1"
-											   { {"in2", FAttributeValueArray{{"type", "float"}, {"value", "360"}}} });
+		if(!AddAttributeFromValueOrInterface(Input, TEXT("RotationAngle"), Rotate3DNode))
+		{
+			//copy "in" input into "in1" and create "amount" input as "in2" with the value 360 because UE's angle is a value between [0,1]
+			mx::NodePtr NewDivideNode = CreateNode(UpstreamNode->getParent()->asA<mx::NodeGraph>(),
+												   UpstreamNode->getName().c_str(),
+												   mx::Category::Divide,
+												   { {"in1", Input} }, // rename input to match <divide> input "in1"
+												   { {"in2", FAttributeValueArray{{"type", "float"}, {"value", "360"}}} });
 
-		// Input now points to the new node
-		Input->setNodeName(NewDivideNode->getName());
+				// Input now points to the new node
+				Input->setNodeName(NewDivideNode->getName());
+		}
 	}
 
 	UInterchangeShaderPortsAPI::ConnectDefaultOuputToInput(ParentShaderNode, InputChannelName, Rotate3DNode->GetUniqueID());
