- Filtering assets for MaterialX Pipeline settings during import, only allowing Material Functions that fulfill the requirements of the Standard Surface material functions, avoiding potential crashes

diff --git a/Engine/Plugins/Interchange/Editor/Source/Pipelines/Private/InterchangeEditorPipelinesModule.cpp b/Engine/Plugins/Interchange/Editor/Source/Pipelines/Private/InterchangeEditorPipelinesModule.cpp
--- a/Engine/Plugins/Interchange/Editor/Source/Pipelines/Private/InterchangeEditorPipelinesModule.cpp
+++ b/Engine/Plugins/Interchange/Editor/Source/Pipelines/Private/InterchangeEditorPipelinesModule.cpp
@@ -4,7 +4,9 @@
 
 #include "InterchangeEditorPipelineDetails.h"
 #include "InterchangeEditorPipelineStyle.h"
+#include "InterchangeGenericMaterialPipeline.h"
 #include "InterchangeManager.h"
+#include "InterchangeMaterialXPipelineSettingsCustomization.h"
 #include "InterchangePipelineBase.h"
 #include "InterchangePipelineFactories.h"
 #include "Misc/CoreDelegates.h"
@@ -78,6 +80,7 @@ void FInterchangeEditorPipelinesModule::AcquireResources()
 	ClassesToUnregisterOnShutdown.Add(UInterchangePipelineBase::StaticClass()->GetFName());
 	PropertyEditorModule.RegisterCustomClassLayout(ClassesToUnregisterOnShutdown.Last(), FOnGetDetailCustomizationInstance::CreateStatic(&FInterchangePipelineBaseDetailsCustomization::MakeInstance));
 
+	PropertyEditorModule.RegisterCustomPropertyTypeLayout(FMaterialXPipelineSettings::StaticStruct()->GetFName(), FOnGetPropertyTypeCustomizationInstance::CreateStatic(&FInterchangeMaterialXPipelineSettingsCustomization::MakeInstance));
 
 	if (!InterchangeEditorPipelineStyle.IsValid())
 	{
@@ -128,6 +131,8 @@ void FInterchangeEditorPipelinesModule::ReleaseResources()
 		AssetTools.UnregisterAssetTypeActions(PipelineBase_TypeActions.ToSharedRef());
 		AssetTools.UnregisterAssetTypeActions(PythonPipelineBase_TypeActions.ToSharedRef());
 	}
+
+	PropertyEditorModule->UnregisterCustomPropertyTypeLayout(FMaterialXPipelineSettings::StaticStruct()->GetFName());
 }
 
 TSharedRef<FPropertySection> FInterchangeEditorPipelinesModule::RegisterPropertySection(FPropertyEditorModule& PropertyModule, FName ClassName, FName SectionName, FText DisplayName)
diff --git a/Engine/Plugins/Interchange/Editor/Source/Pipelines/Private/InterchangeMaterialXPipelineSettingsCustomization.cpp b/Engine/Plugins/Interchange/Editor/Source/Pipelines/Private/InterchangeMaterialXPipelineSettingsCustomization.cpp
new file mode 100644
--- /dev/null
+++ b/Engine/Plugins/Interchange/Editor/Source/Pipelines/Private/InterchangeMaterialXPipelineSettingsCustomization.cpp
@@ -0,0 +1,90 @@
+// Copyright Epic Games, Inc. All Rights Reserved.
+
+#include "InterchangeMaterialXPipelineSettingsCustomization.h"
+#include "IDetailChildrenBuilder.h"
+#include "InterchangeGenericMaterialPipeline.h"
+#include "InterchangeMaterialDefinitions.h"
+#include "PropertyCustomizationHelpers.h"
+
+TSharedRef<IPropertyTypeCustomization> FInterchangeMaterialXPipelineSettingsCustomization::MakeInstance()
+{
+	return MakeShareable(new FInterchangeMaterialXPipelineSettingsCustomization());
+}
+
+void FInterchangeMaterialXPipelineSettingsCustomization::CustomizeHeader(TSharedRef<IPropertyHandle> PropertyHandle, FDetailWidgetRow& HeaderRow, IPropertyTypeCustomizationUtils& CustomizationUtils)
+{
+	// nothing to do here
+}
+
+
+void FInterchangeMaterialXPipelineSettingsCustomization::CustomizeChildren(TSharedRef<IPropertyHandle> PropertyHandle, IDetailChildrenBuilder & ChildBuilder, IPropertyTypeCustomizationUtils & CustomizationUtils)
+{
+	using OnShouldFilterAssetFunc = bool (FInterchangeMaterialXPipelineSettingsCustomization::*)(const FAssetData&);
+
+	// add all the child properties (ApplicationDisplayName & ApplicationDescription)	
+	if(PropertyHandle->IsValidHandle())
+	{
+		TSharedPtr<IPropertyHandle> SurfaceShaderProperty = PropertyHandle->GetChildHandle(GET_MEMBER_NAME_CHECKED(FMaterialXPipelineSettings, SurfaceShader));
+		uint32 NumChildren = 0;
+		SurfaceShaderProperty->GetNumChildren(NumChildren);
+
+		for(uint32 i = 0; i < NumChildren; ++i)
+		{
+			TSharedPtr<IPropertyHandle> ChildPropertyHandle = SurfaceShaderProperty->GetChildHandle(i);
+			TSharedPtr<IPropertyHandle> KeyPropertyHandle = ChildPropertyHandle->GetKeyHandle();
+
+			OnShouldFilterAssetFunc OnShouldFilterAsset = &FInterchangeMaterialXPipelineSettingsCustomization::OnShouldFilterAssetStandardSurface;
+
+			if(KeyPropertyHandle.IsValid())
+			{
+				FString Str;
+				KeyPropertyHandle->GetValueAsDisplayString(Str);
+				if(Str == TEXT("Standard Surface Transmission"))
+				{
+					OnShouldFilterAsset = &FInterchangeMaterialXPipelineSettingsCustomization::OnShouldFilterAssetStandardSurfaceTransmission;
+				}
+			}
+
+			IDetailPropertyRow& SurfaceShaderTypeRow = ChildBuilder.AddProperty(ChildPropertyHandle.ToSharedRef());
+			SurfaceShaderTypeRow.ShowPropertyButtons(false);
+
+			TSharedPtr<SWidget> NameWidget;
+			TSharedPtr<SWidget> ValueWidget;
+			FDetailWidgetRow Row;
+			SurfaceShaderTypeRow.GetDefaultWidgets(NameWidget, ValueWidget, Row);
+
+			FDetailWidgetRow& DetailWidgetRow = SurfaceShaderTypeRow.CustomWidget();
+			DetailWidgetRow
+			.NameContent()
+			.MinDesiredWidth(Row.NameWidget.MinWidth)
+			.MaxDesiredWidth(Row.NameWidget.MaxWidth)
+			[
+				NameWidget.ToSharedRef()
+			]
+			.ValueContent()
+			.MinDesiredWidth(Row.ValueContent().MaxWidth)
+			.MaxDesiredWidth(Row.ValueContent().MaxWidth)
+			[
+				SNew(SHorizontalBox)
+				+ SHorizontalBox::Slot()
+				.HAlign(HAlign_Fill)
+				[
+					SNew(SObjectPropertyEntryBox)
+					.AllowedClass(UMaterialFunction::StaticClass())
+					.PropertyHandle(ChildPropertyHandle)
+					.OnShouldFilterAsset(this, OnShouldFilterAsset)
+				]
+			];
+		}
+	}
+}
+
+bool FInterchangeMaterialXPipelineSettingsCustomization::OnShouldFilterAssetStandardSurface(const FAssetData & InAssetData)
+{
+	return FMaterialXPipelineSettings::ShouldFilterAssets(Cast<UMaterialFunction>(InAssetData.GetAsset()), FMaterialXPipelineSettings::StandardSurfaceInputs, FMaterialXPipelineSettings::StandardSurfaceOutputs);
+}
+
+bool FInterchangeMaterialXPipelineSettingsCustomization::OnShouldFilterAssetStandardSurfaceTransmission(const FAssetData& InAssetData)
+{
+	return FMaterialXPipelineSettings::ShouldFilterAssets(Cast<UMaterialFunction>(InAssetData.GetAsset()), FMaterialXPipelineSettings::TransmissionSurfaceInputs, FMaterialXPipelineSettings::TransmissionSurfaceOutputs);
+}
\ No newline at end of file
diff --git a/Engine/Plugins/Interchange/Editor/Source/Pipelines/Private/InterchangeMaterialXPipelineSettingsCustomization.h b/Engine/Plugins/Interchange/Editor/Source/Pipelines/Private/InterchangeMaterialXPipelineSettingsCustomization.h
new file mode 100644
--- /dev/null
+++ b/Engine/Plugins/Interchange/Editor/Source/Pipelines/Private/InterchangeMaterialXPipelineSettingsCustomization.h
@@ -0,0 +1,24 @@
+// Copyright Epic Games, Inc. All Rights Reserved.
+
+#pragma once
+
+#include "IPropertyTypeCustomization.h"
+
+class FInterchangeMaterialXPipelineSettingsCustomization : public IPropertyTypeCustomization
+{
+public:
+
+	/**
+	 * Creates an instance of this class
+	 */
+	static TSharedRef<IPropertyTypeCustomization> MakeInstance();
+
+	virtual void CustomizeHeader(TSharedRef<IPropertyHandle> PropertyHandle, FDetailWidgetRow& HeaderRow, IPropertyTypeCustomizationUtils& CustomizationUtils) override;
+
+	virtual void CustomizeChildren(TSharedRef<IPropertyHandle> PropertyHandle, IDetailChildrenBuilder& ChildBuilder, IPropertyTypeCustomizationUtils& CustomizationUtils) override;
+
+protected:
+
+	bool OnShouldFilterAssetStandardSurface(const FAssetData & InAssetData);
+	bool OnShouldFilterAssetStandardSurfaceTransmission(const FAssetData& InAssetData);
+};
\ No newline at end of file
diff --git a/Engine/Plugins/Interchange/Runtime/Source/Pipelines/Private/InterchangeGenericMaterialPipeline.cpp b/Engine/Plugins/Interchange/Runtime/Source/Pipelines/Private/InterchangeGenericMaterialPipeline.cpp
--- a/Engine/Plugins/Interchange/Runtime/Source/Pipelines/Private/InterchangeGenericMaterialPipeline.cpp
+++ b/Engine/Plugins/Interchange/Runtime/Source/Pipelines/Private/InterchangeGenericMaterialPipeline.cpp
@@ -81,7 +81,7 @@ FMaterialXPipelineSettings::FMaterialXPipelineSettings()
 
 bool FMaterialXPipelineSettings::AreRequiredPackagesLoaded()
 {
-	auto ArePackagesLoaded = [this](const TMap<EInterchangeMaterialXShaders, FSoftObjectPath>& ObjectPaths)
+	auto ArePackagesLoaded = [&](const TMap<EInterchangeMaterialXShaders, FSoftObjectPath>& ObjectPaths)
 	{
 		bool bAllLoaded = true;
 
@@ -94,11 +94,26 @@ bool FMaterialXPipelineSettings::AreRequiredPackagesLoaded()
 				FString PackagePath = ObjectPath.GetLongPackageName();
 				if(FPackageName::DoesPackageExist(PackagePath))
 				{
-					if(!ObjectPath.TryLoad())
+					UObject* Asset = ObjectPath.TryLoad();
+					if(!Asset)
 					{
 						UE_LOG(LogInterchangePipeline, Warning, TEXT("Couldn't load %s"), *PackagePath);
 						bAllLoaded = false;
 					}
+#if WITH_EDITOR
+					else
+					{
+						if(Pair.Get<EInterchangeMaterialXShaders>() == EInterchangeMaterialXShaders::StandardSurface)
+						{
+							bAllLoaded = !ShouldFilterAssets(Cast<UMaterialFunction>(Asset), StandardSurfaceInputs, StandardSurfaceOutputs);
+						}
+						else if(Pair.Get<EInterchangeMaterialXShaders>() == EInterchangeMaterialXShaders::StandardSurfaceTransmission)
+						{
+							bAllLoaded = !ShouldFilterAssets(Cast<UMaterialFunction>(Asset), TransmissionSurfaceInputs, TransmissionSurfaceOutputs);
+						}
+
+					}
+#endif // WITH_EDITOR
 				}
 				else
 				{
@@ -125,6 +140,139 @@ FString FMaterialXPipelineSettings::GetAssetPathString(EInterchangeMaterialXShad
 	return AssetPath;
 }
 
+#if WITH_EDITOR
+TSet<FName> FMaterialXPipelineSettings::StandardSurfaceInputs
+{
+	UE::Interchange::Materials::StandardSurface::Parameters::Base,
+	UE::Interchange::Materials::StandardSurface::Parameters::BaseColor,
+	UE::Interchange::Materials::StandardSurface::Parameters::DiffuseRoughness,
+	UE::Interchange::Materials::StandardSurface::Parameters::Metalness,
+	UE::Interchange::Materials::StandardSurface::Parameters::Specular,
+	UE::Interchange::Materials::StandardSurface::Parameters::SpecularRoughness,
+	UE::Interchange::Materials::StandardSurface::Parameters::SpecularIOR,
+	UE::Interchange::Materials::StandardSurface::Parameters::SpecularAnisotropy,
+	UE::Interchange::Materials::StandardSurface::Parameters::SpecularRotation,
+	UE::Interchange::Materials::StandardSurface::Parameters::Subsurface,
+	UE::Interchange::Materials::StandardSurface::Parameters::SubsurfaceColor,
+	UE::Interchange::Materials::StandardSurface::Parameters::SubsurfaceRadius,
+	UE::Interchange::Materials::StandardSurface::Parameters::SubsurfaceScale,
+	UE::Interchange::Materials::StandardSurface::Parameters::Sheen,
+	UE::Interchange::Materials::StandardSurface::Parameters::SheenColor,
+	UE::Interchange::Materials::StandardSurface::Parameters::SheenRoughness,
+	UE::Interchange::Materials::StandardSurface::Parameters::Coat,
+	UE::Interchange::Materials::StandardSurface::Parameters::CoatColor,
+	UE::Interchange::Materials::StandardSurface::Parameters::CoatRoughness,
+	UE::Interchange::Materials::StandardSurface::Parameters::CoatNormal,
+	UE::Interchange::Materials::StandardSurface::Parameters::ThinFilmThickness,
+	UE::Interchange::Materials::StandardSurface::Parameters::Emission,
+	UE::Interchange::Materials::StandardSurface::Parameters::EmissionColor,
+	UE::Interchange::Materials::StandardSurface::Parameters::Normal,
+	UE::Interchange::Materials::StandardSurface::Parameters::Tangent
+};
+
+TSet<FName> FMaterialXPipelineSettings::StandardSurfaceOutputs
+{
+	TEXT("Base Color"), // MX_StandardSurface has BaseColor with a whitespace, this should be fixed in further release
+	UE::Interchange::Materials::PBR::Parameters::Metallic,
+	UE::Interchange::Materials::PBR::Parameters::Specular,
+	UE::Interchange::Materials::PBR::Parameters::Roughness,
+	UE::Interchange::Materials::PBR::Parameters::Anisotropy,
+	UE::Interchange::Materials::PBR::Parameters::EmissiveColor,
+	UE::Interchange::Materials::PBR::Parameters::Opacity,
+	UE::Interchange::Materials::PBR::Parameters::Normal,
+	UE::Interchange::Materials::PBR::Parameters::Tangent,
+	UE::Interchange::Materials::Sheen::Parameters::SheenRoughness,
+	UE::Interchange::Materials::Sheen::Parameters::SheenColor,
+	UE::Interchange::Materials::Subsurface::Parameters::SubsurfaceColor,
+	UE::Interchange::Materials::ClearCoat::Parameters::ClearCoat,
+	UE::Interchange::Materials::ClearCoat::Parameters::ClearCoatRoughness,
+	UE::Interchange::Materials::ClearCoat::Parameters::ClearCoatNormal
+};
+
+TSet<FName> FMaterialXPipelineSettings::TransmissionSurfaceInputs
+{
+	UE::Interchange::Materials::StandardSurface::Parameters::Base,
+	UE::Interchange::Materials::StandardSurface::Parameters::BaseColor,
+	UE::Interchange::Materials::StandardSurface::Parameters::DiffuseRoughness,
+	UE::Interchange::Materials::StandardSurface::Parameters::Metalness,
+	UE::Interchange::Materials::StandardSurface::Parameters::Specular,
+	UE::Interchange::Materials::StandardSurface::Parameters::SpecularRoughness,
+	UE::Interchange::Materials::StandardSurface::Parameters::SpecularIOR,
+	UE::Interchange::Materials::StandardSurface::Parameters::SpecularAnisotropy,
+	UE::Interchange::Materials::StandardSurface::Parameters::SpecularRotation,
+	UE::Interchange::Materials::StandardSurface::Parameters::Transmission,
+	UE::Interchange::Materials::StandardSurface::Parameters::TransmissionColor,
+	UE::Interchange::Materials::StandardSurface::Parameters::TransmissionDepth,
+	UE::Interchange::Materials::StandardSurface::Parameters::TransmissionScatter,
+	UE::Interchange::Materials::StandardSurface::Parameters::TransmissionScatterAnisotropy,
+	UE::Interchange::Materials::StandardSurface::Parameters::TransmissionDispersion,
+	UE::Interchange::Materials::StandardSurface::Parameters::TransmissionExtraRoughness,
+	UE::Interchange::Materials::StandardSurface::Parameters::Subsurface,
+	UE::Interchange::Materials::StandardSurface::Parameters::SubsurfaceColor,
+	UE::Interchange::Materials::StandardSurface::Parameters::SubsurfaceRadius,
+	UE::Interchange::Materials::StandardSurface::Parameters::SubsurfaceScale,
+	UE::Interchange::Materials::StandardSurface::Parameters::Sheen,
+	UE::Interchange::Materials::StandardSurface::Parameters::SheenColor,
+	UE::Interchange::Materials::StandardSurface::Parameters::SheenRoughness,
+	UE::Interchange::Materials::StandardSurface::Parameters::Coat,
+	UE::Interchange::Materials::StandardSurface::Parameters::CoatColor,
+	UE::Interchange::Materials::StandardSurface::Parameters::CoatRoughness,
+	UE::Interchange::Materials::StandardSurface::Parameters::CoatNormal,
+	UE::Interchange::Materials::StandardSurface::Parameters::ThinFilmThickness,
+	UE::Interchange::Materials::StandardSurface::Parameters::Emission,
+	UE::Interchange::Materials::StandardSurface::Parameters::EmissionColor,
+	UE::Interchange::Materials::StandardSurface::Parameters::Normal,
+};
+
+TSet<FName> FMaterialXPipelineSettings::TransmissionSurfaceOutputs
+{
+	UE::Interchange::Materials::PBR::Parameters::BaseColor,
+	UE::Interchange::Materials::PBR::Parameters::Metallic,
+	UE::Interchange::Materials::PBR::Parameters::Specular,
+	UE::Interchange::Materials::PBR::Parameters::Roughness,
+	UE::Interchange::Materials::PBR::Parameters::Anisotropy,
+	UE::Interchange::Materials::PBR::Parameters::EmissiveColor,
+	UE::Interchange::Materials::PBR::Parameters::Opacity,
+	UE::Interchange::Materials::PBR::Parameters::Normal,
+	UE::Interchange::Materials::PBR::Parameters::Tangent,
+	UE::Interchange::Materials::PBR::Parameters::Refraction,
+	UE::Interchange::Materials::ThinTranslucent::Parameters::TransmissionColor
+};
+
+
+bool FMaterialXPipelineSettings::ShouldFilterAssets(UMaterialFunction* Asset, const TSet<FName>& Inputs, const TSet<FName>& Outputs)
+{
+	int32 InputMatches = 0;
+	int32 OutputMatches = 0;
+
+	if(Asset != nullptr)
+	{
+		TArray<FFunctionExpressionInput> ExpressionInputs;
+		TArray<FFunctionExpressionOutput> ExpressionOutputs;
+		Asset->GetInputsAndOutputs(ExpressionInputs, ExpressionOutputs);
+
+		for(const FFunctionExpressionInput& ExpressionInput : ExpressionInputs)
+		{
+			if(Inputs.Find(ExpressionInput.Input.InputName))
+			{
+				InputMatches++;
+			}
+		}
+
+		for(const FFunctionExpressionOutput& ExpressionOutput : ExpressionOutputs)
+		{
+			if(Outputs.Find(ExpressionOutput.Output.OutputName))
+			{
+				OutputMatches++;
+			}
+		}
+	}
+
+	// we allow at least one input of the same name, but we should have exactly the same outputs
+	return !(InputMatches > 0 && OutputMatches == Outputs.Num());
+}
+#endif // WITH_EDITOR
+
 namespace UE::Interchange::InterchangeGenericMaterialPipeline::Private
 {
 	bool AreRequiredPackagesLoaded()
diff --git a/Engine/Plugins/Interchange/Runtime/Source/Pipelines/Public/InterchangeGenericMaterialPipeline.h b/Engine/Plugins/Interchange/Runtime/Source/Pipelines/Public/InterchangeGenericMaterialPipeline.h
--- a/Engine/Plugins/Interchange/Runtime/Source/Pipelines/Public/InterchangeGenericMaterialPipeline.h
+++ b/Engine/Plugins/Interchange/Runtime/Source/Pipelines/Public/InterchangeGenericMaterialPipeline.h
@@ -41,19 +41,32 @@ enum class EInterchangeMaterialXShaders : uint8
 };
 
 USTRUCT(BlueprintType)
-struct FMaterialXPipelineSettings
+struct INTERCHANGEPIPELINES_API FMaterialXPipelineSettings
 {
 	GENERATED_BODY()
 
 public:
 	FMaterialXPipelineSettings();
 
-	UPROPERTY(EditAnywhere, BlueprintReadWrite, Category = "MaterialX")
+	UPROPERTY(EditAnywhere, BlueprintReadWrite, Category = "MaterialX", meta=(AllowedClasses="/Script/Engine.MaterialFunction"))
 	TMap<EInterchangeMaterialXShaders, FSoftObjectPath> SurfaceShader;
 
 	bool AreRequiredPackagesLoaded();
 
 	FString GetAssetPathString(EInterchangeMaterialXShaders ShaderType) const;
+
+private:
+
+#if WITH_EDITOR
+	friend class FInterchangeMaterialXPipelineSettingsCustomization;
+
+	static bool ShouldFilterAssets(UMaterialFunction* Asset, const TSet<FName>& Inputs, const TSet<FName>& Outputs);
+
+	static TSet<FName> StandardSurfaceInputs;
+	static TSet<FName> StandardSurfaceOutputs;
+	static TSet<FName> TransmissionSurfaceInputs;
+	static TSet<FName> TransmissionSurfaceOutputs;
+#endif // WITH_EDITOR
 };
 
 UCLASS(BlueprintType, editinlinenew)
