Fixed localization issue

diff --git a/Engine/Plugins/Interchange/Runtime/Source/Import/Private/Gltf/InterchangeGltfTranslator.cpp b/Engine/Plugins/Interchange/Runtime/Source/Import/Private/Gltf/InterchangeGltfTranslator.cpp
--- a/Engine/Plugins/Interchange/Runtime/Source/Import/Private/Gltf/InterchangeGltfTranslator.cpp
+++ b/Engine/Plugins/Interchange/Runtime/Source/Import/Private/Gltf/InterchangeGltfTranslator.cpp
@@ -1030,7 +1030,7 @@ bool UInterchangeGltfTranslator::Translate( UInterchangeBaseNodeContainer& NodeC
 				{
 					Message->SourceAssetName = GetSourceData()->GetFilename();
 					Message->Text = FText::Format(
-						LOCTEXT("TextureCreationFailed", "The extension of the image file, {0}, for texture {1} is not supported."),
+						LOCTEXT("TextureCreationFailedFromFile", "The extension of the image file, {0}, for texture {1} is not supported."),
 						FText::FromString(GltfTexture.Source.FilePath), FText::FromString(GltfTexture.Name));
 				}
 
