[Interchange] Crash when importing particular glb model via content drawer import fix

diff --git a/Engine/Plugins/Interchange/Runtime/Source/Pipelines/Public/InterchangePipelineMeshesUtilities.h b/Engine/Plugins/Interchange/Runtime/Source/Pipelines/Public/InterchangePipelineMeshesUtilities.h
--- a/Engine/Plugins/Interchange/Runtime/Source/Pipelines/Public/InterchangePipelineMeshesUtilities.h
+++ b/Engine/Plugins/Interchange/Runtime/Source/Pipelines/Public/InterchangePipelineMeshesUtilities.h
@@ -320,11 +320,12 @@ namespace UE::Interchange::MeshesUtilities
 
 		TArray<FString> KeyReorder;
 		KeyReorder.Reserve(SlotMaterialDependencies.Num());
-		TArray<FString> MissingSuffixe;
-		MissingSuffixe.Reserve(SlotMaterialDependencies.Num());
+		TArray<FString> MissingSuffixMaterialNames;
+		MissingSuffixMaterialNames.Reserve(SlotMaterialDependencies.Num());
 		//Reorder material using the skinXX workflow
 		for (const TPair<FString, FString>& SlotMaterialDependency : SlotMaterialDependencies)
 		{
+			bool bHasSuffix = false;
 			const FString& MaterialName = SlotMaterialDependency.Key;
 			if (MaterialName.Len() > 6)
 			{
@@ -338,6 +339,7 @@ namespace UE::Interchange::MeshesUtilities
 						int32 TmpIndex = FPlatformString::Atoi(*SkinXXNumber);
 						if (TmpIndex >= 0)
 						{
+							bHasSuffix = true;
 							while (KeyReorder.Num() <= TmpIndex)
 							{
 								KeyReorder.AddDefaulted();
@@ -351,14 +353,11 @@ namespace UE::Interchange::MeshesUtilities
 						}
 					}
 				}
-				else
-				{
-					MissingSuffixe.Add(MaterialName);
-				}
 			}
-			else
+			
+			if (!bHasSuffix)
 			{
-				MissingSuffixe.Add(MaterialName);
+				MissingSuffixMaterialNames.Add(MaterialName);
 			}
 		}
 
@@ -376,8 +375,8 @@ namespace UE::Interchange::MeshesUtilities
 			const FString& SlotMaterialUid = SlotMaterialDependencies.FindChecked(MaterialName);
 			ReorderedSlotMaterialDependencies.Add(MaterialName, SlotMaterialUid);
 		}
-		//Add the missing suffixe material
-		for (const FString& MaterialName : MissingSuffixe)
+		//Add the missing suffix material
+		for (const FString& MaterialName : MissingSuffixMaterialNames)
 		{
 			const FString& SlotMaterialUid = SlotMaterialDependencies.FindChecked(MaterialName);
 			ReorderedSlotMaterialDependencies.Add(MaterialName, SlotMaterialUid);
