Fixing scale and orientation inconsistency between legacy and interchange OBJ.
Reverting back CL 18980183 without the scaling by 100 on the vertices because OBJ is unitless

diff --git a/Engine/Plugins/Interchange/Runtime/Source/Import/Private/Mesh/InterchangeOBJTranslator.cpp b/Engine/Plugins/Interchange/Runtime/Source/Import/Private/Mesh/InterchangeOBJTranslator.cpp
--- a/Engine/Plugins/Interchange/Runtime/Source/Import/Private/Mesh/InterchangeOBJTranslator.cpp
+++ b/Engine/Plugins/Interchange/Runtime/Source/Import/Private/Mesh/InterchangeOBJTranslator.cpp
@@ -139,16 +139,9 @@ private:
 };
 
 
-static FVector3f NormalToUEBasis(const FVector3f& InVector)
-{
-	return FVector3f(InVector.X, InVector.Z, InVector.Y);
-}
-
-
 static FVector3f PositionToUEBasis(const FVector3f& InVector)
 {
-	// Convert from meters to centimeters
-	return FVector3f(InVector.X, InVector.Z, InVector.Y) * 100.0f;
+	return FVector3f(InVector.X, -InVector.Y, InVector.Z);
 }
 
 
@@ -323,7 +316,7 @@ FMeshDescription FObjData::MakeMeshDescriptionForGroup(const FString& GroupName)
 
 			if (VertexData.NormalIndex != INDEX_NONE && Normals.IsValidIndex(VertexData.NormalIndex))
 			{
-				Attributes.GetVertexInstanceNormals()[VertexInstanceID] = NormalToUEBasis(Normals[VertexData.NormalIndex]);
+				Attributes.GetVertexInstanceNormals()[VertexInstanceID] = PositionToUEBasis(Normals[VertexData.NormalIndex]);
 			}
 
 			if (VertexData.UVIndex != INDEX_NONE && UVs.IsValidIndex(VertexData.UVIndex))
