[Interchange] Animation API consolidation.
FBX Rigid Animation support, FBX Multi Skeletal Animation support.

diff --git a/Engine/Plugins/Interchange/Runtime/Source/FactoryNodes/Private/InterchangeAnimSequenceFactoryNode.cpp b/Engine/Plugins/Interchange/Runtime/Source/FactoryNodes/Private/InterchangeAnimSequenceFactoryNode.cpp
--- a/Engine/Plugins/Interchange/Runtime/Source/FactoryNodes/Private/InterchangeAnimSequenceFactoryNode.cpp
+++ b/Engine/Plugins/Interchange/Runtime/Source/FactoryNodes/Private/InterchangeAnimSequenceFactoryNode.cpp
@@ -215,31 +215,6 @@ bool UInterchangeAnimSequenceFactoryNode::SetCustomSkeletonSoftObjectPath(const
 	IMPLEMENT_NODE_ATTRIBUTE_SETTER_NODELEGATE(SkeletonSoftObjectPath, FSoftObjectPath)
 }
 
-int32 UInterchangeAnimSequenceFactoryNode::GetAnimatedMorphTargetDependeciesCount() const
-{
-	return AnimatedMorphTargetDependencies.GetCount();
-}
-
-void UInterchangeAnimSequenceFactoryNode::GetAnimatedMorphTargetDependencies(TArray<FString>& OutDependencies) const
-{
-	AnimatedMorphTargetDependencies.GetItems(OutDependencies);
-}
-
-void UInterchangeAnimSequenceFactoryNode::GetAnimatedMorphTargetDependency(const int32 Index, FString& OutDependency) const
-{
-	AnimatedMorphTargetDependencies.GetItem(Index, OutDependency);
-}
-
-bool UInterchangeAnimSequenceFactoryNode::SetAnimatedMorphTargetDependencyUid(const FString& DependencyUid)
-{
-	return AnimatedMorphTargetDependencies.AddItem(DependencyUid);
-}
-
-bool UInterchangeAnimSequenceFactoryNode::RemoveAnimatedMorphTargetDependencyUid(const FString& DependencyUid)
-{
-	return AnimatedMorphTargetDependencies.RemoveItem(DependencyUid);
-}
-
 
 int32 UInterchangeAnimSequenceFactoryNode::GetAnimatedAttributeCurveNamesCount() const
 {
diff --git a/Engine/Plugins/Interchange/Runtime/Source/FactoryNodes/Public/InterchangeAnimSequenceFactoryNode.h b/Engine/Plugins/Interchange/Runtime/Source/FactoryNodes/Public/InterchangeAnimSequenceFactoryNode.h
--- a/Engine/Plugins/Interchange/Runtime/Source/FactoryNodes/Public/InterchangeAnimSequenceFactoryNode.h
+++ b/Engine/Plugins/Interchange/Runtime/Source/FactoryNodes/Public/InterchangeAnimSequenceFactoryNode.h
@@ -210,26 +210,6 @@ public:
 	UFUNCTION(BlueprintCallable, Category = "Interchange | Node | AnimSequence")
 	bool SetCustomDeleteExistingMorphTargetCurves(const bool& AttributeValue);
 
-	/** return how many animated moprh target (moprh target are translated into a mesh node and can be animated with a float curve) this anim sequence depends on. */
-	UFUNCTION(BlueprintCallable, Category = "Interchange | Node | AnimSequence")
-	int32 GetAnimatedMorphTargetDependeciesCount() const;
-
-	/** Get all animated moprh target unique id (moprh target are translated into a mesh node and can be animated with a float curve) this anim sequence depends on. */
-	UFUNCTION(BlueprintCallable, Category = "Interchange | Node | AnimSequence")
-	void GetAnimatedMorphTargetDependencies(TArray<FString>& OutDependencies) const;
-
-	/** Get an animated moprh target unique id (moprh target are translated into a mesh node and can be animated with a float curve) point by the specified index. */
-	UFUNCTION(BlueprintCallable, Category = "Interchange | Node | AnimSequence")
-	void GetAnimatedMorphTargetDependency(const int32 Index, FString& OutDependency) const;
-
-	/** Add an animated moprh target unique id (moprh target are translated into a mesh node and can be animated with a float curve). */
-	UFUNCTION(BlueprintCallable, Category = "Interchange | Node | AnimSequence")
-	bool SetAnimatedMorphTargetDependencyUid(const FString& DependencyUid);
-
-	/** Remove one animated moprh target unique id (moprh target are translated into a mesh node and can be animated with a float curve). Return false if we cannot remove the moprh target unique id. */
-	UFUNCTION(BlueprintCallable, Category = "Interchange | Node | AnimSequence")
-	bool RemoveAnimatedMorphTargetDependencyUid(const FString& DependencyUid);
-
 	/**
 	 * Morph target curve API End
 	 *********************************************************************************************************/
diff --git a/Engine/Plugins/Interchange/Runtime/Source/Import/Private/Animation/InterchangeAnimSequenceFactory.cpp b/Engine/Plugins/Interchange/Runtime/Source/Import/Private/Animation/InterchangeAnimSequenceFactory.cpp
--- a/Engine/Plugins/Interchange/Runtime/Source/Import/Private/Animation/InterchangeAnimSequenceFactory.cpp
+++ b/Engine/Plugins/Interchange/Runtime/Source/Import/Private/Animation/InterchangeAnimSequenceFactory.cpp
@@ -16,7 +16,6 @@
 #include "InterchangeSkeletonFactoryNode.h"
 #include "InterchangeSourceData.h"
 #include "InterchangeTranslatorBase.h"
-#include "Nodes/InterchangeAnimationAPI.h"
 #include "Nodes/InterchangeBaseNode.h"
 #include "Nodes/InterchangeBaseNodeContainer.h"
 #include "Nodes/InterchangeSourceNode.h"
@@ -298,28 +297,11 @@ namespace UE::Interchange::Private
 			TMap<FString, FString> PayloadKeys;
 			AnimSequenceFactoryNode->GetSceneNodeAnimationPayloadKeys(PayloadKeys);
 
-			if (PayloadKeys.Num() == 0)
+			for (const TTuple<FString, FString>& SceneNodeUidAndPayloadKey : PayloadKeys)
 			{
-				for (const FString& NodeUid : SkeletonNodes)
+				if (const UInterchangeSceneNode* SkeletonSceneNode = Cast<UInterchangeSceneNode>(NodeContainer->GetNode(SceneNodeUidAndPayloadKey.Key)))
 				{
-					if (const UInterchangeSceneNode* SkeletonSceneNode = Cast<UInterchangeSceneNode>(NodeContainer->GetNode(NodeUid)))
-					{
-						FString PayloadKey;
-						if (UInterchangeAnimationAPI::GetCustomNodeTransformPayloadKey(SkeletonSceneNode, PayloadKey))
-						{
-							AnimationPayloads.Add(SkeletonSceneNode, AnimSequenceTranslatorPayloadInterface->GetAnimationBakeTransformPayloadData(PayloadKey, SampleRate, RangeStart, RangeEnd));
-						}
-					}
-				}
-			}
-			else
-			{
-				for (const TTuple<FString, FString>& SceneNodeUidAndPayloadKey : PayloadKeys)
-				{
-					if (const UInterchangeSceneNode* SkeletonSceneNode = Cast<UInterchangeSceneNode>(NodeContainer->GetNode(SceneNodeUidAndPayloadKey.Key)))
-					{
-						AnimationPayloads.Add(SkeletonSceneNode, AnimSequenceTranslatorPayloadInterface->GetAnimationBakeTransformPayloadData(SceneNodeUidAndPayloadKey.Value, SampleRate, RangeStart, RangeEnd));
-					}
+					AnimationPayloads.Add(SkeletonSceneNode, AnimSequenceTranslatorPayloadInterface->GetAnimationBakeTransformPayloadData(SceneNodeUidAndPayloadKey.Value, SampleRate, RangeStart, RangeEnd));
 				}
 			}
 
@@ -518,38 +500,17 @@ namespace UE::Interchange::Private
 				TMap<FString, TFuture<TOptional<UE::Interchange::FAnimationCurvePayloadData>>> AnimationCurvesPayloads;
 				TMap<FString, FString> AnimationCurveMorphTargetNodeNames;
 
-				if (MorphTargetNodeAnimationPayloads.Num() > 0)
+				for (const TPair<FString, FString>& MorphTargetNodeUidAnimationPayload : MorphTargetNodeAnimationPayloads)
 				{
-					for (const TPair<FString, FString>& MorphTargetNodeUidAnimationPayload : MorphTargetNodeAnimationPayloads)
+					FString PayloadKey = MorphTargetNodeUidAnimationPayload.Value;
+					if (PayloadKey.Len() == 0)
 					{
-						FString PayloadKey = MorphTargetNodeUidAnimationPayload.Value;
-						if (PayloadKey.Len() == 0)
-						{
-							continue;
-						}
-						if (const UInterchangeMeshNode* MorphTargetNode = Cast<UInterchangeMeshNode>(NodeContainer->GetNode(MorphTargetNodeUidAnimationPayload.Key)))
-						{
-							AnimationCurvesPayloads.Add(PayloadKey, AnimSequenceTranslatorPayloadInterface->GetAnimationCurvePayloadData(PayloadKey));
-							AnimationCurveMorphTargetNodeNames.Add(PayloadKey, MorphTargetNode->GetDisplayLabel());
-						}
+						continue;
 					}
-				}
-				else
-				{
-					//Import morph target curves (FRichCurve is what anim api )
-					TArray<FString> AnimatedMorphTargetUids;
-					AnimSequenceFactoryNode->GetAnimatedMorphTargetDependencies(AnimatedMorphTargetUids);
-					for (const FString& MorphTargetUid : AnimatedMorphTargetUids)
+					if (const UInterchangeMeshNode* MorphTargetNode = Cast<UInterchangeMeshNode>(NodeContainer->GetNode(MorphTargetNodeUidAnimationPayload.Key)))
 					{
-						if (const UInterchangeMeshNode* MorphTargetNode = Cast<UInterchangeMeshNode>(NodeContainer->GetNode(MorphTargetUid)))
-						{
-							TOptional<FString> PayloadKey = MorphTargetNode->GetAnimationCurvePayLoadKey();
-							if (PayloadKey.IsSet())
-							{
-								AnimationCurvesPayloads.Add(PayloadKey.GetValue(), AnimSequenceTranslatorPayloadInterface->GetAnimationCurvePayloadData(PayloadKey.GetValue()));
-								AnimationCurveMorphTargetNodeNames.Add(PayloadKey.GetValue(), MorphTargetNode->GetDisplayLabel());
-							}
-						}
+						AnimationCurvesPayloads.Add(PayloadKey, AnimSequenceTranslatorPayloadInterface->GetAnimationCurvePayloadData(PayloadKey));
+						AnimationCurveMorphTargetNodeNames.Add(PayloadKey, MorphTargetNode->GetDisplayLabel());
 					}
 				}
 
diff --git a/Engine/Plugins/Interchange/Runtime/Source/Nodes/Private/InterchangeMeshNode.cpp b/Engine/Plugins/Interchange/Runtime/Source/Nodes/Private/InterchangeMeshNode.cpp
--- a/Engine/Plugins/Interchange/Runtime/Source/Nodes/Private/InterchangeMeshNode.cpp
+++ b/Engine/Plugins/Interchange/Runtime/Source/Nodes/Private/InterchangeMeshNode.cpp
@@ -16,12 +16,6 @@ namespace UE
 			return AttributeKey;
 		}
 
-		const FAttributeKey& FMeshNodeStaticData::PayloadAnimationCurveKey()
-		{
-			static FAttributeKey AttributeKey(TEXT("__PayloadAnimationCurveKey__"));
-			return AttributeKey;
-		}
-
 		const FAttributeKey& FMeshNodeStaticData::IsSkinnedMeshKey()
 		{
 			static FAttributeKey AttributeKey(TEXT("__IsSkinnedMeshKey__"));
@@ -83,10 +77,6 @@ FString UInterchangeMeshNode::GetKeyDisplayName(const UE::Interchange::FAttribut
 	{
 		return KeyDisplayName = TEXT("Payload Source Key");
 	}
-	else if (NodeAttributeKey == UE::Interchange::FMeshNodeStaticData::PayloadAnimationCurveKey())
-	{
-		return KeyDisplayName = TEXT("Payload Animation Curve Key");
-	}
 	else if (NodeAttributeKey == UE::Interchange::FMeshNodeStaticData::IsSkinnedMeshKey())
 	{
 		return KeyDisplayName = TEXT("Is a Skinned Mesh");
@@ -320,36 +310,6 @@ void UInterchangeMeshNode::SetPayLoadKey(const FString& PayloadKey)
 	}
 }
 
-const TOptional<FString> UInterchangeMeshNode::GetAnimationCurvePayLoadKey() const
-{
-	if (!Attributes->ContainAttribute(UE::Interchange::FMeshNodeStaticData::PayloadAnimationCurveKey()))
-	{
-		return TOptional<FString>();
-	}
-	UE::Interchange::FAttributeStorage::TAttributeHandle<FString> AttributeHandle = Attributes->GetAttributeHandle<FString>(UE::Interchange::FMeshNodeStaticData::PayloadAnimationCurveKey());
-	if (!AttributeHandle.IsValid())
-	{
-		return TOptional<FString>();
-	}
-	FString PayloadKey;
-	UE::Interchange::EAttributeStorageResult Result = AttributeHandle.Get(PayloadKey);
-	if (!IsAttributeStorageResultSuccess(Result))
-	{
-		LogAttributeStorageErrors(Result, TEXT("UInterchangeMeshNode.GetAnimationCurvePayLoadKey"), UE::Interchange::FMeshNodeStaticData::PayloadAnimationCurveKey());
-		return TOptional<FString>();
-	}
-	return TOptional<FString>(PayloadKey);
-}
-
-void UInterchangeMeshNode::SetAnimationCurvePayLoadKey(const FString& PayloadKey)
-{
-	UE::Interchange::EAttributeStorageResult Result = Attributes->RegisterAttribute(UE::Interchange::FMeshNodeStaticData::PayloadAnimationCurveKey(), PayloadKey);
-	if (!IsAttributeStorageResultSuccess(Result))
-	{
-		LogAttributeStorageErrors(Result, TEXT("UInterchangeMeshNode.SetAnimationCurvePayLoadKey"), UE::Interchange::FMeshNodeStaticData::PayloadAnimationCurveKey());
-	}
-}
-
 bool UInterchangeMeshNode::GetCustomVertexCount(int32& AttributeValue) const
 {
 	IMPLEMENT_NODE_ATTRIBUTE_GETTER(VertexCount, int32);
diff --git a/Engine/Plugins/Interchange/Runtime/Source/Nodes/Public/InterchangeMeshNode.h b/Engine/Plugins/Interchange/Runtime/Source/Nodes/Public/InterchangeMeshNode.h
--- a/Engine/Plugins/Interchange/Runtime/Source/Nodes/Public/InterchangeMeshNode.h
+++ b/Engine/Plugins/Interchange/Runtime/Source/Nodes/Public/InterchangeMeshNode.h
@@ -16,7 +16,6 @@ namespace UE
 		struct INTERCHANGENODES_API FMeshNodeStaticData : public FBaseNodeStaticData
 		{
 			static const FAttributeKey& PayloadSourceKey();
-			static const FAttributeKey& PayloadAnimationCurveKey();
 			static const FAttributeKey& IsSkinnedMeshKey();
 			static const FAttributeKey& IsMorphTargetKey();
 			static const FAttributeKey& MorphTargetNameKey();
@@ -106,11 +105,6 @@ public:
 
 	UFUNCTION(BlueprintCallable, Category = "Interchange | Node | Mesh")
 	virtual void SetPayLoadKey(const FString& PayloadKey);
-
-	virtual const TOptional<FString> GetAnimationCurvePayLoadKey() const;
-
-	UFUNCTION(BlueprintCallable, Category = "Interchange | Node | Mesh")
-	virtual void SetAnimationCurvePayLoadKey(const FString& PayloadKey);
 	
 	/** Query this mesh vertices count. Return false if the attribute was not set.*/
 	UFUNCTION(BlueprintCallable, Category = "Interchange | Node | Mesh")
diff --git a/Engine/Plugins/Interchange/Runtime/Source/Nodes/Public/InterchangeSceneNode.h b/Engine/Plugins/Interchange/Runtime/Source/Nodes/Public/InterchangeSceneNode.h
--- a/Engine/Plugins/Interchange/Runtime/Source/Nodes/Public/InterchangeSceneNode.h
+++ b/Engine/Plugins/Interchange/Runtime/Source/Nodes/Public/InterchangeSceneNode.h
@@ -31,7 +31,6 @@ namespace UE
 
 /**
  * The scene node represent a transform node in the scene
- * Scene node can have animations: Use UInterchangeAnimationAPI to get\set animation datas
  * Scene node can have user defined attribute. Use UInterchangeUserDefinedAttributesAPI to get\set user define attribute data
  */
 UCLASS(BlueprintType, Experimental)
diff --git a/Engine/Plugins/Interchange/Runtime/Source/Parsers/Fbx/Private/FbxAPI.cpp b/Engine/Plugins/Interchange/Runtime/Source/Parsers/Fbx/Private/FbxAPI.cpp
--- a/Engine/Plugins/Interchange/Runtime/Source/Parsers/Fbx/Private/FbxAPI.cpp
+++ b/Engine/Plugins/Interchange/Runtime/Source/Parsers/Fbx/Private/FbxAPI.cpp
@@ -129,6 +129,8 @@ namespace UE
 
 				FFbxScene FbxScene(*this);
 				FbxScene.AddHierarchy(SDKScene, NodeContainer, PayloadContexts);
+				FbxScene.AddAnimation(SDKScene, NodeContainer, PayloadContexts);
+				FbxScene.AddMorphTargetAnimations(SDKScene, NodeContainer, PayloadContexts, FbxMesh.GetMorphTargetAnimationsBuildingData());
 			}
 
 			bool FFbxParser::FetchPayloadData(const FString& PayloadKey, const FString& PayloadFilepath)
diff --git a/Engine/Plugins/Interchange/Runtime/Source/Parsers/Fbx/Private/FbxAPI.h b/Engine/Plugins/Interchange/Runtime/Source/Parsers/Fbx/Private/FbxAPI.h
--- a/Engine/Plugins/Interchange/Runtime/Source/Parsers/Fbx/Private/FbxAPI.h
+++ b/Engine/Plugins/Interchange/Runtime/Source/Parsers/Fbx/Private/FbxAPI.h
@@ -73,6 +73,8 @@ namespace UE
 					Item->SourceAssetName = SourceFilename;
 				}
 
+				FbxScene* GetSDKScene() { return SDKScene; }
+
 				/**
 				 * Critical section to avoid getting multiple payload in same time.
 				 * The FBX evaluator use a cache mechanism for evaluating global transform that is not thread safe.
diff --git a/Engine/Plugins/Interchange/Runtime/Source/Parsers/Fbx/Private/FbxAnimation.cpp b/Engine/Plugins/Interchange/Runtime/Source/Parsers/Fbx/Private/FbxAnimation.cpp
--- a/Engine/Plugins/Interchange/Runtime/Source/Parsers/Fbx/Private/FbxAnimation.cpp
+++ b/Engine/Plugins/Interchange/Runtime/Source/Parsers/Fbx/Private/FbxAnimation.cpp
@@ -7,6 +7,7 @@
 #include "FbxConvert.h"
 #include "FbxHelper.h"
 #include "FbxInclude.h"
+#include "FbxMesh.h"
 #include "Fbx/InterchangeFbxMessages.h"
 #include "InterchangeCommonAnimationPayload.h"
 #include "InterchangeMeshNode.h"
@@ -15,11 +16,11 @@
 #include "MeshDescription.h"
 #include "Misc/FileHelper.h"
 #include "Misc/ScopeLock.h"
-#include "Nodes/InterchangeAnimationAPI.h"
 #include "Serialization/LargeMemoryWriter.h"
 #include "SkeletalMeshAttributes.h"
 #include "StaticMeshAttributes.h"
-
+#include "InterchangeAnimationTrackSetNode.h"
+#include "Nodes/InterchangeBaseNodeContainer.h"
 
 #define LOCTEXT_NAMESPACE "InterchangeFbxMesh"
 
@@ -383,7 +384,7 @@ namespace UE::Interchange::Private
 		DestinationCurve.StringKeyValues = StepCurveValues;
 	}
 
-	bool ImportBakeTransforms(FbxNode* Node, FAnimationBakeTransformPayloadData& AnimationBakeTransformPayloadData)
+	bool ImportBakeTransforms(FNodeTransformFetchPayloadData& FetchPayloadData, FAnimationBakeTransformPayloadData& AnimationBakeTransformPayloadData, FbxScene* SDKScene)
 	{
 		if (!ensure(!FMath::IsNearlyZero(AnimationBakeTransformPayloadData.BakeFrequency)))
 		{
@@ -410,15 +411,18 @@ namespace UE::Interchange::Private
 		const FbxTime TimeComparisonThreshold = (UE_DOUBLE_KINDA_SMALL_NUMBER * static_cast<double>(FBXSDK_TC_SECOND));
 		AnimationBakeTransformPayloadData.Transforms.Empty(NumFrame);
 
+		SDKScene->SetCurrentAnimationStack(FetchPayloadData.CurrentAnimStack);
+
 		for (FbxTime CurTime = StartTime; CurTime < (EndTime + TimeComparisonThreshold); CurTime += TimeStep)
 		{
-			FbxAMatrix NodeTransform = Node->EvaluateGlobalTransform(CurTime);
-			FbxNode* ParentNode = Node->GetParent();
+			FbxAMatrix NodeTransform = FetchPayloadData.Node->EvaluateGlobalTransform(CurTime);
+			FbxNode* ParentNode = FetchPayloadData.Node->GetParent();
 			if (ParentNode)
 			{
 				FbxAMatrix ParentTransform = ParentNode->EvaluateGlobalTransform(CurTime);
 				NodeTransform = ParentTransform.Inverse() * NodeTransform;
 			}
+
 			AnimationBakeTransformPayloadData.Transforms.Add(UE::Interchange::Private::FFbxConvert::ConvertTransform(NodeTransform));
 		}
 
@@ -437,16 +441,11 @@ namespace UE::Interchange::Private
 
 		FbxScene* SDKScene = nullptr;
 		FbxNode* Node = nullptr;
-		TArray<FbxAnimCurve*> TransformChannelCurves;
-		int32 TransformChannelCount = 0;
 		bool IsNodeAnimated = false;
-		FbxTimeSpan ExportedTimeSpan = FbxTimeSpan(FBXSDK_TIME_INFINITE, FBXSDK_TIME_MINUS_INFINITE);
-		FbxTime StartTime = FBXSDK_TIME_INFINITE;
-		FbxTime EndTime = FBXSDK_TIME_MINUS_INFINITE;
-		int32 KeyCount = 0;
+		FbxAnimStack* CurrentAnimStack = nullptr;
 	};
 
-	void GetFbxTransformCurves(FGetFbxTransformCurvesParameters& Parameters)
+	void GetFbxTransformCurves(FGetFbxTransformCurvesParameters& Parameters, const int32& AnimationIndex)
 	{
 		if (!ensure(Parameters.SDKScene) || !ensure(Parameters.Node))
 		{
@@ -462,74 +461,47 @@ namespace UE::Interchange::Private
 		//Scale X
 		//Scale Y
 		//Scale Z
-		int32 NumAnimations = Parameters.SDKScene->GetSrcObjectCount<FbxAnimStack>();
-		//Anim stack should be merge so we expect to have only one stack here
-		ensure(NumAnimations <= 1);
 
 		Parameters.IsNodeAnimated = false;
 
-		Parameters.TransformChannelCount = 9;
-		Parameters.TransformChannelCurves.Reset(Parameters.TransformChannelCount);
-		for (int32 AnimationIndex = 0; AnimationIndex < NumAnimations && !Parameters.IsNodeAnimated; AnimationIndex++)
+		int32 NumAnimations = Parameters.SDKScene->GetSrcObjectCount<FbxAnimStack>();
+		if (AnimationIndex >= NumAnimations)
 		{
-			FbxAnimStack* AnimStack = (FbxAnimStack*)Parameters.SDKScene->GetSrcObject<FbxAnimStack>(AnimationIndex);
-			int32 NumLayers = AnimStack->GetMemberCount();
-			for (int LayerIndex = 0; LayerIndex < NumLayers && !Parameters.IsNodeAnimated; LayerIndex++)
-			{
-				FbxAnimLayer* AnimLayer = (FbxAnimLayer*)AnimStack->GetMember(LayerIndex);
-				// Display curves specific to properties
-				Parameters.TransformChannelCurves.Add(Parameters.Node->LclTranslation.GetCurve(AnimLayer, FBXSDK_CURVENODE_COMPONENT_X, false));
-				Parameters.TransformChannelCurves.Add(Parameters.Node->LclTranslation.GetCurve(AnimLayer, FBXSDK_CURVENODE_COMPONENT_Y, false));
-				Parameters.TransformChannelCurves.Add(Parameters.Node->LclTranslation.GetCurve(AnimLayer, FBXSDK_CURVENODE_COMPONENT_Z, false));
+			return;
+		}
+
+		int32 NumberOfChannels = 9;
+
+		Parameters.CurrentAnimStack = (FbxAnimStack*)Parameters.SDKScene->GetSrcObject<FbxAnimStack>(AnimationIndex);
+		
+		int32 NumLayers = Parameters.CurrentAnimStack->GetMemberCount();
+		for (int LayerIndex = 0; LayerIndex < NumLayers && !Parameters.IsNodeAnimated; LayerIndex++)
+		{
+			FbxAnimLayer* AnimLayer = (FbxAnimLayer*)Parameters.CurrentAnimStack->GetMember(LayerIndex);
+
+			TArray<FbxAnimCurve*> TransformChannelCurves;
+			TransformChannelCurves.Reserve(NumberOfChannels);
+
+			// Display curves specific to properties
+			TransformChannelCurves.Add(Parameters.Node->LclTranslation.GetCurve(AnimLayer, FBXSDK_CURVENODE_COMPONENT_X, false));
+			TransformChannelCurves.Add(Parameters.Node->LclTranslation.GetCurve(AnimLayer, FBXSDK_CURVENODE_COMPONENT_Y, false));
+			TransformChannelCurves.Add(Parameters.Node->LclTranslation.GetCurve(AnimLayer, FBXSDK_CURVENODE_COMPONENT_Z, false));
 
-				Parameters.TransformChannelCurves.Add(Parameters.Node->LclRotation.GetCurve(AnimLayer, FBXSDK_CURVENODE_COMPONENT_X, false));
-				Parameters.TransformChannelCurves.Add(Parameters.Node->LclRotation.GetCurve(AnimLayer, FBXSDK_CURVENODE_COMPONENT_Y, false));
-				Parameters.TransformChannelCurves.Add(Parameters.Node->LclRotation.GetCurve(AnimLayer, FBXSDK_CURVENODE_COMPONENT_Z, false));
+			TransformChannelCurves.Add(Parameters.Node->LclRotation.GetCurve(AnimLayer, FBXSDK_CURVENODE_COMPONENT_X, false));
+			TransformChannelCurves.Add(Parameters.Node->LclRotation.GetCurve(AnimLayer, FBXSDK_CURVENODE_COMPONENT_Y, false));
+			TransformChannelCurves.Add(Parameters.Node->LclRotation.GetCurve(AnimLayer, FBXSDK_CURVENODE_COMPONENT_Z, false));
 
-				Parameters.TransformChannelCurves.Add(Parameters.Node->LclScaling.GetCurve(AnimLayer, FBXSDK_CURVENODE_COMPONENT_X, false));
-				Parameters.TransformChannelCurves.Add(Parameters.Node->LclScaling.GetCurve(AnimLayer, FBXSDK_CURVENODE_COMPONENT_Y, false));
-				Parameters.TransformChannelCurves.Add(Parameters.Node->LclScaling.GetCurve(AnimLayer, FBXSDK_CURVENODE_COMPONENT_Z, false));
+			TransformChannelCurves.Add(Parameters.Node->LclScaling.GetCurve(AnimLayer, FBXSDK_CURVENODE_COMPONENT_X, false));
+			TransformChannelCurves.Add(Parameters.Node->LclScaling.GetCurve(AnimLayer, FBXSDK_CURVENODE_COMPONENT_Y, false));
+			TransformChannelCurves.Add(Parameters.Node->LclScaling.GetCurve(AnimLayer, FBXSDK_CURVENODE_COMPONENT_Z, false));
 
-				if (!Parameters.IsNodeAnimated)
+			if (!Parameters.IsNodeAnimated)
+			{
+				for (int32 ChannelIndex = 0; ChannelIndex < NumberOfChannels; ++ChannelIndex)
 				{
-					for (int32 ChannelIndex = 0; ChannelIndex < Parameters.TransformChannelCount; ++ChannelIndex)
+					if (TransformChannelCurves[ChannelIndex])
 					{
-						if (Parameters.TransformChannelCurves[ChannelIndex])
-						{
-							Parameters.IsNodeAnimated = true;
-							const int32 LocalKeycount = Parameters.TransformChannelCurves[ChannelIndex]->KeyGetCount();
-							if (LocalKeycount > Parameters.KeyCount)
-							{
-								Parameters.KeyCount = LocalKeycount;
-							}
-							if (LocalKeycount > 1)
-							{
-								FbxTimeSpan LocalAnimatedTimeSpan(FBXSDK_TIME_INFINITE, FBXSDK_TIME_MINUS_INFINITE);
-								Parameters.TransformChannelCurves[ChannelIndex]->GetTimeInterval(LocalAnimatedTimeSpan);
-								if (Parameters.StartTime > LocalAnimatedTimeSpan.GetStart())
-								{
-									Parameters.StartTime = LocalAnimatedTimeSpan.GetStart();
-								}
-								if (Parameters.EndTime < LocalAnimatedTimeSpan.GetStop())
-								{
-									Parameters.EndTime = LocalAnimatedTimeSpan.GetStop();
-								}
-							}
-							else if (Parameters.KeyCount == 1)
-							{
-								//When there is only one key there is no interval
-								FbxAnimCurveKey FbxKey = Parameters.TransformChannelCurves[ChannelIndex]->KeyGet(0);
-								FbxTime KeyTime = FbxKey.GetTime();
-								if (Parameters.StartTime > KeyTime)
-								{
-									Parameters.StartTime = KeyTime;
-								}
-								if (Parameters.EndTime < KeyTime)
-								{
-									Parameters.EndTime = KeyTime;
-								}
-							}
-						}
+						Parameters.IsNodeAnimated = true;
 					}
 				}
 			}
@@ -538,7 +510,7 @@ namespace UE::Interchange::Private
 
 	bool FAnimationPayloadContextTransform::FetchPayloadToFile(FFbxParser& Parser, const FString& PayloadFilepath)
 	{
-		if (AttributeFetchPayloadData.IsSet())
+		if (AttributeFetchPayloadData.IsSet() || AttributeNodeTransformFetchPayloadData.IsSet())
 		{
 			return InternalFetchCurveNodePayloadToFile(Parser, PayloadFilepath);
 		}
@@ -573,7 +545,7 @@ namespace UE::Interchange::Private
 		AnimationBakeTransformPayloadData.RangeStartTime = RangeStartTime;
 		AnimationBakeTransformPayloadData.RangeEndTime = RangeEndTime;
 
-		ImportBakeTransforms(FetchPayloadData.Node, AnimationBakeTransformPayloadData);
+		ImportBakeTransforms(FetchPayloadData, AnimationBakeTransformPayloadData, Parser.GetSDKScene());
 		{
 			FLargeMemoryWriter Ar;
 			AnimationBakeTransformPayloadData.Serialize(Ar);
@@ -587,37 +559,38 @@ namespace UE::Interchange::Private
 
 	bool FAnimationPayloadContextTransform::InternalFetchCurveNodePayloadToFile(FFbxParser& Parser, const FString& PayloadFilepath)
 	{
-		check(AttributeFetchPayloadData.IsSet());
-		FAttributeFetchPayloadData& FetchPayloadData = AttributeFetchPayloadData.GetValue();
-		if (!ensure(FetchPayloadData.Node != nullptr))
+		if (AttributeFetchPayloadData.IsSet())
 		{
-			UInterchangeResultError_Generic* Message = Parser.AddMessage<UInterchangeResultError_Generic>();
-			Message->InterchangeKey = Parser.GetFbxHelper()->GetFbxNodeHierarchyName(FetchPayloadData.Node);
-			Message->Text = LOCTEXT("InternalFetchCurveNodePayloadToFile_FBXNodeNull", "Cannot fetch FBX animation curve payload because the FBX node is null.");
-			return false;
-		}
+			FAttributeFetchPayloadData& FetchPayloadData = AttributeFetchPayloadData.GetValue();
+			if (!ensure(FetchPayloadData.Node != nullptr))
+			{
+				UInterchangeResultError_Generic* Message = Parser.AddMessage<UInterchangeResultError_Generic>();
+				Message->InterchangeKey = Parser.GetFbxHelper()->GetFbxNodeHierarchyName(FetchPayloadData.Node);
+				Message->Text = LOCTEXT("InternalFetchCurveNodePayloadToFile_FBXNodeNull", "Cannot fetch FBX animation curve payload because the FBX node is null.");
+				return false;
+			}
 
-		if (!ensure(FetchPayloadData.AnimCurves != nullptr))
-		{
-			UInterchangeResultError_Generic* Message = Parser.AddMessage<UInterchangeResultError_Generic>();
-			Message->InterchangeKey = Parser.GetFbxHelper()->GetFbxNodeHierarchyName(FetchPayloadData.Node);
-			Message->Text = LOCTEXT("InternalFetchCurveNodePayloadToFile_FBXCurveNull", "Cannot fetch FBX user attribute animation curve payload because the FBX anim curve node is null.");
-			return false;
-		}
-		if (FetchPayloadData.bAttributeTypeIsStepCurveAnimation)
-		{
-			//Fetch TArray<FInterchangeStepCurve> step curve, no interpolation
-			TArray<FInterchangeStepCurve> InterchangeStepCurves;
-			const uint32 ChannelCount = FetchPayloadData.AnimCurves->GetChannelsCount();
-			for (uint32 ChannelIndex = 0; ChannelIndex < ChannelCount; ++ChannelIndex)
+			if (!ensure(FetchPayloadData.AnimCurves != nullptr))
+			{
+				UInterchangeResultError_Generic* Message = Parser.AddMessage<UInterchangeResultError_Generic>();
+				Message->InterchangeKey = Parser.GetFbxHelper()->GetFbxNodeHierarchyName(FetchPayloadData.Node);
+				Message->Text = LOCTEXT("InternalFetchCurveNodePayloadToFile_FBXCurveNull", "Cannot fetch FBX user attribute animation curve payload because the FBX anim curve node is null.");
+				return false;
+			}
+			if (FetchPayloadData.bAttributeTypeIsStepCurveAnimation)
 			{
-				const uint32 ChannelCurveCount = FetchPayloadData.AnimCurves->GetCurveCount(ChannelIndex);
-				for (uint32 CurveIndex = 0; CurveIndex < ChannelCurveCount; ++CurveIndex)
+				//Fetch TArray<FInterchangeStepCurve> step curve, no interpolation
+				TArray<FInterchangeStepCurve> InterchangeStepCurves;
+				const uint32 ChannelCount = FetchPayloadData.AnimCurves->GetChannelsCount();
+				for (uint32 ChannelIndex = 0; ChannelIndex < ChannelCount; ++ChannelIndex)
 				{
-					if (FbxAnimCurve* CurrentAnimCurve = FetchPayloadData.AnimCurves->GetCurve(ChannelIndex, CurveIndex))
+					const uint32 ChannelCurveCount = FetchPayloadData.AnimCurves->GetCurveCount(ChannelIndex);
+					for (uint32 CurveIndex = 0; CurveIndex < ChannelCurveCount; ++CurveIndex)
 					{
-						switch (FetchPayloadData.PropertyType)
+						if (FbxAnimCurve* CurrentAnimCurve = FetchPayloadData.AnimCurves->GetCurve(ChannelIndex, CurveIndex))
 						{
+							switch (FetchPayloadData.PropertyType)
+							{
 							case EFbxType::eFbxBool:
 							case EFbxType::eFbxChar:
 							case EFbxType::eFbxUChar:
@@ -644,35 +617,84 @@ namespace UE::Interchange::Private
 							case EFbxType::eFbxString:
 								ImportStringStepCurve(CurrentAnimCurve, FetchPayloadData.Property, InterchangeStepCurves.AddDefaulted_GetRef());
 								break;
+							}
 						}
 					}
 				}
+				{
+					FLargeMemoryWriter Ar;
+					Ar << InterchangeStepCurves;
+					uint8* ArchiveData = Ar.GetData();
+					int64 ArchiveSize = Ar.TotalSize();
+					TArray64<uint8> Buffer(ArchiveData, ArchiveSize);
+					FFileHelper::SaveArrayToFile(Buffer, *PayloadFilepath);
+				}
 			}
+			else
 			{
-				FLargeMemoryWriter Ar;
-				Ar << InterchangeStepCurves;
-				uint8* ArchiveData = Ar.GetData();
-				int64 ArchiveSize = Ar.TotalSize();
-				TArray64<uint8> Buffer(ArchiveData, ArchiveSize);
-				FFileHelper::SaveArrayToFile(Buffer, *PayloadFilepath);
+				//Fetch TArray<FInterchangeCurve> which are float curves with interpolation
+				TArray<FInterchangeCurve> InterchangeCurves;
+				const uint32 ChannelCount = FetchPayloadData.AnimCurves->GetChannelsCount();
+				for (uint32 ChannelIndex = 0; ChannelIndex < ChannelCount; ++ChannelIndex)
+				{
+					const uint32 ChannelCurveCount = FetchPayloadData.AnimCurves->GetCurveCount(ChannelIndex);
+					for (uint32 CurveIndex = 0; CurveIndex < ChannelCurveCount; ++CurveIndex)
+					{
+						if (FbxAnimCurve* CurrentAnimCurve = FetchPayloadData.AnimCurves->GetCurve(ChannelIndex, CurveIndex))
+						{
+							ImportCurve(CurrentAnimCurve, 1.0f, InterchangeCurves.AddDefaulted_GetRef().Keys);
+						}
+					}
+				}
+				{
+					FLargeMemoryWriter Ar;
+					Ar << InterchangeCurves;
+					uint8* ArchiveData = Ar.GetData();
+					int64 ArchiveSize = Ar.TotalSize();
+					TArray64<uint8> Buffer(ArchiveData, ArchiveSize);
+					FFileHelper::SaveArrayToFile(Buffer, *PayloadFilepath);
+				}
 			}
+			return true;
 		}
-		else
+		
+		if (AttributeNodeTransformFetchPayloadData.IsSet())
 		{
+			FAttributeNodeTransformFetchPayloadData& FetchPayloadData = AttributeNodeTransformFetchPayloadData.GetValue();
 			//Fetch TArray<FInterchangeCurve> which are float curves with interpolation
 			TArray<FInterchangeCurve> InterchangeCurves;
-			const uint32 ChannelCount = FetchPayloadData.AnimCurves->GetChannelsCount();
-			for (uint32 ChannelIndex = 0; ChannelIndex < ChannelCount; ++ChannelIndex)
+
+			auto ParseInterchangeCurveCurve = [&InterchangeCurves](FbxAnimCurveNode* AnimCurveNode, const char* pChannelName, float Scale = 1.0f)
 			{
-				const uint32 ChannelCurveCount = FetchPayloadData.AnimCurves->GetCurveCount(ChannelIndex);
-				for (uint32 CurveIndex = 0; CurveIndex < ChannelCurveCount; ++CurveIndex)
+				bool bInterchangeCurveAdded = false;
+				if (AnimCurveNode)
 				{
-					if (FbxAnimCurve* CurrentAnimCurve = FetchPayloadData.AnimCurves->GetCurve(ChannelIndex, CurveIndex))
+					int ChannelIndex = AnimCurveNode->GetChannelIndex(pChannelName);
+					if (ChannelIndex != -1)
 					{
-						ImportCurve(CurrentAnimCurve, 1.0f, InterchangeCurves.AddDefaulted_GetRef().Keys);
+						if (FbxAnimCurve* CurrentAnimCurve = AnimCurveNode->GetCurve(ChannelIndex))
+						{
+							ImportCurve(CurrentAnimCurve, Scale, InterchangeCurves.AddDefaulted_GetRef().Keys);
+							bInterchangeCurveAdded = true;
+						}
 					}
 				}
-			}
+				
+				if (!bInterchangeCurveAdded) InterchangeCurves.AddDefaulted_GetRef();
+			};
+
+			ParseInterchangeCurveCurve(FetchPayloadData.TranlsationCurveNode, FBXSDK_CURVENODE_COMPONENT_X);
+			ParseInterchangeCurveCurve(FetchPayloadData.TranlsationCurveNode, FBXSDK_CURVENODE_COMPONENT_Y, -1.0f); //FBX to UE conversion with scale -1.0
+			ParseInterchangeCurveCurve(FetchPayloadData.TranlsationCurveNode, FBXSDK_CURVENODE_COMPONENT_Z);
+
+			ParseInterchangeCurveCurve(FetchPayloadData.RotationCurveNode, FBXSDK_CURVENODE_COMPONENT_X);
+			ParseInterchangeCurveCurve(FetchPayloadData.RotationCurveNode, FBXSDK_CURVENODE_COMPONENT_Y, -1.0f); //FBX to UE conversion with scale -1.0
+			ParseInterchangeCurveCurve(FetchPayloadData.RotationCurveNode, FBXSDK_CURVENODE_COMPONENT_Z, -1.0f); //FBX to UE conversion with scale -1.0
+
+			ParseInterchangeCurveCurve(FetchPayloadData.ScaleCurveNode, FBXSDK_CURVENODE_COMPONENT_X);
+			ParseInterchangeCurveCurve(FetchPayloadData.ScaleCurveNode, FBXSDK_CURVENODE_COMPONENT_Y);
+			ParseInterchangeCurveCurve(FetchPayloadData.ScaleCurveNode, FBXSDK_CURVENODE_COMPONENT_Z);
+
 			{
 				FLargeMemoryWriter Ar;
 				Ar << InterchangeCurves;
@@ -681,51 +703,34 @@ namespace UE::Interchange::Private
 				TArray64<uint8> Buffer(ArchiveData, ArchiveSize);
 				FFileHelper::SaveArrayToFile(Buffer, *PayloadFilepath);
 			}
+
+			return true;
 		}
-		return true;
+		return false;
 	}
 
 	bool FAnimationPayloadContextTransform::InternalFetchMorphTargetCurvePayloadToFile(FFbxParser& Parser, const FString& PayloadFilepath)
 	{
 		check(MorphTargetFetchPayloadData.IsSet());
 		FMorphTargetFetchPayloadData& FetchPayloadData = MorphTargetFetchPayloadData.GetValue();
-
-		if(!FetchPayloadData.SDKScene)
+	
+		if (!FetchPayloadData.SDKScene)
 		{
 			UInterchangeResultError_Generic* Message = Parser.AddMessage<UInterchangeResultError_Generic>();
 			Message->Text = LOCTEXT("InternalFetchMorphTargetCurvePayloadToFile_FBXSDKSceneNull", "InternalFetchMorphTargetCurvePayloadToFile, fbx sdk is nullptr.");
 			return false;
 		}
-		
-		int32 GeometryCount = FetchPayloadData.SDKScene->GetGeometryCount();
 
 		FbxGeometry* Geometry = FetchPayloadData.SDKScene->GetGeometry(FetchPayloadData.GeometryIndex);
-		if (!Geometry || Geometry->GetAttributeType() != FbxNodeAttribute::eMesh)
-		{
-			UInterchangeResultError_Generic* Message = Parser.AddMessage<UInterchangeResultError_Generic>();
-			Message->Text = LOCTEXT("InternalFetchMorphTargetCurvePayloadToFile_FBXGeometryNull", "Cannot fetch FBX geometry from the scene.");
-			return false;
-		}
-
-		FbxAnimCurve* AnimCurve = nullptr;
-		int32 NumAnimations = FetchPayloadData.SDKScene->GetSrcObjectCount<FbxAnimStack>();
-		//Anim stack should be merge so we expect to have only one stack here
-		ensure(NumAnimations <= 1);
-		if (FbxAnimStack* AnimStack = NumAnimations == 1 ? (FbxAnimStack*)FetchPayloadData.SDKScene->GetSrcObject<FbxAnimStack>(0) : nullptr)
-		{
-			AnimCurve = Geometry->GetShapeChannel(FetchPayloadData.MorphTargetIndex, FetchPayloadData.ChannelIndex, (FbxAnimLayer*)AnimStack->GetMember(0));
-		}
 
-		if (!ensure(AnimCurve != nullptr))
+		if (!Geometry)
 		{
 			UInterchangeResultError_Generic* Message = Parser.AddMessage<UInterchangeResultError_Generic>();
-			if (FbxNode* Node = Geometry->GetNode())
-			{
-				Message->InterchangeKey = Parser.GetFbxHelper()->GetFbxNodeHierarchyName(Node);
-			}
-			Message->Text = LOCTEXT("InternalFetchMorphTargetCurvePayloadToFile_FBXCurveNull", "Cannot fetch FBX morph target animation curve payload because the FBX anim curve node is null.");
+			Message->Text = LOCTEXT("InternalFetchMorphTargetCurvePayloadToFile_FBXSDKSceneNull", "InternalFetchMorphTargetCurvePayloadToFile, Geometry acquisition failed.");
 			return false;
 		}
+		FbxAnimCurve* AnimCurve = Geometry->GetShapeChannel(FetchPayloadData.MorphTargetIndex, FetchPayloadData.ChannelIndex, FetchPayloadData.AnimLayer);
+
 		TArray<FInterchangeCurve> InterchangeCurves;
 		//Morph target curve in fbx are between 0 and 100, in Unreal we are between 0 and 1, so we must scale
 		//The curve with 0.01
@@ -742,38 +747,34 @@ namespace UE::Interchange::Private
 		return true;
 	}
 
-	void FFbxAnimation::AddNodeTransformAnimation(FbxScene* SDKScene
+	bool FFbxAnimation::AddSkeletalTransformAnimation(FbxScene* SDKScene
 		, FFbxParser& Parser
 		, FbxNode* Node
 		, UInterchangeSceneNode* SceneNode
-		, TMap<FString, TSharedPtr<FPayloadContextBase>>& PayloadContexts)
+		, TMap<FString, TSharedPtr<FPayloadContextBase>>& PayloadContexts
+		, UInterchangeSkeletalAnimationTrackNode* SkeletalAnimationTrackNode
+		, const int32& AnimationIndex)
 	{
-		int32 NumAnimations = SDKScene->GetSrcObjectCount<FbxAnimStack>();
-		//Anim stack should be merge so we expect to have only one stack here
-		ensure(NumAnimations <= 1);
-
-		double FrameRate = FbxTime::GetFrameRate(SDKScene->GetGlobalSettings().GetTimeMode());
-
 		FGetFbxTransformCurvesParameters Parameters(SDKScene, Node);
-		GetFbxTransformCurves(Parameters);
+		GetFbxTransformCurves(Parameters, AnimationIndex);
 
 		if (Parameters.IsNodeAnimated)
 		{
-			FString PayLoadKey = Parser.GetFbxHelper()->GetFbxNodeHierarchyName(Node) + TEXT("_AnimationPayloadKey");
+			FString PayLoadKey = Parser.GetFbxHelper()->GetFbxNodeHierarchyName(Node) + TEXT("_") + FString::FromInt(AnimationIndex) + TEXT("_SkeletalAnimationPayloadKey");
 			if (ensure(!PayloadContexts.Contains(PayLoadKey)))
 			{
 				TSharedPtr<FAnimationPayloadContextTransform> AnimPayload = MakeShared<FAnimationPayloadContextTransform>();
 				FNodeTransformFetchPayloadData FetchPayloadData;
 				FetchPayloadData.Node = Node;
+				FetchPayloadData.CurrentAnimStack = Parameters.CurrentAnimStack;
+
 				AnimPayload->NodeTransformFetchPayloadData = FetchPayloadData;
 				PayloadContexts.Add(PayLoadKey, AnimPayload);
 			}
-			UInterchangeAnimationAPI::SetCustomNodeTransformPayloadKey(SceneNode, PayLoadKey);
-			UInterchangeAnimationAPI::SetCustomIsNodeTransformAnimated(SceneNode, Parameters.IsNodeAnimated);
-			UInterchangeAnimationAPI::SetCustomNodeTransformAnimationKeyCount(SceneNode, Parameters.KeyCount);
-			UInterchangeAnimationAPI::SetCustomNodeTransformAnimationStartTime(SceneNode, Parameters.StartTime.GetSecondDouble());
-			UInterchangeAnimationAPI::SetCustomNodeTransformAnimationEndTime(SceneNode, Parameters.EndTime.GetSecondDouble());
+			SkeletalAnimationTrackNode->SetAnimationPayloadKeyForSceneNodeUid(SceneNode->GetUniqueID(), PayLoadKey);
 		}
+
+		return Parameters.IsNodeAnimated;
 	}
 
 	void FFbxAnimation::AddNodeAttributeCurvesAnimation(FFbxParser& Parser
@@ -793,7 +794,7 @@ namespace UE::Interchange::Private
 			FAttributeFetchPayloadData FetchPayloadData;
 			FetchPayloadData.Node = Node;
 			FetchPayloadData.AnimCurves = AnimCurveNode;
-
+			
 			//Any property that is not decimal should import with step curve
 			FetchPayloadData.bAttributeTypeIsStepCurveAnimation = !IsFbxPropertyTypeDecimal(PropertyType);
 			FetchPayloadData.PropertyType = PropertyType;
@@ -804,29 +805,55 @@ namespace UE::Interchange::Private
 		}
 	}
 
-	void FFbxAnimation::AddMorphTargetCurvesAnimation(UInterchangeMeshNode* MorphTargetNode
-		, FbxScene* SDKScene
-		, int32 GeometryIndex
-		, int32 MorphTargetIndex
-		, int32 ChannelIndex
-		, const FString& MorphTargetPayloadKey
-		, TMap<FString, TSharedPtr<FPayloadContextBase>>& PayloadContexts)
+	void FFbxAnimation::AddRigidTransformAnimation(FFbxParser& Parser
+		, FbxNode* Node
+		, FbxAnimCurveNode* TranslationCurveNode
+		, FbxAnimCurveNode* RotationCurveNode
+		, FbxAnimCurveNode* ScaleCurveNode
+		, TMap<FString, TSharedPtr<FPayloadContextBase>>& PayloadContexts
+		, TOptional<FString>& OutPayloadKey)
 	{
-		check(MorphTargetNode);
-		const FString PayLoadKey = MorphTargetPayloadKey + TEXT("_CurveAnimationPayloadKey");
+		const FString PayLoadKey = Parser.GetFbxHelper()->GetFbxNodeHierarchyName(Node) + TEXT("_RigidAnimationPayloadKey");
+		if (ensure(!PayloadContexts.Contains(PayLoadKey)))
+		{
+			TSharedPtr<FAnimationPayloadContextTransform> AnimPayload = MakeShared<FAnimationPayloadContextTransform>();
+			FAttributeNodeTransformFetchPayloadData FetchPayloadData;
+			FetchPayloadData.Node = Node;
+			FetchPayloadData.TranlsationCurveNode = TranslationCurveNode;
+			FetchPayloadData.RotationCurveNode = RotationCurveNode;
+			FetchPayloadData.ScaleCurveNode = ScaleCurveNode;
+
+			//Any property that is not decimal should import with step curve
+			AnimPayload->AttributeNodeTransformFetchPayloadData = FetchPayloadData;
+			PayloadContexts.Add(PayLoadKey, AnimPayload);
+			OutPayloadKey = PayLoadKey;
+		}
+	}
+
+	void FFbxAnimation::AddMorphTargetCurvesAnimation(FbxScene* SDKScene
+		, FFbxParser& Parser
+		, UInterchangeSkeletalAnimationTrackNode* SkeletalAnimationTrackNode
+		, TMap<FString, TSharedPtr<FPayloadContextBase>>& PayloadContexts
+		, const FMorphTargetAnimationBuildingData& MorphTargetAnimationBuildingData)
+	{
+		FString PayLoadKey = MorphTargetAnimationBuildingData.MorphTargetNodeUid + TEXT("\\") + FString::FromInt(MorphTargetAnimationBuildingData.AnimationIndex) + TEXT("\\") + FString::FromInt(MorphTargetAnimationBuildingData.MorphTargetIndex) + TEXT("\\") + FString::FromInt(MorphTargetAnimationBuildingData.ChannelIndex) + TEXT("_CurveAnimationPayloadKey");
+
 		if (ensure(!PayloadContexts.Contains(PayLoadKey)))
 		{
 			TSharedPtr<FAnimationPayloadContextTransform> AnimPayload = MakeShared<FAnimationPayloadContextTransform>();
 			FMorphTargetFetchPayloadData FetchPayloadData;
 			FetchPayloadData.SDKScene = SDKScene;
-			FetchPayloadData.GeometryIndex = GeometryIndex;
-			FetchPayloadData.MorphTargetIndex = MorphTargetIndex;
-			FetchPayloadData.ChannelIndex = ChannelIndex;
+			FetchPayloadData.GeometryIndex = MorphTargetAnimationBuildingData.GeometryIndex;
+			FetchPayloadData.MorphTargetIndex = MorphTargetAnimationBuildingData.MorphTargetIndex;
+			FetchPayloadData.ChannelIndex = MorphTargetAnimationBuildingData.ChannelIndex;
+			FetchPayloadData.AnimLayer = MorphTargetAnimationBuildingData.AnimLayer;
+
 			AnimPayload->MorphTargetFetchPayloadData = FetchPayloadData;
 
 			PayloadContexts.Add(PayLoadKey, AnimPayload);
 		}
-		MorphTargetNode->SetAnimationCurvePayLoadKey(PayLoadKey);
+
+		SkeletalAnimationTrackNode->SetAnimationPayloadKeyForMorphTargetNodeUid(MorphTargetAnimationBuildingData.MorphTargetNodeUid, PayLoadKey);
 	}
 
 	bool FFbxAnimation::IsFbxPropertyTypeSupported(EFbxType PropertyType)
diff --git a/Engine/Plugins/Interchange/Runtime/Source/Parsers/Fbx/Private/FbxAnimation.h b/Engine/Plugins/Interchange/Runtime/Source/Parsers/Fbx/Private/FbxAnimation.h
--- a/Engine/Plugins/Interchange/Runtime/Source/Parsers/Fbx/Private/FbxAnimation.h
+++ b/Engine/Plugins/Interchange/Runtime/Source/Parsers/Fbx/Private/FbxAnimation.h
@@ -11,14 +11,20 @@
 struct FMeshDescription;
 class UInterchangeMeshNode;
 class UInterchangeSceneNode;
+class UInterchangeSkeletalAnimationTrackNode;
 
 namespace UE::Interchange::Private
 {
+	struct FMorphTargetAnimationBuildingData;
+
+	//Skeletal/Rigged Animation
 	struct FNodeTransformFetchPayloadData
 	{
 		FbxNode* Node = nullptr;
+		FbxAnimStack* CurrentAnimStack = nullptr;
 	};
 
+	//User defined animations:
 	struct FAttributeFetchPayloadData
 	{
 		FbxNode* Node = nullptr;
@@ -30,12 +36,29 @@ namespace UE::Interchange::Private
 		FbxProperty Property;
 	};
 
+	//Rigid Animation
+	struct FAttributeNodeTransformFetchPayloadData
+	{
+		FbxNode* Node = nullptr;
+		FbxAnimCurveNode* TranlsationCurveNode = nullptr;
+		FbxAnimCurveNode* RotationCurveNode = nullptr;
+		FbxAnimCurveNode* ScaleCurveNode = nullptr;
+	};
+
 	struct FMorphTargetFetchPayloadData
 	{
-		FbxScene* SDKScene = nullptr;
-		int32 GeometryIndex = 0;
-		int32 MorphTargetIndex = 0;
-		int32 ChannelIndex = 0;
+		//We cannot just store the FbxAnimCurve
+		// becase! it is refernced from FbxMesh which gets replaced when we are triangulating it upon PayloadData acquisition
+		// for that reason we need to re-acquire it all the way from the Scene, based on GeometryIndex, MorphTargetIndex, ChannelIndex and AnimLayer.
+		// (This is happening because upon Triangulation we are passing bReplace = true, see @FMeshDescriptionImporter::FillMeshDescriptionFromFbxMesh)
+		//FbxAnimCurve* MorphTargetAnimCurve = nullptr;
+
+		FbxScene* SDKScene;
+
+		int32 GeometryIndex;
+		int32 MorphTargetIndex;
+		int32 ChannelIndex;
+		FbxAnimLayer* AnimLayer;
 	};
 
 	
@@ -50,6 +73,7 @@ namespace UE::Interchange::Private
 		
 		TOptional<FNodeTransformFetchPayloadData> NodeTransformFetchPayloadData;
 		TOptional<FAttributeFetchPayloadData> AttributeFetchPayloadData;
+		TOptional<FAttributeNodeTransformFetchPayloadData> AttributeNodeTransformFetchPayloadData;
 		TOptional<FMorphTargetFetchPayloadData> MorphTargetFetchPayloadData;
 	private:
 		bool InternalFetchCurveNodePayloadToFile(FFbxParser& Parser, const FString& PayloadFilepath);
@@ -60,11 +84,13 @@ namespace UE::Interchange::Private
 	{
 	public:
 		/** This function add the payload key if the scene node transform is animated. */
-		static void AddNodeTransformAnimation(FbxScene* SDKScene
+		static bool AddSkeletalTransformAnimation(FbxScene* SDKScene
 			, FFbxParser& Parser
 			, FbxNode* Node
 			, UInterchangeSceneNode* UnrealNode
-			, TMap<FString, TSharedPtr<FPayloadContextBase>>& PayloadContexts);
+			, TMap<FString, TSharedPtr<FPayloadContextBase>>& PayloadContexts
+			, UInterchangeSkeletalAnimationTrackNode* SkeletalAnimationTrackNode
+			, const int32& AnimationIndex);
 
 		/** This function add the payload key for an animated node user attribute. */
 		static void AddNodeAttributeCurvesAnimation(FFbxParser& Parser
@@ -76,14 +102,20 @@ namespace UE::Interchange::Private
 			, EFbxType PropertyType
 			, TOptional<FString>& OutPayloadKey);
 
+		static void AddRigidTransformAnimation(FFbxParser& Parser
+			, FbxNode* Node
+			, FbxAnimCurveNode* TranslationCurveNode
+			, FbxAnimCurveNode* RotationCurveNode
+			, FbxAnimCurveNode* ScaleCurveNode
+			, TMap<FString, TSharedPtr<FPayloadContextBase>>& PayloadContexts
+			, TOptional<FString>& OutPayloadKey);
+
 		/** This function add the payload key for an animated curve for morph target. */
-		static void AddMorphTargetCurvesAnimation(UInterchangeMeshNode* ShapeNode
-			, FbxScene* SDKScene
-			, int32 GeometryIndex
-			, int32 MorphTargetIndex
-			, int32 ChannelIndex
-			, const FString& MorphTargetPayloadKey
-			, TMap<FString, TSharedPtr<FPayloadContextBase>>& PayloadContexts);
+		static void AddMorphTargetCurvesAnimation(FbxScene* SDKScene
+			, FFbxParser& Parser
+			, UInterchangeSkeletalAnimationTrackNode* SkeletalAnimationTrackNode
+			, TMap<FString, TSharedPtr<FPayloadContextBase>>& PayloadContexts
+			, const FMorphTargetAnimationBuildingData& MorphTargetAnimationBuildingData);
 
 		static bool IsFbxPropertyTypeSupported(EFbxType PropertyType);
 		static bool IsFbxPropertyTypeDecimal(EFbxType PropertyType);
diff --git a/Engine/Plugins/Interchange/Runtime/Source/Parsers/Fbx/Private/FbxConvert.cpp b/Engine/Plugins/Interchange/Runtime/Source/Parsers/Fbx/Private/FbxConvert.cpp
--- a/Engine/Plugins/Interchange/Runtime/Source/Parsers/Fbx/Private/FbxConvert.cpp
+++ b/Engine/Plugins/Interchange/Runtime/Source/Parsers/Fbx/Private/FbxConvert.cpp
@@ -26,15 +26,19 @@ namespace UE
 				//Set the original framerate from the current fbx file
 				float FbxFramerate = FbxTime::GetFrameRate(TimeMode);
 
-				//Merge the anim stack before the conversion since the above 0 layer will not be converted
 				int32 AnimStackCount = SDKScene->GetSrcObjectCount<FbxAnimStack>();
-				//Merge the animation stack layer before converting the scene
 				for (int32 AnimStackIndex = 0; AnimStackIndex < AnimStackCount; AnimStackIndex++)
 				{
-					FbxAnimStack* CurAnimStack = SDKScene->GetSrcObject<FbxAnimStack>(AnimStackIndex);
-					if (CurAnimStack->GetMemberCount() > 1)
+					FbxAnimStack* CurrentAnimStack = SDKScene->GetSrcObject<FbxAnimStack>(AnimStackIndex);
+					int32 NumLayers = CurrentAnimStack->GetMemberCount();
+					for (int LayerIndex = 0; LayerIndex < NumLayers; LayerIndex++)
 					{
-						MergeAllLayerAnimation(SDKScene, CurAnimStack, FbxFramerate);
+						FbxAnimLayer* AnimLayer = (FbxAnimLayer*)CurrentAnimStack->GetMember(LayerIndex);
+
+						// always apply unroll filter
+						FbxAnimCurveFilterUnroll UnrollFilter;
+						UnrollFilter.Reset();
+						ApplyUnroll(SDKScene->GetRootNode(), AnimLayer, &UnrollFilter);
 					}
 				}
 
@@ -241,26 +245,6 @@ namespace UE
 					ApplyUnroll(Node->GetChild(i), Layer, UnrollFilter);
 				}
 			}
-
-			void FFbxConvert::MergeAllLayerAnimation(FbxScene* SDKScene, FbxAnimStack* AnimStack, float ResampleRate)
-			{
-				if (!ensure(SDKScene) || !ensure(AnimStack))
-				{
-					return;
-				}
-				FbxTime FramePeriod;
-				FramePeriod.SetSecondDouble(1.0 / (double)ResampleRate);
-
-				FbxTimeSpan TimeSpan = AnimStack->GetLocalTimeSpan();
-				AnimStack->BakeLayers(SDKScene->GetAnimationEvaluator(), TimeSpan.GetStart(), TimeSpan.GetStop(), FramePeriod);
-
-				// always apply unroll filter
-				FbxAnimCurveFilterUnroll UnrollFilter;
-
-				FbxAnimLayer* Layer = AnimStack->GetMember<FbxAnimLayer>(0);
-				UnrollFilter.Reset();
-				ApplyUnroll(SDKScene->GetRootNode(), Layer, &UnrollFilter);
-			}
 		}//ns Private
 	}//ns Interchange
 }//ns UE
diff --git a/Engine/Plugins/Interchange/Runtime/Source/Parsers/Fbx/Private/FbxConvert.h b/Engine/Plugins/Interchange/Runtime/Source/Parsers/Fbx/Private/FbxConvert.h
--- a/Engine/Plugins/Interchange/Runtime/Source/Parsers/Fbx/Private/FbxConvert.h
+++ b/Engine/Plugins/Interchange/Runtime/Source/Parsers/Fbx/Private/FbxConvert.h
@@ -77,7 +77,6 @@ namespace UE
 				 * rotation curves from the nodes property. This can become time consuming but we have no choice.
 				 */
 				static void ApplyUnroll(FbxNode* pNode, FbxAnimLayer* pLayer, FbxAnimCurveFilterUnroll* pUnrollFilter);
-				static void MergeAllLayerAnimation(FbxScene* SDKScene, FbxAnimStack* AnimStack, float ResampleRate);
 
 				/** Scene Conversion Private Implementation End */
 				//////////////////////////////////////////////////////////////////////////
diff --git a/Engine/Plugins/Interchange/Runtime/Source/Parsers/Fbx/Private/FbxMesh.cpp b/Engine/Plugins/Interchange/Runtime/Source/Parsers/Fbx/Private/FbxMesh.cpp
--- a/Engine/Plugins/Interchange/Runtime/Source/Parsers/Fbx/Private/FbxMesh.cpp
+++ b/Engine/Plugins/Interchange/Runtime/Source/Parsers/Fbx/Private/FbxMesh.cpp
@@ -1394,10 +1394,6 @@ void FFbxMesh::AddAllMeshes(FbxScene* SDKScene, FbxGeometryConverter* SDKGeometr
 		MeshNode->SetPayLoadKey(PayLoadKey);
 
 		int32 NumAnimations = SDKScene->GetSrcObjectCount<FbxAnimStack>();
-			
-		//Anim stack should be merge so we expect to have only one stack here
-		ensure(NumAnimations <= 1);
-		FbxAnimStack* AnimStack = NumAnimations == 1 ? (FbxAnimStack*)SDKScene->GetSrcObject<FbxAnimStack>(0) : nullptr;
 
 		//Add all MorphTargets for this geometry node
 
@@ -1472,14 +1468,22 @@ void FFbxMesh::AddAllMeshes(FbxScene* SDKScene, FbxGeometryConverter* SDKGeometr
 						}
 						MorphTargetNode->SetPayLoadKey(MorphTargetPayLoadKey);
 
-						//Add the morph target animation float curve payload to the morph target.
-						if (AnimStack)
+						for (int32 AnimationIndex = 0; AnimationIndex < NumAnimations; AnimationIndex++)
 						{
-							//If the morph target node is animated prepare a curve payload to retrieve the morph target animation
-							FbxAnimCurve* Curve = Geometry->GetShapeChannel(MorphTargetIndex, ChannelIndex, (FbxAnimLayer*)AnimStack->GetMember(0));
-							if ((Curve && Curve->KeyGetCount() > 0))
+							FbxAnimStack* AnimStack = (FbxAnimStack*)SDKScene->GetSrcObject<FbxAnimStack>(AnimationIndex);
+							if (AnimStack)
 							{
-								UE::Interchange::Private::FFbxAnimation::AddMorphTargetCurvesAnimation(MorphTargetNode, SDKScene, GeometryIndex, MorphTargetIndex, ChannelIndex, MorphTargetPayLoadKey, PayloadContexts);
+								FbxAnimLayer* AnimLayer = (FbxAnimLayer*)AnimStack->GetMember(0);
+								
+								FbxAnimCurve* Curve = Mesh->GetShapeChannel(MorphTargetIndex, ChannelIndex, AnimLayer);
+								if (Curve && Curve->KeyGetCount() > 0)
+								{
+									FbxTimeSpan TimeInterval;
+									Curve->GetTimeInterval(TimeInterval);
+
+									//In order to appropriately identify the Skeleton Node Uids we have to process the MorphTarget animations once the hierarchy is processed
+									MorphTargetAnimationsBuildingData.Add(FMorphTargetAnimationBuildingData(TimeInterval.GetStart().GetSecondDouble(), TimeInterval.GetStop().GetSecondDouble(), MeshNode, GeometryIndex, AnimationIndex, AnimLayer, MorphTargetIndex, ChannelIndex, MorphTargetUniqueID));
+								}
 							}
 						}
 					}
diff --git a/Engine/Plugins/Interchange/Runtime/Source/Parsers/Fbx/Private/FbxMesh.h b/Engine/Plugins/Interchange/Runtime/Source/Parsers/Fbx/Private/FbxMesh.h
--- a/Engine/Plugins/Interchange/Runtime/Source/Parsers/Fbx/Private/FbxMesh.h
+++ b/Engine/Plugins/Interchange/Runtime/Source/Parsers/Fbx/Private/FbxMesh.h
@@ -18,6 +18,41 @@ namespace UE
 	{
 		namespace Private
 		{
+			struct FMorphTargetAnimationBuildingData
+			{
+				double StartTime;
+				double StopTime;
+				UInterchangeMeshNode* InterchangeMeshNode;
+				int32 GeometryIndex;
+				int32 AnimationIndex;
+				FbxAnimLayer* AnimLayer;
+				int32 MorphTargetIndex;
+				int32 ChannelIndex;
+				FString MorphTargetNodeUid;
+
+				FMorphTargetAnimationBuildingData(double InStartTime
+						, double InStopTime
+						, UInterchangeMeshNode* InInterchangeMeshNode
+						, int32 InGeometryIndex
+						, int32 InAnimationIndex
+						, FbxAnimLayer* InAnimLayer
+						, int32 InMorphTargetIndex
+						, int32 InChannelIndex
+						, FString InMorphTargetNodeUid)
+					: StartTime(InStartTime)
+					, StopTime(InStopTime)
+					, InterchangeMeshNode(InInterchangeMeshNode)
+					, GeometryIndex(InGeometryIndex)
+					, AnimationIndex(InAnimationIndex)
+					, AnimLayer(InAnimLayer)
+					, MorphTargetIndex(InMorphTargetIndex)
+					, ChannelIndex(InChannelIndex)
+					, MorphTargetNodeUid(InMorphTargetNodeUid)
+				{
+
+				}
+			};
+
 			class FFbxMesh;
 
 			class FMeshDescriptionImporter
@@ -115,6 +150,8 @@ namespace UE
 
 				void AddAllMeshes(FbxScene* SDKScene, FbxGeometryConverter* SDKGeometryConverter, UInterchangeBaseNodeContainer& NodeContainer, TMap<FString, TSharedPtr<FPayloadContextBase>>& PayloadContexts);
 				static bool GetGlobalJointBindPoseTransform(FbxScene* SDKScene, FbxNode* Joint, FbxAMatrix& GlobalBindPoseJointMatrix);
+
+				TArray<FMorphTargetAnimationBuildingData>& GetMorphTargetAnimationsBuildingData() { return MorphTargetAnimationsBuildingData; }
 			protected:
 				/** Add joint to the interchange mesh node joint dependencies. Return false if there is no valid joint (not a valid skinned mesh) */
 				bool ExtractSkinnedMeshNodeJoints(FbxScene* SDKScene, UInterchangeBaseNodeContainer& NodeContainer, FbxMesh* Mesh, UInterchangeMeshNode* MeshNode);
@@ -122,6 +159,9 @@ namespace UE
 
 			private:
 				FFbxParser& Parser;
+
+				//In order to appropriately identify the Skeleton Node Uids we have to process the MorphTarget animations once the hierarchy is processed
+				TArray<FMorphTargetAnimationBuildingData> MorphTargetAnimationsBuildingData;
 			};
 		}//ns Private
 	}//ns Interchange
diff --git a/Engine/Plugins/Interchange/Runtime/Source/Parsers/Fbx/Private/FbxScene.cpp b/Engine/Plugins/Interchange/Runtime/Source/Parsers/Fbx/Private/FbxScene.cpp
--- a/Engine/Plugins/Interchange/Runtime/Source/Parsers/Fbx/Private/FbxScene.cpp
+++ b/Engine/Plugins/Interchange/Runtime/Source/Parsers/Fbx/Private/FbxScene.cpp
@@ -18,6 +18,7 @@
 #include "Nodes/InterchangeBaseNodeContainer.h"
 #include "Nodes/InterchangeSourceNode.h"
 #include "Nodes/InterchangeUserDefinedAttribute.h"
+#include "InterchangeAnimationTrackSetNode.h"
 
 #define LOCTEXT_NAMESPACE "InterchangeFbxScene"
 
@@ -279,8 +280,6 @@ namespace UE
 						}
 					}
 				}
-				//Scene node transform can be animated, add the transform animation payload key.
-				FFbxAnimation::AddNodeTransformAnimation(SDKScene, Parser, Node, UnrealNode, PayloadContexts);
 
 				//Add all custom Attributes for the node
 				FbxProperty Property = Node->GetFirstProperty();
@@ -443,51 +442,351 @@ namespace UE
 			void FFbxScene::AddHierarchy(FbxScene* SDKScene, UInterchangeBaseNodeContainer& NodeContainer, TMap<FString, TSharedPtr<FPayloadContextBase, ESPMode::ThreadSafe>>& PayloadContexts)
 			{
 				 FbxNode* RootNode = SDKScene->GetRootNode();
-				 //Create a source node where we can store any general file info
-				 {
-					 UInterchangeSourceNode* SourceNode = UInterchangeSourceNode::FindOrCreateUniqueInstance(&NodeContainer);
-					 if (!ensure(SourceNode))
-					 {
-						 UInterchangeResultError_Generic* Message = Parser.AddMessage<UInterchangeResultError_Generic>();
-						 Message->Text = LOCTEXT("NodeAllocationError", "Unable to allocate a node when importing FBX.");
-						 return;
-					 }
-
-					 //Store the fbx frame rate
-					 {
-						 double FrameRate = FbxTime::GetFrameRate(SDKScene->GetGlobalSettings().GetTimeMode());
-						 SourceNode->SetCustomSourceFrameRateNumerator(FrameRate);
-						 constexpr double Denominator = 1.0;
-						 SourceNode->SetCustomSourceFrameRateDenominator(Denominator);
-					 }
-
-					 //Store the fbx timeline
-					 {
-						 //Timeline time span
-						 FbxTimeSpan TimelineTimeSpan(FBXSDK_TIME_INFINITE, FBXSDK_TIME_MINUS_INFINITE);
-						 SDKScene->GetGlobalSettings().GetTimelineDefaultTimeSpan(TimelineTimeSpan);
-						 SourceNode->SetCustomSourceTimelineStart(TimelineTimeSpan.GetStart().GetSecondDouble());
-						 SourceNode->SetCustomSourceTimelineEnd(TimelineTimeSpan.GetStop().GetSecondDouble());
-
-						 //Animated time span
-						 FbxTimeSpan AnimatedTimeSpan(FBXSDK_TIME_INFINITE, FBXSDK_TIME_MINUS_INFINITE);
-						 int32 AnimCurveNodeCount = SDKScene->GetSrcObjectCount<FbxAnimCurveNode>();
-						 for (int32 AnimCurveNodeIndex = 0; AnimCurveNodeIndex < AnimCurveNodeCount; AnimCurveNodeIndex++)
-						 {
-							 FbxAnimCurveNode* CurAnimCruveNode = SDKScene->GetSrcObject<FbxAnimCurveNode>(AnimCurveNodeIndex);
-							 if (CurAnimCruveNode->IsAnimated(true))
-							 {
-								 FbxTimeSpan CurveTimeSpan(FBXSDK_TIME_INFINITE, FBXSDK_TIME_MINUS_INFINITE);
-								 CurAnimCruveNode->GetAnimationInterval(CurveTimeSpan);
-								 AnimatedTimeSpan.UnionAssignment(CurveTimeSpan);
-							 }
-						 }
-						 SourceNode->SetCustomAnimatedTimeStart(AnimatedTimeSpan.GetStart().GetSecondDouble());
-						 SourceNode->SetCustomAnimatedTimeEnd(AnimatedTimeSpan.GetStop().GetSecondDouble());
-					 }
-				 }
+
 				 AddHierarchyRecursively(nullptr, RootNode, SDKScene, NodeContainer, PayloadContexts);
 			}
+
+			void FFbxScene::AddRigidAnimation(FbxNode* Node
+				, UInterchangeSceneNode* UnrealNode
+				, UInterchangeBaseNodeContainer& NodeContainer
+				, TMap<FString, TSharedPtr<FPayloadContextBase, ESPMode::ThreadSafe>>& PayloadContexts)
+			{
+				FbxAnimCurveNode* TranlsationCurveNode = nullptr;
+				FbxAnimCurveNode* RotationCurveNode = nullptr;
+				FbxAnimCurveNode* ScaleCurveNode = nullptr;
+
+				FbxProperty Property = Node->GetFirstProperty();
+				while (Property.IsValid())
+				{
+					EFbxType PropertyType = Property.GetPropertyDataType().GetType();
+
+					if (FFbxAnimation::IsFbxPropertyTypeSupported(PropertyType))
+					{
+						FbxAnimCurveNode* CurveNode = Property.GetCurveNode();
+
+						//only translation/rotation/scale is supported
+						if (CurveNode && CurveNode->IsAnimated())
+						{
+							//(currently FBXSDK_CURVENODE_TRANSFORM is not supported for Curve based animations)
+
+							const char* CurveNodeName = CurveNode->GetName(); //which lets us know the component that we are animating:
+							if (std::strcmp(CurveNodeName, FBXSDK_CURVENODE_TRANSLATION) == 0)
+							{
+								TranlsationCurveNode = CurveNode;
+							}
+							else if (std::strcmp(CurveNodeName, FBXSDK_CURVENODE_ROTATION) == 0)
+							{
+								RotationCurveNode = CurveNode;
+							}
+							else if (std::strcmp(CurveNodeName, FBXSDK_CURVENODE_SCALING) == 0)
+							{
+								ScaleCurveNode = CurveNode;
+							}
+						}
+					}
+					Property = Node->GetNextProperty(Property);
+				}
+
+				//
+				constexpr int32 TranslationChannel = 0x0001 | 0x0002 | 0x0004;
+				constexpr int32 RotationChannel = 0x0008 | 0x0010 | 0x0020;
+				constexpr int32 ScaleChannel = 0x0040 | 0x0080 | 0x0100;
+
+				int32 UsedChannels = 0;
+				if (TranlsationCurveNode)
+				{
+					UsedChannels |= TranslationChannel;
+				}
+				if (RotationCurveNode)
+				{
+					UsedChannels |= RotationChannel;
+				}
+				if (ScaleCurveNode)
+				{
+					UsedChannels |= ScaleChannel;
+				}
+
+				if (UsedChannels)
+				{
+					TOptional<FString> PayloadKey;
+
+					FFbxAnimation::AddRigidTransformAnimation(Parser, Node, TranlsationCurveNode, RotationCurveNode, ScaleCurveNode, PayloadContexts, PayloadKey);
+
+					if (PayloadKey.IsSet())
+					{
+						UInterchangeTransformAnimationTrackNode* TransformAnimTrackNode = NewObject< UInterchangeTransformAnimationTrackNode >(&NodeContainer);
+
+						const FString TransformAnimTrackNodeName = FString::Printf(TEXT("%s"), *UnrealNode->GetDisplayLabel());
+						const FString TransformAnimTrackNodeUid = TEXT("\\AnimationTrack\\") + TransformAnimTrackNodeName;
+
+						TransformAnimTrackNode->InitializeNode(TransformAnimTrackNodeUid, TransformAnimTrackNodeName, EInterchangeNodeContainerType::TranslatedAsset);
+
+						TransformAnimTrackNode->SetCustomActorDependencyUid(*UnrealNode->GetUniqueID());
+
+						TransformAnimTrackNode->SetCustomAnimationPayloadKey(PayloadKey.GetValue());
+
+						TransformAnimTrackNode->SetCustomUsedChannels(UsedChannels);
+
+						NodeContainer.AddNode(TransformAnimTrackNode);
+					}
+				}
+			}
+
+			void FFbxScene::AddAnimationRecursively(FbxNode* Node
+				, FbxScene* SDKScene
+				, UInterchangeBaseNodeContainer& NodeContainer
+				, TMap<FString, TSharedPtr<FPayloadContextBase, ESPMode::ThreadSafe>>& PayloadContexts
+				, UInterchangeSkeletalAnimationTrackNode* SkeletalAnimationTrackNode, bool SkeletalAnimationAddedToContainer
+				, const FString& RootSceneNodeUid, const TSet<FString>& SkeletonRootNodeUids
+				, const int32& AnimationIndex)
+			{
+				FString NodeUniqueID = Parser.GetFbxHelper()->GetFbxNodeHierarchyName(Node);
+				if (UInterchangeSceneNode* UnrealNode = const_cast<UInterchangeSceneNode*>(Cast< UInterchangeSceneNode >(NodeContainer.GetNode(NodeUniqueID))))
+				{
+					bool HasSkeletonAttribute = false;
+					int32 AttributeCount = Node->GetNodeAttributeCount();
+
+					for (int32 AttributeIndex = 0; AttributeIndex < AttributeCount && !HasSkeletonAttribute; ++AttributeIndex)
+					{
+						FbxNodeAttribute* NodeAttribute = Node->GetNodeAttributeByIndex(AttributeIndex);
+						switch (NodeAttribute->GetAttributeType())
+						{
+						case FbxNodeAttribute::eNull:
+						case FbxNodeAttribute::eSkeleton:
+							HasSkeletonAttribute = true;
+							if (SkeletonRootNodeUids.Contains(NodeUniqueID))
+							{
+								SkeletalAnimationTrackNode = NewObject< UInterchangeSkeletalAnimationTrackNode >(&NodeContainer);
+								FString TrackNodeUid = "\\SkeletalAnimation\\" + UnrealNode->GetUniqueID() + "_" + FString::FromInt(AnimationIndex);
+								FString DisplayString = "Anim_" + FString::FromInt(AnimationIndex) + "_" + UnrealNode->GetDisplayLabel();
+								SkeletalAnimationTrackNode->InitializeNode(TrackNodeUid, DisplayString, EInterchangeNodeContainerType::TranslatedAsset);
+
+								double FrameRate = FbxTime::GetFrameRate(SDKScene->GetGlobalSettings().GetTimeMode());
+								SkeletalAnimationTrackNode->SetCustomAnimationSampleRate(FrameRate);
+
+								SkeletalAnimationTrackNode->SetCustomSkeletonNodeUid(UnrealNode->GetUniqueID());
+
+								//Calculate AnimationTime:
+								FbxAnimStack* CurrentAnimationStack = (FbxAnimStack*)SDKScene->GetSrcObject<FbxAnimStack>(AnimationIndex);
+								FbxTimeSpan TimeSpan = CurrentAnimationStack->GetLocalTimeSpan();
+
+								SkeletalAnimationTrackNode->SetCustomAnimationStartTime(TimeSpan.GetStart().GetSecondDouble());
+								SkeletalAnimationTrackNode->SetCustomAnimationStopTime(TimeSpan.GetStop().GetSecondDouble());
+							}
+							break;
+						default:
+							break;
+						}
+					}
+
+					if (!HasSkeletonAttribute)
+					{
+						//in case the joint node "hierarchy finished" then the SkeletalAnimationTrackNode should be reset:
+						//as on the next occurance of a Joint node a New skeleton will start:
+						SkeletalAnimationTrackNode = nullptr;
+
+						SkeletalAnimationAddedToContainer = false;
+
+						//if No Skeleton Attribute then it potentially has a Rigid animation:
+						//AnimationIndex does not change the Skeletal Attributes nor the Properties from which we gather the Rigid Animation Data:
+						if (AnimationIndex == 0)
+						{
+							AddRigidAnimation(Node, UnrealNode, NodeContainer, PayloadContexts);
+						}
+					}
+					else if (SkeletalAnimationTrackNode)
+					{
+						//Scene node transform can be animated, add the transform animation payload key.
+						if (FFbxAnimation::AddSkeletalTransformAnimation(SDKScene, Parser, Node, UnrealNode, PayloadContexts, SkeletalAnimationTrackNode, AnimationIndex)
+							&& !SkeletalAnimationAddedToContainer)
+						{
+							SkeletalAnimationAddedToContainer = true;
+							NodeContainer.AddNode(SkeletalAnimationTrackNode);
+						}
+					}
+				}
+			
+
+				const int32 ChildCount = Node->GetChildCount();
+				for (int32 ChildIndex = 0; ChildIndex < ChildCount; ++ChildIndex)
+				{
+					FbxNode* ChildNode = Node->GetChild(ChildIndex);
+					AddAnimationRecursively(ChildNode, SDKScene, NodeContainer, PayloadContexts, SkeletalAnimationTrackNode, SkeletalAnimationAddedToContainer, RootSceneNodeUid, SkeletonRootNodeUids, AnimationIndex);
+				}
+			}
+
+			void FFbxScene::AddAnimation(FbxScene* SDKScene, UInterchangeBaseNodeContainer& NodeContainer, TMap<FString, TSharedPtr<FPayloadContextBase, ESPMode::ThreadSafe>>& PayloadContexts)
+			{
+				FbxNode* RootNode = SDKScene->GetRootNode();
+				FString RootSceneNodeUniqueID = Parser.GetFbxHelper()->GetFbxNodeHierarchyName(RootNode);
+
+				//acquire Skeletal Node Uids from Meshes (via the skeletondependencies: )
+				TSet<FString> SkeletonRootNodeUids;
+				NodeContainer.IterateNodesOfType<UInterchangeMeshNode>([&](const FString& NodeUid, UInterchangeMeshNode* MeshNode)
+					{
+						//Find the root joint for this MeshGeometry
+						FString JointNodeUid;
+						FString ParentNodeUid;
+						MeshNode->GetSkeletonDependency(0, JointNodeUid);
+						ParentNodeUid = JointNodeUid;
+						
+						while (!JointNodeUid.Equals(UInterchangeBaseNode::InvalidNodeUid()))
+						{
+							if (const UInterchangeSceneNode* Node = Cast< UInterchangeSceneNode >(NodeContainer.GetNode(ParentNodeUid)))
+							{
+								if (Node->IsSpecializedTypeContains(FSceneNodeStaticData::GetJointSpecializeTypeString()))
+								{
+									JointNodeUid = ParentNodeUid;
+									ParentNodeUid = Node->GetParentUid();
+								}
+								else
+								{
+									break;
+								}
+							}
+							else
+							{
+								break;
+							}
+						}
+
+						if (!JointNodeUid.Equals(UInterchangeBaseNode::InvalidNodeUid()))
+						{
+							SkeletonRootNodeUids.Add(JointNodeUid);
+						}
+					});
+
+				int32 NumAnimations = SDKScene->GetSrcObjectCount<FbxAnimStack>();
+
+
+				for (int32 AnimationIndex = 0; AnimationIndex < NumAnimations; AnimationIndex++)
+				{
+					AddAnimationRecursively(RootNode, SDKScene, NodeContainer, PayloadContexts, nullptr, false, RootSceneNodeUniqueID, SkeletonRootNodeUids, AnimationIndex);
+				}
+
+				TArray<FString> TransformAnimTrackNodeUids;
+				NodeContainer.IterateNodesOfType<UInterchangeTransformAnimationTrackNode>([&](const FString& NodeUid, UInterchangeTransformAnimationTrackNode* TransformAnimationTrackNode)
+					{
+						TransformAnimTrackNodeUids.Add(NodeUid);
+					});
+
+				//Only one Track Set Node per fbx file:
+				if (TransformAnimTrackNodeUids.Num() > 0)
+				{
+					UInterchangeAnimationTrackSetNode* TrackSetNode = NewObject< UInterchangeAnimationTrackSetNode >(&NodeContainer);
+
+					const FString AnimTrackSetNodeUid = TEXT("\\Animation\\") + FString(RootNode->GetName());
+					const FString AnimTrackSetNodeDisplayLabel = FString(RootNode->GetName()) + TEXT("_TrackSetNode");
+					TrackSetNode->InitializeNode(AnimTrackSetNodeUid, AnimTrackSetNodeDisplayLabel, EInterchangeNodeContainerType::TranslatedAsset);
+
+					NodeContainer.AddNode(TrackSetNode);
+
+					for (const FString& TransformAnimTrackNodeUid : TransformAnimTrackNodeUids)
+					{
+						TrackSetNode->AddCustomAnimationTrackUid(TransformAnimTrackNodeUid);
+					}
+				}
+			}
+
+			void FFbxScene::AddMorphTargetAnimations(FbxScene* SDKScene, UInterchangeBaseNodeContainer& NodeContainer, TMap<FString, TSharedPtr<FPayloadContextBase, ESPMode::ThreadSafe>>& PayloadContexts, const TArray<FMorphTargetAnimationBuildingData>& MorphTargetAnimationsBuildingData)
+			{
+				//Group the Morph Target animations based on SkeletonNodeUid and AnimationIndex
+				TMap<const UInterchangeSceneNode*, TMap<int32, TArray<FMorphTargetAnimationBuildingData>>> MorphTargetAnimationsBuildingDataGrouped;
+
+				for (const FMorphTargetAnimationBuildingData& MorphTargetAnimationBuildingData : MorphTargetAnimationsBuildingData)
+				{
+					if (MorphTargetAnimationBuildingData.StartTime == MorphTargetAnimationBuildingData.StopTime)
+					{
+						//in case the interval is 0 skip the MorphTargetAnimation.
+						continue;
+					}
+					
+					FString SkeletonUid;
+					if (MorphTargetAnimationBuildingData.InterchangeMeshNode->IsSkinnedMesh())
+					{
+						//Find the root joint for this MeshGeometry
+						FString JointNodeUid;
+						FString ParentNodeUid;
+						MorphTargetAnimationBuildingData.InterchangeMeshNode->GetSkeletonDependency(0, JointNodeUid);
+						ParentNodeUid = JointNodeUid;
+
+						while (!JointNodeUid.Equals(UInterchangeBaseNode::InvalidNodeUid()))
+						{
+							if (const UInterchangeSceneNode* Node = Cast< UInterchangeSceneNode >(NodeContainer.GetNode(ParentNodeUid)))
+							{
+								if (Node->IsSpecializedTypeContains(FSceneNodeStaticData::GetJointSpecializeTypeString()))
+								{
+									JointNodeUid = ParentNodeUid;
+									ParentNodeUid = Node->GetParentUid();
+								}
+								else
+								{
+									break;
+								}
+							}
+							else
+							{
+								break;
+							}
+						}
+
+						if (!JointNodeUid.Equals(UInterchangeBaseNode::InvalidNodeUid()))
+						{
+							SkeletonUid = JointNodeUid;
+						}
+					}
+					else
+					{
+						//If it is not skinned then the only way to get this animation to show up is by mesh type forcing to Skeleton,
+						//where the root node id is going to be the Skeleton id:
+						FbxNode* RootNode = SDKScene->GetRootNode();
+						SkeletonUid = Parser.GetFbxHelper()->GetFbxNodeHierarchyName(RootNode);
+					}
+
+					if (const UInterchangeSceneNode* SkeletonNode = Cast< UInterchangeSceneNode >(NodeContainer.GetNode(SkeletonUid)))
+					{
+						//For the given skeleton:
+						TMap<int32, TArray<FMorphTargetAnimationBuildingData>>& MorphTargetAnimationPerAnimationIndex = MorphTargetAnimationsBuildingDataGrouped.FindOrAdd(SkeletonNode);
+						//For the given skeleton and animationindex:
+						TArray<FMorphTargetAnimationBuildingData>& MorphTargetAnimations = MorphTargetAnimationPerAnimationIndex.FindOrAdd(MorphTargetAnimationBuildingData.AnimationIndex);
+						MorphTargetAnimations.Add(MorphTargetAnimationBuildingData);
+					}
+				}
+				
+				for (const TPair<const UInterchangeSceneNode*, TMap<int32, TArray<FMorphTargetAnimationBuildingData>>>& MorphTargetAnimationBuildingDataGrouped : MorphTargetAnimationsBuildingDataGrouped)
+				{
+					const UInterchangeSceneNode* SkeletonNode = MorphTargetAnimationBuildingDataGrouped.Key;
+					FString SkeletonDisplayLabel = SkeletonNode->GetDisplayLabel();
+
+					for (const TPair<int32, TArray<FMorphTargetAnimationBuildingData>>& MorphTargetAnimationsBuildingDataPerSkeleton : MorphTargetAnimationBuildingDataGrouped.Value)
+					{
+						int32 AnimationIndex = MorphTargetAnimationsBuildingDataPerSkeleton.Key;
+
+						FbxAnimStack* CurrentAnimationStack = (FbxAnimStack*)SDKScene->GetSrcObject<FbxAnimStack>(AnimationIndex);
+						FbxTimeSpan TimeSpan = CurrentAnimationStack->GetLocalTimeSpan();
+
+						UInterchangeSkeletalAnimationTrackNode* SkeletalAnimationTrackNode = NewObject< UInterchangeSkeletalAnimationTrackNode >(&NodeContainer);
+						FString TrackNodeUid = "\\MorphTargetAnimation\\" + SkeletonNode->GetUniqueID() + TEXT("\\") + FString::FromInt(AnimationIndex);
+						FString DisplayString = SkeletonNode->GetDisplayLabel() + TEXT("_MorphAnim_") + FString::FromInt(AnimationIndex);
+						SkeletalAnimationTrackNode->InitializeNode(TrackNodeUid, DisplayString, EInterchangeNodeContainerType::TranslatedAsset);
+
+						double FrameRate = FbxTime::GetFrameRate(SDKScene->GetGlobalSettings().GetTimeMode());
+						SkeletalAnimationTrackNode->SetCustomAnimationSampleRate(FrameRate);
+
+						SkeletalAnimationTrackNode->SetCustomSkeletonNodeUid(SkeletonNode->GetUniqueID());
+
+						SkeletalAnimationTrackNode->SetCustomAnimationStartTime(TimeSpan.GetStart().GetSecondDouble());
+						SkeletalAnimationTrackNode->SetCustomAnimationStopTime(TimeSpan.GetStop().GetSecondDouble());
+
+						NodeContainer.AddNode(SkeletalAnimationTrackNode);
+						
+						for (const FMorphTargetAnimationBuildingData& MorphTargetAnimationBuildingDataPerSkeletonPerAnimationIndex : MorphTargetAnimationsBuildingDataPerSkeleton.Value)
+						{
+							UE::Interchange::Private::FFbxAnimation::AddMorphTargetCurvesAnimation(SDKScene, Parser, SkeletalAnimationTrackNode, PayloadContexts, MorphTargetAnimationBuildingDataPerSkeletonPerAnimationIndex);
+						}
+					}
+				}
+			}
 		} //ns Private
 	} //ns Interchange
 }//ns UE
diff --git a/Engine/Plugins/Interchange/Runtime/Source/Parsers/Fbx/Private/FbxScene.h b/Engine/Plugins/Interchange/Runtime/Source/Parsers/Fbx/Private/FbxScene.h
--- a/Engine/Plugins/Interchange/Runtime/Source/Parsers/Fbx/Private/FbxScene.h
+++ b/Engine/Plugins/Interchange/Runtime/Source/Parsers/Fbx/Private/FbxScene.h
@@ -12,6 +12,7 @@ namespace UE::Interchange::Private
 }
 class UInterchangeBaseNodeContainer;
 class UInterchangeSceneNode;
+class UInterchangeSkeletalAnimationTrackNode;
 
 namespace UE
 {
@@ -20,6 +21,7 @@ namespace UE
 		namespace Private
 		{
 			class FFbxParser;
+			struct FMorphTargetAnimationBuildingData;
 
 			class FFbxScene
 			{
@@ -29,6 +31,8 @@ namespace UE
 				{}
 
 				void AddHierarchy(FbxScene* SDKScene, UInterchangeBaseNodeContainer& NodeContainer, TMap<FString, TSharedPtr<FPayloadContextBase, ESPMode::ThreadSafe>>& PayloadContexts);
+				void AddAnimation(FbxScene* SDKScene, UInterchangeBaseNodeContainer& NodeContainer, TMap<FString, TSharedPtr<FPayloadContextBase, ESPMode::ThreadSafe>>& PayloadContexts);
+				void AddMorphTargetAnimations(FbxScene* SDKScene, UInterchangeBaseNodeContainer& NodeContainer, TMap<FString, TSharedPtr<FPayloadContextBase, ESPMode::ThreadSafe>>& PayloadContexts, const TArray<FMorphTargetAnimationBuildingData>& MorphTargetAnimationsBuildingData);
 				UInterchangeSceneNode* CreateTransformNode(UInterchangeBaseNodeContainer& NodeContainer, const FString& NodeName, const FString& NodeUniqueID);
 
 			protected:
@@ -41,7 +45,20 @@ namespace UE
 					, UInterchangeBaseNodeContainer& NodeContainer
 					, TMap<FString, TSharedPtr<FPayloadContextBase, ESPMode::ThreadSafe>>& PayloadContexts);
 
+				void AddAnimationRecursively(FbxNode* Node
+					, FbxScene* SDKScene
+					, UInterchangeBaseNodeContainer& NodeContainer
+					, TMap<FString, TSharedPtr<FPayloadContextBase, ESPMode::ThreadSafe>>& PayloadContexts
+					, UInterchangeSkeletalAnimationTrackNode* SkeletalAnimationTrackNode, bool SkeletalAnimationAddedToContainer
+					, const FString& RootSceneNodeUid, const TSet<FString>& SkeletonRootNodeUids
+					, const int32& AnimationIndex);
+
 			private:
+				void AddRigidAnimation(FbxNode* Node
+					, UInterchangeSceneNode* UnrealNode
+					, UInterchangeBaseNodeContainer& NodeContainer
+					, TMap<FString, TSharedPtr<FPayloadContextBase, ESPMode::ThreadSafe>>& PayloadContexts);
+
 				FFbxParser& Parser;
 			};
 		}//ns Private
diff --git a/Engine/Plugins/Interchange/Runtime/Source/Pipelines/Private/InterchangeGenericAnimationPipeline.cpp b/Engine/Plugins/Interchange/Runtime/Source/Pipelines/Private/InterchangeGenericAnimationPipeline.cpp
--- a/Engine/Plugins/Interchange/Runtime/Source/Pipelines/Private/InterchangeGenericAnimationPipeline.cpp
+++ b/Engine/Plugins/Interchange/Runtime/Source/Pipelines/Private/InterchangeGenericAnimationPipeline.cpp
@@ -14,7 +14,6 @@
 #include "InterchangeSkeletonFactoryNode.h"
 #include "InterchangeSkeletonHelper.h"
 #include "InterchangeSourceData.h"
-#include "Nodes/InterchangeAnimationAPI.h"
 #include "Nodes/InterchangeBaseNode.h"
 #include "Nodes/InterchangeBaseNodeContainer.h"
 #include "Nodes/InterchangeSourceNode.h"
@@ -22,33 +21,6 @@
 
 #include UE_INLINE_GENERATED_CPP_BY_NAME(InterchangeGenericAnimationPipeline)
 
-namespace UE::Interchange::Private
-{
-	/**
-	 * Return true if there is one animated node in the scene node hierarchy under NodeUid
-	 */
-	bool IsSkeletonAnimatedRecursive(const FString& NodeUid, UInterchangeBaseNodeContainer* BaseNodeContainer)
-	{
-		if (const UInterchangeSceneNode* SceneNode = Cast<UInterchangeSceneNode>(BaseNodeContainer->GetNode(NodeUid)))
-		{
-			bool bIsAnimated = false;
-			if (UInterchangeAnimationAPI::GetCustomIsNodeTransformAnimated(SceneNode, bIsAnimated) && bIsAnimated)
-			{
-				return true;
-			}
-		}
-		TArray<FString> Children = BaseNodeContainer->GetNodeChildrenUids(NodeUid);
-		for (const FString& ChildUid : Children)
-		{
-			if (IsSkeletonAnimatedRecursive(ChildUid, BaseNodeContainer))
-			{
-				return true;
-			}
-		}
-		return false;
-	}
-}
-
 void UInterchangeGenericAnimationPipeline::AdjustSettingsForContext(EInterchangePipelineContext ImportType, TObjectPtr<UObject> ReimportAsset)
 {
 	Super::AdjustSettingsForContext(ImportType, ReimportAsset);
@@ -140,436 +112,12 @@ void UInterchangeGenericAnimationPipeline::ExecutePipeline(UInterchangeBaseNodeC
 			TrackNodes.Add(Node);
 		});
 
-	if (TrackNodes.Num() > 0)
-	{
-		//UInterchangeSkeletalAnimationTrackNode is exclusively and only used by the GLTF translator
-		//The class will be re-visited to unify management of skeletal animations in Interchange.
-		for (UInterchangeSkeletalAnimationTrackNode* TrackNode : TrackNodes)
-		{
-			if (TrackNode)
-			{
-				CreateAnimSequenceFactoryNode(*TrackNode);
-			}
-		}
-		return;
-	}
-
-	double SampleRate = 30.0;
-	double RangeStart = 0;
-	double RangeStop = 0;
-	bool bRangeIsValid = false;
-	const UInterchangeSourceNode* SourceNode = UInterchangeSourceNode::GetUniqueInstance(BaseNodeContainer);
-	if (SourceNode)
-	{
-		if (bImportBoneTracks)
-		{
-			int32 Numerator, Denominator;
-			if (!bUse30HzToBakeBoneAnimation && CustomBoneAnimationSampleRate == 0 && SourceNode->GetCustomSourceFrameRateNumerator(Numerator))
-			{
-				if (SourceNode->GetCustomSourceFrameRateDenominator(Denominator) && Denominator > 0 && Numerator > 0)
-				{
-					SampleRate = static_cast<double>(Numerator) / static_cast<double>(Denominator);
-				}
-			}
-			else if ((!bUse30HzToBakeBoneAnimation && CustomBoneAnimationSampleRate > 0))
-			{
-				SampleRate = static_cast<double>(CustomBoneAnimationSampleRate);
-			}
-
-			if (AnimationRange == EInterchangeAnimationRange::Timeline)
-			{
-				if (SourceNode->GetCustomSourceTimelineStart(RangeStart))
-				{
-					if (SourceNode->GetCustomSourceTimelineEnd(RangeStop))
-					{
-						bRangeIsValid = true;
-					}
-				}
-			}
-			else if (AnimationRange == EInterchangeAnimationRange::Animated)
-			{
-				if (SourceNode->GetCustomAnimatedTimeStart(RangeStart))
-				{
-					if (SourceNode->GetCustomAnimatedTimeEnd(RangeStop))
-					{
-						bRangeIsValid = true;
-					}
-				}
-			}
-			else if (AnimationRange == EInterchangeAnimationRange::SetRange)
-			{
-				RangeStart = static_cast<double>(FrameImportRange.Min) / SampleRate;
-				RangeStop = static_cast<double>(FrameImportRange.Max) / SampleRate;
-				bRangeIsValid = true;
-			}
-		}
-	}
-	else
+	for (UInterchangeSkeletalAnimationTrackNode* TrackNode : TrackNodes)
 	{
-		if (bImportBoneTracks)
-		{
-			if ((!bUse30HzToBakeBoneAnimation && CustomBoneAnimationSampleRate > 0))
-			{
-				SampleRate = CustomBoneAnimationSampleRate;
-			}
-			//Find the range by iteration the scene node
-			TArray<FString> SceneNodeUids;
-			BaseNodeContainer->GetNodes(UInterchangeSceneNode::StaticClass(), SceneNodeUids);
-			for (const FString& SceneNodeUid : SceneNodeUids)
-			{
-				if (const UInterchangeSceneNode* SceneNode = Cast<const UInterchangeSceneNode>(BaseNodeContainer->GetNode(SceneNodeUid)))
-				{
-					double SceneNodeAnimStart;
-					double SceneNodeAnimStop;
-					if (UInterchangeAnimationAPI::GetCustomNodeTransformAnimationStartTime(SceneNode, SceneNodeAnimStart))
-					{
-						if (UInterchangeAnimationAPI::GetCustomNodeTransformAnimationEndTime(SceneNode, SceneNodeAnimStop))
-						{
-							if (RangeStart > SceneNodeAnimStart)
-							{
-								RangeStart = SceneNodeAnimStart;
-							}
-							if (RangeStop < SceneNodeAnimStop)
-							{
-								RangeStop = SceneNodeAnimStop;
-							}
-							bRangeIsValid = true;
-						}
-					}
-				}
-			}
-		}
-	}
-
-	//Retrieve all skeletons and all morph target anim data
-	TMap<UInterchangeSkeletonFactoryNode*, TArray<const UInterchangeMeshNode*>> MorphTargetsPerSkeletons;
-
-	BaseNodeContainer->IterateNodes([&MorphTargetsPerSkeletons, LocalNodeContainer = BaseNodeContainer](const FString& NodeUid, UInterchangeBaseNode* Node)
-		{
-			if (UInterchangeSkeletonFactoryNode* SkeletonFactoryNode = Cast<UInterchangeSkeletonFactoryNode>(Node))
-			{
-				const UInterchangeSkeletalMeshFactoryNode* SkeletalMeshFactoryNode = nullptr;
-				//Iterate dependencies
-				TArray<FString> SkeletalMeshNodeUids;
-				LocalNodeContainer->GetNodes(UInterchangeSkeletalMeshFactoryNode::StaticClass(), SkeletalMeshNodeUids);
-				for (const FString& SkelMeshFactoryNodeUid : SkeletalMeshNodeUids)
-				{
-					if (const UInterchangeSkeletalMeshFactoryNode* CurrentSkeletalMeshFactoryNode = Cast<const UInterchangeSkeletalMeshFactoryNode>(LocalNodeContainer->GetFactoryNode(SkelMeshFactoryNodeUid)))
-					{
-						TArray<FString> SkeletalMeshDependencies;
-						CurrentSkeletalMeshFactoryNode->GetFactoryDependencies(SkeletalMeshDependencies);
-						for (const FString& SkeletalMeshDependencyUid : SkeletalMeshDependencies)
-						{
-							if (NodeUid.Equals(SkeletalMeshDependencyUid))
-							{
-								SkeletalMeshFactoryNode = CurrentSkeletalMeshFactoryNode;
-								break;
-							}
-						}
-					}
-				}
-
-				FString RootSceneNodeUid;
-				SkeletonFactoryNode->GetCustomRootJointUid(RootSceneNodeUid);
-				if(UE::Interchange::Private::IsSkeletonAnimatedRecursive(RootSceneNodeUid, LocalNodeContainer))
-				{
-					MorphTargetsPerSkeletons.FindOrAdd(SkeletonFactoryNode);
-				}
-				
-				if (SkeletalMeshFactoryNode)
-				{
-					//Find the skeletalmesh morph targets
-					FString SkeletalMeshFactoryNodeUid = SkeletalMeshFactoryNode->GetUniqueID();
-					//Get the LOD data factory node
-					TArray<FString> LodDataChildren;
-					SkeletalMeshFactoryNode->GetLodDataUniqueIds(LodDataChildren);
-					for (const FString& ChildUid : LodDataChildren)
-					{
-						if (const UInterchangeSkeletalMeshLodDataNode* LodData = Cast<UInterchangeSkeletalMeshLodDataNode>(LocalNodeContainer->GetNode(ChildUid)))
-						{
-							TArray<FString> MeshUids;
-							LodData->GetMeshUids(MeshUids);
-							for (const FString& MeshUid : MeshUids)
-							{
-								if (const UInterchangeBaseNode* BaseNode = LocalNodeContainer->GetNode(MeshUid))
-								{
-									const UInterchangeMeshNode* MeshNode = nullptr;
-									FString RealMeshUid;
-									if (const UInterchangeSceneNode* SceneNode = Cast<UInterchangeSceneNode>(BaseNode))
-									{
-										SceneNode->GetCustomAssetInstanceUid(RealMeshUid);
-									}
-									else
-									{
-										RealMeshUid = MeshUid;
-									}
-									MeshNode = Cast<UInterchangeMeshNode>(LocalNodeContainer->GetNode(RealMeshUid));
-									if (MeshNode)
-									{
-										TArray<FString> MorphTargetUids;
-										MeshNode->GetMorphTargetDependencies(MorphTargetUids);
-										bool bIsMorphTargetAnimated = false;
-										TArray<const UInterchangeMeshNode*> MorphTargets;
-										for (const FString& MorphTargetUid : MorphTargetUids)
-										{
-											if (const UInterchangeMeshNode* MorphTargetNode = Cast<UInterchangeMeshNode>(LocalNodeContainer->GetNode(MorphTargetUid)))
-											{
-												MorphTargets.AddUnique(MorphTargetNode);
-												TOptional<FString> MorphTargetAnimationPayloadKey = MorphTargetNode->GetAnimationCurvePayLoadKey();
-												if (MorphTargetAnimationPayloadKey.IsSet())
-												{
-													bIsMorphTargetAnimated = true;
-												}
-											}
-										}
-										if (bIsMorphTargetAnimated)
-										{
-											TArray<const UInterchangeMeshNode*>& MorphTargetNodes = MorphTargetsPerSkeletons.FindOrAdd(SkeletonFactoryNode);
-											MorphTargetNodes.Reset(MorphTargets.Num());
-											MorphTargetNodes.Append(MorphTargets);
-										}
-									}
-								}
-							}
-						}
-					}
-					
-					//If we do not already create an anim sequence for this skeleton, add the anim sequence if we have
-					//at least one animated user defined attribute.
-					if (!MorphTargetsPerSkeletons.Contains(SkeletonFactoryNode))
-					{
-						//Add anim sequence if we have at least one animated user defined attributes
-						LocalNodeContainer->BreakableIterateNodeChildren(RootSceneNodeUid, [&MorphTargetsPerSkeletons, SkeletonFactoryNode](const UInterchangeBaseNode* Node)
-							{
-								if (const UInterchangeSceneNode* SceneNode = Cast<UInterchangeSceneNode>(Node))
-								{
-									TArray<FInterchangeUserDefinedAttributeInfo> AttributeInfos;
-									UInterchangeUserDefinedAttributesAPI::GetUserDefinedAttributeInfos(SceneNode, AttributeInfos);
-									for (const FInterchangeUserDefinedAttributeInfo& AttributeInfo : AttributeInfos)
-									{
-										if (AttributeInfo.PayloadKey.IsSet())
-										{
-											MorphTargetsPerSkeletons.FindOrAdd(SkeletonFactoryNode);
-											return true;
-										}
-									}
-								}
-								return false;
-							});
-					}
-				}
-			}
-		});
-
-	const TArray<FString> CustomAttributeNamesToImport = UAnimationSettings::Get()->GetBoneCustomAttributeNamesToImport();
-
-	//for each animated skeleton create one anim sequence factory node
-	for (const TPair<UInterchangeSkeletonFactoryNode*, TArray<const UInterchangeMeshNode*>>& SkeletonAndMorphTargets : MorphTargetsPerSkeletons)
-	{
-		UInterchangeSkeletonFactoryNode* SkeletonFactoryNode = SkeletonAndMorphTargets.Key;
-		const FString AnimSequenceUid = TEXT("\\AnimSequence") + SkeletonFactoryNode->GetUniqueID();
-		FString AnimSequenceName = SkeletonFactoryNode->GetDisplayLabel();
-		if (AnimSequenceName.EndsWith(TEXT("_Skeleton")))
-		{
-			AnimSequenceName.LeftChopInline(9);
-		}
-		AnimSequenceName += TEXT("_Anim");
-
-		if (bImportBoneTracks)
-		{
-			FFrameRate FrameRate = UE::Interchange::Animation::ConvertSampleRatetoFrameRate(SampleRate);
-
-			const double SequenceLength = FMath::Max<double>(RangeStop - RangeStart, MINIMUM_ANIMATION_LENGTH);
-
-			const float SubFrame = FrameRate.AsFrameTime(SequenceLength).GetSubFrame();
-
-			if (!FMath::IsNearlyZero(SubFrame, KINDA_SMALL_NUMBER) && !FMath::IsNearlyEqual(SubFrame, 1.0f, KINDA_SMALL_NUMBER))
-			{
-			    if (bSnapToClosestFrameBoundary)
-			    {
-				    // Figure out whether start or stop has to be adjusted
-				    const FFrameTime StartFrameTime = FrameRate.AsFrameTime(RangeStart);
-				    const FFrameTime StopFrameTime = FrameRate.AsFrameTime(RangeStop);
-				    FFrameNumber StartFrameNumber = StartFrameTime.GetFrame().Value, StopFrameNumber = StopFrameTime.GetFrame().Value;
-				    double NewStartTime = RangeStart, NewStopTime = RangeStop;
-    
-				    if (!FMath::IsNearlyZero(StartFrameTime.GetSubFrame()))
-				    {
-					    StartFrameNumber = StartFrameTime.RoundToFrame();
-					    NewStartTime = FrameRate.AsSeconds(StartFrameNumber);
-				    }
-    
-				    if (!FMath::IsNearlyZero(StopFrameTime.GetSubFrame()))
-				    {
-					    StopFrameNumber = StopFrameTime.RoundToFrame();
-					    NewStopTime = FrameRate.AsSeconds(StopFrameNumber);
-				    }
-    
-				    UInterchangeResultWarning_Generic* Message = AddMessage<UInterchangeResultWarning_Generic>();
-				    Message->SourceAssetName = SourceDatas[0]->GetFilename();
-				    Message->DestinationAssetName = AnimSequenceName;
-				    Message->AssetType = UAnimSequence::StaticClass();
-				    Message->Text = FText::Format(NSLOCTEXT("UInterchangeGenericAnimationPipeline", "Info_ImportLengthSnap", "Animation length has been adjusted to align with frame borders using import frame-rate {0}.\n\nOriginal timings:\n\t\tStart: {1} ({2})\n\t\tStop: {3} ({4})\nAligned timings:\n\t\tStart: {5} ({6})\n\t\tStop: {7} ({8})"),
-					    FrameRate.ToPrettyText(),
-					    FText::AsNumber(RangeStart),
-					    FText::AsNumber(StartFrameTime.AsDecimal()),
-					    FText::AsNumber(RangeStop),
-					    FText::AsNumber(StopFrameTime.AsDecimal()),
-					    FText::AsNumber(NewStartTime),
-					    FText::AsNumber(StartFrameNumber.Value),
-					    FText::AsNumber(NewStopTime),
-					    FText::AsNumber(StopFrameNumber.Value));
-    
-				    RangeStart = NewStartTime;
-				    RangeStop = NewStopTime;
-			    }
-			    else
-			    {
-				    UInterchangeResultError_Generic* Message = AddMessage<UInterchangeResultError_Generic>();
-				    Message->SourceAssetName = SourceDatas[0]->GetFilename();
-				    Message->DestinationAssetName = AnimSequenceName;
-				    Message->AssetType = UAnimSequence::StaticClass();
-				    Message->Text = FText::Format(NSLOCTEXT("UInterchangeGenericAnimationPipeline", "WrongSequenceLength", "Animation length {0} is not compatible with import frame-rate {1} (sub frame {2}), animation has to be frame-border aligned."),
-					    FText::AsNumber(SequenceLength), FrameRate.ToPrettyText(), FText::AsNumber(SubFrame));
-				    //Skip this anim sequence factory node
-				    continue;
-			    }
-			}
-		}
-
-		UInterchangeAnimSequenceFactoryNode* AnimSequenceFactoryNode = NewObject<UInterchangeAnimSequenceFactoryNode>(BaseNodeContainer, NAME_None);
-		AnimSequenceFactoryNode->InitializeAnimSequenceNode(AnimSequenceUid, AnimSequenceName);
-		
-		AnimSequenceFactoryNode->SetCustomSkeletonFactoryNodeUid(SkeletonFactoryNode->GetUniqueID());
-		AnimSequenceFactoryNode->SetCustomImportBoneTracks(bImportBoneTracks);
-		AnimSequenceFactoryNode->SetCustomImportBoneTracksSampleRate(SampleRate);
-		if (bRangeIsValid)
-		{
-			AnimSequenceFactoryNode->SetCustomImportBoneTracksRangeStart(RangeStart);
-			AnimSequenceFactoryNode->SetCustomImportBoneTracksRangeStop(RangeStop);
-		}
-
-		AnimSequenceFactoryNode->SetCustomImportAttributeCurves(bImportCustomAttribute);
-		AnimSequenceFactoryNode->SetCustomAddCurveMetadataToSkeleton(bAddCurveMetadataToSkeleton);
-		AnimSequenceFactoryNode->SetCustomDoNotImportCurveWithZero(bDoNotImportCurveWithZero);
-		AnimSequenceFactoryNode->SetCustomRemoveCurveRedundantKeys(bRemoveCurveRedundantKeys);
-		AnimSequenceFactoryNode->SetCustomDeleteExistingMorphTargetCurves(bDeleteExistingMorphTargetCurves);
-		AnimSequenceFactoryNode->SetCustomDeleteExistingCustomAttributeCurves(bDeleteExistingCustomAttributeCurves);
-		AnimSequenceFactoryNode->SetCustomDeleteExistingNonCurveCustomAttributes(bDeleteExistingNonCurveCustomAttributes);
-
-		AnimSequenceFactoryNode->SetCustomMaterialDriveParameterOnCustomAttribute(bSetMaterialDriveParameterOnCustomAttribute);
-		for (const FString& MaterialSuffixe : MaterialCurveSuffixes)
+		if (TrackNode)
 		{
-			AnimSequenceFactoryNode->SetAnimatedMaterialCurveSuffixe(MaterialSuffixe);
-		}
-
-		//Add the animated morph targets uid so the factory can import them
-		if (SkeletonAndMorphTargets.Value.Num() > 0)
-		{
-			for (const UInterchangeMeshNode* MorphTargetNode : SkeletonAndMorphTargets.Value)
-			{
-				AnimSequenceFactoryNode->SetAnimatedMorphTargetDependencyUid(MorphTargetNode->GetUniqueID());
-			}
-		}
-
-		//USkeleton cannot be created without a valid skeletalmesh
-		const FString SkeletonUid = SkeletonFactoryNode->GetUniqueID();
-		AnimSequenceFactoryNode->AddFactoryDependencyUid(SkeletonUid);
-
-		FString RootJointUid;
-		if(SkeletonFactoryNode->GetCustomRootJointUid(RootJointUid))
-		{
-#if WITH_EDITOR
-			//Iterate all joints to set the meta data value in the anim sequence factory node
-			UE::Interchange::Private::FSkeletonHelper::RecursiveAddSkeletonMetaDataValues(BaseNodeContainer, AnimSequenceFactoryNode, RootJointUid);
-#endif //WITH_EDITOR
-			BaseNodeContainer->IterateNodeChildren(RootJointUid, [&AnimSequenceFactoryNode, &CustomAttributeNamesToImport](const UInterchangeBaseNode* Node)
-			{
-				if (const UInterchangeSceneNode* SceneNode = Cast<UInterchangeSceneNode>(Node))
-				{
-					FString BoneName = SceneNode->GetDisplayLabel();
-					bool bImportAllAttributesOnBone = UAnimationSettings::Get()->BoneNamesWithCustomAttributes.Contains(BoneName);
-						
-					TArray<FInterchangeUserDefinedAttributeInfo> AttributeInfos;
-					UInterchangeUserDefinedAttributesAPI::GetUserDefinedAttributeInfos(SceneNode, AttributeInfos);
-					for (const FInterchangeUserDefinedAttributeInfo& AttributeInfo : AttributeInfos)
-					{
-						if (AttributeInfo.PayloadKey.IsSet())
-						{
-							bool bDecimalType = false;
-							switch (AttributeInfo.Type)
-							{
-							case UE::Interchange::EAttributeTypes::Float:
-							case UE::Interchange::EAttributeTypes::Float16:
-							case UE::Interchange::EAttributeTypes::Double:
-								{
-									bDecimalType = true;
-								}
-								break;
-							}
-
-							const bool bForceImportBoneCustomAttribute = CustomAttributeNamesToImport.Contains(AttributeInfo.Name);
-							//Material attribute curve
-							if (!bImportAllAttributesOnBone && bDecimalType && !bForceImportBoneCustomAttribute)
-							{
-								AnimSequenceFactoryNode->SetAnimatedAttributeCurveName(AttributeInfo.Name);
-							}
-							else if (bForceImportBoneCustomAttribute || bImportAllAttributesOnBone)
-							{
-								AnimSequenceFactoryNode->SetAnimatedAttributeStepCurveName(AttributeInfo.Name);
-							}
-						}
-					}
-				}
-			});
+			CreateAnimSequenceFactoryNode(*TrackNode);
 		}
-
-		const bool bImportOnlyAnimation = CommonSkeletalMeshesAndAnimationsProperties->bImportOnlyAnimations;
-		
-		//Iterate dependencies
-		{
-			TArray<FString> SkeletalMeshNodeUids;
-			BaseNodeContainer->GetNodes(UInterchangeSkeletalMeshFactoryNode::StaticClass(), SkeletalMeshNodeUids);
-			for (const FString& SkelMeshFactoryNodeUid : SkeletalMeshNodeUids)
-			{
-				if (const UInterchangeSkeletalMeshFactoryNode* SkeletalMeshFactoryNode = Cast<const UInterchangeSkeletalMeshFactoryNode>(BaseNodeContainer->GetFactoryNode(SkelMeshFactoryNodeUid)))
-				{
-					TArray<FString> SkeletalMeshDependencies;
-					SkeletalMeshFactoryNode->GetFactoryDependencies(SkeletalMeshDependencies);
-					for (const FString& SkeletalMeshDependencyUid : SkeletalMeshDependencies)
-					{
-						if (SkeletonUid.Equals(SkeletalMeshDependencyUid))
-						{
-							AnimSequenceFactoryNode->AddFactoryDependencyUid(SkelMeshFactoryNodeUid);
-							break;
-						}
-					}
-				}
-			}
-		}
-
-		if (CommonSkeletalMeshesAndAnimationsProperties->Skeleton.IsValid())
-		{
-			bool bSkeletonCompatible = true;
-
-			//TODO: support skeleton helper in runtime
-#if WITH_EDITOR
-			bSkeletonCompatible = UE::Interchange::Private::FSkeletonHelper::IsCompatibleSkeleton(CommonSkeletalMeshesAndAnimationsProperties->Skeleton.Get(), RootJointUid, BaseNodeContainer);
-#endif
-			if(bSkeletonCompatible)
-			{
-				FSoftObjectPath SkeletonSoftObjectPath(CommonSkeletalMeshesAndAnimationsProperties->Skeleton.Get());
-				AnimSequenceFactoryNode->SetCustomSkeletonSoftObjectPath(SkeletonSoftObjectPath);
-			}
-			else
-			{
-				UInterchangeResultError_Generic* Message = AddMessage<UInterchangeResultError_Generic>();
-				Message->Text = FText::Format(NSLOCTEXT("UInterchangeGenericAnimationPipeline", "IncompatibleSkeleton", "Incompatible skeleton {0} when importing AnimSequence {1}."),
-					FText::FromString(CommonSkeletalMeshesAndAnimationsProperties->Skeleton->GetName()),
-					FText::FromString(AnimSequenceName));
-			}
-		}
-		BaseNodeContainer->AddNode(AnimSequenceFactoryNode);
 	}
 }
 
@@ -652,67 +200,6 @@ void UInterchangeGenericAnimationPipeline::CreateAnimSequenceFactoryNode(UInterc
 		SkeletalMeshFactoryNode = Cast<const UInterchangeSkeletalMeshFactoryNode>(BaseNodeContainer->GetFactoryNode(SkeletalMeshFactoryNodeUid));
 	}
 
-	// NOTE: See whether it would not make sense to add a method to UInterchangeSkeletalAnimationTrackNode
-	// to collect all the mesh unique ids of the morph targets.
-	// In that case, GetCustomSkeletalMeshNodeUid might not be necessary.
-	TArray<const UInterchangeMeshNode*> MorphTargetNodes;
-	if (SkeletalMeshFactoryNode)
-	{
-		//Get the LOD data factory node
-		TArray<FString> LodDataChildren;
-		SkeletalMeshFactoryNode->GetLodDataUniqueIds(LodDataChildren);
-
-		for (const FString& ChildUid : LodDataChildren)
-		{
-			if (const UInterchangeSkeletalMeshLodDataNode* LodData = Cast<UInterchangeSkeletalMeshLodDataNode>(BaseNodeContainer->GetNode(ChildUid)))
-			{
-				TArray<FString> MeshUids;
-				LodData->GetMeshUids(MeshUids);
-				for (const FString& MeshUid : MeshUids)
-				{
-					if (const UInterchangeBaseNode* BaseNode = BaseNodeContainer->GetNode(MeshUid))
-					{
-						FString RealMeshUid;
-						if (const UInterchangeSceneNode* SceneNode = Cast<UInterchangeSceneNode>(BaseNode))
-						{
-							SceneNode->GetCustomAssetInstanceUid(RealMeshUid);
-						}
-						else
-						{
-							RealMeshUid = MeshUid;
-						}
-
-						if (const UInterchangeMeshNode* MeshNode = Cast<UInterchangeMeshNode>(BaseNodeContainer->GetNode(RealMeshUid)))
-						{
-							TArray<FString> MorphTargetUids;
-							MeshNode->GetMorphTargetDependencies(MorphTargetUids);
-
-							TArray<const UInterchangeMeshNode*> LocalMorphTargetNodes;
-
-							for (const FString& MorphTargetUid : MorphTargetUids)
-							{
-								if (const UInterchangeMeshNode* MorphTargetNode = Cast<UInterchangeMeshNode>(BaseNodeContainer->GetNode(MorphTargetUid)))
-								{
-									TOptional<FString> MorphTargetAnimationPayloadKey = MorphTargetNode->GetAnimationCurvePayLoadKey();
-									if (MorphTargetAnimationPayloadKey.IsSet())
-									{
-										LocalMorphTargetNodes.AddUnique(MorphTargetNode);
-									}
-								}
-							}
-							if (LocalMorphTargetNodes.Num() > 0)
-							{
-								// This is weird logic. Why not break?
-								MorphTargetNodes.Reset(LocalMorphTargetNodes.Num());
-								MorphTargetNodes.Append(LocalMorphTargetNodes);
-							}
-						}
-					}
-				}
-			}
-		}
-	}
-
 	double SampleRate = 30.;
 	double StartTime = 0.;
 	double StopTime = 0.;
@@ -730,10 +217,6 @@ void UInterchangeGenericAnimationPipeline::CreateAnimSequenceFactoryNode(UInterc
 					SampleRate = static_cast<double>(Numerator) / static_cast<double>(Denominator);
 				}
 			}
-			else if ((!bUse30HzToBakeBoneAnimation && CustomBoneAnimationSampleRate > 0))
-			{
-				SampleRate = static_cast<double>(CustomBoneAnimationSampleRate);
-			}
 
 			if (AnimationRange == EInterchangeAnimationRange::Timeline)
 			{
@@ -769,6 +252,11 @@ void UInterchangeGenericAnimationPipeline::CreateAnimSequenceFactoryNode(UInterc
 			bTimeRangeIsValid &= TrackNode.GetCustomAnimationStopTime(StopTime);
 		}
 
+		if (!bUse30HzToBakeBoneAnimation && CustomBoneAnimationSampleRate > 0)
+		{
+			SampleRate = static_cast<double>(CustomBoneAnimationSampleRate);
+		}
+
 		FFrameRate FrameRate = UE::Interchange::Animation::ConvertSampleRatetoFrameRate(SampleRate);
 
 		const double SequenceLength = FMath::Max<double>(StopTime - StartTime, MINIMUM_ANIMATION_LENGTH);
@@ -849,6 +337,7 @@ void UInterchangeGenericAnimationPipeline::CreateAnimSequenceFactoryNode(UInterc
 	}
 
 	AnimSequenceFactoryNode->SetCustomImportAttributeCurves(bImportCustomAttribute);
+	AnimSequenceFactoryNode->SetCustomAddCurveMetadataToSkeleton(bAddCurveMetadataToSkeleton);
 	AnimSequenceFactoryNode->SetCustomDoNotImportCurveWithZero(bDoNotImportCurveWithZero);
 	AnimSequenceFactoryNode->SetCustomRemoveCurveRedundantKeys(bRemoveCurveRedundantKeys);
 	AnimSequenceFactoryNode->SetCustomDeleteExistingMorphTargetCurves(bDeleteExistingMorphTargetCurves);
@@ -861,15 +350,6 @@ void UInterchangeGenericAnimationPipeline::CreateAnimSequenceFactoryNode(UInterc
 		AnimSequenceFactoryNode->SetAnimatedMaterialCurveSuffixe(MaterialSuffixe);
 	}
 
-	//Add the animated morph targets uid so the factory can import them
-	if (MorphTargetNodes.Num() > 0)
-	{
-		for (const UInterchangeMeshNode* MorphTargetNode : MorphTargetNodes)
-		{
-			AnimSequenceFactoryNode->SetAnimatedMorphTargetDependencyUid(MorphTargetNode->GetUniqueID());
-		}
-	}
-
 	//USkeleton cannot be created without a valid skeletal mesh
 	const FString SkeletonUid = SkeletonFactoryNode->GetUniqueID();
 	AnimSequenceFactoryNode->AddFactoryDependencyUid(SkeletonUid);
diff --git a/Engine/Source/Runtime/Interchange/Core/Private/Nodes/InterchangeAnimationAPI.cpp b/Engine/Source/Runtime/Interchange/Core/Private/Nodes/InterchangeAnimationAPI.cpp
deleted file mode 100644
diff --git a/Engine/Source/Runtime/Interchange/Core/Public/Nodes/InterchangeAnimationAPI.h b/Engine/Source/Runtime/Interchange/Core/Public/Nodes/InterchangeAnimationAPI.h
deleted file mode 100644
