Interchange: Fix assert if trying to create a UObject outside of the game thread.

#rnx

diff --git a/Engine/Plugins/Interchange/Runtime/Source/Import/Private/Animation/InterchangeAnimSequenceFactory.cpp b/Engine/Plugins/Interchange/Runtime/Source/Import/Private/Animation/InterchangeAnimSequenceFactory.cpp
--- a/Engine/Plugins/Interchange/Runtime/Source/Import/Private/Animation/InterchangeAnimSequenceFactory.cpp
+++ b/Engine/Plugins/Interchange/Runtime/Source/Import/Private/Animation/InterchangeAnimSequenceFactory.cpp
@@ -861,8 +861,15 @@ UObject* UInterchangeAnimSequenceFactory::ImportObjectSourceData(const FImportAs
 	{
 		//NewObject is not thread safe, the asset registry directory watcher tick on the main thread can trig before we finish initializing the UObject and will crash
 		//The UObject should have been create by calling CreateEmptyAsset on the main thread.
-		check(IsInGameThread());
-		AnimSequenceObject = NewObject<UObject>(Arguments.Parent, AnimSequenceClass, *Arguments.AssetName, RF_Public | RF_Standalone);
+		if (IsInGameThread())
+		{
+			AnimSequenceObject = NewObject<UObject>(Arguments.Parent, AnimSequenceClass, *Arguments.AssetName, RF_Public | RF_Standalone);
+		}
+		else
+		{
+			UE_LOG(LogInterchangeImport, Error, TEXT("Could not create AnimSequence asset [%s] outside of the game thread"), *Arguments.AssetName);
+			return nullptr;
+		}
 	}
 	else if (ExistingAsset->GetClass()->IsChildOf(AnimSequenceClass))
 	{
diff --git a/Engine/Plugins/Interchange/Runtime/Source/Import/Private/Animation/InterchangeAnimationTrackSetFactory.cpp b/Engine/Plugins/Interchange/Runtime/Source/Import/Private/Animation/InterchangeAnimationTrackSetFactory.cpp
--- a/Engine/Plugins/Interchange/Runtime/Source/Import/Private/Animation/InterchangeAnimationTrackSetFactory.cpp
+++ b/Engine/Plugins/Interchange/Runtime/Source/Import/Private/Animation/InterchangeAnimationTrackSetFactory.cpp
@@ -670,8 +670,15 @@ UObject* UInterchangeAnimationTrackSetFactory::ImportObjectSourceData(const FImp
 	{
 		//NewObject is not thread safe, the asset registry directory watcher tick on the main thread can trig before we finish initializing the UObject and will crash
 		//The UObject should have been create by calling CreateEmptyAsset on the main thread.
-		check(IsInGameThread());
-		LevelSequence = NewObject<ULevelSequence>(Arguments.Parent, LevelSequenceClass, *Arguments.AssetName, RF_Public | RF_Standalone);
+		if (IsInGameThread())
+		{
+			LevelSequence = NewObject<ULevelSequence>(Arguments.Parent, LevelSequenceClass, *Arguments.AssetName, RF_Public | RF_Standalone);
+		}
+		else
+		{
+			UE_LOG(LogInterchangeImport, Error, TEXT("Could not create LevelSequence asset [%s] outside of the game thread"), *Arguments.AssetName);
+			return nullptr;
+		}
 	}
 	else if (ExistingAsset->GetClass()->IsChildOf(LevelSequenceClass))
 	{
diff --git a/Engine/Plugins/Interchange/Runtime/Source/Import/Private/Material/InterchangeMaterialFactory.cpp b/Engine/Plugins/Interchange/Runtime/Source/Import/Private/Material/InterchangeMaterialFactory.cpp
--- a/Engine/Plugins/Interchange/Runtime/Source/Import/Private/Material/InterchangeMaterialFactory.cpp
+++ b/Engine/Plugins/Interchange/Runtime/Source/Import/Private/Material/InterchangeMaterialFactory.cpp
@@ -345,20 +345,34 @@ UObject* UInterchangeMaterialFactory::ImportAssetObject_GameThread(const FImport
 {
 	UObject* Material = nullptr;
 
+	auto CouldNotCreateMaterialLog = [this, &Arguments](const FText& Info)
+	{
+		UInterchangeResultError_Generic* Message = AddMessage<UInterchangeResultError_Generic>();
+		Message->SourceAssetName = Arguments.SourceData->GetFilename();
+		Message->DestinationAssetName = Arguments.AssetName;
+		Message->AssetType = GetFactoryClass();
+		Message->Text = FText::Format(LOCTEXT("MatFactory_CouldNotCreateMat", "Could not create Material asset %s. Reason: %s"), FText::FromString(Arguments.AssetName), Info);
+	};
+	
+	const FText MissMatchClassText = LOCTEXT("MatFactory_CouldNotCreateMat_MissMatchClass", "Missmatch between interchange material factory node class and factory class.");
+
 	if (!Arguments.AssetNode || !Arguments.AssetNode->GetObjectClass()->IsChildOf(GetFactoryClass()))
 	{
+		CouldNotCreateMaterialLog(MissMatchClassText);
 		return nullptr;
 	}
 
 	const UInterchangeBaseMaterialFactoryNode* MaterialFactoryNode = Cast<UInterchangeBaseMaterialFactoryNode>(Arguments.AssetNode);
 	if (MaterialFactoryNode == nullptr)
 	{
+		CouldNotCreateMaterialLog(LOCTEXT("MatFactory_CouldNotCreateMat_CannotCastFactoryNode", "Cannot cast interchange factory node to UInterchangeBaseMaterialFactoryNode."));
 		return nullptr;
 	}
 
 	const UClass* MaterialClass = MaterialFactoryNode->GetObjectClass();
 	if (!ensure(MaterialClass && MaterialClass->IsChildOf(GetFactoryClass())))
 	{
+		CouldNotCreateMaterialLog(MissMatchClassText);
 		return nullptr;
 	}
 
@@ -393,7 +407,7 @@ UObject* UInterchangeMaterialFactory::ImportAssetObject_GameThread(const FImport
 
 	if (!Material)
 	{
-		UE_LOG(LogInterchangeImport, Warning, TEXT("Could not create Material asset %s"), *Arguments.AssetName);
+		CouldNotCreateMaterialLog(LOCTEXT("MatFactory_CouldNotCreateMat_MaterialCreationFail", "Material creation fail."));
 		return nullptr;
 	}
 
@@ -442,8 +456,15 @@ UObject* UInterchangeMaterialFactory::ImportAssetObject_Async(const FImportAsset
 	{
 		//NewObject is not thread safe, the asset registry directory watcher tick on the main thread can trig before we finish initializing the UObject and will crash
 		//The UObject should have been created by calling CreateEmptyAsset on the main thread.
-		check(IsInGameThread());
-		MaterialObject = NewObject<UObject>(Arguments.Parent, MaterialClass, *Arguments.AssetName, RF_Public | RF_Standalone);
+		if (IsInGameThread())
+		{
+			MaterialObject = NewObject<UObject>(Arguments.Parent, MaterialClass, *Arguments.AssetName, RF_Public | RF_Standalone);
+		}
+		else
+		{
+			UE_LOG(LogInterchangeImport, Error, TEXT("Could not create Material asset [%s] outside of the game thread"), *Arguments.AssetName);
+			return nullptr;
+		}
 	}
 	else if(ExistingAsset->GetClass()->IsChildOf(MaterialClass))
 	{
@@ -1152,20 +1173,34 @@ UObject* UInterchangeMaterialFunctionFactory::ImportAssetObject_GameThread(const
 {
 	UObject* Material = nullptr;
 
+	auto CouldNotCreateMaterialLog = [this, &Arguments](const FText& Info)
+	{
+		UInterchangeResultError_Generic* Message = AddMessage<UInterchangeResultError_Generic>();
+		Message->SourceAssetName = Arguments.SourceData->GetFilename();
+		Message->DestinationAssetName = Arguments.AssetName;
+		Message->AssetType = GetFactoryClass();
+		Message->Text = FText::Format(LOCTEXT("MatFunc_CouldNotCreateMat", "Could not create Material asset %s. Reason: %s"), FText::FromString(Arguments.AssetName), Info);
+	};
+
+	const FText MissMatchClassText = LOCTEXT("MatFunc_CouldNotCreateMat_MissMatchClass", "Missmatch between interchange material factory node class and factory class.");
+
 	if (!Arguments.AssetNode || !Arguments.AssetNode->GetObjectClass()->IsChildOf(GetFactoryClass()))
 	{
+		CouldNotCreateMaterialLog(MissMatchClassText);
 		return nullptr;
 	}
 
 	const UInterchangeMaterialFunctionFactoryNode* MaterialFactoryNode = Cast<UInterchangeMaterialFunctionFactoryNode>(Arguments.AssetNode);
 	if (!ensure(MaterialFactoryNode))
 	{
+		CouldNotCreateMaterialLog(LOCTEXT("MatFunc_CouldNotCreateMat_CannotCastFactoryNode", "Cannot cast interchange factory node to UInterchangeBaseMaterialFactoryNode."));
 		return nullptr;
 	}
 
 	const UClass* MaterialClass = MaterialFactoryNode->GetObjectClass();
 	if (!ensure(MaterialClass && MaterialClass->IsChildOf(GetFactoryClass())))
 	{
+		CouldNotCreateMaterialLog(MissMatchClassText);
 		return nullptr;
 	}
 
@@ -1185,7 +1220,7 @@ UObject* UInterchangeMaterialFunctionFactory::ImportAssetObject_GameThread(const
 
 	if (!Material)
 	{
-		UE_LOG(LogInterchangeImport, Warning, TEXT("Could not create Material asset %s"), *Arguments.AssetName);
+		CouldNotCreateMaterialLog(LOCTEXT("MatFunc_CouldNotCreateMat_MaterialCreationFail", "Material creation fail."));
 		return nullptr;
 	}
 
@@ -1221,8 +1256,15 @@ UObject* UInterchangeMaterialFunctionFactory::ImportAssetObject_Async(const FImp
 	{
 		//NewObject is not thread safe, the asset registry directory watcher tick on the main thread can trig before we finish initializing the UObject and will crash
 		//The UObject should have been created by calling CreateEmptyAsset on the main thread.
-		check(IsInGameThread());
-		MaterialObject = NewObject<UObject>(Arguments.Parent, MaterialClass, *Arguments.AssetName, RF_Public | RF_Standalone);
+		if (IsInGameThread())
+		{
+			MaterialObject = NewObject<UObject>(Arguments.Parent, MaterialClass, *Arguments.AssetName, RF_Public | RF_Standalone);
+		}
+		else
+		{
+			UE_LOG(LogInterchangeImport, Error, TEXT("Could not create Material asset [%s] outside of the game thread"), *Arguments.AssetName);
+			return nullptr;
+		}
 	}
 	else if (ExistingAsset->GetClass()->IsChildOf(MaterialClass))
 	{
diff --git a/Engine/Plugins/Interchange/Runtime/Source/Import/Private/Mesh/InterchangePhysicsAssetFactory.cpp b/Engine/Plugins/Interchange/Runtime/Source/Import/Private/Mesh/InterchangePhysicsAssetFactory.cpp
--- a/Engine/Plugins/Interchange/Runtime/Source/Import/Private/Mesh/InterchangePhysicsAssetFactory.cpp
+++ b/Engine/Plugins/Interchange/Runtime/Source/Import/Private/Mesh/InterchangePhysicsAssetFactory.cpp
@@ -90,8 +90,15 @@ UObject* UInterchangePhysicsAssetFactory::ImportAssetObject_Async(const FImportA
 	{
 		//NewObject is not thread safe, the asset registry directory watcher tick on the main thread can trig before we finish initializing the UObject and will crash
 		//The UObject should have been create by calling ImportAssetObject_GameThread on the main thread.
-		check(IsInGameThread());
-		PhysicsAssetObject = NewObject<UObject>(Arguments.Parent, PhysicsAssetClass, *Arguments.AssetName, RF_Public | RF_Standalone);
+		if (IsInGameThread())
+		{
+			PhysicsAssetObject = NewObject<UObject>(Arguments.Parent, PhysicsAssetClass, *Arguments.AssetName, RF_Public | RF_Standalone);
+		}
+		else
+		{
+			UE_LOG(LogInterchangeImport, Error, TEXT("Could not create Physics asset [%s] outside of the game thread"), *Arguments.AssetName);
+			return nullptr;
+		}
 	}
 	else if(ExistingAsset->GetClass()->IsChildOf(PhysicsAssetClass))
 	{
diff --git a/Engine/Plugins/Interchange/Runtime/Source/Import/Private/Mesh/InterchangeSkeletalMeshFactory.cpp b/Engine/Plugins/Interchange/Runtime/Source/Import/Private/Mesh/InterchangeSkeletalMeshFactory.cpp
--- a/Engine/Plugins/Interchange/Runtime/Source/Import/Private/Mesh/InterchangeSkeletalMeshFactory.cpp
+++ b/Engine/Plugins/Interchange/Runtime/Source/Import/Private/Mesh/InterchangeSkeletalMeshFactory.cpp
@@ -971,8 +971,15 @@ UObject* UInterchangeSkeletalMeshFactory::ImportAssetObject_Async(const FImportA
 	{
 		//NewObject is not thread safe, the asset registry directory watcher tick on the main thread can trig before we finish initializing the UObject and will crash
 		//The UObject should have been create by calling CreateEmptyAsset on the main thread.
-		check(IsInGameThread());
-		SkeletalMeshObject = NewObject<UObject>(Arguments.Parent, SkeletalMeshClass, *Arguments.AssetName, RF_Public | RF_Standalone);
+		if (IsInGameThread())
+		{
+			SkeletalMeshObject = NewObject<UObject>(Arguments.Parent, SkeletalMeshClass, *Arguments.AssetName, RF_Public | RF_Standalone);
+		}
+		else
+		{
+			UE_LOG(LogInterchangeImport, Error, TEXT("Could not create SkeletalMesh asset [%s] outside of the game thread"), *Arguments.AssetName);
+			return nullptr;
+		}
 	}
 	else if(ExistingAsset->GetClass()->IsChildOf(SkeletalMeshClass))
 	{
diff --git a/Engine/Plugins/Interchange/Runtime/Source/Import/Private/Mesh/InterchangeSkeletonFactory.cpp b/Engine/Plugins/Interchange/Runtime/Source/Import/Private/Mesh/InterchangeSkeletonFactory.cpp
--- a/Engine/Plugins/Interchange/Runtime/Source/Import/Private/Mesh/InterchangeSkeletonFactory.cpp
+++ b/Engine/Plugins/Interchange/Runtime/Source/Import/Private/Mesh/InterchangeSkeletonFactory.cpp
@@ -93,9 +93,16 @@ UObject* UInterchangeSkeletonFactory::ImportAssetObject_Async(const FImportAsset
 	{
 		//NewObject is not thread safe, the asset registry directory watcher tick on the main thread can trig before we finish initializing the UObject and will crash
 		//The UObject should have been create by calling CreateEmptyAsset on the main thread.
-		check(IsInGameThread());
-		//We should not do a NewObject if we are doing a reimport
-		check(!bIsReImport);
+		if (!IsInGameThread())
+		{
+			UE_LOG(LogInterchangeImport, Error, TEXT("Could not create Skeleton asset [%s] outside of the game thread"), *Arguments.AssetName);
+			return nullptr;
+		}
+		if (bIsReImport)
+		{
+			UE_LOG(LogInterchangeImport, Error, TEXT("Could not create Skeleton asset [%s] when re-importing"), *Arguments.AssetName);
+			return nullptr;
+		}
 		SkeletonObject = NewObject<UObject>(Arguments.Parent, SkeletonClass, *Arguments.AssetName, RF_Public | RF_Standalone);
 	}
 	else if(ExistingAsset->GetClass()->IsChildOf(SkeletonClass))
diff --git a/Engine/Plugins/Interchange/Runtime/Source/Import/Private/Mesh/InterchangeStaticMeshFactory.cpp b/Engine/Plugins/Interchange/Runtime/Source/Import/Private/Mesh/InterchangeStaticMeshFactory.cpp
--- a/Engine/Plugins/Interchange/Runtime/Source/Import/Private/Mesh/InterchangeStaticMeshFactory.cpp
+++ b/Engine/Plugins/Interchange/Runtime/Source/Import/Private/Mesh/InterchangeStaticMeshFactory.cpp
@@ -137,8 +137,15 @@ UObject* UInterchangeStaticMeshFactory::ImportAssetObject_Async(const FImportAss
 	{
 		// NewObject is not thread safe, the asset registry directory watcher tick on the main thread can trig before we finish initializing the UObject and will crash
 		// The UObject should have been created by calling CreateEmptyAsset on the main thread.
-		check(IsInGameThread());
-		StaticMeshObject = NewObject<UObject>(Arguments.Parent, StaticMeshClass, *Arguments.AssetName, RF_Public | RF_Standalone);
+		if (IsInGameThread())
+		{
+			StaticMeshObject = NewObject<UObject>(Arguments.Parent, StaticMeshClass, *Arguments.AssetName, RF_Public | RF_Standalone);
+		}
+		else
+		{
+			UE_LOG(LogInterchangeImport, Error, TEXT("Could not create StaticMesh asset [%s] outside of the game thread"), *Arguments.AssetName);
+			return nullptr;
+		}
 	}
 	else if(ExistingAsset->GetClass()->IsChildOf(StaticMeshClass))
 	{
diff --git a/Engine/Plugins/Interchange/Runtime/Source/Import/Private/Texture/InterchangeTextureFactory.cpp b/Engine/Plugins/Interchange/Runtime/Source/Import/Private/Texture/InterchangeTextureFactory.cpp
--- a/Engine/Plugins/Interchange/Runtime/Source/Import/Private/Texture/InterchangeTextureFactory.cpp
+++ b/Engine/Plugins/Interchange/Runtime/Source/Import/Private/Texture/InterchangeTextureFactory.cpp
@@ -1464,23 +1464,32 @@ UObject* UInterchangeTextureFactory::ImportAssetObject_GameThread(const FImportA
 
 	UObject* Texture = nullptr;
 
+	auto CouldNotCreateTextureLog = [this, &Arguments](const FText& Info)
+	{
+		UInterchangeResultError_Generic* Message = AddMessage<UInterchangeResultError_Generic>();
+		Message->SourceAssetName = Arguments.SourceData->GetFilename();
+		Message->DestinationAssetName = Arguments.AssetName;
+		Message->AssetType = GetFactoryClass();
+		Message->Text = FText::Format(LOCTEXT("TexFact_CouldNotCreateMat", "UInterchangeTextureFactory: Could not create texture asset %s. Reason: %s"), FText::FromString(Arguments.AssetName), Info);
+	};
+
 	if (!Arguments.AssetNode)
 	{
-		UE_LOG(LogInterchangeImport, Error, TEXT("UInterchangeTextureFactory: Asset node parameter is null."));
+		CouldNotCreateTextureLog(LOCTEXT("TextureFactory_AssetNodeNull", "Asset node parameter is null."));
 		return nullptr;
 	}
 
 	const UClass* TextureClass = Arguments.AssetNode->GetObjectClass();
 	if (!TextureClass || !TextureClass->IsChildOf(UTexture::StaticClass()))
 	{
-		UE_LOG(LogInterchangeImport, Error, TEXT("UInterchangeTextureFactory: Asset node parameter class doesnt derive from UTexture."));
+		CouldNotCreateTextureLog(LOCTEXT("TextureFactory_NodeClassMissmatch", "Asset node parameter class doesnt derive from UTexture."));
 		return nullptr;
 	}
 
 	UClass* SupportedFactoryNodeClass = GetSupportedFactoryNodeClass(Arguments.AssetNode);
 	if (SupportedFactoryNodeClass == nullptr)
 	{
-		UE_LOG(LogInterchangeImport, Error, TEXT("UInterchangeTextureFactory: Asset node parameter is not a UInterchangeTextureFactoryNode or UInterchangeTextureCubeFactoryNode."));
+		CouldNotCreateTextureLog(LOCTEXT("TextureFactory_NodeClassWrong", "Asset node parameter is not a UInterchangeTextureFactoryNode or UInterchangeTextureCubeFactoryNode."));
 		return nullptr;
 	}
 
@@ -1488,15 +1497,15 @@ UObject* UInterchangeTextureFactory::ImportAssetObject_GameThread(const FImportA
 	FTextureNodeVariant TextureNodeVariant = GetTextureNodeVariantFromFactoryVariant(GetAsTextureFactoryNodeVariant(Arguments.AssetNode, SupportedFactoryNodeClass), Arguments.NodeContainer);
 	if (TextureNodeVariant.IsType<FEmptyVariantState>())
 	{
-		UE_LOG(LogInterchangeImport, Error, TEXT("UInterchangeTextureFactory: Asset factory node (%s) do not reference a valid texture translated node.")
-				, *SupportedFactoryNodeClass->GetAuthoredName()
-			);
+		FText Info = FText::Format(LOCTEXT("TextureFactory_InvalidTextureTranslated", "Asset factory node (%s) do not reference a valid texture translated node.")
+			, FText::FromString(SupportedFactoryNodeClass->GetAuthoredName()));
+		CouldNotCreateTextureLog(Info);
 		return nullptr;
 	}
 
 	if (!HasPayloadKey(TextureNodeVariant))
 	{
-		UE_LOG(LogInterchangeImport, Error, TEXT("UInterchangeTextureFactory: Texture translated node doesnt have a payload key."));
+		CouldNotCreateTextureLog(LOCTEXT("TextureFactory_InvalidPayloadKey", "Texture translated node doesnt have a payload key."));
 		return nullptr;
 	}
 
@@ -1542,7 +1551,7 @@ UObject* UInterchangeTextureFactory::ImportAssetObject_GameThread(const FImportA
 
 	if (!Texture)
 	{
-		UE_LOG(LogInterchangeImport, Warning, TEXT("Could not create texture asset %s"), *Arguments.AssetName);
+		CouldNotCreateTextureLog(LOCTEXT("TextureFactory_TextureCreateFail", "Texture creation fail."));
 		return nullptr;
 	}
 
@@ -1556,23 +1565,32 @@ UObject* UInterchangeTextureFactory::ImportAssetObject_Async(const FImportAssetO
 
 	using namespace UE::Interchange::Private::InterchangeTextureFactory;
 
+	auto ImportTextureErrorLog = [this, &Arguments](const FText& Info)
+	{
+		UInterchangeResultError_Generic* Message = AddMessage<UInterchangeResultError_Generic>();
+		Message->SourceAssetName = Arguments.SourceData->GetFilename();
+		Message->DestinationAssetName = Arguments.AssetName;
+		Message->AssetType = GetFactoryClass();
+		Message->Text = Info;
+	};
+
 	if (!Arguments.AssetNode)
 	{
-		UE_LOG(LogInterchangeImport, Error, TEXT("UInterchangeTextureFactory: Asset node parameter is null."));
+		ImportTextureErrorLog(LOCTEXT("TextureFactory_Async_AssetNodeNull", "UInterchangeTextureFactory: Asset node parameter is null."));
 		return nullptr;
 	}
 
 	const UClass* TextureClass = Arguments.AssetNode->GetObjectClass();
 	if (!TextureClass || !TextureClass->IsChildOf(UTexture::StaticClass()))
 	{
-		UE_LOG(LogInterchangeImport, Error, TEXT("UInterchangeTextureFactory: Asset node parameter class doesnt derive from UTexture."));
+		ImportTextureErrorLog(LOCTEXT("TextureFactory_Async_MissMatchClass", "UInterchangeTextureFactory: Asset node parameter class doesnt derive from UTexture."));
 		return nullptr;
 	}
 
 	UClass* SupportedFactoryNodeClass = GetSupportedFactoryNodeClass(Arguments.AssetNode);
 	if (SupportedFactoryNodeClass == nullptr)
 	{
-		UE_LOG(LogInterchangeImport, Error, TEXT("UInterchangeTextureFactory: Asset node parameter is not a child of UInterchangeTextureFactoryNode"));
+		ImportTextureErrorLog(LOCTEXT("TextureFactory_Async_NodeWrongClass", "UInterchangeTextureFactory: Asset node parameter is not a child of UInterchangeTextureFactoryNode."));
 		return nullptr;
 	}
 
@@ -1580,16 +1598,16 @@ UObject* UInterchangeTextureFactory::ImportAssetObject_Async(const FImportAssetO
 	FTextureNodeVariant TextureNodeVariant = GetTextureNodeVariantFromFactoryVariant(TextureFactoryNodeVariant, Arguments.NodeContainer);
 	if (TextureNodeVariant.IsType<FEmptyVariantState>())
 	{
-		UE_LOG(LogInterchangeImport, Error, TEXT("UInterchangeTextureFactory: Asset factory node (%s) do not reference a valid texture translated node.")
-				, *SupportedFactoryNodeClass->GetAuthoredName()
-			);
+		FText Info = FText::Format(LOCTEXT("TextureFactory_Async_InvalidTextureTranslated", "UInterchangeTextureFactory: Asset factory node (%s) do not reference a valid texture translated node.")
+			, FText::FromString(SupportedFactoryNodeClass->GetAuthoredName()));
+		ImportTextureErrorLog(Info);
 		return nullptr;
 	}
 
 	const TOptional<FString>& PayLoadKey = GetPayloadKey(TextureNodeVariant);
 	if (!PayLoadKey.IsSet())
 	{
-		UE_LOG(LogInterchangeImport, Error, TEXT("UInterchangeTextureFactory: Texture translated node (UInterchangeTexture2DNode) doesnt have a payload key."));
+		ImportTextureErrorLog(LOCTEXT("TextureFactory_Async_InvalidPayloadKey", "UInterchangeTextureFactory: Texture translated node (UInterchangeTexture2DNode) doesnt have a payload key."));
 		return nullptr;
 	}
 
@@ -1597,7 +1615,7 @@ UObject* UInterchangeTextureFactory::ImportAssetObject_Async(const FImportAssetO
 
 	if(TexturePayload.IsType<FEmptyVariantState>())
 	{
-		UE_LOG(LogInterchangeImport, Error, TEXT("UInterchangeTextureFactory: Invalid translator couldn't retrive a payload."));
+		ImportTextureErrorLog(LOCTEXT("TextureFactory_Async_CannotRetrievePayload", "UInterchangeTextureFactory: Invalid translator couldn't retrive a payload."));
 		return nullptr;
 	}
 
@@ -1611,8 +1629,15 @@ UObject* UInterchangeTextureFactory::ImportAssetObject_Async(const FImportAssetO
 	{
 		//NewObject is not thread safe, the asset registry directory watcher tick on the main thread can trig before we finish initializing the UObject and will crash
 		//The UObject should have been create by calling CreateEmptyAsset on the main thread.
-		check(IsInGameThread());
-		Texture = NewObject<UTexture>(Arguments.Parent, TextureClass, *Arguments.AssetName, RF_Public | RF_Standalone);
+		if (IsInGameThread())
+		{
+			Texture = NewObject<UTexture>(Arguments.Parent, TextureClass, *Arguments.AssetName, RF_Public | RF_Standalone);
+		}
+		else
+		{
+			ImportTextureErrorLog(LOCTEXT("TextureFactory_Async_CannotCreateAsync", "UInterchangeTextureFactory: Could not create Texture asset outside of the game thread."));
+			return nullptr;
+		}
 	}
 	else if(ExistingAsset->GetClass()->IsChildOf(TextureClass))
 	{
@@ -1623,7 +1648,7 @@ UObject* UInterchangeTextureFactory::ImportAssetObject_Async(const FImportAssetO
 
 	if (!Texture)
 	{
-		UE_LOG(LogInterchangeImport, Error, TEXT("UInterchangeTextureFactory: Could not create texture asset %s"), *Arguments.AssetName);
+		ImportTextureErrorLog(LOCTEXT("TextureFactory_Async_FailCreatingAsset", "UInterchangeTextureFactory: Could not create texture asset."));
 		return nullptr;
 	}
 
diff --git a/Engine/Source/Runtime/Interchange/Engine/Private/InterchangeManager.cpp b/Engine/Source/Runtime/Interchange/Engine/Private/InterchangeManager.cpp
--- a/Engine/Source/Runtime/Interchange/Engine/Private/InterchangeManager.cpp
+++ b/Engine/Source/Runtime/Interchange/Engine/Private/InterchangeManager.cpp
@@ -938,6 +938,9 @@ void UInterchangeManager::StartQueuedTasks(bool bCancelAllTasks /*= false*/)
 			check(QueuedTaskData.AsyncHelper->Translators.Num() == QueuedTaskData.AsyncHelper->SourceDatas.Num());
 			for (int32 SourceDataIndex = 0; SourceDataIndex < QueuedTaskData.AsyncHelper->SourceDatas.Num(); ++SourceDataIndex)
 			{
+				//Log the source we begin importing
+				UE_LOG(LogInterchangeEngine, Display, TEXT("Interchange start importing source [%s]"), *QueuedTaskData.AsyncHelper->SourceDatas[SourceDataIndex]->ToDisplayString());
+
 				int32 TranslatorTaskIndex = QueuedTaskData.AsyncHelper->TranslatorTasks.Add(TGraphTask<UE::Interchange::FTaskTranslator>::CreateTask().ConstructAndDispatchWhenReady(SourceDataIndex, WeakAsyncHelper));
 				PipelinePrerequistes.Add(QueuedTaskData.AsyncHelper->TranslatorTasks[TranslatorTaskIndex]);
 			}
