Interchange: Make sure InterchangeWorker do not use the task graph to "process command" since interchange factory task graph are waitting on this "process command", it can deadlock if task graph run out of thread pool. We will simply process results in the interchange worker thread.

#rnx

diff --git a/Engine/Plugins/Interchange/Runtime/Source/Dispatcher/Private/InterchangeWorkerHandler.cpp b/Engine/Plugins/Interchange/Runtime/Source/Dispatcher/Private/InterchangeWorkerHandler.cpp
--- a/Engine/Plugins/Interchange/Runtime/Source/Dispatcher/Private/InterchangeWorkerHandler.cpp
+++ b/Engine/Plugins/Interchange/Runtime/Source/Dispatcher/Private/InterchangeWorkerHandler.cpp
@@ -6,9 +6,9 @@
 #include "InterchangeDispatcherConfig.h"
 #include "InterchangeDispatcherLog.h"
 
-#include "Async/TaskGraphInterfaces.h"
 #include "HAL/FileManager.h"
 #include "HAL/PlatformProcess.h"
+#include "HAL/PlatformTime.h"
 #include "Misc/Paths.h"
 #include "Misc/ScopeLock.h"
 #include "Sockets.h"
@@ -18,46 +18,6 @@ namespace UE
 {
 	namespace Interchange
 	{
-		namespace Dispatcher
-		{
-			class FTaskProcessCommand
-			{
-			private:
-				FInterchangeWorkerHandler* TaskOwner = nullptr;
-				TSharedPtr<ICommand> Command = nullptr;
-
-			public:
-				FTaskProcessCommand(FInterchangeWorkerHandler* InTaskOwner, TSharedPtr<ICommand> InCommand)
-					: TaskOwner(InTaskOwner)
-					, Command(InCommand)
-				{
-				}
-
-				static FORCEINLINE ENamedThreads::Type GetDesiredThread()
-				{
-					return ENamedThreads::AnyBackgroundThreadNormalTask;
-				}
-				static FORCEINLINE ESubsequentsMode::Type GetSubsequentsMode()
-				{
-					return ESubsequentsMode::TrackSubsequents;
-				}
-
-				FORCEINLINE TStatId GetStatId() const
-				{
-					RETURN_QUICK_DECLARE_CYCLE_STAT(FTaskProcessCommand, STATGROUP_TaskGraphTasks);
-				}
-
-				void DoTask(ENamedThreads::Type CurrentThread, const FGraphEventRef& MyCompletionGraphEvent)
-				{
-					//Verify if the task was cancel
-					if (TaskOwner == nullptr || !TaskOwner->IsAlive() || !Command.IsValid())
-					{
-						return;
-					}
-					TaskOwner->ProcessCommand(*Command);
-				}
-			};
-		}
 		static FString GetWorkerExecutablePath()
 		{
 			static FString ProcessorPath = [&]()
@@ -323,8 +283,7 @@ namespace UE
 							// consume task
 							if (TSharedPtr<ICommand> Command = CommandIO.GetNextCommand(Config::IdleLoopDelay))
 							{
-								//Make the processing asynchronous
-								TGraphTask<UE::Interchange::Dispatcher::FTaskProcessCommand>::CreateTask().ConstructAndDispatchWhenReady(this, Command);
+								ProcessCommand(*Command);
 							}
 
 							//Need to terminate?
