[Interchange] FbxAnimation loctext warning fix

diff --git a/Engine/Plugins/Interchange/Runtime/Source/Parsers/Fbx/Private/FbxAnimation.cpp b/Engine/Plugins/Interchange/Runtime/Source/Parsers/Fbx/Private/FbxAnimation.cpp
--- a/Engine/Plugins/Interchange/Runtime/Source/Parsers/Fbx/Private/FbxAnimation.cpp
+++ b/Engine/Plugins/Interchange/Runtime/Source/Parsers/Fbx/Private/FbxAnimation.cpp
@@ -726,7 +726,7 @@ namespace UE::Interchange::Private
 		if (!Geometry)
 		{
 			UInterchangeResultError_Generic* Message = Parser.AddMessage<UInterchangeResultError_Generic>();
-			Message->Text = LOCTEXT("InternalFetchMorphTargetCurvePayloadToFile_FBXSDKSceneNull", "InternalFetchMorphTargetCurvePayloadToFile, Geometry acquisition failed.");
+			Message->Text = LOCTEXT("InternalFetchMorphTargetCurvePayloadToFile_FBXGeometryNull", "Cannot fetch FBX geometry from the scene.");
 			return false;
 		}
 		FbxAnimCurve* AnimCurve = Geometry->GetShapeChannel(FetchPayloadData.MorphTargetIndex, FetchPayloadData.ChannelIndex, FetchPayloadData.AnimLayer);
