Interchange: Fix the import sphere collision

#rnx

diff --git a/Engine/Plugins/Interchange/Runtime/Source/Import/Private/Mesh/InterchangeStaticMeshFactory.cpp b/Engine/Plugins/Interchange/Runtime/Source/Import/Private/Mesh/InterchangeStaticMeshFactory.cpp
--- a/Engine/Plugins/Interchange/Runtime/Source/Import/Private/Mesh/InterchangeStaticMeshFactory.cpp
+++ b/Engine/Plugins/Interchange/Runtime/Source/Import/Private/Mesh/InterchangeStaticMeshFactory.cpp
@@ -947,7 +947,7 @@ bool UInterchangeStaticMeshFactory::AddSphereGeomFromVertices(const FImportAsset
 
 	for (const FVector3f& VertexPosition : VertexPositions.GetRawArray())
 	{
-		FVector3f CToV = VertexPosition - FVector3f(Center);
+		FVector CToV = Transform.TransformPosition(FVector(VertexPosition)) - Center;
 		float RSqr = CToV.SizeSquared();
 
 		MaxR = FMath::Max(RSqr, MaxR);
