[Interchange] Actor hierarchy generation fix (fix for RigidMesh.fbx)

diff --git a/Engine/Plugins/Interchange/Runtime/Source/Pipelines/Private/InterchangeGenericScenesPipeline.cpp b/Engine/Plugins/Interchange/Runtime/Source/Pipelines/Private/InterchangeGenericScenesPipeline.cpp
--- a/Engine/Plugins/Interchange/Runtime/Source/Pipelines/Private/InterchangeGenericScenesPipeline.cpp
+++ b/Engine/Plugins/Interchange/Runtime/Source/Pipelines/Private/InterchangeGenericScenesPipeline.cpp
@@ -150,8 +150,14 @@ void UInterchangeGenericLevelPipeline::ExecuteSceneNodePreImport(const FTransfor
 				}
 			}
 		}
+		else
+		{
+			//If it does not have a SkeletonFactoryNode with the given Uid, then it is not a root joint node:
+			bRootJointNode = false;
+		}
 	}
-	else
+	
+	if (!bRootJointNode)
 	{
 		FString AssetInstanceUid;
 		if (SceneNode->GetCustomAssetInstanceUid(AssetInstanceUid))
