MaterialX:
- Adding material expressions:
-- Adding UMaterialExpressionTextureSampleParameterBlur, that allows the possibility to blur a texture using either a gaussian or box linear filter, only 3x3, 5x5 and 7x7 kernels are supported
- Translator:
-- Connecting <heighttonormal> and <blur> nodes, only <image> nodes aresupported, we duplicate the inputs of the image node to the blur/heighttonormal node and then remove it from the graph node
- Material pipeline:
-- Adding TextureBlurNode to correctly create the proper material expression, the nodes are handled the same way as is if it was a TextureSampleParameter2D
-- Adding the support of TextureObject and NormalFromHeightMap material function

diff --git a/Engine/Plugins/Interchange/Runtime/Source/Import/Private/MaterialX/InterchangeMaterialXDefinitions.h b/Engine/Plugins/Interchange/Runtime/Source/Import/Private/MaterialX/InterchangeMaterialXDefinitions.h
--- a/Engine/Plugins/Interchange/Runtime/Source/Import/Private/MaterialX/InterchangeMaterialXDefinitions.h
+++ b/Engine/Plugins/Interchange/Runtime/Source/Import/Private/MaterialX/InterchangeMaterialXDefinitions.h
@@ -119,6 +119,18 @@ MATERIALX_NAMESPACE_BEGIN
 		}
 	}
 
+	namespace SurfaceUnlit
+	{
+		namespace Input
+		{
+			static constexpr const char* Emission = "emission";
+			static constexpr const char* EmissionColor = "emission_color";
+			static constexpr const char* Transmission = "transmission";
+			static constexpr const char* TransmissionColor = "transmission_color";
+			static constexpr const char* Opacity = "opacity";
+		}
+	}
+
 	namespace Lights
 	{
 		//There's no input per se in a Light, but we can find some common inputs among those lights
@@ -279,11 +291,17 @@ MATERIALX_NAMESPACE_BEGIN
 		static constexpr const char* Time= "time";
 		// PBR
 		static constexpr const char* StandardSurface = "standard_surface";
+		// Shader
+		static constexpr const char* SurfaceUnlit = "surface_unlit";
+		// Convolution
+		static constexpr const char* Blur = "blur";
+		static constexpr const char* HeightToNormal = "heighttonormal";
 	}
 
 	namespace NodeDefinition
 	{
 		static constexpr const char* StandardSurface = "ND_standard_surface_surfaceshader";
+		static constexpr const char* SurfaceUnlit = "ND_surface_unlit";
 		static constexpr const char* PointLight = "ND_point_light";
 		static constexpr const char* DirectionalLight = "ND_directional_light";
 		static constexpr const char* SpotLight = "ND_spot_light";
diff --git a/Engine/Plugins/Interchange/Runtime/Source/Import/Private/MaterialX/InterchangeMaterialXTranslator.cpp b/Engine/Plugins/Interchange/Runtime/Source/Import/Private/MaterialX/InterchangeMaterialXTranslator.cpp
--- a/Engine/Plugins/Interchange/Runtime/Source/Import/Private/MaterialX/InterchangeMaterialXTranslator.cpp
+++ b/Engine/Plugins/Interchange/Runtime/Source/Import/Private/MaterialX/InterchangeMaterialXTranslator.cpp
@@ -11,6 +11,7 @@
 #include "InterchangeSceneNode.h"
 #include "InterchangeMaterialXDefinitions.h"
 #include "InterchangeTexture2DNode.h"
+#include "InterchangeTextureBlurNode.h"
 #include "Nodes/InterchangeSourceNode.h"
 #include "UObject/GCObjectScopeGuard.h"
 
@@ -19,6 +20,8 @@
 #include "Materials/MaterialExpressionTransformPosition.h"
 #include "Materials/MaterialExpressionVectorNoise.h"
 
+#include "Misc/PackageName.h"
+
 #include UE_INLINE_GENERATED_CPP_BY_NAME(InterchangeMaterialXTranslator)
 
 #define LOCTEXT_NAMESPACE "InterchangeMaterialXTranslator"
@@ -33,71 +36,105 @@ static FAutoConsoleVariableRef CCvarInterchangeEnableMaterialXImport(
 
 namespace mx = MaterialX;
 
+namespace UE::Interchange::MaterialX
+{
+	bool IsMaterialFunctionPackageLoaded()
+	{
+		auto IsPackageLoaded = [](const FString& TextPath) -> bool
+		{
+			const FString FunctionPath{ FPackageName::ExportTextPathToObjectPath(TextPath) };
+			if(FPackageName::DoesPackageExist(FunctionPath))
+			{
+				if(FSoftObjectPath(FunctionPath).TryLoad())
+				{
+					return true;
+				}
+				else
+				{
+					UE_LOG(LogInterchangeImport, Warning, TEXT("Couldn't load %s"), *FunctionPath);
+				}
+			}
+			else
+			{
+				UE_LOG(LogInterchangeImport, Warning, TEXT("Couldn't find %s"), *FunctionPath);
+			}
+
+			return false;
+		};
+
+		static const bool bNormalFromHeightMapPackageLoaded = 
+			IsPackageLoaded(TEXT("MaterialFunction'/Engine/Functions/Engine_MaterialFunctions03/Procedurals/NormalFromHeightmap.NormalFromHeightmap'"));
+
+		return bNormalFromHeightMapPackageLoaded;
+	}
+}
+
 UInterchangeMaterialXTranslator::UInterchangeMaterialXTranslator()
 #if WITH_EDITOR
 	:
 InputNamesMaterialX2UE
 {
-	{{TEXT(""),                         TEXT("amplitude")}, TEXT("Amplitude")},
-	{{TEXT(""),                         TEXT("bg")},        TEXT("B")},
-	{{TEXT(""),                         TEXT("center")},    TEXT("Center")},
-	{{TEXT(""),                         TEXT("diminish")},  TEXT("Diminish")},
-	{{TEXT(""),                         TEXT("fg")},        TEXT("A")},
-	{{TEXT(""),                         TEXT("high")},      TEXT("Max")},
-	{{TEXT(""),                         TEXT("in")},        TEXT("Input")},
-	{{TEXT(""),                         TEXT("in1")},       TEXT("A")},
-	{{TEXT(""),                         TEXT("in2")},       TEXT("B")},
-	{{TEXT(""),                         TEXT("in3")},       TEXT("C")},
-	{{TEXT(""),                         TEXT("in4")},       TEXT("D")},
-	{{TEXT(""),                         TEXT("inlow")},     TEXT("InputLow")},
-	{{TEXT(""),                         TEXT("inhigh")},    TEXT("InputHigh")},
-	{{TEXT(""),                         TEXT("lacunarity")},TEXT("Lacunarity")},
-	{{TEXT(""),                         TEXT("low")},       TEXT("Min")},
-	{{TEXT(""),                         TEXT("lumacoeffs")},TEXT("LuminanceFactors")}, // for the moment not yet handled by Interchange, because of the attribute being an advanced pin
-	{{TEXT(""),                         TEXT("mix")},       TEXT("Alpha")},
-	{{TEXT(""),                         TEXT("position")},  TEXT("Position")},
-	{{TEXT(""),                         TEXT("texcoord")},  TEXT("Coordinates")},
-	{{TEXT(""),                         TEXT("octaves")},   TEXT("Octaves")},
-	{{TEXT(""),                         TEXT("outlow")},    TEXT("TargetLow")},
-	{{TEXT(""),                         TEXT("outhigh")},   TEXT("TargetHigh")},
-	{{TEXT(""),                         TEXT("valuel")},    TEXT("A")},
-	{{TEXT(""),                         TEXT("valuer")},    TEXT("B")},
-	{{MaterialX::Category::Atan2,       TEXT("in1")},       TEXT("Y")},
-	{{MaterialX::Category::Atan2,       TEXT("in2")},       TEXT("X")},
-	{{MaterialX::Category::IfGreater,   TEXT("value1")},    TEXT("A")},
-	{{MaterialX::Category::IfGreater,   TEXT("value2")},    TEXT("B")},
-	{{MaterialX::Category::IfGreater,   TEXT("in1")},       TEXT("AGreaterThanB")},
-	{{MaterialX::Category::IfGreater,   TEXT("in2")},       TEXT("ALessThanB")}, //another input is added for the case equal see ConnectIfGreater
-	{{MaterialX::Category::IfGreaterEq, TEXT("value1")},    TEXT("A")},
-	{{MaterialX::Category::IfGreaterEq, TEXT("value2")},    TEXT("B")},
-	{{MaterialX::Category::IfGreaterEq, TEXT("in1")},       TEXT("AGreaterThanB")}, //another is added for the case equal see ConnectIfGreaterEq
-	{{MaterialX::Category::IfGreaterEq, TEXT("in2")},       TEXT("ALessThanB")},
-	{{MaterialX::Category::IfEqual,     TEXT("value1")},    TEXT("A")},
-	{{MaterialX::Category::IfEqual,     TEXT("value2")},    TEXT("B")},
-	{{MaterialX::Category::IfEqual,     TEXT("in1")},       TEXT("AEqualsB")},
-	{{MaterialX::Category::IfEqual,     TEXT("in2")},       TEXT("ALessThanB")},  // another input is added for the case greater see ConnectIfEqual
-	{{MaterialX::Category::Inside,      TEXT("in")},        TEXT("A")},			  // Inside is treated as a Multiply node
-	{{MaterialX::Category::Inside,      TEXT("mask")},      TEXT("B")},			  // Inside is treated as a Multiply node
-	{{MaterialX::Category::Invert,      TEXT("amount")},    TEXT("A")},
-	{{MaterialX::Category::Invert,      TEXT("in")},        TEXT("B")},
-	{{MaterialX::Category::Magnitude,   TEXT("in")},        TEXT("A")},
-	{{MaterialX::Category::Mix,         TEXT("fg")},        TEXT("B")},
-	{{MaterialX::Category::Mix,         TEXT("bg")},        TEXT("A")},
-	{{MaterialX::Category::Mix,         TEXT("mix")},		TEXT("Factor")},
-	{{MaterialX::Category::Noise3D,     TEXT("amplitude")}, TEXT("B")},            // The amplitude of the noise is connected to a multiply node
-	{{MaterialX::Category::Noise3D,     TEXT("pivot")},     TEXT("B")},            // The pivot of the noise is connected to a add node
-	{{MaterialX::Category::Normalize,   TEXT("in")},        TEXT("VectorInput")},
-	{{MaterialX::Category::Outside,     TEXT("in")},        TEXT("A")},				// Outside is treated as Multiply node
-	{{MaterialX::Category::Outside,     TEXT("mask")},      TEXT("B")},				// Outside is treated as Multiply node
-	{{MaterialX::Category::Power,       TEXT("in1")},       TEXT("Base")},
-	{{MaterialX::Category::Power,	    TEXT("in2")},       TEXT("Exponent")},
-	{{MaterialX::Category::Rotate2D,    TEXT("amount")},    TEXT("RotationAngle")},
-	{{MaterialX::Category::Rotate2D,    TEXT("in")},        TEXT("Position")},
-	{{MaterialX::Category::Rotate3D,    TEXT("amount")},    TEXT("RotationAngle")},
-	{{MaterialX::Category::Rotate3D,    TEXT("axis")},		TEXT("NormalizedRotationAxis")},
-	{{MaterialX::Category::Rotate3D,    TEXT("in")},        TEXT("Position")},
-	{{MaterialX::Category::Saturate,    TEXT("amount")},    TEXT("Fraction")},
-	{{MaterialX::Category::Smoothstep,  TEXT("in")},        TEXT("Value")},
+	{{TEXT(""),                            TEXT("amplitude")}, TEXT("Amplitude")},
+	{{TEXT(""),                            TEXT("bg")},        TEXT("B")},
+	{{TEXT(""),                            TEXT("center")},    TEXT("Center")},
+	{{TEXT(""),                            TEXT("diminish")},  TEXT("Diminish")},
+	{{TEXT(""),                            TEXT("fg")},        TEXT("A")},
+	{{TEXT(""),                            TEXT("high")},      TEXT("Max")},
+	{{TEXT(""),                            TEXT("in")},        TEXT("Input")},
+	{{TEXT(""),                            TEXT("in1")},       TEXT("A")},
+	{{TEXT(""),                            TEXT("in2")},       TEXT("B")},
+	{{TEXT(""),                            TEXT("in3")},       TEXT("C")},
+	{{TEXT(""),                            TEXT("in4")},       TEXT("D")},
+	{{TEXT(""),                            TEXT("inlow")},     TEXT("InputLow")},
+	{{TEXT(""),                            TEXT("inhigh")},    TEXT("InputHigh")},
+	{{TEXT(""),                            TEXT("lacunarity")},TEXT("Lacunarity")},
+	{{TEXT(""),                            TEXT("low")},       TEXT("Min")},
+	{{TEXT(""),                            TEXT("lumacoeffs")},TEXT("LuminanceFactors")}, // for the moment not yet handled by Interchange, because of the attribute being an advanced pin
+	{{TEXT(""),                            TEXT("mix")},       TEXT("Alpha")},
+	{{TEXT(""),                            TEXT("position")},  TEXT("Position")},
+	{{TEXT(""),                            TEXT("texcoord")},  TEXT("Coordinates")},
+	{{TEXT(""),                            TEXT("octaves")},   TEXT("Octaves")},
+	{{TEXT(""),                            TEXT("outlow")},    TEXT("TargetLow")},
+	{{TEXT(""),                            TEXT("outhigh")},   TEXT("TargetHigh")},
+	{{TEXT(""),                            TEXT("valuel")},    TEXT("A")},
+	{{TEXT(""),                            TEXT("valuer")},    TEXT("B")},
+	{{MaterialX::Category::Atan2,          TEXT("in1")},       TEXT("Y")},
+	{{MaterialX::Category::Atan2,          TEXT("in2")},       TEXT("X")},
+	{{MaterialX::Category::HeightToNormal, TEXT("scale")},     UE::Interchange::Materials::Standard::Nodes::NormalFromHeightMap::Inputs::Intensity.ToString()},
+	{{MaterialX::Category::IfGreater,      TEXT("value1")},    TEXT("A")},
+	{{MaterialX::Category::IfGreater,      TEXT("value2")},    TEXT("B")},
+	{{MaterialX::Category::IfGreater,      TEXT("in1")},       TEXT("AGreaterThanB")},
+	{{MaterialX::Category::IfGreater,      TEXT("in2")},       TEXT("ALessThanB")}, //another input is added for the case equal see ConnectIfGreater
+	{{MaterialX::Category::IfGreaterEq,    TEXT("value1")},    TEXT("A")},
+	{{MaterialX::Category::IfGreaterEq,    TEXT("value2")},    TEXT("B")},
+	{{MaterialX::Category::IfGreaterEq,    TEXT("in1")},       TEXT("AGreaterThanB")}, //another is added for the case equal see ConnectIfGreaterEq
+	{{MaterialX::Category::IfGreaterEq,    TEXT("in2")},       TEXT("ALessThanB")},
+	{{MaterialX::Category::IfEqual,        TEXT("value1")},    TEXT("A")},
+	{{MaterialX::Category::IfEqual,        TEXT("value2")},    TEXT("B")},
+	{{MaterialX::Category::IfEqual,        TEXT("in1")},       TEXT("AEqualsB")},
+	{{MaterialX::Category::IfEqual,        TEXT("in2")},       TEXT("ALessThanB")},  // another input is added for the case greater see ConnectIfEqual
+	{{MaterialX::Category::Inside,         TEXT("in")},        TEXT("A")},			  // Inside is treated as a Multiply node
+	{{MaterialX::Category::Inside,         TEXT("mask")},      TEXT("B")},			  // Inside is treated as a Multiply node
+	{{MaterialX::Category::Invert,         TEXT("amount")},    TEXT("A")},
+	{{MaterialX::Category::Invert,         TEXT("in")},        TEXT("B")},
+	{{MaterialX::Category::Magnitude,      TEXT("in")},        TEXT("A")},
+	{{MaterialX::Category::Mix,            TEXT("fg")},        TEXT("B")},
+	{{MaterialX::Category::Mix,            TEXT("bg")},        TEXT("A")},
+	{{MaterialX::Category::Mix,            TEXT("mix")},		TEXT("Factor")},
+	{{MaterialX::Category::Noise3D,        TEXT("amplitude")}, TEXT("B")},            // The amplitude of the noise is connected to a multiply node
+	{{MaterialX::Category::Noise3D,        TEXT("pivot")},     TEXT("B")},            // The pivot of the noise is connected to a add node
+	{{MaterialX::Category::Normalize,      TEXT("in")},        TEXT("VectorInput")},
+	{{MaterialX::Category::Outside,        TEXT("in")},        TEXT("A")},				// Outside is treated as Multiply node
+	{{MaterialX::Category::Outside,        TEXT("mask")},      TEXT("B")},				// Outside is treated as Multiply node
+	{{MaterialX::Category::Power,          TEXT("in1")},       TEXT("Base")},
+	{{MaterialX::Category::Power,	       TEXT("in2")},       TEXT("Exponent")},
+	{{MaterialX::Category::Rotate2D,       TEXT("amount")},    TEXT("RotationAngle")},
+	{{MaterialX::Category::Rotate2D,       TEXT("in")},        TEXT("Position")},
+	{{MaterialX::Category::Rotate3D,       TEXT("amount")},    TEXT("RotationAngle")},
+	{{MaterialX::Category::Rotate3D,       TEXT("axis")},		TEXT("NormalizedRotationAxis")},
+	{{MaterialX::Category::Rotate3D,       TEXT("in")},        TEXT("Position")},
+	{{MaterialX::Category::Saturate,       TEXT("amount")},    TEXT("Fraction")},
+	{{MaterialX::Category::Smoothstep,     TEXT("in")},        TEXT("Value")},
 },
 NodeNamesMaterialX2UE{
 	// Math nodes
@@ -219,7 +256,7 @@ TArray<FString> UInterchangeMaterialXTranslator::GetSupportedFormats() const
 		return TArray<FString>{};
 	}
 
-	return TArray<FString>{ TEXT("mtlx;MaterialX File Format") };
+	return UE::Interchange::MaterialX::IsMaterialFunctionPackageLoaded() ? TArray<FString>{ TEXT("mtlx;MaterialX File Format") } : TArray<FString>{};
 }
 
 bool UInterchangeMaterialXTranslator::Translate(UInterchangeBaseNodeContainer& BaseNodeContainer) const
@@ -1027,7 +1064,15 @@ bool UInterchangeMaterialXTranslator::ConnectNodeGraphOutputToInput(MaterialX::I
 				else if(UpstreamNode->getCategory() == mx::Category::WorleyNoise3D)
 				{
 					ConnectWorleyNoise3DInputToOutput(UpstreamNode, ParentShaderNode, InputChannelName, NamesToShaderNodes, NodeContainer);
-				}				
+				}
+				else if(UpstreamNode->getCategory() == mx::Category::Blur)
+				{
+					ConnectBlurInputToOutput(UpstreamNode, ParentShaderNode, InputChannelName, NamesToShaderNodes, NodeContainer);
+				}
+				else if(UpstreamNode->getCategory() == mx::Category::HeightToNormal)
+				{
+					ConnectHeightToNormalInputToOutput(UpstreamNode, ParentShaderNode, InputChannelName, NamesToShaderNodes, NodeContainer);
+				}
 				else
 				{
 					UInterchangeResultWarning_Generic* Message = AddMessage<UInterchangeResultWarning_Generic>();
@@ -1065,46 +1110,6 @@ bool UInterchangeMaterialXTranslator::ConnectNodeOutputToInput(MaterialX::NodePt
 	return bIsConnected;
 }
 
-UInterchangeTextureNode* UInterchangeMaterialXTranslator::CreateTextureNode(MaterialX::NodePtr Node, UInterchangeBaseNodeContainer& NodeContainer) const
-{
-	UInterchangeTextureNode* TextureNode = nullptr;
-
-	//A node image should have an input file otherwise the user should check its default value
-	if(Node)
-	{
-		if(mx::InputPtr InputFile = Node->getInput("file"); InputFile && InputFile->hasValue())
-		{
-			FString Filepath{ InputFile->getValueString().c_str() };
-			const FString FilePrefix = GetFilePrefix(InputFile);
-			Filepath = FPaths::Combine(FilePrefix, Filepath);
-			const FString Filename = FPaths::GetCleanFilename(Filepath);
-			const FString TextureNodeUID = TEXT("\\Texture\\") + Filename;
-
-			//Only add the TextureNode once
-			TextureNode = const_cast<UInterchangeTextureNode*>(Cast<UInterchangeTextureNode>(NodeContainer.GetNode(TextureNodeUID)));
-			if(TextureNode == nullptr)
-			{
-				TextureNode = NewObject<UInterchangeTexture2DNode>(&NodeContainer);
-				TextureNode->InitializeNode(TextureNodeUID, Filename, EInterchangeNodeContainerType::TranslatedAsset);
-				NodeContainer.AddNode(TextureNode);
-
-				if(FPaths::IsRelative(Filepath))
-				{
-					Filepath = FPaths::ConvertRelativePathToFull(FPaths::GetPath(Node->getActiveSourceUri().c_str()), Filepath);
-				}
-
-				TextureNode->SetPayLoadKey(Filepath);
-
-				const FString ColorSpace = GetColorSpace(InputFile);
-				const bool bIsSRGB = ColorSpace == TEXT("srgb_texture");
-				TextureNode->SetCustomSRGB(bIsSRGB);
-			}
-		}
-	}
-
-	return TextureNode;
-}
-
 const FString& UInterchangeMaterialXTranslator::GetMatchedInputName(MaterialX::NodePtr Node, MaterialX::InputPtr Input) const
 {
 	static FString EmptyString;
@@ -1514,7 +1519,7 @@ void UInterchangeMaterialXTranslator::ConnectImageInputToOutput(MaterialX::NodeP
 {
 	using namespace UE::Interchange::Materials::Standard::Nodes;
 
-	if(UInterchangeTextureNode* TextureNode = CreateTextureNode(UpstreamNode, NodeContainer))
+	if(UInterchangeTextureNode* TextureNode = CreateTextureNode<UInterchangeTexture2DNode>(UpstreamNode, NodeContainer))
 	{
 		//By default set the output of a texture to RGB
 		FString OutputChannel{ TEXT("RGB") };
@@ -2051,6 +2056,123 @@ void UInterchangeMaterialXTranslator::ConnectWorleyNoise3DInputToOutput(Material
 	UInterchangeShaderPortsAPI::ConnectDefaultOuputToInput(ParentShaderNode, InputChannelName, NoiseNode->GetUniqueID());
 }
 
+void UInterchangeMaterialXTranslator::ConnectHeightToNormalInputToOutput(MaterialX::NodePtr UpstreamNode, UInterchangeShaderNode* ParentShaderNode, const FString& InputChannelName, TMap<FString, UInterchangeShaderNode*>& NamesToShaderNodes, UInterchangeBaseNodeContainer& NodeContainer) const
+{
+	if(mx::InputPtr Input = UpstreamNode->getInput("in"))
+	{
+		// Image node will become this node
+		if(mx::NodePtr ConnectedNode = Input->getConnectedNode();
+		   ConnectedNode && ConnectedNode->getCategory() == mx::Category::Image)
+		{
+			//we need to copy the content of the image node to this node
+			UpstreamNode->copyContentFrom(ConnectedNode);
+
+			RenameNodeInputs(UpstreamNode);
+
+			//the copy overwrite every attribute of the node, so we need to get them back, essentially the type and the renaming
+			// the output is always a vec3
+			UpstreamNode->setType(mx::Type::Vector3);
+			
+			mx::NodeGraphPtr Graph = UpstreamNode->getParent()->asA<mx::NodeGraph>();
+			Graph->removeNode(ConnectedNode->getName());
+
+			using namespace UE::Interchange::Materials::Standard::Nodes;
+
+			if(UInterchangeTextureNode* TextureNode = CreateTextureNode<UInterchangeTexture2DNode>(UpstreamNode, NodeContainer))
+			{
+				UInterchangeShaderNode* HeightMapNode = CreateShaderNode<UInterchangeShaderNode>(UpstreamNode->getName().c_str(), NormalFromHeightMap::Name.ToString(), NamesToShaderNodes, NodeContainer);
+				UInterchangeShaderPortsAPI::ConnectDefaultOuputToInput(ParentShaderNode, InputChannelName, HeightMapNode->GetUniqueID());
+
+				const FString TextureNodeName = UpstreamNode->getName().c_str() + FString{ "_texture" };
+				UInterchangeShaderNode* TextureShaderNode = CreateShaderNode<UInterchangeShaderNode>(TextureNodeName, TextureObject::Name.ToString(), NamesToShaderNodes, NodeContainer);
+				TextureShaderNode->AddStringAttribute(UInterchangeShaderPortsAPI::MakeInputValueKey(TextureObject::Inputs::Texture.ToString()), TextureNode->GetUniqueID());
+				UInterchangeShaderPortsAPI::ConnectDefaultOuputToInput(HeightMapNode, NormalFromHeightMap::Inputs::HeightMap.ToString(), TextureShaderNode->GetUniqueID());
+
+				AddAttributeFromValueOrInterface(UpstreamNode->getInput("scale"), NormalFromHeightMap::Inputs::Intensity.ToString(), HeightMapNode);
+			}
+			else
+			{
+				AddAttributeFromValueOrInterface(UpstreamNode->getInput(mx::NodeGroup::Texture2D::Inputs::Default), InputChannelName, ParentShaderNode);
+			}
+		}
+	}
+	else
+	{
+		// For the moment it doesn't make sense to plug a value to it, so let's plug directly the child to the parent, in the future we could implement a Sobel and handle a multi output
+		RenameInput(Input, TCHAR_TO_UTF8(*InputChannelName)); //let's take the parent node's input name
+		NamesToShaderNodes.Add(UpstreamNode->getName().c_str(), ParentShaderNode);
+	}
+}
+
+void UInterchangeMaterialXTranslator::ConnectBlurInputToOutput(MaterialX::NodePtr UpstreamNode, UInterchangeShaderNode* ParentShaderNode, const FString& InputChannelName, TMap<FString, UInterchangeShaderNode*>& NamesToShaderNodes, UInterchangeBaseNodeContainer& NodeContainer) const
+{
+	if(mx::InputPtr Input = UpstreamNode->getInput("in"))
+	{
+		// Image node will become this node
+		if(mx::NodePtr ConnectedNode = Input->getConnectedNode();
+		   ConnectedNode && ConnectedNode->getCategory() == mx::Category::Image)
+		{
+			std::string NodeType = UpstreamNode->getType();
+
+			//we need to copy the content of the image node to this node
+			UpstreamNode->copyContentFrom(ConnectedNode);
+
+			RenameNodeInputs(UpstreamNode);
+
+			//the copy overwrites every attribute of the node, so we need to get them back, essentially the type and the renaming
+			UpstreamNode->setType(NodeType);
+
+			mx::NodeGraphPtr Graph = UpstreamNode->getParent()->asA<mx::NodeGraph>();
+			Graph->removeNode(ConnectedNode->getName());
+
+			using namespace UE::Interchange::Materials::Standard::Nodes;
+
+			if(UInterchangeTextureNode* TextureNode = CreateTextureNode<UInterchangeTextureBlurNode>(UpstreamNode, NodeContainer))
+			{
+				FString OutputChannel{ TEXT("RGB") };
+
+				if(NodeType == mx::Type::Vector4 || NodeType == mx::Type::Color4)
+				{
+					OutputChannel = TEXT("RGBA");
+				}
+				else if(NodeType == mx::Type::Float)
+				{
+					OutputChannel = TEXT("R");
+				}
+
+				UInterchangeShaderNode* TextureShaderNode = CreateShaderNode<UInterchangeShaderNode>(UpstreamNode->getName().c_str(), TextureSampleBlur::Name.ToString(), NamesToShaderNodes, NodeContainer);
+				TextureShaderNode->AddStringAttribute(UInterchangeShaderPortsAPI::MakeInputValueKey(TextureSampleBlur::Inputs::Texture.ToString()), TextureNode->GetUniqueID());
+				UInterchangeShaderPortsAPI::ConnectOuputToInput(ParentShaderNode, InputChannelName, TextureShaderNode->GetUniqueID(), OutputChannel);
+
+				if(mx::InputPtr InputKernel = UpstreamNode->getInput("filtertype"))
+				{
+					// By default TextureSampleBox uses gaussian filter
+					if(InputKernel->getValueString() == "box")
+					{
+						TextureShaderNode->AddInt32Attribute(TextureSampleBlur::Attributes::Filter, 0);
+					}
+				}
+
+				if(mx::InputPtr InputKernel = UpstreamNode->getInput("size"))
+				{
+					float KernelSize = mx::fromValueString<float>(InputKernel->getValueString());
+					TextureShaderNode->AddFloatAttribute(TextureSampleBlur::Attributes::KernelSize, KernelSize);
+				}
+			}
+			else
+			{
+				AddAttributeFromValueOrInterface(UpstreamNode->getInput(mx::NodeGroup::Texture2D::Inputs::Default), InputChannelName, ParentShaderNode);
+			}
+		}
+	}
+	else
+	{
+		// For a blur it doesn't make sense if there's no image input
+		RenameInput(Input, TCHAR_TO_UTF8(*InputChannelName)); //let's take the parent node's input name
+		NamesToShaderNodes.Add(UpstreamNode->getName().c_str(), ParentShaderNode);
+	}
+}
+
 UInterchangeShaderNode* UInterchangeMaterialXTranslator::CreateMaskShaderNode(uint8 RGBA, const FString& NodeName, TMap<FString, UInterchangeShaderNode*>& NamesToShaderNodes, UInterchangeBaseNodeContainer& NodeContainer) const
 {
 	bool bR = (0b1000 & RGBA) >> 3;
diff --git a/Engine/Plugins/Interchange/Runtime/Source/Import/Private/MaterialX/MaterialExpressionTextureSampleParameterBlur.cpp b/Engine/Plugins/Interchange/Runtime/Source/Import/Private/MaterialX/MaterialExpressionTextureSampleParameterBlur.cpp
new file mode 100644
--- /dev/null
+++ b/Engine/Plugins/Interchange/Runtime/Source/Import/Private/MaterialX/MaterialExpressionTextureSampleParameterBlur.cpp
@@ -0,0 +1,246 @@
+// Copyright Epic Games, Inc. All Rights Reserved.
+
+#include "MaterialX/MaterialExpressionTextureSampleParameterBlur.h"
+#include "Materials/MaterialExpressionAdd.h"
+#include "Materials/MaterialExpressionConstant.h"
+#include "Materials/MaterialExpressionConstant2Vector.h"
+#include "Materials/MaterialExpressionCustom.h"
+#include "Materials/MaterialExpressionMultiply.h"
+#include "Materials/MaterialExpressionTextureCoordinate.h"
+
+#include UE_INLINE_GENERATED_CPP_BY_NAME(MaterialExpressionTextureSampleParameterBlur)
+
+#define LOCTEXT_NAMESPACE "MaterialExpressionTextureSampleBlur"
+
+namespace Gauss
+{
+	namespace
+	{
+		static constexpr float KernelWeights3x3[9] =
+		{
+			0.077847, 0.123317, 0.077847,
+			0.123317, 0.195346, 0.123317,
+			0.077847, 0.123317, 0.077847
+		};
+
+		static constexpr float KernelWeights5x5[25] =
+		{
+			0.003765, 0.015019, 0.023792, 0.015019, 0.003765,
+			0.015019, 0.059912, 0.094907, 0.059912, 0.015019,
+			0.023792, 0.094907, 0.150342, 0.094907, 0.023792,
+			0.015019, 0.059912, 0.094907, 0.059912, 0.015019,
+			0.003765, 0.015019, 0.023792, 0.015019, 0.003765
+		};
+
+		static constexpr float KernelWeights7x7[49] =
+		{
+			0.000036, 0.000363, 0.001446, 0.002291, 0.001446, 0.000363, 0.000036,
+			0.000363, 0.003676, 0.014662, 0.023226, 0.014662, 0.003676, 0.000363,
+			0.001446, 0.014662, 0.058488, 0.092651, 0.058488, 0.014662, 0.001446,
+			0.002291, 0.023226, 0.092651, 0.146768, 0.092651, 0.023226, 0.002291,
+			0.001446, 0.014662, 0.058488, 0.092651, 0.058488, 0.014662, 0.001446,
+			0.000363, 0.003676, 0.014662, 0.023226, 0.014662, 0.003676, 0.000363,
+			0.000036, 0.000363, 0.001446, 0.002291, 0.001446, 0.000363, 0.000036
+		};
+	}
+}
+
+namespace Box
+{
+	namespace
+	{
+		static constexpr float KernelWeights3x3[9] =
+		{
+			1.f / 9.f, 1.f / 9.f, 1.f / 9.f,
+			1.f / 9.f, 1.f / 9.f, 1.f / 9.f,
+			1.f / 9.f, 1.f / 9.f, 1.f / 9.f,
+		};
+
+		static constexpr float KernelWeights5x5[25] =
+		{
+			1.f / 25.f, 1.f / 25.f, 1.f / 25.f, 1.f / 25.f, 1.f / 25.f,
+			1.f / 25.f, 1.f / 25.f, 1.f / 25.f, 1.f / 25.f, 1.f / 25.f,
+			1.f / 25.f, 1.f / 25.f, 1.f / 25.f, 1.f / 25.f, 1.f / 25.f,
+			1.f / 25.f, 1.f / 25.f, 1.f / 25.f, 1.f / 25.f, 1.f / 25.f,
+			1.f / 25.f, 1.f / 25.f, 1.f / 25.f, 1.f / 25.f, 1.f / 25.f,
+		};
+
+		static constexpr float KernelWeights7x7[49] =
+		{
+			1.f / 49.f, 1.f / 49.f, 1.f / 49.f, 1.f / 49.f, 1.f / 49.f, 1.f / 49.f, 1.f / 49.f,
+			1.f / 49.f, 1.f / 49.f, 1.f / 49.f, 1.f / 49.f, 1.f / 49.f, 1.f / 49.f, 1.f / 49.f,
+			1.f / 49.f, 1.f / 49.f, 1.f / 49.f, 1.f / 49.f, 1.f / 49.f, 1.f / 49.f, 1.f / 49.f,
+			1.f / 49.f, 1.f / 49.f, 1.f / 49.f, 1.f / 49.f, 1.f / 49.f, 1.f / 49.f, 1.f / 49.f,
+			1.f / 49.f, 1.f / 49.f, 1.f / 49.f, 1.f / 49.f, 1.f / 49.f, 1.f / 49.f, 1.f / 49.f,
+			1.f / 49.f, 1.f / 49.f, 1.f / 49.f, 1.f / 49.f, 1.f / 49.f, 1.f / 49.f, 1.f / 49.f,
+			1.f / 49.f, 1.f / 49.f, 1.f / 49.f, 1.f / 49.f, 1.f / 49.f, 1.f / 49.f, 1.f / 49.f,
+		};
+	}
+}
+
+UMaterialExpressionTextureSampleParameterBlur::UMaterialExpressionTextureSampleParameterBlur(const FObjectInitializer& ObjectInitializer)
+	: Super(ObjectInitializer)
+{
+	// Structure to hold one-time initialization
+	struct FConstructorStatics
+	{
+		FText NAME_MaterialX;
+		FText NAME_Convolution;
+		FConstructorStatics()
+			: NAME_MaterialX(LOCTEXT("MaterialX", "MaterialX"))
+			, NAME_Convolution(LOCTEXT("Convolution", "Convolution"))
+		{}
+	};
+	static FConstructorStatics ConstructorStatics;
+
+#if WITH_EDITORONLY_DATA
+	MenuCategories.Add(ConstructorStatics.NAME_MaterialX);
+	MenuCategories.Add(ConstructorStatics.NAME_Convolution);
+#endif
+}
+
+#if WITH_EDITOR
+int32 UMaterialExpressionTextureSampleParameterBlur::Compile(FMaterialCompiler* Compiler, int32 OutputIndex)
+{
+	if(KernelSize > 0)
+	{
+		FExpressionInput TexCoordInput;
+		if(Coordinates.GetTracedInput().Expression)
+		{
+			TexCoordInput = Coordinates;
+		}
+		else
+		{
+			UMaterialExpressionTextureCoordinate* DefaultTexCoord = NewObject<UMaterialExpressionTextureCoordinate>();
+			DefaultTexCoord->CoordinateIndex = ConstCoordinate;
+			TexCoordInput.Connect(0, DefaultTexCoord);
+		}
+
+		UMaterialExpressionCustom* ComputeSampleSize = NewObject<UMaterialExpressionCustom>();
+
+		ComputeSampleSize->Inputs[0].InputName = TEXT("uv");
+		ComputeSampleSize->Inputs[0].Input = TexCoordInput;
+
+		auto ComputeSampleSizeInput = [&](float Value, FName InputName)
+		{
+			UMaterialExpressionConstant* FilterSizeConstant = NewObject<UMaterialExpressionConstant>();
+			FilterSizeConstant->R = Value;
+			FCustomInput & Input = ComputeSampleSize->Inputs.Add_GetRef({});
+			Input.InputName = InputName;
+			Input.Input.Expression = FilterSizeConstant;
+		};
+
+		ComputeSampleSizeInput(FilterSize, TEXT("filterSize"));
+		ComputeSampleSizeInput(FilterOffset, TEXT("filterOffset"));
+
+		ComputeSampleSize->OutputType = ECustomMaterialOutputType::CMOT_Float2;
+		ComputeSampleSize->Code = TEXT(
+			R"(float2 derivUVx = ddx(uv) * 0.5f;
+   float2 derivUVy = ddy(uv) * 0.5f;
+   float derivX = abs(derivUVx.x) + abs(derivUVy.x);
+   float derivY = abs(derivUVx.y) + abs(derivUVy.y);
+   float sampleSizeU = 2.0f * filterSize * derivX + filterOffset;
+   if (sampleSizeU < 1.0E-05f)
+       sampleSizeU = 1.0E-05f;
+   float sampleSizeV = 2.0f * filterSize * derivY + filterOffset;
+   if (sampleSizeV < 1.0E-05f)
+       sampleSizeV = 1.0E-05f;
+   return float2(sampleSizeU, sampleSizeV);)"
+		);
+
+		const float* Kernel;
+		int32 FilterWidth = 1;
+
+		if(KernelSize <= (1.f / 3.f))
+		{
+			FilterWidth = 3 / 2;
+			if(Filter == ETextureSampleBlurFilter::Box)
+			{
+				Kernel = Box::KernelWeights3x3;
+			}
+			else
+			{
+				Kernel = Gauss::KernelWeights3x3;
+			}
+		}
+		else if(KernelSize <= (2.f / 3.f))
+		{
+			FilterWidth = 5 / 2;
+			if(Filter == ETextureSampleBlurFilter::Box)
+			{
+				Kernel = Box::KernelWeights5x5;
+			}
+			else
+			{
+				Kernel = Gauss::KernelWeights5x5;
+			}
+		}
+		else
+		{
+			FilterWidth = 7 / 2;
+			if(Filter == ETextureSampleBlurFilter::Box)
+			{
+				Kernel = Box::KernelWeights7x7;
+			}
+			else
+			{
+				Kernel = Gauss::KernelWeights7x7;
+			}
+		}
+
+		UMaterialExpressionCustom* Convolution = NewObject<UMaterialExpressionCustom>();
+		Convolution->OutputType = ECustomMaterialOutputType::CMOT_Float4;
+		Convolution->Inputs.Empty();
+		Convolution->Inputs.Reserve(25);
+		Convolution->Code = TEXT("float4 result = float4(0,0,0,0);\n");
+
+		int32 i = 0;
+
+		for(int Row = -FilterWidth; Row <= FilterWidth; ++Row)
+		{
+			for(int Col = -FilterWidth; Col <= FilterWidth; ++Col)
+			{
+				UMaterialExpressionConstant2Vector* Constant = NewObject<UMaterialExpressionConstant2Vector>();
+				Constant->R = Col;
+				Constant->G = Row;
+				UMaterialExpressionMultiply* MultiplyCoord = NewObject<UMaterialExpressionMultiply>();
+				MultiplyCoord->A.Expression = ComputeSampleSize;
+				MultiplyCoord->B.Expression = Constant;
+				UMaterialExpressionAdd* AddCoord = NewObject<UMaterialExpressionAdd>();
+				AddCoord->A = TexCoordInput;
+				AddCoord->B.Expression = MultiplyCoord;
+
+				UMaterialExpressionTextureSampleParameter2D* TextureSample = NewObject<UMaterialExpressionTextureSampleParameter2D>();
+				TextureSample->ConstCoordinate = ConstCoordinate;
+				TextureSample->Coordinates.Expression = AddCoord;
+				TextureSample->TextureObject = TextureObject;
+				TextureSample->Texture = Texture;
+				TextureSample->SamplerSource = SamplerSource;
+				TextureSample->SamplerType = SamplerType;
+
+				Convolution->Inputs.Add({});
+				FString InputName{ TEXT("S") + FString::FromInt(i) };
+				Convolution->Inputs[i].InputName = FName(InputName);// FName(TEXT("S") + Index);
+				Convolution->Inputs[i].Input.Expression = TextureSample;
+				Convolution->Code += (TEXT("result += ") + InputName + TEXT("*") + FString::SanitizeFloat(Kernel[i]) + TEXT(";\n"));
+				++i;
+			}
+		}
+		Convolution->Code += TEXT("return result;");
+
+		return Convolution->Compile(Compiler, 0);
+	}
+	else
+	{
+		return Super::Compile(Compiler, OutputIndex);
+	}
+}
+
+void UMaterialExpressionTextureSampleParameterBlur::GetCaption(TArray<FString>& OutCaptions) const
+{
+	OutCaptions.Add(TEXT("ParamBlur"));
+	OutCaptions.Add(FString::Printf(TEXT("'%s'"), *ParameterName.ToString()));
+}
+#endif // WITH_EDITOR
+
+#undef LOCTEXT_NAMESPACE
\ No newline at end of file
diff --git a/Engine/Plugins/Interchange/Runtime/Source/Import/Public/MaterialX/InterchangeMaterialXTranslator.h b/Engine/Plugins/Interchange/Runtime/Source/Import/Public/MaterialX/InterchangeMaterialXTranslator.h
--- a/Engine/Plugins/Interchange/Runtime/Source/Import/Public/MaterialX/InterchangeMaterialXTranslator.h
+++ b/Engine/Plugins/Interchange/Runtime/Source/Import/Public/MaterialX/InterchangeMaterialXTranslator.h
@@ -3,6 +3,7 @@
 
 #include "InterchangeShaderGraphNode.h"
 #include "InterchangeTranslatorBase.h"
+#include "InterchangeTextureNode.h"
 #include "Texture/InterchangeTexturePayloadInterface.h"
 
 #if WITH_EDITOR
@@ -171,7 +172,47 @@ protected:
 	 * 
 	 * @return The texture node that was created
 	 */
-	UInterchangeTextureNode* CreateTextureNode(MaterialX::NodePtr Node, UInterchangeBaseNodeContainer& NodeContainer) const;
+	template<typename TextureTypeNode>
+	UInterchangeTextureNode* CreateTextureNode(MaterialX::NodePtr Node, UInterchangeBaseNodeContainer& NodeContainer) const
+	{
+		static_assert(std::is_convertible_v<TextureTypeNode*, UInterchangeTextureNode*>, "CreateTextureNode only accepts type that derived from UInterchangeTextureNode");
+		UInterchangeTextureNode* TextureNode = nullptr;
+
+		//A node image should have an input file otherwise the user should check its default value
+		if(Node)
+		{
+			if(MaterialX::InputPtr InputFile = Node->getInput("file"); InputFile && InputFile->hasValue())
+			{
+				FString Filepath{ InputFile->getValueString().c_str() };
+				const FString FilePrefix = GetFilePrefix(InputFile);
+				Filepath = FPaths::Combine(FilePrefix, Filepath);
+				const FString Filename = FPaths::GetCleanFilename(Filepath);
+				const FString TextureNodeUID = TEXT("\\Texture\\") + Filename;
+
+				//Only add the TextureNode once
+				TextureNode = const_cast<UInterchangeTextureNode*>(Cast<UInterchangeTextureNode>(NodeContainer.GetNode(TextureNodeUID)));
+				if(TextureNode == nullptr)
+				{
+					TextureNode = NewObject<TextureTypeNode>(&NodeContainer);
+					TextureNode->InitializeNode(TextureNodeUID, Filename, EInterchangeNodeContainerType::TranslatedAsset);
+					NodeContainer.AddNode(TextureNode);
+
+					if(FPaths::IsRelative(Filepath))
+					{
+						Filepath = FPaths::ConvertRelativePathToFull(FPaths::GetPath(Node->getActiveSourceUri().c_str()), Filepath);
+					}
+
+					TextureNode->SetPayLoadKey(Filepath);
+
+					const FString ColorSpace = GetColorSpace(InputFile);
+					const bool bIsSRGB = ColorSpace == TEXT("srgb_texture");
+					TextureNode->SetCustomSRGB(bIsSRGB);
+				}
+			}
+		}
+
+		return TextureNode;
+	}
 
 	/**
 	 * Get the UE corresponding name of a MaterialX Node category and input for a material
@@ -334,27 +375,78 @@ protected:
 	 */
 	FLinearColor GetLinearColor(MaterialX::InputPtr Input) const;
 
+	/** Begin Connect MaterialX nodes*/
+	
+	/** <constant> */
 	void ConnectConstantInputToOutput(MaterialX::NodePtr UpstreamNode, UInterchangeShaderNode* ParentShaderNode, const FString& InputChannelName, TMap<FString, UInterchangeShaderNode*>& NamesToShaderNodes, UInterchangeBaseNodeContainer& NodeContainer) const;
+	
+	/** <extract> */
 	void ConnectExtractInputToOutput(MaterialX::NodePtr UpstreamNode, UInterchangeShaderNode* ParentShaderNode, const FString& InputChannelName, TMap<FString, UInterchangeShaderNode*>& NamesToShaderNodes, UInterchangeBaseNodeContainer& NodeContainer) const;
+	
+	/** <dot> */
 	void ConnectDotInputToOutput(MaterialX::NodePtr UpstreamNode, UInterchangeShaderNode* ParentShaderNode, const FString& InputChannelName, TMap<FString, UInterchangeShaderNode*>& NamesToShaderNodes, UInterchangeBaseNodeContainer& NodeContainer) const;
+	
+	/** <transformpoint> */
 	void ConnectTransformPositionInputToOutput(MaterialX::NodePtr UpstreamNode, UInterchangeShaderNode* ParentShaderNode, const FString& InputChannelName, TMap<FString, UInterchangeShaderNode*>& NamesToShaderNodes, UInterchangeBaseNodeContainer& NodeContainer) const;
+	
+	/** <transformvector> */
 	void ConnectTransformVectorInputToOutput(MaterialX::NodePtr UpstreamNode, UInterchangeShaderNode* ParentShaderNode, const FString& InputChannelName, TMap<FString, UInterchangeShaderNode*>& NamesToShaderNodes, UInterchangeBaseNodeContainer& NodeContainer) const;
+	
+	/** <rotate2d> */
 	void ConnectRotate2DInputToOutput(MaterialX::NodePtr UpstreamNode, UInterchangeShaderNode* ParentShaderNode, const FString& InputChannelName, TMap<FString, UInterchangeShaderNode*>& NamesToShaderNodes, UInterchangeBaseNodeContainer& NodeContainer) const;
+	
+	/** <rotate3d> */
 	void ConnectRotate3DInputToOutput(MaterialX::NodePtr UpstreamNode, UInterchangeShaderNode* ParentShaderNode, const FString& InputChannelName, TMap<FString, UInterchangeShaderNode*>& NamesToShaderNodes, UInterchangeBaseNodeContainer& NodeContainer) const;
+	
+	/** <image> */
 	void ConnectImageInputToOutput(MaterialX::NodePtr UpstreamNode, UInterchangeShaderNode* ParentShaderNode, const FString& InputChannelName, TMap<FString, UInterchangeShaderNode*>& NamesToShaderNodes, UInterchangeBaseNodeContainer& NodeContainer) const;
+	
+	/** <convert> */
 	void ConnectConvertInputToOutput(MaterialX::NodePtr UpstreamNode, UInterchangeShaderNode* ParentShaderNode, const FString& InputChannelName, TMap<FString, UInterchangeShaderNode*>& NamesToShaderNodes, UInterchangeBaseNodeContainer& NodeContainer) const;
+	
+	/** <ifgreater> */
 	void ConnectIfGreaterInputToOutput(MaterialX::NodePtr UpstreamNode, UInterchangeShaderNode* ParentShaderNode, const FString& InputChannelName, TMap<FString, UInterchangeShaderNode*>& NamesToShaderNodes, UInterchangeBaseNodeContainer& NodeContainer) const;
+	
+	/** <ifgreatereq> */
 	void ConnectIfGreaterEqInputToOutput(MaterialX::NodePtr UpstreamNode, UInterchangeShaderNode* ParentShaderNode, const FString& InputChannelName, TMap<FString, UInterchangeShaderNode*>& NamesToShaderNodes, UInterchangeBaseNodeContainer& NodeContainer) const;
+	
+	/** <ifequal> */
 	void ConnectIfEqualInputToOutput(MaterialX::NodePtr UpstreamNode, UInterchangeShaderNode* ParentShaderNode, const FString& InputChannelName, TMap<FString, UInterchangeShaderNode*>& NamesToShaderNodes, UInterchangeBaseNodeContainer& NodeContainer) const;
+	
+	/** <outside> */
 	void ConnectOutsideInputToOutput(MaterialX::NodePtr UpstreamNode, UInterchangeShaderNode* ParentShaderNode, const FString& InputChannelName, TMap<FString, UInterchangeShaderNode*>& NamesToShaderNodes, UInterchangeBaseNodeContainer& NodeContainer) const;
+	
+	/** <position> */
 	void ConnectPositionInputToOutput(MaterialX::NodePtr UpstreamNode, UInterchangeShaderNode* ParentShaderNode, const FString& InputChannelName, TMap<FString, UInterchangeShaderNode*>& NamesToShaderNodes, UInterchangeBaseNodeContainer& NodeContainer) const;
+	
+	/** <normal> */
 	void ConnectNormalInputToOutput(MaterialX::NodePtr UpstreamNode, UInterchangeShaderNode* ParentShaderNode, const FString& InputChannelName, TMap<FString, UInterchangeShaderNode*>& NamesToShaderNodes, UInterchangeBaseNodeContainer& NodeContainer) const;
+	
+	/** <tangent> */
 	void ConnectTangentInputToOutput(MaterialX::NodePtr UpstreamNode, UInterchangeShaderNode* ParentShaderNode, const FString& InputChannelName, TMap<FString, UInterchangeShaderNode*>& NamesToShaderNodes, UInterchangeBaseNodeContainer& NodeContainer) const;
+	
+	/** <bitangent> */
 	void ConnectBitangentInputToOutput(MaterialX::NodePtr UpstreamNode, UInterchangeShaderNode* ParentShaderNode, const FString& InputChannelName, TMap<FString, UInterchangeShaderNode*>& NamesToShaderNodes, UInterchangeBaseNodeContainer& NodeContainer) const;
+	
+	/** <time> */
 	void ConnectTimeInputToOutput(MaterialX::NodePtr UpstreamNode, UInterchangeShaderNode* ParentShaderNode, const FString& InputChannelName, TMap<FString, UInterchangeShaderNode*>& NamesToShaderNodes, UInterchangeBaseNodeContainer& NodeContainer) const;
+	
+	/** <noise3d> */
 	void ConnectNoise3DInputToOutput(MaterialX::NodePtr UpstreamNode, UInterchangeShaderNode* ParentShaderNode, const FString& InputChannelName, TMap<FString, UInterchangeShaderNode*>& NamesToShaderNodes, UInterchangeBaseNodeContainer& NodeContainer) const;
+	
+	/** <cellnoise3d> */
 	void ConnectCellNoise3DInputToOutput(MaterialX::NodePtr UpstreamNode, UInterchangeShaderNode* ParentShaderNode, const FString& InputChannelName, TMap<FString, UInterchangeShaderNode*>& NamesToShaderNodes, UInterchangeBaseNodeContainer& NodeContainer) const;
+	
+	/** <worleynoise3d> */
 	void ConnectWorleyNoise3DInputToOutput(MaterialX::NodePtr UpstreamNode, UInterchangeShaderNode* ParentShaderNode, const FString& InputChannelName, TMap<FString, UInterchangeShaderNode*>& NamesToShaderNodes, UInterchangeBaseNodeContainer& NodeContainer) const;
+	
+	/** <heighttonormal> */
+	void ConnectHeightToNormalInputToOutput(MaterialX::NodePtr UpstreamNode, UInterchangeShaderNode* ParentShaderNode, const FString& InputChannelName, TMap<FString, UInterchangeShaderNode*>& NamesToShaderNodes, UInterchangeBaseNodeContainer& NodeContainer) const;
+	
+	/** <blur> */
+	void ConnectBlurInputToOutput(MaterialX::NodePtr UpstreamNode, UInterchangeShaderNode* ParentShaderNode, const FString& InputChannelName, TMap<FString, UInterchangeShaderNode*>& NamesToShaderNodes, UInterchangeBaseNodeContainer& NodeContainer) const;
+
+	/** End Connect MaterialX nodes*/
 
 	UInterchangeShaderNode* CreateMaskShaderNode(uint8 RGBA, const FString & NodeName, TMap<FString, UInterchangeShaderNode*>& NamesToShaderNodes, UInterchangeBaseNodeContainer& NodeContainer) const;
 
diff --git a/Engine/Plugins/Interchange/Runtime/Source/Import/Public/MaterialX/MaterialExpressionTextureSampleParameterBlur.h b/Engine/Plugins/Interchange/Runtime/Source/Import/Public/MaterialX/MaterialExpressionTextureSampleParameterBlur.h
new file mode 100644
--- /dev/null
+++ b/Engine/Plugins/Interchange/Runtime/Source/Import/Public/MaterialX/MaterialExpressionTextureSampleParameterBlur.h
@@ -0,0 +1,47 @@
+// Copyright Epic Games, Inc. All Rights Reserved.
+#pragma once
+
+#include "Materials/MaterialExpressionTextureSampleParameter2D.h"
+#include "MaterialExpressionTextureSampleParameterBlur.generated.h"
+
+UENUM()
+enum class ETextureSampleBlurFilter : uint8
+{
+	Box,
+	Gaussian
+};
+
+UCLASS(MinimalAPI, collapsecategories, hidecategories = Object)
+class UMaterialExpressionTextureSampleParameterBlur: public UMaterialExpressionTextureSampleParameter2D
+{
+	GENERATED_UCLASS_BODY()
+
+	/** The size of the blur kernel, relative to 0-1 UV space, only 3x3, 5x5 and 7x7 kernel are supported.
+	  * Basically depending on the KernelSize value:
+	  * KernelSize == 0   -> 1x1 Kernel (Meaning no Blur)
+	  * KernelSize <= 1/3 -> 3x3 Kernel
+	  * KernelSize <= 2/3 -> 5x5 Kernel
+	  * Otherwise         -> 7x7 Kernel
+	  */
+	UPROPERTY(EditAnywhere, Category = Filtering, Meta = (ShowAsInputPin = "Advanced"))
+	float KernelSize = 0.f;
+
+	/** Size of the filter when we sample a texture coordinate */
+	UPROPERTY(EditAnywhere, Category = Filtering, Meta = (ShowAsInputPin = "Advanced"))
+	float FilterSize = 1.f;
+
+	/** Offset of the filter when we sample a texture coordinate */
+	UPROPERTY(EditAnywhere, Category = Filtering, Meta = (ShowAsInputPin = "Advanced"))
+	float FilterOffset = 0.f;
+
+	/** Filter to use when we blur a Texture: Gaussian or Box Linear filter*/
+	UPROPERTY(EditAnywhere, Category = Filtering, Meta = (ShowAsInputPin = "Advanced"))
+	ETextureSampleBlurFilter Filter = ETextureSampleBlurFilter::Gaussian;
+
+	//~ Begin UMaterialExpression Interface
+#if WITH_EDITOR
+	virtual int32 Compile(class FMaterialCompiler* Compiler, int32 OutputIndex) override;
+	virtual void GetCaption(TArray<FString>& OutCaptions) const override;
+#endif
+	//~ End UMaterialExpression Interface
+};
\ No newline at end of file
diff --git a/Engine/Plugins/Interchange/Runtime/Source/Nodes/Private/InterchangeTextureBlurNode.cpp b/Engine/Plugins/Interchange/Runtime/Source/Nodes/Private/InterchangeTextureBlurNode.cpp
new file mode 100644
--- /dev/null
+++ b/Engine/Plugins/Interchange/Runtime/Source/Nodes/Private/InterchangeTextureBlurNode.cpp
@@ -0,0 +1,22 @@
+// Copyright Epic Games, Inc. All Rights Reserved.
+
+#include "InterchangeTextureBlurNode.h"
+#include "Nodes/InterchangeBaseNodeContainer.h"
+
+#include UE_INLINE_GENERATED_CPP_BY_NAME(InterchangeTextureBlurNode)
+
+FString UInterchangeTextureBlurNode::MakeNodeUid(const FStringView NodeName)
+{
+	return UInterchangeTextureNode::MakeNodeUid(NodeName);
+}
+
+UInterchangeTextureBlurNode* UInterchangeTextureBlurNode::Create(UInterchangeBaseNodeContainer* NodeContainer, const FStringView TextureNodeName)
+{
+	UInterchangeTextureBlurNode* TextureNode = NewObject<UInterchangeTextureBlurNode>(NodeContainer);
+	const FString TextureNodeUid = UInterchangeTextureNode::MakeNodeUid(TextureNodeName);
+
+	TextureNode->InitializeNode(TextureNodeUid, FString(TextureNodeName), EInterchangeNodeContainerType::TranslatedAsset);
+	NodeContainer->AddNode(TextureNode);
+
+	return TextureNode;
+}
diff --git a/Engine/Plugins/Interchange/Runtime/Source/Nodes/Public/InterchangeMaterialDefinitions.h b/Engine/Plugins/Interchange/Runtime/Source/Nodes/Public/InterchangeMaterialDefinitions.h
--- a/Engine/Plugins/Interchange/Runtime/Source/Nodes/Public/InterchangeMaterialDefinitions.h
+++ b/Engine/Plugins/Interchange/Runtime/Source/Nodes/Public/InterchangeMaterialDefinitions.h
@@ -113,6 +113,20 @@ namespace Interchange
 					}
 				}
 
+				namespace NormalFromHeightMap
+				{
+					const FName Name = TEXT("NormalFromHeightMap");
+
+					namespace Inputs
+					{
+						const FName HeightMap = TEXT("Height Map");                // Type: FString (unique id of a texture node)
+						const FName Intensity = TEXT("Normal Map Intensity");      // Type: float
+						const FName Offset = TEXT("Height Map UV Offset");         // Type: float
+						const FName Coordinates = TEXT("Coordinates");             // Type: vec2
+						const FName Channel = TEXT("Height Map Channel Selector"); // Type: vec4
+					}
+				}
+
 				namespace OneMinus
 				{
 					const FName Name = TEXT("OneMinus");
@@ -139,6 +153,16 @@ namespace Interchange
 					}
 				}
 
+				namespace TextureObject
+				{
+					const FName Name = TEXT("TextureObject");
+
+					namespace Inputs
+					{
+						const FName Texture = TEXT("TextureUid"); // Type: FString (unique id of a texture node)
+					}
+				}
+
 				namespace TextureSample
 				{
 					const FName Name = TEXT("TextureSample");
@@ -160,6 +184,19 @@ namespace Interchange
 					}
 				}
 
+				namespace TextureSampleBlur
+				{
+					using namespace TextureSample;					
+
+					namespace Attributes
+					{
+						const FName KernelSize = TEXT("KernelSize"); // Type: float
+						const FName FilterSize = TEXT("FilterSize"); // Type: float
+						const FName FilterOffset= TEXT("FilterOffset"); // Type: float
+						const FName Filter = TEXT("Filter"); // Type: int
+					}
+				}
+
 				namespace Time
 				{
 					const FName Name = TEXT("Time");
diff --git a/Engine/Plugins/Interchange/Runtime/Source/Nodes/Public/InterchangeTextureBlurNode.h b/Engine/Plugins/Interchange/Runtime/Source/Nodes/Public/InterchangeTextureBlurNode.h
new file mode 100644
--- /dev/null
+++ b/Engine/Plugins/Interchange/Runtime/Source/Nodes/Public/InterchangeTextureBlurNode.h
@@ -0,0 +1,35 @@
+// Copyright Epic Games, Inc. All Rights Reserved.
+
+#pragma once
+
+#include "InterchangeTexture2DNode.h"
+
+#include "InterchangeTextureBlurNode.generated.h"
+
+class UInterchangeBaseNodeContainer;
+
+UCLASS(BlueprintType, Experimental)
+class INTERCHANGENODES_API UInterchangeTextureBlurNode : public UInterchangeTexture2DNode
+{
+	GENERATED_BODY()
+
+public:
+
+	/**
+	 * Build and return a UID name for a texture 2D node.
+	 */
+	static FString MakeNodeUid(const FStringView NodeName);
+
+	/**
+	 * Creates a new UInterchangeTexture2DNode and adds it to NodeContainer as a translated node.
+	 */
+	static UInterchangeTextureBlurNode* Create(UInterchangeBaseNodeContainer* NodeContainer, const FStringView TextureNodeName);
+
+	/**
+	 * Return the node type name of the class, we use this when reporting error
+	 */
+	virtual FString GetTypeName() const override
+	{
+		return TEXT("TextureBlurNode");
+	}
+};
diff --git a/Engine/Plugins/Interchange/Runtime/Source/Pipelines/Private/InterchangeGenericMaterialPipeline.cpp b/Engine/Plugins/Interchange/Runtime/Source/Pipelines/Private/InterchangeGenericMaterialPipeline.cpp
--- a/Engine/Plugins/Interchange/Runtime/Source/Pipelines/Private/InterchangeGenericMaterialPipeline.cpp
+++ b/Engine/Plugins/Interchange/Runtime/Source/Pipelines/Private/InterchangeGenericMaterialPipeline.cpp
@@ -10,6 +10,7 @@
 #include "InterchangeShaderGraphNode.h"
 #include "InterchangeTexture2DNode.h"
 #include "InterchangeTexture2DArrayNode.h"
+#include "InterchangeTextureBlurNode.h"
 #include "InterchangeTextureCubeNode.h"
 #include "InterchangeTextureFactoryNode.h"
 #include "InterchangeTextureNode.h"
@@ -30,6 +31,7 @@
 #include "Materials/MaterialExpressionOneMinus.h"
 #include "Materials/MaterialExpressionScalarParameter.h"
 #include "Materials/MaterialExpressionTextureCoordinate.h"
+#include "Materials/MaterialExpressionTextureObject.h"
 #include "Materials/MaterialExpressionTextureSample.h"
 #include "Materials/MaterialExpressionTextureSampleParameter2D.h"
 #include "Materials/MaterialExpressionTextureSampleParameter2DArray.h"
@@ -42,6 +44,7 @@
 #include "Materials/MaterialFunction.h"
 #include "Materials/MaterialInstanceConstant.h"
 #include "Materials/MaterialInstanceDynamic.h"
+#include "MaterialX/MaterialExpressionTextureSampleParameterBlur.h"
 #include "Misc/CoreMisc.h"
 #include "Misc/Paths.h"
 #include "Nodes/InterchangeBaseNode.h"
@@ -1354,6 +1357,77 @@ void UInterchangeGenericMaterialPipeline::HandleFlattenNormalNode(const UInterch
 	}
 }
 
+void UInterchangeGenericMaterialPipeline::HandleNormalFromHeightMapNode(const UInterchangeShaderNode* ShaderNode, UInterchangeBaseMaterialFactoryNode* MaterialFactoryNode, UInterchangeMaterialExpressionFactoryNode* NormalFromHeightMapFactoryNode)
+{
+	using namespace UE::Interchange::Materials::Standard::Nodes::NormalFromHeightMap;
+
+	NormalFromHeightMapFactoryNode->SetCustomExpressionClassName(UMaterialExpressionMaterialFunctionCall::StaticClass()->GetName());
+
+	const FName MaterialFunctionMemberName = GET_MEMBER_NAME_CHECKED(UMaterialExpressionMaterialFunctionCall, MaterialFunction);
+	NormalFromHeightMapFactoryNode->AddStringAttribute(MaterialFunctionMemberName, TEXT("/Engine/Functions/Engine_MaterialFunctions03/Procedurals/NormalFromHeightmap.NormalFromHeightmap"));
+	NormalFromHeightMapFactoryNode->AddApplyAndFillDelegates<FString>(MaterialFunctionMemberName, UMaterialExpressionMaterialFunctionCall::StaticClass(), MaterialFunctionMemberName);
+
+	// Heightmap
+	{
+		const FString HeightMapInput = Inputs::HeightMap.ToString();
+		TTuple<UInterchangeMaterialExpressionFactoryNode*, FString> HeightMapExpression =
+			CreateMaterialExpressionForInput(MaterialFactoryNode, ShaderNode, HeightMapInput, NormalFromHeightMapFactoryNode->GetUniqueID());
+
+		if(HeightMapExpression.Get<0>())
+		{
+			UInterchangeShaderPortsAPI::ConnectOuputToInput(NormalFromHeightMapFactoryNode, HeightMapInput, HeightMapExpression.Get<0>()->GetUniqueID(), HeightMapExpression.Get<1>());
+		}
+	}
+
+	// Intensity
+	{
+		const FString Intensity = Inputs::Intensity.ToString();
+		TTuple<UInterchangeMaterialExpressionFactoryNode*, FString> IntensityExpression =
+			CreateMaterialExpressionForInput(MaterialFactoryNode, ShaderNode, Intensity, NormalFromHeightMapFactoryNode->GetUniqueID());
+
+		if(IntensityExpression.Get<0>())
+		{
+			UInterchangeShaderPortsAPI::ConnectOuputToInput(NormalFromHeightMapFactoryNode, Intensity, IntensityExpression.Get<0>()->GetUniqueID(), IntensityExpression.Get<1>());
+		}
+	}
+
+	// Offset
+	{
+		const FString Offset = Inputs::Offset.ToString();
+		TTuple<UInterchangeMaterialExpressionFactoryNode*, FString> OffsetExpression =
+			CreateMaterialExpressionForInput(MaterialFactoryNode, ShaderNode, Offset, NormalFromHeightMapFactoryNode->GetUniqueID());
+
+		if(OffsetExpression.Get<0>())
+		{
+			UInterchangeShaderPortsAPI::ConnectOuputToInput(NormalFromHeightMapFactoryNode, Offset, OffsetExpression.Get<0>()->GetUniqueID(), OffsetExpression.Get<1>());
+		}
+	}
+
+	// Coordinates
+	{
+		const FString Coordinates = Inputs::Coordinates.ToString();
+		TTuple<UInterchangeMaterialExpressionFactoryNode*, FString> CoordinatesExpression =
+			CreateMaterialExpressionForInput(MaterialFactoryNode, ShaderNode, Coordinates, NormalFromHeightMapFactoryNode->GetUniqueID());
+
+		if(CoordinatesExpression.Get<0>())
+		{
+			UInterchangeShaderPortsAPI::ConnectOuputToInput(NormalFromHeightMapFactoryNode, Coordinates, CoordinatesExpression.Get<0>()->GetUniqueID(), CoordinatesExpression.Get<1>());
+		}
+	}
+
+	// Channel
+	{
+		const FString Channel = Inputs::Channel.ToString();
+		TTuple<UInterchangeMaterialExpressionFactoryNode*, FString> ChannelExpression =
+			CreateMaterialExpressionForInput(MaterialFactoryNode, ShaderNode, Channel, NormalFromHeightMapFactoryNode->GetUniqueID());
+
+		if(ChannelExpression.Get<0>())
+		{
+			UInterchangeShaderPortsAPI::ConnectOuputToInput(NormalFromHeightMapFactoryNode, Channel, ChannelExpression.Get<0>()->GetUniqueID(), ChannelExpression.Get<1>());
+		}
+	}
+}
+
 void UInterchangeGenericMaterialPipeline::HandleMakeFloat3Node(const UInterchangeShaderNode* ShaderNode, UInterchangeBaseMaterialFactoryNode* MaterialFactoryNode,
 	UInterchangeMaterialExpressionFactoryNode* MakeFloat3FactoryNode)
 {
@@ -1390,6 +1464,90 @@ void UInterchangeGenericMaterialPipeline::HandleMakeFloat3Node(const UInterchang
 	}
 }
 
+void UInterchangeGenericMaterialPipeline::HandleTextureNode(const UInterchangeTextureNode* TextureNode, UInterchangeBaseMaterialFactoryNode* MaterialFactoryNode, UInterchangeMaterialExpressionFactoryNode* TextureBaseFactoryNode, const FString & ExpressionClassName)
+{
+	using namespace UE::Interchange::Materials::Standard::Nodes::TextureSample;
+
+	FString TextureFactoryUid;
+	TArray<FString> TextureTargetNodes;
+	TextureNode->GetTargetNodeUids(TextureTargetNodes);
+
+	if(TextureTargetNodes.Num() > 0)
+	{
+		TextureFactoryUid = TextureTargetNodes[0];
+	}
+
+	TextureBaseFactoryNode->SetCustomExpressionClassName(ExpressionClassName);
+	TextureBaseFactoryNode->AddStringAttribute(UInterchangeShaderPortsAPI::MakeInputValueKey(Inputs::Texture.ToString()), TextureFactoryUid);
+
+	if(UInterchangeTextureFactoryNode* TextureFactoryNode = Cast<UInterchangeTextureFactoryNode>(BaseNodeContainer->GetFactoryNode(TextureFactoryUid)))
+	{
+		EMaterialInputType TextureUsage = EMaterialInputType::Unknown;
+		TextureFactoryNode->GetAttribute(TEXT("TextureUsage"), TextureUsage);
+
+		const bool bIsOutputLinear = MaterialExpressionCreationContextStack.Top().OutputName.Equals(Outputs::A.ToString());
+		const EMaterialInputType DesiredTextureUsage = MaterialCreationContext.InputTypeBeingProcessed == EMaterialInputType::Scalar && bIsOutputLinear ?
+			EMaterialInputType::Unknown : // Alpha channels are always in linear space so ignore them when determining texture usage
+			MaterialCreationContext.InputTypeBeingProcessed;
+
+		if(TextureUsage == EMaterialInputType::Unknown)
+		{
+			if(DesiredTextureUsage == EMaterialInputType::Vector)
+			{
+				TextureFactoryNode->SetCustomCompressionSettings(TextureCompressionSettings::TC_Normalmap);
+			}
+			else if(DesiredTextureUsage == EMaterialInputType::Scalar)
+			{
+				bool bSRGB;
+				if(!TextureNode->GetCustomSRGB(bSRGB))
+				{
+					// Only set CustomSRGB if it wasn't set by the InterchangeGenericTexturePipeline before
+					TextureFactoryNode->SetCustomSRGB(false);
+				}
+			}
+
+			TextureFactoryNode->SetAttribute(TEXT("TextureUsage"), DesiredTextureUsage);
+		}
+		else if(TextureUsage != DesiredTextureUsage && DesiredTextureUsage != EMaterialInputType::Unknown)
+		{
+			UInterchangeResultWarning_Generic* TextureUsageWarning = AddMessage<UInterchangeResultWarning_Generic>();
+			TextureUsageWarning->DestinationAssetName = TextureFactoryNode->GetAssetName();
+			TextureUsageWarning->AssetType = TextureFactoryNode->GetObjectClass();
+
+			TextureUsageWarning->Text = FText::Format(LOCTEXT("TextureUsageMismatch", "{0} is being used as both {1} and {2} which aren't compatible."),
+													  FText::FromString(TextureFactoryNode->GetAssetName()), FText::FromString(LexToString(TextureUsage)), FText::FromString(LexToString(DesiredTextureUsage)));
+
+			// Flipping the green channel only makes sense for vector data as it's used to compensate for different handedness.
+			// Clear it if we're not gonna be used only as a vector map. This normally happens when a normal map is also used as a color map.
+			bool bFlipGreenChannel;
+			if(TextureFactoryNode->GetCustombFlipGreenChannel(bFlipGreenChannel))
+			{
+				TextureFactoryNode->SetCustombFlipGreenChannel(false);
+			}
+		}
+	}
+}
+
+void UInterchangeGenericMaterialPipeline::HandleTextureObjectNode(const UInterchangeShaderNode* ShaderNode, UInterchangeBaseMaterialFactoryNode* MaterialFactoryNode, UInterchangeMaterialExpressionFactoryNode* TextureObjectFactoryNode)
+{
+	using namespace UE::Interchange::Materials::Standard::Nodes::TextureObject;
+
+	FString TextureUid;
+	ShaderNode->GetStringAttribute(UInterchangeShaderPortsAPI::MakeInputValueKey(Inputs::Texture.ToString()), TextureUid);
+	FString ExpressionClassName;
+	FString TextureFactoryUid;
+
+	if(const UInterchangeTextureNode* TextureNode = Cast<const UInterchangeTextureNode>(BaseNodeContainer->GetNode(TextureUid)))
+	{
+		HandleTextureNode(TextureNode, MaterialFactoryNode, TextureObjectFactoryNode, UMaterialExpressionTextureObject::StaticClass()->GetName());
+	}
+	else
+	{
+		TextureObjectFactoryNode->SetCustomExpressionClassName(UMaterialExpressionTextureObject::StaticClass()->GetName());
+		TextureObjectFactoryNode->AddStringAttribute(UInterchangeShaderPortsAPI::MakeInputValueKey(Inputs::Texture.ToString()), TextureFactoryUid);
+	}
+}
+
 void UInterchangeGenericMaterialPipeline::HandleTextureSampleNode(const UInterchangeShaderNode* ShaderNode, UInterchangeBaseMaterialFactoryNode* MaterialFactoryNode, UInterchangeMaterialExpressionFactoryNode* TextureSampleFactoryNode)
 {
 	using namespace UE::Interchange::Materials::Standard::Nodes::TextureSample;
@@ -1410,6 +1568,10 @@ void UInterchangeGenericMaterialPipeline::HandleTextureSampleNode(const UInterch
 		{
 			ExpressionClassName = UMaterialExpressionTextureSampleParameter2DArray::StaticClass()->GetName();
 		}
+		else if(TextureNode->IsA<UInterchangeTextureBlurNode>())
+		{
+			ExpressionClassName = UMaterialExpressionTextureSampleParameterBlur::StaticClass()->GetName();
+		}
 		else if (TextureNode->IsA<UInterchangeTexture2DNode>())
 		{
 			ExpressionClassName = UMaterialExpressionTextureSampleParameter2D::StaticClass()->GetName();
@@ -1419,63 +1581,7 @@ void UInterchangeGenericMaterialPipeline::HandleTextureSampleNode(const UInterch
 			ExpressionClassName = UMaterialExpressionTextureSampleParameter2D::StaticClass()->GetName();
 		}
 
-		TArray<FString> TextureTargetNodes;
-		TextureNode->GetTargetNodeUids(TextureTargetNodes);
-
-		if (TextureTargetNodes.Num() > 0)
-		{
-			TextureFactoryUid = TextureTargetNodes[0];
-		}
-
-		TextureSampleFactoryNode->SetCustomExpressionClassName(ExpressionClassName);
-		TextureSampleFactoryNode->AddStringAttribute(UInterchangeShaderPortsAPI::MakeInputValueKey(Inputs::Texture.ToString()), TextureFactoryUid);
-
-		if (UInterchangeTextureFactoryNode* TextureFactoryNode = Cast<UInterchangeTextureFactoryNode>(BaseNodeContainer->GetFactoryNode(TextureFactoryUid)))
-		{
-			EMaterialInputType TextureUsage = EMaterialInputType::Unknown;
-			TextureFactoryNode->GetAttribute(TEXT("TextureUsage"), TextureUsage);
-
-			const bool bIsOutputLinear = MaterialExpressionCreationContextStack.Top().OutputName.Equals(Outputs::A.ToString());
-			const EMaterialInputType DesiredTextureUsage = MaterialCreationContext.InputTypeBeingProcessed == EMaterialInputType::Scalar && bIsOutputLinear ?
-																		EMaterialInputType::Unknown : // Alpha channels are always in linear space so ignore them when determining texture usage
-																		MaterialCreationContext.InputTypeBeingProcessed;
-
-			if (TextureUsage == EMaterialInputType::Unknown)
-			{
-				if (DesiredTextureUsage == EMaterialInputType::Vector)
-				{
-					TextureFactoryNode->SetCustomCompressionSettings(TextureCompressionSettings::TC_Normalmap);
-				}
-				else if (DesiredTextureUsage == EMaterialInputType::Scalar)
-				{
-					bool bSRGB;
-					if (!TextureNode->GetCustomSRGB(bSRGB))
-					{
-						// Only set CustomSRGB if it wasn't set by the InterchangeGenericTexturePipeline before
-						TextureFactoryNode->SetCustomSRGB(false);
-					}
-				}
-
-				TextureFactoryNode->SetAttribute(TEXT("TextureUsage"), DesiredTextureUsage);
-			}
-			else if (TextureUsage != DesiredTextureUsage && DesiredTextureUsage != EMaterialInputType::Unknown)
-			{
-				UInterchangeResultWarning_Generic* TextureUsageWarning = AddMessage<UInterchangeResultWarning_Generic>();
-				TextureUsageWarning->DestinationAssetName = TextureFactoryNode->GetAssetName();
-				TextureUsageWarning->AssetType = TextureFactoryNode->GetObjectClass();
-
-				TextureUsageWarning->Text = FText::Format(LOCTEXT("TextureUsageMismatch", "{0} is being used as both {1} and {2} which aren't compatible."),
-					FText::FromString(TextureFactoryNode->GetAssetName()), FText::FromString(LexToString(TextureUsage)), FText::FromString(LexToString(DesiredTextureUsage)));
-
-				// Flipping the green channel only makes sense for vector data as it's used to compensate for different handedness.
-				// Clear it if we're not gonna be used only as a vector map. This normally happens when a normal map is also used as a color map.
-				bool bFlipGreenChannel;
-				if (TextureFactoryNode->GetCustombFlipGreenChannel(bFlipGreenChannel))
-				{
-					TextureFactoryNode->SetCustombFlipGreenChannel(false);
-				}
-			}
-		}
+		HandleTextureNode(TextureNode, MaterialFactoryNode, TextureSampleFactoryNode, ExpressionClassName);
 	}
 	else
 	{
@@ -1494,6 +1600,48 @@ void UInterchangeGenericMaterialPipeline::HandleTextureSampleNode(const UInterch
 				CoordinatesExpression.Get<0>()->GetUniqueID(), CoordinatesExpression.Get<1>());
 		}
 	}
+
+	if(ExpressionClassName == UMaterialExpressionTextureSampleParameterBlur::StaticClass()->GetName())
+	{
+		HandleTextureSampleBlurNode(ShaderNode, MaterialFactoryNode, TextureSampleFactoryNode);
+	}
+}
+
+void UInterchangeGenericMaterialPipeline::HandleTextureSampleBlurNode(const UInterchangeShaderNode* ShaderNode, UInterchangeBaseMaterialFactoryNode* MaterialFactoryNode, UInterchangeMaterialExpressionFactoryNode* TextureSampleFactoryNode)
+{
+	using namespace UE::Interchange::Materials::Standard::Nodes;
+
+	// KernelSize
+	if(float KernelSize; ShaderNode->GetFloatAttribute(TextureSampleBlur::Attributes::KernelSize, KernelSize))
+	{
+		const FName KernelSizeMemberName = GET_MEMBER_NAME_CHECKED(UMaterialExpressionTextureSampleParameterBlur, KernelSize);
+		TextureSampleFactoryNode->AddFloatAttribute(KernelSizeMemberName, KernelSize);
+		TextureSampleFactoryNode->AddApplyAndFillDelegates<float>(KernelSizeMemberName, UMaterialExpressionTextureSampleParameterBlur::StaticClass(), KernelSizeMemberName);
+	}
+
+	// FilterSize
+	if(float FilterSize; ShaderNode->GetFloatAttribute(TextureSampleBlur::Attributes::FilterSize, FilterSize))
+	{
+		const FName FilterSizeMemberName = GET_MEMBER_NAME_CHECKED(UMaterialExpressionTextureSampleParameterBlur, FilterSize);
+		TextureSampleFactoryNode->AddFloatAttribute(FilterSizeMemberName, FilterSize);
+		TextureSampleFactoryNode->AddApplyAndFillDelegates<float>(FilterSizeMemberName, UMaterialExpressionTextureSampleParameterBlur::StaticClass(), FilterSizeMemberName);
+	}
+
+	// FilterOffset
+	if(float FilterOffset; ShaderNode->GetFloatAttribute(TextureSampleBlur::Attributes::FilterOffset, FilterOffset))
+	{
+		const FName FilterOffsetMemberName = GET_MEMBER_NAME_CHECKED(UMaterialExpressionTextureSampleParameterBlur, FilterOffset);
+		TextureSampleFactoryNode->AddFloatAttribute(FilterOffsetMemberName, FilterOffset);
+		TextureSampleFactoryNode->AddApplyAndFillDelegates<float>(FilterOffsetMemberName, UMaterialExpressionTextureSampleParameterBlur::StaticClass(), FilterOffsetMemberName);
+	}
+
+	// Filter
+	if(int32 Filter; ShaderNode->GetInt32Attribute(TextureSampleBlur::Attributes::Filter, Filter))
+	{
+		const FName FilterMemberName = GET_MEMBER_NAME_CHECKED(UMaterialExpressionTextureSampleParameterBlur, Filter);
+		TextureSampleFactoryNode->AddInt32Attribute(FilterMemberName, Filter);
+		TextureSampleFactoryNode->AddApplyAndFillDelegates<int>(FilterMemberName, UMaterialExpressionTextureSampleParameterBlur::StaticClass(), FilterMemberName);
+	}
 }
 
 void UInterchangeGenericMaterialPipeline::HandleTextureCoordinateNode(const UInterchangeShaderNode* ShaderNode, UInterchangeBaseMaterialFactoryNode* MaterialFactoryNode,
@@ -2036,10 +2184,18 @@ UInterchangeMaterialExpressionFactoryNode* UInterchangeGenericMaterialPipeline::
 	{
 		HandleNoiseNode(ShaderNode, MaterialFactoryNode, MaterialExpression);
 	}
+	else if(*ShaderType == Nodes::NormalFromHeightMap::Name)
+	{
+		HandleNormalFromHeightMapNode(ShaderNode, MaterialFactoryNode, MaterialExpression);
+	}
 	else if (*ShaderType == Nodes::TextureCoordinate::Name)
 	{
 		HandleTextureCoordinateNode(ShaderNode, MaterialFactoryNode, MaterialExpression);
 	}
+	else if(*ShaderType == Nodes::TextureObject::Name)
+	{
+		HandleTextureObjectNode(ShaderNode, MaterialFactoryNode, MaterialExpression);
+	}
 	else if (*ShaderType == Nodes::TextureSample::Name)
 	{
 		HandleTextureSampleNode(ShaderNode, MaterialFactoryNode, MaterialExpression);
diff --git a/Engine/Plugins/Interchange/Runtime/Source/Pipelines/Private/InterchangeGenericTexturePipeline.cpp b/Engine/Plugins/Interchange/Runtime/Source/Pipelines/Private/InterchangeGenericTexturePipeline.cpp
--- a/Engine/Plugins/Interchange/Runtime/Source/Pipelines/Private/InterchangeGenericTexturePipeline.cpp
+++ b/Engine/Plugins/Interchange/Runtime/Source/Pipelines/Private/InterchangeGenericTexturePipeline.cpp
@@ -10,6 +10,7 @@
 #include "InterchangeTexture2DArrayNode.h"
 #include "InterchangeTexture2DFactoryNode.h"
 #include "InterchangeTexture2DNode.h"
+#include "InterchangeTextureBlurNode.h"
 #include "InterchangeTextureCubeArrayFactoryNode.h"
 #include "InterchangeTextureCubeArrayNode.h"
 #include "InterchangeTextureCubeFactoryNode.h"
@@ -64,6 +65,11 @@ namespace UE::Interchange::Private
 			return UInterchangeVolumeTextureFactoryNode::StaticClass();
 		}
 
+		if(UInterchangeTextureBlurNode::StaticClass() == NodeClass)
+		{
+			return UInterchangeTexture2DFactoryNode::StaticClass();
+		}
+
 		return nullptr;
 	}
 
diff --git a/Engine/Plugins/Interchange/Runtime/Source/Pipelines/Public/InterchangeGenericMaterialPipeline.h b/Engine/Plugins/Interchange/Runtime/Source/Pipelines/Public/InterchangeGenericMaterialPipeline.h
--- a/Engine/Plugins/Interchange/Runtime/Source/Pipelines/Public/InterchangeGenericMaterialPipeline.h
+++ b/Engine/Plugins/Interchange/Runtime/Source/Pipelines/Public/InterchangeGenericMaterialPipeline.h
@@ -22,6 +22,7 @@ class UInterchangeMaterialFunctionFactoryNode;
 class UInterchangeMaterialInstanceFactoryNode;
 class UInterchangeResult;
 class UMaterialFunction;
+class UInterchangeTextureNode;
 
 UENUM(BlueprintType)
 enum class EInterchangeMaterialImportOption : uint8
@@ -44,7 +45,7 @@ enum class EInterchangeMaterialXShaders : uint8
 USTRUCT(BlueprintType)
 struct INTERCHANGEPIPELINES_API FMaterialXPipelineSettings
 {
-	GENERATED_BODY()
+	GENERATED_USTRUCT_BODY()
 
 public:
 	FMaterialXPipelineSettings();
@@ -164,8 +165,12 @@ private:
 	bool HandleUnlitModel(const UInterchangeShaderGraphNode* ShaderGraphNode, UInterchangeMaterialFactoryNode* MaterialFactoryNode);
 
 	void HandleFlattenNormalNode(const UInterchangeShaderNode* ShaderNode, UInterchangeBaseMaterialFactoryNode* MaterialFactoryNode, UInterchangeMaterialExpressionFactoryNode* FlattenNormalFactoryNode);
+	void HandleNormalFromHeightMapNode(const UInterchangeShaderNode* ShaderNode, UInterchangeBaseMaterialFactoryNode* MaterialFactoryNode, UInterchangeMaterialExpressionFactoryNode* NormalFromHeightMapFactoryNode);
 	void HandleMakeFloat3Node(const UInterchangeShaderNode* ShaderNode, UInterchangeBaseMaterialFactoryNode* MaterialFactoryNode, UInterchangeMaterialExpressionFactoryNode* MakeFloat3FactoryNode);
+	void HandleTextureNode(const UInterchangeTextureNode* TextureNode, UInterchangeBaseMaterialFactoryNode* MaterialFactoryNode, UInterchangeMaterialExpressionFactoryNode* TextureBaseFactoryNode, const FString& ExpressionClassName);
+	void HandleTextureObjectNode(const UInterchangeShaderNode* ShaderNode, UInterchangeBaseMaterialFactoryNode* MaterialFactoryNode, UInterchangeMaterialExpressionFactoryNode* TextureObjectFactoryNode);
 	void HandleTextureSampleNode(const UInterchangeShaderNode* ShaderNode, UInterchangeBaseMaterialFactoryNode* MaterialFactoryNode, UInterchangeMaterialExpressionFactoryNode* TextureSampleFactoryNode);
+	void HandleTextureSampleBlurNode(const UInterchangeShaderNode* ShaderNode, UInterchangeBaseMaterialFactoryNode* MaterialFactoryNode, UInterchangeMaterialExpressionFactoryNode* TextureSampleFactoryNode);
 	void HandleTextureCoordinateNode(const UInterchangeShaderNode* ShaderNode, UInterchangeBaseMaterialFactoryNode* MaterialFactoryNode, UInterchangeMaterialExpressionFactoryNode*& TextureSampleFactoryNode);
 	void HandleLerpNode(const UInterchangeShaderNode* ShaderNode, UInterchangeBaseMaterialFactoryNode* MaterialFactoryNode, UInterchangeMaterialExpressionFactoryNode* LerpFactoryNode);
 	void HandleMaskNode(const UInterchangeShaderNode* ShaderNode, UInterchangeBaseMaterialFactoryNode* MaterialFactoryNode, UInterchangeMaterialExpressionFactoryNode* MaskFactoryNode);
