Fixing possible dereferencing of a null pointer

diff --git a/Engine/Plugins/Interchange/Runtime/Source/Import/Private/MaterialX/InterchangeMaterialXTranslator.cpp b/Engine/Plugins/Interchange/Runtime/Source/Import/Private/MaterialX/InterchangeMaterialXTranslator.cpp
--- a/Engine/Plugins/Interchange/Runtime/Source/Import/Private/MaterialX/InterchangeMaterialXTranslator.cpp
+++ b/Engine/Plugins/Interchange/Runtime/Source/Import/Private/MaterialX/InterchangeMaterialXTranslator.cpp
@@ -2095,12 +2095,12 @@ void UInterchangeMaterialXTranslator::ConnectHeightToNormalInputToOutput(Materia
 				AddAttributeFromValueOrInterface(UpstreamNode->getInput(mx::NodeGroup::Texture2D::Inputs::Default), InputChannelName, ParentShaderNode);
 			}
 		}
-	}
-	else
-	{
-		// For the moment it doesn't make sense to plug a value to it, so let's plug directly the child to the parent, in the future we could implement a Sobel and handle a multi output
-		RenameInput(Input, TCHAR_TO_UTF8(*InputChannelName)); //let's take the parent node's input name
-		NamesToShaderNodes.Add(UpstreamNode->getName().c_str(), ParentShaderNode);
+		else
+		{
+			// For the moment it doesn't make sense to plug a value to it, so let's plug directly the child to the parent, in the future we could implement a Sobel and handle a multi output
+			RenameInput(Input, TCHAR_TO_UTF8(*InputChannelName)); //let's take the parent node's input name
+			NamesToShaderNodes.Add(UpstreamNode->getName().c_str(), ParentShaderNode);
+		}
 	}
 }
 
@@ -2164,12 +2164,12 @@ void UInterchangeMaterialXTranslator::ConnectBlurInputToOutput(MaterialX::NodePt
 				AddAttributeFromValueOrInterface(UpstreamNode->getInput(mx::NodeGroup::Texture2D::Inputs::Default), InputChannelName, ParentShaderNode);
 			}
 		}
-	}
-	else
-	{
-		// For a blur it doesn't make sense if there's no image input
-		RenameInput(Input, TCHAR_TO_UTF8(*InputChannelName)); //let's take the parent node's input name
-		NamesToShaderNodes.Add(UpstreamNode->getName().c_str(), ParentShaderNode);
+		else
+		{
+			// For a blur it doesn't make sense if there's no image input
+			RenameInput(Input, TCHAR_TO_UTF8(*InputChannelName)); //let's take the parent node's input name
+			NamesToShaderNodes.Add(UpstreamNode->getName().c_str(), ParentShaderNode);
+		}
 	}
 }
 
