[Interchange] Refactor GetMeshPayloadData interfaces.
Also fixing gltf asset creation for mesh type forcing scenarios

diff --git a/Engine/Plugins/Experimental/Enterprise/DatasmithInterchange/Source/DatasmithInterchange/Private/InterchangeDatasmithTranslator.cpp b/Engine/Plugins/Experimental/Enterprise/DatasmithInterchange/Source/DatasmithInterchange/Private/InterchangeDatasmithTranslator.cpp
--- a/Engine/Plugins/Experimental/Enterprise/DatasmithInterchange/Source/DatasmithInterchange/Private/InterchangeDatasmithTranslator.cpp
+++ b/Engine/Plugins/Experimental/Enterprise/DatasmithInterchange/Source/DatasmithInterchange/Private/InterchangeDatasmithTranslator.cpp
@@ -199,7 +199,7 @@ bool UInterchangeDatasmithTranslator::Translate(UInterchangeBaseNodeContainer& B
 				const FString DisplayLabel = StaticMeshNameProvider.GenerateUniqueName(MeshElement->GetLabel());
 
 				MeshNode->InitializeNode(MeshNodeUid, DisplayLabel, EInterchangeNodeContainerType::TranslatedAsset);
-				MeshNode->SetPayLoadKey(LexToString(MeshIndex));
+				MeshNode->SetPayLoadKey(LexToString(MeshIndex), EInterchangeMeshPayLoadType::STATIC);
 				MeshNode->SetSkinnedMesh(false);
 				MeshNode->SetCustomHasVertexNormal(true);
 				// TODO: Interchange expect each LOD to have its own mesh node and to declare the number of vertices, however we don't know the content of a datasmith mesh until its bulk data is loaded.
@@ -547,10 +547,10 @@ TOptional<UE::Interchange::FImportImage> UInterchangeDatasmithTranslator::GetTex
 	return TextureTranslator->GetTexturePayloadData(PayloadSourceData, TextureElement->GetFile());
 }
 
-TFuture<TOptional<UE::Interchange::FStaticMeshPayloadData>> UInterchangeDatasmithTranslator::GetStaticMeshPayloadData(const FString& PayloadKey) const
+TFuture<TOptional<UE::Interchange::FMeshPayloadData>> UInterchangeDatasmithTranslator::GetMeshPayloadData(const FInterchangeMeshPayLoadKey& PayLoadKey) const
 {
-	TPromise<TOptional<UE::Interchange::FStaticMeshPayloadData>> EmptyPromise;
-	EmptyPromise.SetValue(TOptional<UE::Interchange::FStaticMeshPayloadData>());
+	TPromise<TOptional<UE::Interchange::FMeshPayloadData>> EmptyPromise;
+	EmptyPromise.SetValue(TOptional<UE::Interchange::FMeshPayloadData>());
 
 	if (!LoadedExternalSource || !LoadedExternalSource->GetDatasmithScene())
 	{
@@ -558,7 +558,7 @@ TFuture<TOptional<UE::Interchange::FStaticMeshPayloadData>> UInterchangeDatasmit
 	}
 
 	int32 MeshIndex = 0;
-	LexFromString(MeshIndex, *PayloadKey);
+	LexFromString(MeshIndex, *PayLoadKey.UniqueId);
 	TSharedPtr<IDatasmithScene> DatasmithScene = LoadedExternalSource->GetDatasmithScene();
 	if (MeshIndex < 0 || MeshIndex >= DatasmithScene->GetMeshesCount())
 	{
@@ -573,14 +573,14 @@ TFuture<TOptional<UE::Interchange::FStaticMeshPayloadData>> UInterchangeDatasmit
 
 	return Async(EAsyncExecution::TaskGraph, [this, MeshElement = MoveTemp(MeshElement)]
 		{
-			TOptional<UE::Interchange::FStaticMeshPayloadData> Result;
+			TOptional<UE::Interchange::FMeshPayloadData> Result;
 
 			FDatasmithMeshElementPayload DatasmithMeshPayload;
 			if (LoadedExternalSource->GetAssetTranslator()->LoadStaticMesh(MeshElement.ToSharedRef(), DatasmithMeshPayload))
 			{
 				if (DatasmithMeshPayload.LodMeshes.Num() > 0)
 				{
-					UE::Interchange::FStaticMeshPayloadData StaticMeshPayloadData;
+					UE::Interchange::FMeshPayloadData StaticMeshPayloadData;
 					StaticMeshPayloadData.MeshDescription = MoveTemp(DatasmithMeshPayload.LodMeshes[0]);
 
 					Result.Emplace(MoveTemp(StaticMeshPayloadData));
diff --git a/Engine/Plugins/Experimental/Enterprise/DatasmithInterchange/Source/DatasmithInterchange/Public/InterchangeDatasmithTranslator.h b/Engine/Plugins/Experimental/Enterprise/DatasmithInterchange/Source/DatasmithInterchange/Public/InterchangeDatasmithTranslator.h
--- a/Engine/Plugins/Experimental/Enterprise/DatasmithInterchange/Source/DatasmithInterchange/Public/InterchangeDatasmithTranslator.h
+++ b/Engine/Plugins/Experimental/Enterprise/DatasmithInterchange/Source/DatasmithInterchange/Public/InterchangeDatasmithTranslator.h
@@ -6,10 +6,9 @@
 
 #include "Animation/InterchangeAnimationPayloadInterface.h"
 #include "InterchangeTranslatorBase.h"
-#include "Mesh/InterchangeStaticMeshPayloadInterface.h"
+#include "Mesh/InterchangeMeshPayloadInterface.h"
 #include "Texture/InterchangeTexturePayloadData.h"
 #include "Texture/InterchangeTexturePayloadInterface.h"
-#include "Mesh/InterchangeStaticMeshPayloadInterface.h"
 #include "Scene/InterchangeVariantSetPayloadInterface.h"
 
 #include "Async/Async.h"
@@ -33,7 +32,7 @@ class UInterchangeSceneNode;
 namespace UE::Interchange
 {
 	struct FImportImage;
-	struct FStaticMeshPayloadData;
+	struct FMeshPayloadData;
 	struct FAnimationTransformPayloadData;
 	struct FAnimationBakeTransformPayloadData;
 }
@@ -53,7 +52,7 @@ namespace UE::DatasmithInterchange::AnimUtils
 UCLASS(BlueprintType, Experimental)
 class DATASMITHINTERCHANGE_API UInterchangeDatasmithTranslator : public UInterchangeTranslatorBase
 	, public IInterchangeTexturePayloadInterface
-	, public IInterchangeStaticMeshPayloadInterface
+	, public IInterchangeMeshPayloadInterface
 	, public IInterchangeAnimationPayloadInterface
 	, public IInterchangeVariantSetPayloadInterface
 {
@@ -80,7 +79,7 @@ public:
 	/* IInterchangeTexturePayloadInterface End */
 
 	/* IInterchangeStaticMeshPayloadInterface Begin */
-	virtual TFuture<TOptional<UE::Interchange::FStaticMeshPayloadData>> GetStaticMeshPayloadData(const FString& PayloadKey) const override;
+	virtual TFuture<TOptional<UE::Interchange::FMeshPayloadData>> GetMeshPayloadData(const FInterchangeMeshPayLoadKey& PayLoadKey) const override;
 	/* IInterchangeStaticMeshPayloadInterface End */
 
 	/* IInterchangeAnimationPayloadInterface Begin */
diff --git a/Engine/Plugins/Interchange/Runtime/Source/Import/Private/Fbx/InterchangeFbxTranslator.cpp b/Engine/Plugins/Interchange/Runtime/Source/Import/Private/Fbx/InterchangeFbxTranslator.cpp
--- a/Engine/Plugins/Interchange/Runtime/Source/Import/Private/Fbx/InterchangeFbxTranslator.cpp
+++ b/Engine/Plugins/Interchange/Runtime/Source/Import/Private/Fbx/InterchangeFbxTranslator.cpp
@@ -10,8 +10,7 @@
 #include "InterchangeDispatcherTask.h"
 #include "InterchangeManager.h"
 #include "InterchangeImportLog.h"
-#include "Mesh/InterchangeSkeletalMeshPayload.h"
-#include "Mesh/InterchangeStaticMeshPayload.h"
+#include "Mesh/InterchangeMeshPayload.h"
 #include "Misc/ConfigCacheIni.h"
 #include "Misc/FileHelper.h"
 #include "Misc/Paths.h"
@@ -183,19 +182,19 @@ TOptional<UE::Interchange::FImportImage> UInterchangeFbxTranslator::GetTexturePa
 	return TextureTranslator->GetTexturePayloadData(PayloadSourceData, PayLoadKey);
 }
 
-TFuture<TOptional<UE::Interchange::FStaticMeshPayloadData>> UInterchangeFbxTranslator::GetStaticMeshPayloadData(const FString& PayLoadKey) const
+TFuture<TOptional<UE::Interchange::FMeshPayloadData>> UInterchangeFbxTranslator::GetMeshPayloadData(const FInterchangeMeshPayLoadKey& PayLoadKey) const
 {
-	TSharedPtr<TPromise<TOptional<UE::Interchange::FStaticMeshPayloadData>>> Promise = MakeShared<TPromise<TOptional<UE::Interchange::FStaticMeshPayloadData>>>();
+	TSharedPtr<TPromise<TOptional<UE::Interchange::FMeshPayloadData>>> Promise = MakeShared<TPromise<TOptional<UE::Interchange::FMeshPayloadData>>>();
 
 	if (!Dispatcher.IsValid())
 	{
-		Promise->SetValue(TOptional<UE::Interchange::FStaticMeshPayloadData>());
+		Promise->SetValue(TOptional<UE::Interchange::FMeshPayloadData>());
 		return Promise->GetFuture();
 	}
 
 	// Create a json command to read the fbx file
-	FString JsonCommand = CreateFetchPayloadFbxCommand(PayLoadKey);
-	const int32 CreatedTaskIndex = Dispatcher->AddTask(JsonCommand, FInterchangeDispatcherTaskCompleted::CreateLambda([this, Promise](const int32 TaskIndex)
+	FString JsonCommand = CreateFetchPayloadFbxCommand(PayLoadKey.UniqueId);
+	const int32 CreatedTaskIndex = Dispatcher->AddTask(JsonCommand, FInterchangeDispatcherTaskCompleted::CreateLambda([this, Promise, PayLoadKey](const int32 TaskIndex)
 	{
 		TRACE_CPUPROFILER_EVENT_SCOPE("UInterchangeFbxTranslator::GetStaticMeshPayloadData::Dispatcher->AddTaskDone")
 		UE::Interchange::ETaskState TaskState;
@@ -211,7 +210,7 @@ TFuture<TOptional<UE::Interchange::FStaticMeshPayloadData>> UInterchangeFbxTrans
 
 		if (TaskState != UE::Interchange::ETaskState::ProcessOk)
 		{
-			Promise->SetValue(TOptional<UE::Interchange::FStaticMeshPayloadData>());
+			Promise->SetValue(TOptional<UE::Interchange::FMeshPayloadData>());
 			return;
 		}
 
@@ -223,12 +222,12 @@ TFuture<TOptional<UE::Interchange::FStaticMeshPayloadData>> UInterchangeFbxTrans
 		if (!ensure(FPaths::FileExists(StaticMeshPayloadFilename)))
 		{
 			// TODO log an error saying the payload file does not exist even if the get payload command succeeded
-			Promise->SetValue(TOptional<UE::Interchange::FStaticMeshPayloadData>());
+			Promise->SetValue(TOptional<UE::Interchange::FMeshPayloadData>());
 			return;
 		}
 
-		UE::Interchange::FStaticMeshPayloadData StaticMeshPayload;
-		StaticMeshPayload.MeshDescription.Empty();
+		UE::Interchange::FMeshPayloadData MeshPayloadData;
+		MeshPayloadData.MeshDescription.Empty();
 
 		// All sub object should be gone with the reset
 		TArray64<uint8> Buffer;
@@ -238,178 +237,47 @@ TFuture<TOptional<UE::Interchange::FStaticMeshPayloadData>> UInterchangeFbxTrans
 		if (FileDataSize < 1)
 		{
 			// Nothing to load from this file
-			Promise->SetValue(TOptional<UE::Interchange::FStaticMeshPayloadData>());
+			Promise->SetValue(TOptional<UE::Interchange::FMeshPayloadData>());
 			return;
 		}
 
-		// Buffer keeps the ownership of the data, the large memory reader is use to serialize the TMap
-		FLargeMemoryReader Ar(FileData, FileDataSize);
-		StaticMeshPayload.MeshDescription.Serialize(Ar);
-
-		// This is a static mesh payload can contain skinned data if we need to convert skeletalmesh to staticmesh
-		bool bFetchSkinnedData = false;
-		Ar << bFetchSkinnedData;
-		if (bFetchSkinnedData)
+		switch (PayLoadKey.Type)
 		{
-			TArray<FString> JointNames;
-			//Read into a dummy structure, we will not use the influence channel of this skinned meshdescription
-			Ar << JointNames;
-		}
-		Promise->SetValue(MoveTemp(StaticMeshPayload));
-	}));
-
-	// The task was not added to the dispatcher
-	if (CreatedTaskIndex == INDEX_NONE)
-	{
-		Promise->SetValue(TOptional<UE::Interchange::FStaticMeshPayloadData>{});
-	}
-
-	return Promise->GetFuture();
-}
-
-TFuture<TOptional<UE::Interchange::FSkeletalMeshLodPayloadData>> UInterchangeFbxTranslator::GetSkeletalMeshLodPayloadData(const FString& PayLoadKey) const
-{
-	TSharedPtr<TPromise<TOptional<UE::Interchange::FSkeletalMeshLodPayloadData>>> Promise = MakeShared<TPromise<TOptional<UE::Interchange::FSkeletalMeshLodPayloadData>>>();
-	if (!Dispatcher.IsValid())
-	{
-		Promise->SetValue(TOptional<UE::Interchange::FSkeletalMeshLodPayloadData>{});
-		return Promise->GetFuture();
-	}
-
-	//Create a json command to read the fbx file
-	const FString JsonCommand = CreateFetchPayloadFbxCommand(PayLoadKey);
-	const int32 CreatedTaskIndex = Dispatcher->AddTask(JsonCommand, FInterchangeDispatcherTaskCompleted::CreateLambda([this, Promise](const int32 TaskIndex)
-	{
-		TRACE_CPUPROFILER_EVENT_SCOPE("UInterchangeFbxTranslator::GetSkeletalMeshLodPayloadData::Dispatcher->AddTaskDone")
-		UE::Interchange::ETaskState TaskState;
-		FString JsonResult;
-		TArray<FString> JsonMessages;
-		Dispatcher->GetTaskState(TaskIndex, TaskState, JsonResult, JsonMessages);
-
-		// Parse the Json messages into UInterchangeResults
-		for (const FString& JsonMessage : JsonMessages)
+		case EInterchangeMeshPayLoadType::STATIC:
+		case EInterchangeMeshPayLoadType::SKELETAL:
+			{
+				// Buffer keeps the ownership of the data, the large memory reader is use to serialize the TMap
+				FLargeMemoryReader Ar(FileData, FileDataSize);
+				MeshPayloadData.MeshDescription.Serialize(Ar);
+
+				// This is a static mesh payload can contain skinned data if we need to convert skeletalmesh to staticmesh
+				bool bFetchSkinnedData = false;
+				Ar << bFetchSkinnedData;
+				if (bFetchSkinnedData)
+				{
+					Ar << MeshPayloadData.JointNames;
+				}
+			}
+			break;
+		case EInterchangeMeshPayLoadType::MORPHTARGET:
 		{
-			AddMessage(UInterchangeResult::FromJson(JsonMessage));
+			//Buffer keep the ownership of the data, the large memory reader is use to serialize the TMap
+			FLargeMemoryReader Ar(FileData, FileDataSize);
+			MeshPayloadData.MeshDescription.Serialize(Ar);
 		}
-
-		if (TaskState != UE::Interchange::ETaskState::ProcessOk)
-		{
-			Promise->SetValue(TOptional<UE::Interchange::FSkeletalMeshLodPayloadData>{});
-			return;
+			break;
+		default:
+			break;
 		}
-		//Grab the result file and fill the BaseNodeContainer
-		UE::Interchange::FJsonFetchPayloadCmd::JsonResultParser ResultParser;
-		ResultParser.FromJson(JsonResult);
-		FString SkeletalMeshPayloadFilename = ResultParser.GetResultFilename();
 
-		if (!ensure(FPaths::FileExists(SkeletalMeshPayloadFilename)))
-		{
-			//TODO log an error saying the payload file do not exist even if the get payload command succeed
-			Promise->SetValue(TOptional<UE::Interchange::FSkeletalMeshLodPayloadData>{});
-			return;
-		}
 		
-		//All sub object should be gone with the reset
-		TArray64<uint8> Buffer;
-		FFileHelper::LoadFileToArray(Buffer, *SkeletalMeshPayloadFilename);
-		uint8* FileData = Buffer.GetData();
-		int64 FileDataSize = Buffer.Num();
-		if (FileDataSize < 1)
-		{
-			//Nothing to load from this file
-			Promise->SetValue(TOptional<UE::Interchange::FSkeletalMeshLodPayloadData>{});
-			return;
-		}
-
-		UE::Interchange::FSkeletalMeshLodPayloadData SkeletalMeshLodPayload;
-		SkeletalMeshLodPayload.LodMeshDescription.Empty();
-
-		//Buffer keep the ownership of the data, the large memory reader is use to serialize the TMap
-		FLargeMemoryReader Ar(FileData, FileDataSize);
-		SkeletalMeshLodPayload.LodMeshDescription.Serialize(Ar);
-		bool bFetchSkinnedData = false;
-		Ar << bFetchSkinnedData;
-		if (bFetchSkinnedData)
-		{
-			//Read the bone Name to remap the influence correctly
-			Ar << SkeletalMeshLodPayload.JointNames;
-		}
-		Promise->SetValue(MoveTemp(SkeletalMeshLodPayload));
-	}));
-
-	//The task was not added to the dispatcher
-	if (CreatedTaskIndex == INDEX_NONE)
-	{
-		Promise->SetValue(TOptional<UE::Interchange::FSkeletalMeshLodPayloadData>{});
-	}
-
-	return Promise->GetFuture();
-}
-
-TFuture<TOptional<UE::Interchange::FSkeletalMeshMorphTargetPayloadData>> UInterchangeFbxTranslator::GetSkeletalMeshMorphTargetPayloadData(const FString& PayLoadKey) const
-{
-	TSharedPtr<TPromise<TOptional<UE::Interchange::FSkeletalMeshMorphTargetPayloadData>>> Promise = MakeShared<TPromise<TOptional<UE::Interchange::FSkeletalMeshMorphTargetPayloadData>>>();
-	if (!Dispatcher.IsValid())
-	{
-		Promise->SetValue(TOptional<UE::Interchange::FSkeletalMeshMorphTargetPayloadData>{});
-		return Promise->GetFuture();
-	}
-
-	//Create a json command to read the fbx file
-	const FString JsonCommand = CreateFetchPayloadFbxCommand(PayLoadKey);
-	const int32 CreatedTaskIndex = Dispatcher->AddTask(JsonCommand, FInterchangeDispatcherTaskCompleted::CreateLambda([this, Promise](const int32 TaskIndex)
-	{
-		TRACE_CPUPROFILER_EVENT_SCOPE("UInterchangeFbxTranslator::GetSkeletalMeshMorphTargetPayloadData::Dispatcher->AddTaskDone")
-		UE::Interchange::ETaskState TaskState;
-		FString JsonResult;
-		TArray<FString> JsonMessages;
-		Dispatcher->GetTaskState(TaskIndex, TaskState, JsonResult, JsonMessages);
-
-		// Parse the Json messages into UInterchangeResults
-		for (const FString& JsonMessage : JsonMessages)
-		{
-			AddMessage(UInterchangeResult::FromJson(JsonMessage));
-		}
-
-		if (TaskState != UE::Interchange::ETaskState::ProcessOk)
-		{
-			Promise->SetValue(TOptional<UE::Interchange::FSkeletalMeshMorphTargetPayloadData>{});
-			return;
-		}
-		//Grab the result file and fill the BaseNodeContainer
-		UE::Interchange::FJsonFetchPayloadCmd::JsonResultParser ResultParser;
-		ResultParser.FromJson(JsonResult);
-		FString SkeletalMeshPayloadFilename = ResultParser.GetResultFilename();
-
-		if (!ensure(FPaths::FileExists(SkeletalMeshPayloadFilename)))
-		{
-			//TODO log an error saying the payload file do not exist even if the get payload command succeed
-			Promise->SetValue(TOptional<UE::Interchange::FSkeletalMeshMorphTargetPayloadData>{});
-			return;
-		}
-		//All sub object should be gone with the reset
-		TArray64<uint8> Buffer;
-		FFileHelper::LoadFileToArray(Buffer, *SkeletalMeshPayloadFilename);
-		uint8* FileData = Buffer.GetData();
-		int64 FileDataSize = Buffer.Num();
-		if (FileDataSize < 1)
-		{
-			//Nothing to load from this file
-			Promise->SetValue(TOptional<UE::Interchange::FSkeletalMeshMorphTargetPayloadData>{});
-			return;
-		}
-		UE::Interchange::FSkeletalMeshMorphTargetPayloadData SkeletalMeshMorphTargetPayload;
-		SkeletalMeshMorphTargetPayload.LodMeshDescription.Empty();
-		//Buffer keep the ownership of the data, the large memory reader is use to serialize the TMap
-		FLargeMemoryReader Ar(FileData, FileDataSize);
-		SkeletalMeshMorphTargetPayload.LodMeshDescription.Serialize(Ar);
-		Promise->SetValue(MoveTemp(SkeletalMeshMorphTargetPayload));
+		Promise->SetValue(MoveTemp(MeshPayloadData));
 	}));
 
-	//The task was not added to the dispatcher
+	// The task was not added to the dispatcher
 	if (CreatedTaskIndex == INDEX_NONE)
 	{
-		Promise->SetValue(TOptional<UE::Interchange::FSkeletalMeshMorphTargetPayloadData>{});
+		Promise->SetValue(TOptional<UE::Interchange::FMeshPayloadData>{});
 	}
 
 	return Promise->GetFuture();
diff --git a/Engine/Plugins/Interchange/Runtime/Source/Import/Private/Gltf/InterchangeGltfMesh.cpp b/Engine/Plugins/Interchange/Runtime/Source/Import/Private/Gltf/InterchangeGltfMesh.cpp
--- a/Engine/Plugins/Interchange/Runtime/Source/Import/Private/Gltf/InterchangeGltfMesh.cpp
+++ b/Engine/Plugins/Interchange/Runtime/Source/Import/Private/Gltf/InterchangeGltfMesh.cpp
@@ -13,7 +13,7 @@
 #include "StaticMeshOperations.h"
 #include "SkeletalMeshOperations.h"
 
-#include "Mesh/InterchangeStaticMeshPayload.h"
+#include "Mesh/InterchangeMeshPayload.h"
 
 namespace UE::Interchange::Gltf::Private
 {
@@ -73,16 +73,8 @@ namespace UE::Interchange::Gltf::Private
 
 	bool GetSkeletalMeshDescriptionForPayLoadKey(const GLTF::FAsset& GltfAsset, const FString& PayLoadKey, FMeshDescription& MeshDescription, TArray<FString>* OutJointUniqueNames)
 	{
-		TArray<FString> MorphTargetCheckPayLoadKeys;
-		PayLoadKey.ParseIntoArray(MorphTargetCheckPayLoadKeys, TEXT("~"));
-
-		if (MorphTargetCheckPayLoadKeys.Num() == 2)
-		{
-			return GetStaticMeshPayloadDataForPayLoadKey(GltfAsset, MorphTargetCheckPayLoadKeys[1], MeshDescription);
-		}
-		
 		TArray<FString> PayLoadKeys;
-		MorphTargetCheckPayLoadKeys[0].ParseIntoArray(PayLoadKeys, TEXT(":"));
+		PayLoadKey.ParseIntoArray(PayLoadKeys, TEXT(":"));
 		TMap<int32, TArray<int32>> MeshIndexToSkinIndicesMap;
 
 		int32 MeshIndex = 0;
diff --git a/Engine/Plugins/Interchange/Runtime/Source/Import/Private/Gltf/InterchangeGltfPrivate.h b/Engine/Plugins/Interchange/Runtime/Source/Import/Private/Gltf/InterchangeGltfPrivate.h
--- a/Engine/Plugins/Interchange/Runtime/Source/Import/Private/Gltf/InterchangeGltfPrivate.h
+++ b/Engine/Plugins/Interchange/Runtime/Source/Import/Private/Gltf/InterchangeGltfPrivate.h
@@ -14,7 +14,7 @@ namespace UE::Interchange
 {
 	struct FAnimationCurvePayloadData;
 	struct FAnimationBakeTransformPayloadData;
-	struct FStaticMeshPayloadData;
+	struct FMeshPayloadData;
 
 	namespace Gltf::Private
 	{
diff --git a/Engine/Plugins/Interchange/Runtime/Source/Import/Private/Gltf/InterchangeGltfTranslator.cpp b/Engine/Plugins/Interchange/Runtime/Source/Import/Private/Gltf/InterchangeGltfTranslator.cpp
--- a/Engine/Plugins/Interchange/Runtime/Source/Import/Private/Gltf/InterchangeGltfTranslator.cpp
+++ b/Engine/Plugins/Interchange/Runtime/Source/Import/Private/Gltf/InterchangeGltfTranslator.cpp
@@ -1258,20 +1258,6 @@ bool UInterchangeGltfTranslator::Translate( UInterchangeBaseNodeContainer& NodeC
 	return true;
 }
 
-TFuture< TOptional< UE::Interchange::FStaticMeshPayloadData > > UInterchangeGltfTranslator::GetStaticMeshPayloadData( const FString& PayLoadKey ) const
-{
-	return Async(EAsyncExecution::TaskGraph, [this, PayLoadKey]
-		{
-			UE::Interchange::FStaticMeshPayloadData StaticMeshPayloadData;
-			UE::Interchange::Gltf::Private::GetStaticMeshPayloadDataForPayLoadKey(GltfAsset, PayLoadKey, StaticMeshPayloadData.MeshDescription);
-
-			TOptional<UE::Interchange::FStaticMeshPayloadData> Result;
-			Result.Emplace(StaticMeshPayloadData);
-
-			return Result;
-		});
-}
-
 TOptional< UE::Interchange::FImportImage > UInterchangeGltfTranslator::GetTexturePayloadData( const UInterchangeSourceData* InSourceData, const FString& PayLoadKey ) const
 {
 	int32 TextureIndex = 0;
@@ -1451,10 +1437,10 @@ void UInterchangeGltfTranslator::HandleGltfAnimation(UInterchangeBaseNodeContain
 							{
 								if (MorphTargetNodeConst->GetPayLoadKey().IsSet())
 								{
-									FString PayLoadKey = MorphTargetNodeConst->GetPayLoadKey().GetValue(); // meshindex : MorphTargetIndex
-									PayLoadKey = TEXT("MorphTargetAnimation~") + LexToString(AnimationIndex) + TEXT(":") + LexToString(ChannelIndex) + TEXT(":") + PayLoadKey;
+									FInterchangeMeshPayLoadKey PayLoadKey = MorphTargetNodeConst->GetPayLoadKey().GetValue();
+									FString PayLoadKeyUniqueId = TEXT("MorphTargetAnimation~") + LexToString(AnimationIndex) + TEXT(":") + LexToString(ChannelIndex) + TEXT(":") + PayLoadKey.UniqueId;
 
-									AnimationPayloadKeyForMorphTargetNodeUids.Add(MorphTargetDependencyUid, PayLoadKey);
+									AnimationPayloadKeyForMorphTargetNodeUids.Add(MorphTargetDependencyUid, PayLoadKeyUniqueId);
 								}
 							}
 						}
@@ -1905,32 +1891,34 @@ bool UInterchangeGltfTranslator::GetVariantSetPayloadData(UE::Interchange::FVari
 	return true;
 }
 
-TFuture<TOptional<UE::Interchange::FSkeletalMeshLodPayloadData>> UInterchangeGltfTranslator::GetSkeletalMeshLodPayloadData(const FString& PayLoadKey) const
+TFuture< TOptional< UE::Interchange::FMeshPayloadData > > UInterchangeGltfTranslator::GetMeshPayloadData(const FInterchangeMeshPayLoadKey& PayLoadKey) const
 {
 	return Async(EAsyncExecution::TaskGraph, [this, PayLoadKey]
 		{
-			TOptional<UE::Interchange::FSkeletalMeshLodPayloadData> Result;
-			UE::Interchange::FSkeletalMeshLodPayloadData SkeletalMeshPayloadData;
-			
-			if (UE::Interchange::Gltf::Private::GetSkeletalMeshDescriptionForPayLoadKey(GltfAsset, PayLoadKey, SkeletalMeshPayloadData.LodMeshDescription, &SkeletalMeshPayloadData.JointNames))
+			UE::Interchange::FMeshPayloadData MeshPayLoadData;
+			bool bSuccessfullAcquisition = false;
+
+			switch (PayLoadKey.Type)
 			{
-				Result.Emplace(SkeletalMeshPayloadData);
+			case EInterchangeMeshPayLoadType::STATIC:
+				bSuccessfullAcquisition = UE::Interchange::Gltf::Private::GetStaticMeshPayloadDataForPayLoadKey(GltfAsset, PayLoadKey.UniqueId, MeshPayLoadData.MeshDescription);
+				break;
+			case EInterchangeMeshPayLoadType::SKELETAL:
+				bSuccessfullAcquisition = UE::Interchange::Gltf::Private::GetSkeletalMeshDescriptionForPayLoadKey(GltfAsset, PayLoadKey.UniqueId, MeshPayLoadData.MeshDescription, &MeshPayLoadData.JointNames);
+				break;
+			case EInterchangeMeshPayLoadType::MORPHTARGET:
+				//GLTF handles morph targets as simple Meshes
+				bSuccessfullAcquisition = UE::Interchange::Gltf::Private::GetStaticMeshPayloadDataForPayLoadKey(GltfAsset, PayLoadKey.UniqueId, MeshPayLoadData.MeshDescription);
+				break;
+			default:
+				break;
 			}
 
-			return Result;
-		});
-}
-
-TFuture<TOptional<UE::Interchange::FSkeletalMeshMorphTargetPayloadData>> UInterchangeGltfTranslator::GetSkeletalMeshMorphTargetPayloadData(const FString& PayLoadKey) const
-{
-	return Async(EAsyncExecution::TaskGraph, [this, PayLoadKey]
-		{
-			UE::Interchange::FSkeletalMeshMorphTargetPayloadData StaticMeshPayloadData;
-
-			UE::Interchange::Gltf::Private::GetStaticMeshPayloadDataForPayLoadKey(GltfAsset, PayLoadKey, StaticMeshPayloadData.LodMeshDescription);
-
-			TOptional<UE::Interchange::FSkeletalMeshMorphTargetPayloadData> Result;
-			Result.Emplace(StaticMeshPayloadData);
+			TOptional<UE::Interchange::FMeshPayloadData> Result;
+			if (bSuccessfullAcquisition)
+			{
+				Result.Emplace(MeshPayLoadData);
+			}
 
 			return Result;
 		});
@@ -2001,7 +1989,7 @@ void UInterchangeGltfTranslator::HandleGltfSkeletons(UInterchangeBaseNodeContain
 				}
 				Payload += LexToString(SkinnedMeshNode.MeshIndex | (SkinnedMeshNode.Skindex << 16));
 			}
-			SkeletalMeshNode->SetPayLoadKey(Payload);
+			SkeletalMeshNode->SetPayLoadKey(Payload, EInterchangeMeshPayLoadType::SKELETAL);
 
 			//set the root joint node as the skeletonDependency:
 			int32 RootJointNodeIndex = RootJointToSkinnedMeshNodes.Key;
@@ -2059,18 +2047,8 @@ UInterchangeMeshNode* UInterchangeGltfTranslator::HandleGltfMesh(UInterchangeBas
 	MeshNode->InitializeNode(MeshNodeUid, MeshName, EInterchangeNodeContainerType::TranslatedAsset);
 
 	//Generate Mesh Payload:
-	FString PayloadKey;
-	//if it has MorphTargets but is not a Skeletal (aka not called from HandleGltfSkeletons, aka SkeletalName.len = 0
-	//then the Payload processing should be calling the StaticMesh acquisition from the SkeletalMesh acquisition
-	// SkeletalMesh acqusition is called because UE handles MorphTargets in SkeletalMeshes 
-	// (hence Meshes with MorphTargets are marked to be skinned and the node that references it, is set to be Joint)
-	if (GltfMesh.MorphTargetNames.Num() > 0 && SkeletalName.Len() == 0)
-	{
-		PayloadKey = "StaticMeshWithMorphTarget~";
-		//if it is called from HandleGltfSkeletons then the payloadkey will be overwritten. (in HandleGltfSkeletons)
-	}
-	PayloadKey += LexToString(MeshIndex);
-	MeshNode->SetPayLoadKey(PayloadKey);
+	FString PayloadKey = LexToString(MeshIndex);
+	MeshNode->SetPayLoadKey(PayloadKey, EInterchangeMeshPayLoadType::STATIC);
 
 	NodeContainer.AddNode(MeshNode);
 
@@ -2113,7 +2091,7 @@ UInterchangeMeshNode* UInterchangeGltfTranslator::HandleGltfMesh(UInterchangeBas
 
 			//Generate Payload:
 			FString MorphTargetPayLoadKey = LexToString(MeshIndex) + TEXT(":") + LexToString(MorphTargetIndex);
-			MorphTargetMeshNode->SetPayLoadKey(MorphTargetPayLoadKey);
+			MorphTargetMeshNode->SetPayLoadKey(MorphTargetPayLoadKey, EInterchangeMeshPayLoadType::MORPHTARGET);
 
 			//set mesh as a morph target:
 			MorphTargetMeshNode->SetMorphTarget(true);
diff --git a/Engine/Plugins/Interchange/Runtime/Source/Import/Private/Mesh/InterchangeOBJTranslator.cpp b/Engine/Plugins/Interchange/Runtime/Source/Import/Private/Mesh/InterchangeOBJTranslator.cpp
--- a/Engine/Plugins/Interchange/Runtime/Source/Import/Private/Mesh/InterchangeOBJTranslator.cpp
+++ b/Engine/Plugins/Interchange/Runtime/Source/Import/Private/Mesh/InterchangeOBJTranslator.cpp
@@ -979,7 +979,7 @@ bool UInterchangeOBJTranslator::Translate(UInterchangeBaseNodeContainer& BaseNod
 
 			// The payload key is the group name.
 			// A MeshDescription is generated for each group of faces.
-			MeshNode->SetPayLoadKey(GroupName);
+			MeshNode->SetPayLoadKey(GroupName, EInterchangeMeshPayLoadType::STATIC);
 
 			MeshNode->SetCustomBoundingBox(ObjDataPtr->GetGroupBoundingBox(GroupName));
 			MeshNode->SetCustomVertexCount(ObjDataPtr->GetGroupVertexCount(GroupName));
@@ -1028,16 +1028,16 @@ bool UInterchangeOBJTranslator::Translate(UInterchangeBaseNodeContainer& BaseNod
 }
 
 
-TFuture<TOptional<UE::Interchange::FStaticMeshPayloadData>> UInterchangeOBJTranslator::GetStaticMeshPayloadData(const FString& PayloadKey) const
+TFuture<TOptional<UE::Interchange::FMeshPayloadData>> UInterchangeOBJTranslator::GetMeshPayloadData(const FInterchangeMeshPayLoadKey& PayLoadKey) const
 {
-	return Async(EAsyncExecution::TaskGraph, [this, PayloadKey]
+	return Async(EAsyncExecution::TaskGraph, [this, PayLoadKey]
 		{
 			using namespace UE::Interchange;
 
-			FStaticMeshPayloadData Payload;
-			Payload.MeshDescription = ObjDataPtr->MakeMeshDescriptionForGroup(PayloadKey);
+			FMeshPayloadData Payload;
+			Payload.MeshDescription = ObjDataPtr->MakeMeshDescriptionForGroup(PayLoadKey.UniqueId);
 
-			return TOptional<FStaticMeshPayloadData>(Payload);
+			return TOptional<FMeshPayloadData>(Payload);
 		}
 	);
 }
diff --git a/Engine/Plugins/Interchange/Runtime/Source/Import/Private/Mesh/InterchangeSkeletalMeshFactory.cpp b/Engine/Plugins/Interchange/Runtime/Source/Import/Private/Mesh/InterchangeSkeletalMeshFactory.cpp
--- a/Engine/Plugins/Interchange/Runtime/Source/Import/Private/Mesh/InterchangeSkeletalMeshFactory.cpp
+++ b/Engine/Plugins/Interchange/Runtime/Source/Import/Private/Mesh/InterchangeSkeletalMeshFactory.cpp
@@ -28,8 +28,8 @@
 #include "MaterialDomain.h"
 #include "Materials/Material.h"
 #include "Math/GenericOctree.h"
-#include "Mesh/InterchangeSkeletalMeshPayload.h"
-#include "Mesh/InterchangeSkeletalMeshPayloadInterface.h"
+#include "Mesh/InterchangeMeshPayload.h"
+#include "Mesh/InterchangeMeshPayloadInterface.h"
 #include "Misc/MessageDialog.h"
 #include "Nodes/InterchangeBaseNode.h"
 #include "Nodes/InterchangeBaseNodeContainer.h"
@@ -61,12 +61,12 @@ namespace UE
 				const UInterchangeMeshNode* MeshNode = nullptr;
 				const UInterchangeSceneNode* SceneNode = nullptr;
 				TOptional<FTransform> SceneGlobalTransform;
-				FString TranslatorPayloadKey;
+				FInterchangeMeshPayLoadKey TranslatorPayloadKey;
 			};
 
 			void FillMorphTargetMeshDescriptionsPerMorphTargetName(const FMeshNodeContext& MeshNodeContext
-																 , TMap<FString, TOptional<UE::Interchange::FSkeletalMeshMorphTargetPayloadData>>& MorphTargetMeshDescriptionsPerMorphTargetName
-																 , const IInterchangeSkeletalMeshPayloadInterface* SkeletalMeshTranslatorPayloadInterface
+																 , TMap<FString, TOptional<UE::Interchange::FMeshPayloadData>>& MorphTargetMeshDescriptionsPerMorphTargetName
+																 , const IInterchangeMeshPayloadInterface* MeshTranslatorPayloadInterface
 																 , const int32 VertexOffset
 																 , const UInterchangeBaseNodeContainer* NodeContainer
 																 , FString AssetName)
@@ -74,21 +74,22 @@ namespace UE
 				TRACE_CPUPROFILER_EVENT_SCOPE("FillMorphTargetMeshDescriptionsPerMorphTargetName")
 				TArray<FString> MorphTargetUids;
 				MeshNodeContext.MeshNode->GetMorphTargetDependencies(MorphTargetUids);
-				TMap<FString, TFuture<TOptional<UE::Interchange::FSkeletalMeshMorphTargetPayloadData>>> TempMorphTargetMeshDescriptionsPerMorphTargetName;
+				TMap<FString, TFuture<TOptional<UE::Interchange::FMeshPayloadData>>> TempMorphTargetMeshDescriptionsPerMorphTargetName;
 				TempMorphTargetMeshDescriptionsPerMorphTargetName.Reserve(MorphTargetUids.Num());
 				for (const FString& MorphTargetUid : MorphTargetUids)
 				{
 					if (const UInterchangeMeshNode* MorphTargetMeshNode = Cast<UInterchangeMeshNode>(NodeContainer->GetNode(MorphTargetUid)))
 					{
-						TOptional<FString> MorphTargetPayloadKey = MorphTargetMeshNode->GetPayLoadKey();
-						if (!MorphTargetPayloadKey.IsSet())
+						TOptional<FInterchangeMeshPayLoadKey> OptionalPayLoadKey = MorphTargetMeshNode->GetPayLoadKey();
+						if (!OptionalPayLoadKey.IsSet())
 						{
 							UE_LOG(LogInterchangeImport, Warning, TEXT("Empty LOD morph target mesh reference payload when importing SkeletalMesh asset %s"), *AssetName);
 							continue;
 						}
-						const FString PayloadKey = MorphTargetPayloadKey.GetValue();
+						FInterchangeMeshPayLoadKey& PayLoadKey = OptionalPayLoadKey.GetValue();
+
 						//Add the map entry key, the translator will be call after to bulk get all the needed payload
-						TempMorphTargetMeshDescriptionsPerMorphTargetName.Add(PayloadKey, SkeletalMeshTranslatorPayloadInterface->GetSkeletalMeshMorphTargetPayloadData(PayloadKey));
+						TempMorphTargetMeshDescriptionsPerMorphTargetName.Add(PayLoadKey.UniqueId, MeshTranslatorPayloadInterface->GetMeshPayloadData(PayLoadKey));
 					}
 				}
 
@@ -97,18 +98,20 @@ namespace UE
 					if (const UInterchangeMeshNode* MorphTargetMeshNode = Cast<UInterchangeMeshNode>(NodeContainer->GetNode(MorphTargetUid)))
 					{
 						
-						TOptional<FString> MorphTargetPayloadKey = MorphTargetMeshNode->GetPayLoadKey();
-						if (!MorphTargetPayloadKey.IsSet())
+						TOptional<FInterchangeMeshPayLoadKey> OptionalPayLoadKey = MorphTargetMeshNode->GetPayLoadKey();
+						if (!OptionalPayLoadKey.IsSet())
 						{
 							continue;
 						}
-						const FString& MorphTargetPayloadKeyString = MorphTargetPayloadKey.GetValue();
+						FInterchangeMeshPayLoadKey& PayLoadKey = OptionalPayLoadKey.GetValue();
+
+						const FString& MorphTargetPayloadKeyString = PayLoadKey.UniqueId;
 						if (!ensure(TempMorphTargetMeshDescriptionsPerMorphTargetName.Contains(MorphTargetPayloadKeyString)))
 						{
 							continue;
 						}
 
-						TOptional<UE::Interchange::FSkeletalMeshMorphTargetPayloadData> MorphTargetMeshPayload = TempMorphTargetMeshDescriptionsPerMorphTargetName.FindChecked(MorphTargetPayloadKeyString).Get();
+						TOptional<UE::Interchange::FMeshPayloadData> MorphTargetMeshPayload = TempMorphTargetMeshDescriptionsPerMorphTargetName.FindChecked(MorphTargetPayloadKeyString).Get();
 						if (!MorphTargetMeshPayload.IsSet())
 						{
 							UE_LOG(LogInterchangeImport, Warning, TEXT("Invalid Skeletal mesh morph target payload key [%s] SkeletalMesh asset %s"), *MorphTargetPayloadKeyString, *AssetName);
@@ -135,23 +138,23 @@ namespace UE
 				}
 			}
 
-			void CopyMorphTargetsMeshDescriptionToSkeletalMeshImportData(const TMap<FString, TOptional<UE::Interchange::FSkeletalMeshMorphTargetPayloadData>>& LodMorphTargetMeshDescriptions, FSkeletalMeshImportData& DestinationSkeletalMeshImportData)
+			void CopyMorphTargetsMeshDescriptionToSkeletalMeshImportData(const TMap<FString, TOptional<UE::Interchange::FMeshPayloadData>>& LodMorphTargetMeshDescriptions, FSkeletalMeshImportData& DestinationSkeletalMeshImportData)
 			{
 				TRACE_CPUPROFILER_EVENT_SCOPE("CopyMorphTargetsMeshDescriptionToSkeletalMeshImportData")
 				const int32 OriginalMorphTargetCount = LodMorphTargetMeshDescriptions.Num();
 				TArray<FString> Keys;
 				int32 MorphTargetCount = 0;
-				for (const TPair<FString, TOptional<UE::Interchange::FSkeletalMeshMorphTargetPayloadData>>& Pair : LodMorphTargetMeshDescriptions)
+				for (const TPair<FString, TOptional<UE::Interchange::FMeshPayloadData>>& Pair : LodMorphTargetMeshDescriptions)
 				{
 					const FString MorphTargetName(Pair.Key);
-					const TOptional<UE::Interchange::FSkeletalMeshMorphTargetPayloadData>& MorphTargetPayloadData = Pair.Value;
+					const TOptional<UE::Interchange::FMeshPayloadData>& MorphTargetPayloadData = Pair.Value;
 					if (!MorphTargetPayloadData.IsSet())
 					{
 						UE_LOG(LogInterchangeImport, Error, TEXT("Empty morph target optional payload data [%s]"), *MorphTargetName);
 						continue;
 					}
 
-					const FMeshDescription& SourceMeshDescription = MorphTargetPayloadData.GetValue().LodMeshDescription;
+					const FMeshDescription& SourceMeshDescription = MorphTargetPayloadData.GetValue().MeshDescription;
 					const int32 VertexOffset = MorphTargetPayloadData->VertexOffset;
 					const int32 SourceMeshVertexCount = SourceMeshDescription.Vertices().Num();
 					const int32 DestinationVertexIndexMax = VertexOffset + SourceMeshVertexCount;
@@ -199,14 +202,14 @@ namespace UE
 							break;
 						}
 						const FString MorphTargetKey(Keys[MorphTargetIndex]);
-						const TOptional<UE::Interchange::FSkeletalMeshMorphTargetPayloadData>& MorphTargetPayloadData = LodMorphTargetMeshDescriptions.FindChecked(MorphTargetKey);
+						const TOptional<UE::Interchange::FMeshPayloadData>& MorphTargetPayloadData = LodMorphTargetMeshDescriptions.FindChecked(MorphTargetKey);
 						if (!ensure(MorphTargetPayloadData.IsSet()))
 						{
 							//This error was suppose to be catch in the pre parallel for loop
 							break;
 						}
 
-						const FMeshDescription& SourceMeshDescription = MorphTargetPayloadData.GetValue().LodMeshDescription;
+						const FMeshDescription& SourceMeshDescription = MorphTargetPayloadData.GetValue().MeshDescription;
 						const FTransform GlobalTransform = MorphTargetPayloadData->GlobalTransform.IsSet() ? MorphTargetPayloadData->GlobalTransform.GetValue() : FTransform::Identity;
 						const int32 VertexOffset = MorphTargetPayloadData->VertexOffset;
 						const int32 SourceMeshVertexCount = SourceMeshDescription.Vertices().Num();
@@ -275,13 +278,13 @@ namespace UE
 				return nullptr;
 			}
 
-			void SkinVertexPositionToTimeZero(UE::Interchange::FSkeletalMeshLodPayloadData& LodMeshPayload
+			void SkinVertexPositionToTimeZero(UE::Interchange::FMeshPayloadData& LodMeshPayload
 				, const UInterchangeBaseNodeContainer* NodeContainer
 				, const FString& RootJointNodeId
 				, const FTransform& MeshGlobalTransform)
 			{
 				TRACE_CPUPROFILER_EVENT_SCOPE("SkinVertexPositionToTimeZero")
-				FMeshDescription& MeshDescription = LodMeshPayload.LodMeshDescription;
+				FMeshDescription& MeshDescription = LodMeshPayload.MeshDescription;
 				const int32 VertexCount = MeshDescription.Vertices().Num();
 				const TArray<FString>& JointNames = LodMeshPayload.JointNames;
 				// Create a copy of the vertex array to receive vertex deformations.
@@ -407,13 +410,13 @@ namespace UE
 																  , TArray<FMeshNodeContext>& MeshReferences
 																  , TArray<SkeletalMeshImportData::FBone>& RefBonesBinary
 																  , const UInterchangeSkeletalMeshFactory::FImportAssetObjectParams& Arguments
-																  , const IInterchangeSkeletalMeshPayloadInterface* SkeletalMeshTranslatorPayloadInterface
+																  , const IInterchangeMeshPayloadInterface* MeshTranslatorPayloadInterface
 																  , const bool bSkinControlPointToTimeZero
 																  , const UInterchangeBaseNodeContainer* NodeContainer
 																  , const FString& RootJointNodeId)
 			{
 				TRACE_CPUPROFILER_EVENT_SCOPE("RetrieveAllSkeletalMeshPayloadsAndFillImportData")
-				if (!SkeletalMeshTranslatorPayloadInterface)
+				if (!MeshTranslatorPayloadInterface)
 				{
 					return;
 				}
@@ -429,16 +432,16 @@ namespace UE
 				bool bImportMorphTarget = true;
 				SkeletalMeshFactoryNode->GetCustomImportMorphTarget(bImportMorphTarget);
 
-				TMap<FString, TFuture<TOptional<UE::Interchange::FSkeletalMeshLodPayloadData>>> LodMeshPayloadPerTranslatorPayloadKey;
+				TMap<FString, TFuture<TOptional<UE::Interchange::FMeshPayloadData>>> LodMeshPayloadPerTranslatorPayloadKey;
 				LodMeshPayloadPerTranslatorPayloadKey.Reserve(MeshReferences.Num());
 
-				TMap<FString, TOptional<UE::Interchange::FSkeletalMeshMorphTargetPayloadData>> MorphTargetMeshDescriptionsPerMorphTargetName;
+				TMap<FString, TOptional<UE::Interchange::FMeshPayloadData>> MorphTargetMeshDescriptionsPerMorphTargetName;
 				int32 MorphTargetCount = 0;
 
 				for (const FMeshNodeContext& MeshNodeContext : MeshReferences)
 				{
 					//Add the payload entry key, the payload data will be fill later in bulk by the translator
-					LodMeshPayloadPerTranslatorPayloadKey.Add(MeshNodeContext.TranslatorPayloadKey, SkeletalMeshTranslatorPayloadInterface->GetSkeletalMeshLodPayloadData(MeshNodeContext.TranslatorPayloadKey));
+					LodMeshPayloadPerTranslatorPayloadKey.Add(MeshNodeContext.TranslatorPayloadKey.UniqueId, MeshTranslatorPayloadInterface->GetMeshPayloadData(MeshNodeContext.TranslatorPayloadKey));
 					//Count the morph target dependencies so we can reserve the right amount
 					MorphTargetCount += (bImportMorphTarget && MeshNodeContext.MeshNode) ? MeshNodeContext.MeshNode->GetMorphTargetDependeciesCount() : 0;
 				}
@@ -448,10 +451,10 @@ namespace UE
 				for (const FMeshNodeContext& MeshNodeContext : MeshReferences)
 				{
 					TRACE_CPUPROFILER_EVENT_SCOPE("RetrieveAllSkeletalMeshPayloadsAndFillImportData::GetPayload")
-					TOptional<UE::Interchange::FSkeletalMeshLodPayloadData> LodMeshPayload = LodMeshPayloadPerTranslatorPayloadKey.FindChecked(MeshNodeContext.TranslatorPayloadKey).Get();
+					TOptional<UE::Interchange::FMeshPayloadData> LodMeshPayload = LodMeshPayloadPerTranslatorPayloadKey.FindChecked(MeshNodeContext.TranslatorPayloadKey.UniqueId).Get();
 					if (!LodMeshPayload.IsSet())
 					{
-						UE_LOG(LogInterchangeImport, Warning, TEXT("Invalid Skeletal mesh payload key [%s] SkeletalMesh asset %s"), *MeshNodeContext.TranslatorPayloadKey, *Arguments.AssetName);
+						UE_LOG(LogInterchangeImport, Warning, TEXT("Invalid Skeletal mesh payload key [%s] SkeletalMesh asset %s"), *MeshNodeContext.TranslatorPayloadKey.UniqueId, *Arguments.AssetName);
 						continue;
 					}
 					const int32 VertexOffset = LodMeshDescription.Vertices().Num();
@@ -461,7 +464,7 @@ namespace UE
 					{
 						TRACE_CPUPROFILER_EVENT_SCOPE("RetrieveAllSkeletalMeshPayloadsAndFillImportData::CompactPayload")
 						FElementIDRemappings ElementIDRemappings;
-						LodMeshPayload->LodMeshDescription.Compact(ElementIDRemappings);
+						LodMeshPayload->MeshDescription.Compact(ElementIDRemappings);
 					}
 
 					const bool bIsRigidMesh = LodMeshPayload->JointNames.Num() <= 0 && MeshNodeContext.SceneNode;
@@ -517,7 +520,7 @@ namespace UE
 						}
 						//Add the skinning in the mesh description
 						{
-							FSkeletalMeshAttributes PayloadSkeletalMeshAttributes(LodMeshPayload->LodMeshDescription);
+							FSkeletalMeshAttributes PayloadSkeletalMeshAttributes(LodMeshPayload->MeshDescription);
 							constexpr bool bKeepExistingAttribute = true;
 							PayloadSkeletalMeshAttributes.Register(bKeepExistingAttribute);
 							using namespace UE::AnimationCore;
@@ -526,7 +529,7 @@ namespace UE
 							BoneWeight.SetBoneIndex(0);
 							BoneWeight.SetWeight(1.0f);
 							FSkinWeightsVertexAttributesRef PayloadVertexSkinWeights = PayloadSkeletalMeshAttributes.GetVertexSkinWeights();
-							for (const FVertexID& PayloadVertexID : LodMeshPayload->LodMeshDescription.Vertices().GetElementIDs())
+							for (const FVertexID& PayloadVertexID : LodMeshPayload->MeshDescription.Vertices().GetElementIDs())
 							{
 								PayloadVertexSkinWeights.Set(PayloadVertexID, BoneWeights);
 							}
@@ -541,16 +544,16 @@ namespace UE
 					{
 						AppendSettings.MeshTransform.Reset();
 					}
-					FStaticMeshOperations::AppendMeshDescription(LodMeshPayload->LodMeshDescription, LodMeshDescription, AppendSettings);
+					FStaticMeshOperations::AppendMeshDescription(LodMeshPayload->MeshDescription, LodMeshDescription, AppendSettings);
 					if (MeshNodeContext.MeshNode->IsSkinnedMesh() || bIsRigidMesh)
 					{
-						FSkeletalMeshOperations::AppendSkinWeight(LodMeshPayload->LodMeshDescription, LodMeshDescription, SkeletalMeshAppendSettings);
+						FSkeletalMeshOperations::AppendSkinWeight(LodMeshPayload->MeshDescription, LodMeshDescription, SkeletalMeshAppendSettings);
 					}
 					if (bImportMorphTarget)
 					{
 						FillMorphTargetMeshDescriptionsPerMorphTargetName(MeshNodeContext
 																		, MorphTargetMeshDescriptionsPerMorphTargetName
-																		, SkeletalMeshTranslatorPayloadInterface
+																		, MeshTranslatorPayloadInterface
 																		, VertexOffset
 																		, Arguments.NodeContainer
 																		, Arguments.AssetName);
@@ -953,8 +956,8 @@ UObject* UInterchangeSkeletalMeshFactory::ImportAssetObject_Async(const FImportA
 		return nullptr;
 	}
 
-	const IInterchangeSkeletalMeshPayloadInterface* SkeletalMeshTranslatorPayloadInterface = Cast<IInterchangeSkeletalMeshPayloadInterface>(Arguments.Translator);
-	if (!SkeletalMeshTranslatorPayloadInterface)
+	const IInterchangeMeshPayloadInterface* MeshTranslatorPayloadInterface = Cast<IInterchangeMeshPayloadInterface>(Arguments.Translator);
+	if (!MeshTranslatorPayloadInterface)
 	{
 		UE_LOG(LogInterchangeImport, Error, TEXT("Cannot import skeletalMesh, the translator do not implement the IInterchangeSkeletalMeshPayloadInterface."));
 		return nullptr;
@@ -1259,10 +1262,11 @@ UObject* UInterchangeSkeletalMeshFactory::ImportAssetObject_Async(const FImportA
 					UE_LOG(LogInterchangeImport, Warning, TEXT("Invalid LOD mesh reference when importing SkeletalMesh asset %s"), *Arguments.AssetName);
 					continue;
 				}
-				TOptional<FString> MeshPayloadKey = MeshReference.MeshNode->GetPayLoadKey();
-				if (MeshPayloadKey.IsSet())
+
+				TOptional<FInterchangeMeshPayLoadKey> OptionalPayLoadKey = MeshReference.MeshNode->GetPayLoadKey();
+				if (OptionalPayLoadKey.IsSet())
 				{
-					MeshReference.TranslatorPayloadKey = MeshPayloadKey.GetValue();
+					MeshReference.TranslatorPayloadKey = OptionalPayLoadKey.GetValue();
 				}
 				else
 				{
@@ -1282,7 +1286,7 @@ UObject* UInterchangeSkeletalMeshFactory::ImportAssetObject_Async(const FImportA
 																					, MeshReferences
 																					, RefBonesBinary
 																					, Arguments
-																					, SkeletalMeshTranslatorPayloadInterface
+																					, MeshTranslatorPayloadInterface
 																					, bSkinControlPointToTimeZero
 																					, Arguments.NodeContainer
 																					, RootJointNodeId);
diff --git a/Engine/Plugins/Interchange/Runtime/Source/Import/Private/Mesh/InterchangeStaticMeshFactory.cpp b/Engine/Plugins/Interchange/Runtime/Source/Import/Private/Mesh/InterchangeStaticMeshFactory.cpp
--- a/Engine/Plugins/Interchange/Runtime/Source/Import/Private/Mesh/InterchangeStaticMeshFactory.cpp
+++ b/Engine/Plugins/Interchange/Runtime/Source/Import/Private/Mesh/InterchangeStaticMeshFactory.cpp
@@ -24,7 +24,7 @@
 #include "Materials/Material.h"
 #include "Materials/MaterialInstance.h"
 #include "Materials/MaterialInterface.h"
-#include "Mesh/InterchangeStaticMeshPayloadInterface.h"
+#include "Mesh/InterchangeMeshPayloadInterface.h"
 #include "MeshBudgetProjectSettings.h"
 #include "Model.h"
 #include "Nodes/InterchangeBaseNode.h"
@@ -274,7 +274,7 @@ UObject* UInterchangeStaticMeshFactory::ImportAssetObject_Async(const FImportAss
 		bool bFirstValidMoved = false;
 		for (FMeshPayload& MeshPayload : MeshPayloads)
 		{
-			const TOptional<UE::Interchange::FStaticMeshPayloadData>& LodMeshPayload = MeshPayload.PayloadData.Get();
+			const TOptional<UE::Interchange::FMeshPayloadData>& LodMeshPayload = MeshPayload.PayloadData.Get();
 			if (!LodMeshPayload.IsSet())
 			{
 				UE_LOG(LogInterchangeImport, Warning, TEXT("Invalid static mesh payload key, StaticMesh asset %s"), *Arguments.AssetName);
@@ -558,10 +558,10 @@ TArray<UInterchangeStaticMeshFactory::FMeshPayload> UInterchangeStaticMeshFactor
 	TArray<FMeshPayload> Payloads;
 	Payloads.Reserve(MeshUids.Num());
 
-	const IInterchangeStaticMeshPayloadInterface* StaticMeshTranslatorPayloadInterface = Cast<IInterchangeStaticMeshPayloadInterface>(Arguments.Translator);
-	if (!StaticMeshTranslatorPayloadInterface)
+	const IInterchangeMeshPayloadInterface* MeshTranslatorPayloadInterface = Cast<IInterchangeMeshPayloadInterface>(Arguments.Translator);
+	if (!MeshTranslatorPayloadInterface)
 	{
-		UE_LOG(LogInterchangeImport, Error, TEXT("Cannot import static mesh, the translator does not implement the IInterchangeStaticMeshPayloadInterface."));
+		UE_LOG(LogInterchangeImport, Error, TEXT("Cannot import static mesh, the translator does not implement the IInterchangeMeshPayloadInterface."));
 		return Payloads;
 	}
 
@@ -619,15 +619,17 @@ TArray<UInterchangeStaticMeshFactory::FMeshPayload> UInterchangeStaticMeshFactor
 			continue;
 		}
 
-		TOptional<FString> MeshPayloadKey = MeshNode->GetPayLoadKey();
-		if (!ensure(MeshPayloadKey.IsSet()))
+		TOptional<FInterchangeMeshPayLoadKey> OptionalPayLoadKey = MeshNode->GetPayLoadKey();
+		if (!ensure(OptionalPayLoadKey.IsSet()))
 		{
 			UE_LOG(LogInterchangeImport, Warning, TEXT("Empty LOD mesh reference payload when importing StaticMesh asset %s"), *Arguments.AssetName);
 			continue;
 		}
 
-		Payload.MeshName = *MeshPayloadKey;
-		Payload.PayloadData = StaticMeshTranslatorPayloadInterface->GetStaticMeshPayloadData(*MeshPayloadKey);
+		FInterchangeMeshPayLoadKey& PayLoadKey = OptionalPayLoadKey.GetValue();
+
+		Payload.MeshName = PayLoadKey.UniqueId;
+		Payload.PayloadData = MeshTranslatorPayloadInterface->GetMeshPayloadData(PayLoadKey);
 
 		Payloads.Emplace(MoveTemp(Payload));
 	}
@@ -1062,7 +1064,7 @@ bool UInterchangeStaticMeshFactory::ImportBoxCollision(const FImportAssetObjectP
 	for (const FMeshPayload& MeshPayload : MeshPayloads)
 	{
 		const FTransform& Transform = MeshPayload.Transform;
-		const TOptional<FStaticMeshPayloadData>& PayloadData = MeshPayload.PayloadData.Get();
+		const TOptional<FMeshPayloadData>& PayloadData = MeshPayload.PayloadData.Get();
 
 		if (!PayloadData.IsSet())
 		{
@@ -1110,7 +1112,7 @@ bool UInterchangeStaticMeshFactory::ImportCapsuleCollision(const FImportAssetObj
 	for (const FMeshPayload& MeshPayload : MeshPayloads)
 	{
 		const FTransform& Transform = MeshPayload.Transform;
-		const TOptional<FStaticMeshPayloadData>& PayloadData = MeshPayload.PayloadData.Get();
+		const TOptional<FMeshPayloadData>& PayloadData = MeshPayload.PayloadData.Get();
 
 		if (!PayloadData.IsSet())
 		{
@@ -1158,7 +1160,7 @@ bool UInterchangeStaticMeshFactory::ImportSphereCollision(const FImportAssetObje
 	for (const FMeshPayload& MeshPayload : MeshPayloads)
 	{
 		const FTransform& Transform = MeshPayload.Transform;
-		const TOptional<FStaticMeshPayloadData>& PayloadData = MeshPayload.PayloadData.Get();
+		const TOptional<FMeshPayloadData>& PayloadData = MeshPayload.PayloadData.Get();
 
 		if (!PayloadData.IsSet())
 		{
@@ -1209,7 +1211,7 @@ bool UInterchangeStaticMeshFactory::ImportConvexCollision(const FImportAssetObje
 		for (const FMeshPayload& MeshPayload : MeshPayloads)
 		{
 			const FTransform& Transform = MeshPayload.Transform;
-			const TOptional<FStaticMeshPayloadData>& PayloadData = MeshPayload.PayloadData.Get();
+			const TOptional<FMeshPayloadData>& PayloadData = MeshPayload.PayloadData.Get();
 
 			if (!PayloadData.IsSet())
 			{
@@ -1234,7 +1236,7 @@ bool UInterchangeStaticMeshFactory::ImportConvexCollision(const FImportAssetObje
 		for (const FMeshPayload& MeshPayload : MeshPayloads)
 		{
 			const FTransform& Transform = MeshPayload.Transform;
-			TOptional<FStaticMeshPayloadData> PayloadData = MeshPayload.PayloadData.Get();
+			TOptional<FMeshPayloadData> PayloadData = MeshPayload.PayloadData.Get();
 
 			if (!PayloadData.IsSet())
 			{
diff --git a/Engine/Plugins/Interchange/Runtime/Source/Import/Public/Fbx/InterchangeFbxTranslator.h b/Engine/Plugins/Interchange/Runtime/Source/Import/Public/Fbx/InterchangeFbxTranslator.h
--- a/Engine/Plugins/Interchange/Runtime/Source/Import/Public/Fbx/InterchangeFbxTranslator.h
+++ b/Engine/Plugins/Interchange/Runtime/Source/Import/Public/Fbx/InterchangeFbxTranslator.h
@@ -5,10 +5,8 @@
 #include "CoreMinimal.h"
 #include "InterchangeTranslatorBase.h"
 #include "InterchangeDispatcher.h"
-#include "Mesh/InterchangeSkeletalMeshPayload.h"
-#include "Mesh/InterchangeSkeletalMeshPayloadInterface.h"
-#include "Mesh/InterchangeStaticMeshPayload.h"
-#include "Mesh/InterchangeStaticMeshPayloadInterface.h"
+#include "Mesh/InterchangeMeshPayload.h"
+#include "Mesh/InterchangeMeshPayloadInterface.h"
 #include "Nodes/InterchangeBaseNodeContainer.h"
 #include "Texture/InterchangeTexturePayloadInterface.h"
 #include "UObject/Object.h"
@@ -21,8 +19,7 @@
 UCLASS(BlueprintType, Experimental)
 class INTERCHANGEIMPORT_API UInterchangeFbxTranslator : public UInterchangeTranslatorBase
 , public IInterchangeTexturePayloadInterface
-, public IInterchangeStaticMeshPayloadInterface
-, public IInterchangeSkeletalMeshPayloadInterface
+, public IInterchangeMeshPayloadInterface
 , public IInterchangeAnimationPayloadInterface
 {
 	GENERATED_BODY()
@@ -56,7 +53,7 @@ public:
 
 
 	//////////////////////////////////////////////////////////////////////////
-	/* IInterchangeStaticMeshPayloadInterface Begin */
+	/* IInterchangeMeshPayloadInterface Begin */
 
 	/**
 	 * Once the translation is done, the import process need a way to retrieve payload data.
@@ -66,18 +63,9 @@ public:
 	 * @param PayloadKey - The key to retrieve the a particular payload contain into the specified source data.
 	 * @return a PayloadData containing the imported data. The TOptional will not be set if there is an error.
 	 */
-	virtual TFuture<TOptional<UE::Interchange::FStaticMeshPayloadData>> GetStaticMeshPayloadData(const FString& PayLoadKey) const override;
+	virtual TFuture<TOptional<UE::Interchange::FMeshPayloadData>> GetMeshPayloadData(const FInterchangeMeshPayLoadKey& PayLoadKey) const override;
 
-	/* IInterchangeStaticMeshPayloadInterface End */
-
-
-	//////////////////////////////////////////////////////////////////////////
-	/* IInterchangeSkeletalMeshPayloadInterface Begin */
-
-	virtual TFuture<TOptional<UE::Interchange::FSkeletalMeshLodPayloadData>> GetSkeletalMeshLodPayloadData(const FString& PayLoadKey) const override;
-	virtual TFuture<TOptional<UE::Interchange::FSkeletalMeshMorphTargetPayloadData>> GetSkeletalMeshMorphTargetPayloadData(const FString& PayLoadKey) const override;
-	
-	/* IInterchangeSkeletalMeshPayloadInterface End */
+	///* IInterchangeMeshPayloadInterface End */
 
 	//////////////////////////////////////////////////////////////////////////
 	/* IInterchangeAnimationPayloadInterface Begin */
diff --git a/Engine/Plugins/Interchange/Runtime/Source/Import/Public/Gltf/InterchangeGltfTranslator.h b/Engine/Plugins/Interchange/Runtime/Source/Import/Public/Gltf/InterchangeGltfTranslator.h
--- a/Engine/Plugins/Interchange/Runtime/Source/Import/Public/Gltf/InterchangeGltfTranslator.h
+++ b/Engine/Plugins/Interchange/Runtime/Source/Import/Public/Gltf/InterchangeGltfTranslator.h
@@ -5,10 +5,8 @@
 #include "GLTFMaterial.h"
 #include "InterchangeTranslatorBase.h"
 #include "Animation/InterchangeAnimationPayloadInterface.h"
-#include "Mesh/InterchangeStaticMeshPayload.h"
-#include "Mesh/InterchangeStaticMeshPayloadInterface.h"
-#include "Mesh/InterchangeSkeletalMeshPayload.h"
-#include "Mesh/InterchangeSkeletalMeshPayloadInterface.h"
+#include "Mesh/InterchangeMeshPayload.h"
+#include "Mesh/InterchangeMeshPayloadInterface.h"
 #include "Nodes/InterchangeBaseNodeContainer.h"
 #include "Scene/InterchangeVariantSetPayloadInterface.h"
 #include "Texture/InterchangeTexturePayloadInterface.h"
@@ -26,11 +24,10 @@ class UInterchangeMeshNode;
 
 UCLASS(BlueprintType)
 class UInterchangeGltfTranslator : public UInterchangeTranslatorBase
-	, public IInterchangeStaticMeshPayloadInterface
+	, public IInterchangeMeshPayloadInterface
 	, public IInterchangeTexturePayloadInterface
 	, public IInterchangeAnimationPayloadInterface
 	, public IInterchangeVariantSetPayloadInterface
-	, public IInterchangeSkeletalMeshPayloadInterface
 {
 	GENERATED_BODY()
 
@@ -44,7 +41,7 @@ public:
 
 	/* IInterchangeStaticMeshPayloadInterface Begin */
 
-	virtual TFuture< TOptional< UE::Interchange::FStaticMeshPayloadData > > GetStaticMeshPayloadData( const FString& PayLoadKey ) const override;
+	virtual TFuture<TOptional<UE::Interchange::FMeshPayloadData>> GetMeshPayloadData(const FInterchangeMeshPayLoadKey& PayLoadKey) const override;
 
 	/* IInterchangeStaticMeshPayloadInterface End */
 
@@ -64,11 +61,6 @@ public:
 	virtual TFuture<TOptional<UE::Interchange::FVariantSetPayloadData>> GetVariantSetPayloadData(const FString& PayloadKey) const override;
 	/* IInterchangeVariantSetPayloadInterface End */
 
-	/* IInterchangeSkeletalMeshPayloadInterface Begin */
-	virtual TFuture<TOptional<UE::Interchange::FSkeletalMeshLodPayloadData>> GetSkeletalMeshLodPayloadData(const FString& PayLoadKey) const override;
-	virtual TFuture<TOptional<UE::Interchange::FSkeletalMeshMorphTargetPayloadData>> GetSkeletalMeshMorphTargetPayloadData(const FString& PayLoadKey) const override;
-	/* IInterchangeSkeletalMeshPayloadInterface End */
-
 protected:
 	using FNodeUidMap = TMap<const GLTF::FNode*, FString>;
 
diff --git a/Engine/Plugins/Interchange/Runtime/Source/Import/Public/Mesh/InterchangeMeshPayload.h b/Engine/Plugins/Interchange/Runtime/Source/Import/Public/Mesh/InterchangeMeshPayload.h
new file mode 100644
--- /dev/null
+++ b/Engine/Plugins/Interchange/Runtime/Source/Import/Public/Mesh/InterchangeMeshPayload.h
@@ -0,0 +1,30 @@
+// Copyright Epic Games, Inc. All Rights Reserved.
+
+#pragma once
+
+#include "CoreMinimal.h"
+#include "MeshDescription.h"
+
+namespace UE
+{
+	namespace Interchange
+	{
+		struct FMeshPayloadData
+		{
+			/* MESH */
+			//Currently the skeletalmesh payload data is editor only, we have to move to something available at runtime
+			FMeshDescription MeshDescription;
+
+			/* SKELETAL */
+			//This map the indice use in the meshdescription to the bone name, so we can use this information to remap properly the skinning when we merge the meshdescription
+			TArray<FString> JointNames;
+
+			/* MORPH */
+			TOptional<FTransform> GlobalTransform;
+			//The vertex offset of the morph target in case we combine mesh node together
+			int32 VertexOffset;
+			//The name of the morph target
+			FString MorphTargetName;
+		};
+	}//ns Interchange
+}//ns UE
diff --git a/Engine/Plugins/Interchange/Runtime/Source/Import/Public/Mesh/InterchangeMeshPayloadInterface.h b/Engine/Plugins/Interchange/Runtime/Source/Import/Public/Mesh/InterchangeMeshPayloadInterface.h
new file mode 100644
--- /dev/null
+++ b/Engine/Plugins/Interchange/Runtime/Source/Import/Public/Mesh/InterchangeMeshPayloadInterface.h
@@ -0,0 +1,39 @@
+// Copyright Epic Games, Inc. All Rights Reserved.
+#pragma once
+
+#include "CoreMinimal.h"
+#include "Async/Future.h"
+#include "InterchangeSourceData.h"
+#include "Mesh/InterchangeMeshPayload.h"
+#include "InterchangeMeshNode.h"
+#include "UObject/Object.h"
+#include "UObject/ObjectMacros.h"
+#include "UObject/Interface.h"
+
+#include "InterchangeMeshPayloadInterface.generated.h"
+
+UINTERFACE()
+class INTERCHANGEIMPORT_API UInterchangeMeshPayloadInterface : public UInterface
+{
+	GENERATED_BODY()
+};
+
+/**
+ * Static mesh payload interface. Derive from this interface if your payload can import static mesh
+ */
+class INTERCHANGEIMPORT_API IInterchangeMeshPayloadInterface
+{
+	GENERATED_BODY()
+public:
+
+	/**
+	 * Once the translation is done, the import process need a way to retrieve payload data.
+	 * This payload will be use by the factories to create the asset.
+	 *
+	 * @param PayloadKey - The key to retrieve the a particular payload contain into the specified source data.
+	 * @return a PayloadData containing the the data ask with the key.
+	 */
+	virtual TFuture<TOptional<UE::Interchange::FMeshPayloadData>> GetMeshPayloadData(const FInterchangeMeshPayLoadKey& PayLoadKey) const = 0;
+};
+
+
diff --git a/Engine/Plugins/Interchange/Runtime/Source/Import/Public/Mesh/InterchangeOBJTranslator.h b/Engine/Plugins/Interchange/Runtime/Source/Import/Public/Mesh/InterchangeOBJTranslator.h
--- a/Engine/Plugins/Interchange/Runtime/Source/Import/Public/Mesh/InterchangeOBJTranslator.h
+++ b/Engine/Plugins/Interchange/Runtime/Source/Import/Public/Mesh/InterchangeOBJTranslator.h
@@ -7,8 +7,8 @@
 #include "UObject/ObjectMacros.h"
 #include "InterchangeTranslatorBase.h"
 #include "Nodes/InterchangeBaseNodeContainer.h"
-#include "Mesh/InterchangeStaticMeshPayload.h"
-#include "Mesh/InterchangeStaticMeshPayloadInterface.h"
+#include "Mesh/InterchangeMeshPayload.h"
+#include "Mesh/InterchangeMeshPayloadInterface.h"
 #include "Texture/InterchangeTexturePayloadInterface.h"
 
 #include "InterchangeOBJTranslator.generated.h"
@@ -19,7 +19,7 @@ struct FObjData;
 
 UCLASS(BlueprintType, Experimental)
 class INTERCHANGEIMPORT_API UInterchangeOBJTranslator : public UInterchangeTranslatorBase,
-	                                                    public IInterchangeStaticMeshPayloadInterface,
+	                                                    public IInterchangeMeshPayloadInterface,
 	                                                    public IInterchangeTexturePayloadInterface
 {
 	GENERATED_BODY()
@@ -50,7 +50,7 @@ public:
 	 * @param PayloadKey - The key to retrieve the a particular payload contain into the specified source data.
 	 * @return a PayloadData containing the the data ask with the key.
 	 */
-	virtual TFuture<TOptional<UE::Interchange::FStaticMeshPayloadData>> GetStaticMeshPayloadData(const FString& PayloadKey) const override;
+	virtual TFuture<TOptional<UE::Interchange::FMeshPayloadData>> GetMeshPayloadData(const FInterchangeMeshPayLoadKey& PayLoadKey) const override;
 
 	/* IInterchangeStaticMeshPayloadInterface End */
 
diff --git a/Engine/Plugins/Interchange/Runtime/Source/Import/Public/Mesh/InterchangeSkeletalMeshPayload.h b/Engine/Plugins/Interchange/Runtime/Source/Import/Public/Mesh/InterchangeSkeletalMeshPayload.h
deleted file mode 100644
diff --git a/Engine/Plugins/Interchange/Runtime/Source/Import/Public/Mesh/InterchangeSkeletalMeshPayloadInterface.h b/Engine/Plugins/Interchange/Runtime/Source/Import/Public/Mesh/InterchangeSkeletalMeshPayloadInterface.h
deleted file mode 100644
diff --git a/Engine/Plugins/Interchange/Runtime/Source/Import/Public/Mesh/InterchangeStaticMeshFactory.h b/Engine/Plugins/Interchange/Runtime/Source/Import/Public/Mesh/InterchangeStaticMeshFactory.h
--- a/Engine/Plugins/Interchange/Runtime/Source/Import/Public/Mesh/InterchangeStaticMeshFactory.h
+++ b/Engine/Plugins/Interchange/Runtime/Source/Import/Public/Mesh/InterchangeStaticMeshFactory.h
@@ -6,7 +6,7 @@
 #include "InterchangeFactoryBase.h"
 #include "UObject/Object.h"
 #include "UObject/ObjectMacros.h"
-#include "Mesh/InterchangeStaticMeshPayload.h"
+#include "Mesh/InterchangeMeshPayload.h"
 
 #include "InterchangeStaticMeshFactory.generated.h"
 
@@ -44,7 +44,7 @@ private:
 	struct FMeshPayload
 	{
 		FString MeshName;
-		TFuture<TOptional<UE::Interchange::FStaticMeshPayloadData>> PayloadData;
+		TFuture<TOptional<UE::Interchange::FMeshPayloadData>> PayloadData;
 		FTransform Transform = FTransform::Identity;
 	};
 
diff --git a/Engine/Plugins/Interchange/Runtime/Source/Import/Public/Mesh/InterchangeStaticMeshPayload.h b/Engine/Plugins/Interchange/Runtime/Source/Import/Public/Mesh/InterchangeStaticMeshPayload.h
deleted file mode 100644
diff --git a/Engine/Plugins/Interchange/Runtime/Source/Import/Public/Mesh/InterchangeStaticMeshPayloadInterface.h b/Engine/Plugins/Interchange/Runtime/Source/Import/Public/Mesh/InterchangeStaticMeshPayloadInterface.h
deleted file mode 100644
diff --git a/Engine/Plugins/Interchange/Runtime/Source/Nodes/Private/InterchangeMeshNode.cpp b/Engine/Plugins/Interchange/Runtime/Source/Nodes/Private/InterchangeMeshNode.cpp
--- a/Engine/Plugins/Interchange/Runtime/Source/Nodes/Private/InterchangeMeshNode.cpp
+++ b/Engine/Plugins/Interchange/Runtime/Source/Nodes/Private/InterchangeMeshNode.cpp
@@ -9,10 +9,15 @@ namespace UE
 {
 	namespace Interchange
 	{
+		const FAttributeKey& FMeshNodeStaticData::PayLoadKey()
+		{
+			static FAttributeKey AttributeKey(TEXT("__PayloadKey__"));
+			return AttributeKey;
+		}
 
-		const FAttributeKey& FMeshNodeStaticData::PayloadSourceKey()
+		const FAttributeKey& FMeshNodeStaticData::PayLoadTypeKey()
 		{
-			static FAttributeKey AttributeKey(TEXT("__PayloadSourceKey__"));
+			static FAttributeKey AttributeKey(TEXT("__PayloadTypeKey__"));
 			return AttributeKey;
 		}
 
@@ -73,7 +78,7 @@ FString UInterchangeMeshNode::GetKeyDisplayName(const UE::Interchange::FAttribut
 {
 	FString KeyDisplayName = NodeAttributeKey.ToString();
 	const FString NodeAttributeKeyString = KeyDisplayName;
-	if (NodeAttributeKey == UE::Interchange::FMeshNodeStaticData::PayloadSourceKey())
+	if (NodeAttributeKey == UE::Interchange::FMeshNodeStaticData::PayLoadKey())
 	{
 		return KeyDisplayName = TEXT("Payload Source Key");
 	}
@@ -280,33 +285,67 @@ bool UInterchangeMeshNode::SetMorphTargetName(const FString& MorphTargetName)
 	return false;
 }
 
-const TOptional<FString> UInterchangeMeshNode::GetPayLoadKey() const
+const TOptional<FInterchangeMeshPayLoadKey> UInterchangeMeshNode::GetPayLoadKey() const
 {
-	if (!Attributes->ContainAttribute(UE::Interchange::FMeshNodeStaticData::PayloadSourceKey()))
-	{
-		return TOptional<FString>();
-	}
-	UE::Interchange::FAttributeStorage::TAttributeHandle<FString> AttributeHandle = Attributes->GetAttributeHandle<FString>(UE::Interchange::FMeshNodeStaticData::PayloadSourceKey());
-	if (!AttributeHandle.IsValid())
+	FString PayloadKey;
+	EInterchangeMeshPayLoadType PayLoadType;
+
+	//PayLoadKey
 	{
-		return TOptional<FString>();
+		if (!Attributes->ContainAttribute(UE::Interchange::FMeshNodeStaticData::PayLoadKey()))
+		{
+			return TOptional<FInterchangeMeshPayLoadKey>();
+		}
+		UE::Interchange::FAttributeStorage::TAttributeHandle<FString> AttributeHandle = Attributes->GetAttributeHandle<FString>(UE::Interchange::FMeshNodeStaticData::PayLoadKey());
+		if (!AttributeHandle.IsValid())
+		{
+			return TOptional<FInterchangeMeshPayLoadKey>();
+		}
+		UE::Interchange::EAttributeStorageResult Result = AttributeHandle.Get(PayloadKey);
+		if (!IsAttributeStorageResultSuccess(Result))
+		{
+			LogAttributeStorageErrors(Result, TEXT("UInterchangeMeshNode.GetPayLoadKey"), UE::Interchange::FMeshNodeStaticData::PayLoadKey());
+			return TOptional<FInterchangeMeshPayLoadKey>();
+		}
 	}
-	FString PayloadKey;
-	UE::Interchange::EAttributeStorageResult Result = AttributeHandle.Get(PayloadKey);
-	if (!IsAttributeStorageResultSuccess(Result))
+
+	//PayLoadType
 	{
-		LogAttributeStorageErrors(Result, TEXT("UInterchangeMeshNode.GetPayLoadKey"), UE::Interchange::FMeshNodeStaticData::PayloadSourceKey());
-		return TOptional<FString>();
+		if (!Attributes->ContainAttribute(UE::Interchange::FMeshNodeStaticData::PayLoadTypeKey()))
+		{
+			return TOptional<FInterchangeMeshPayLoadKey>();
+		}
+		UE::Interchange::FAttributeStorage::TAttributeHandle<EInterchangeMeshPayLoadType> AttributeHandle = Attributes->GetAttributeHandle<EInterchangeMeshPayLoadType>(UE::Interchange::FMeshNodeStaticData::PayLoadTypeKey());
+		if (!AttributeHandle.IsValid())
+		{
+			return TOptional<FInterchangeMeshPayLoadKey>();
+		}
+
+		UE::Interchange::EAttributeStorageResult Result = AttributeHandle.Get(PayLoadType);
+		if (!IsAttributeStorageResultSuccess(Result))
+		{
+			LogAttributeStorageErrors(Result, TEXT("UInterchangeMeshNode.GetPayLoadTypeKey"), UE::Interchange::FMeshNodeStaticData::PayLoadTypeKey());
+			return TOptional<FInterchangeMeshPayLoadKey>();
+		}
 	}
-	return TOptional<FString>(PayloadKey);
+	
+	return TOptional<FInterchangeMeshPayLoadKey>(FInterchangeMeshPayLoadKey(PayloadKey, PayLoadType));
 }
 
-void UInterchangeMeshNode::SetPayLoadKey(const FString& PayloadKey)
+void UInterchangeMeshNode::SetPayLoadKey(const FString& PayLoadKey, const EInterchangeMeshPayLoadType& PayLoadType)
 {
-	UE::Interchange::EAttributeStorageResult Result = Attributes->RegisterAttribute(UE::Interchange::FMeshNodeStaticData::PayloadSourceKey(), PayloadKey);
+	UE::Interchange::EAttributeStorageResult Result = Attributes->RegisterAttribute(UE::Interchange::FMeshNodeStaticData::PayLoadKey(), PayLoadKey);
 	if (!IsAttributeStorageResultSuccess(Result))
 	{
-		LogAttributeStorageErrors(Result, TEXT("UInterchangeMeshNode.SetPayLoadKey"), UE::Interchange::FMeshNodeStaticData::PayloadSourceKey());
+		LogAttributeStorageErrors(Result, TEXT("UInterchangeMeshNode.SetPayLoadKey"), UE::Interchange::FMeshNodeStaticData::PayLoadKey());
+	}
+	else
+	{
+		Result = Attributes->RegisterAttribute(UE::Interchange::FMeshNodeStaticData::PayLoadTypeKey(), PayLoadType);
+		if (!IsAttributeStorageResultSuccess(Result))
+		{
+			LogAttributeStorageErrors(Result, TEXT("UInterchangeMeshNode.SetPayLoadTypeKey"), UE::Interchange::FMeshNodeStaticData::PayLoadTypeKey());
+		}
 	}
 }
 
diff --git a/Engine/Plugins/Interchange/Runtime/Source/Nodes/Public/InterchangeMeshNode.h b/Engine/Plugins/Interchange/Runtime/Source/Nodes/Public/InterchangeMeshNode.h
--- a/Engine/Plugins/Interchange/Runtime/Source/Nodes/Public/InterchangeMeshNode.h
+++ b/Engine/Plugins/Interchange/Runtime/Source/Nodes/Public/InterchangeMeshNode.h
@@ -15,7 +15,8 @@ namespace UE
 
 		struct INTERCHANGENODES_API FMeshNodeStaticData : public FBaseNodeStaticData
 		{
-			static const FAttributeKey& PayloadSourceKey();
+			static const FAttributeKey& PayLoadKey();
+			static const FAttributeKey& PayLoadTypeKey();
 			static const FAttributeKey& IsSkinnedMeshKey();
 			static const FAttributeKey& IsMorphTargetKey();
 			static const FAttributeKey& MorphTargetNameKey();
@@ -28,6 +29,38 @@ namespace UE
 	}//ns Interchange
 }//ns UE
 
+UENUM(BlueprintType)
+enum class EInterchangeMeshPayLoadType : uint8
+{
+	STATIC = 0,
+	SKELETAL,
+	MORPHTARGET
+};
+
+
+USTRUCT(BlueprintType)
+struct INTERCHANGENODES_API FInterchangeMeshPayLoadKey
+{
+	GENERATED_BODY()
+
+public:
+
+	UPROPERTY(EditAnywhere, BlueprintReadWrite, Category = "Interchange | Mesh")
+	FString UniqueId;
+
+	UPROPERTY(EditAnywhere, BlueprintReadWrite, Category = "Interchange | Mesh")
+	EInterchangeMeshPayLoadType Type;
+
+	FInterchangeMeshPayLoadKey() {}
+
+	FInterchangeMeshPayLoadKey(const FString& InUniqueId, const EInterchangeMeshPayLoadType& InType)
+		: UniqueId(InUniqueId)
+		, Type(InType)
+	{
+	}
+};
+
+
 UCLASS(BlueprintType, Experimental)
 class INTERCHANGENODES_API UInterchangeMeshNode : public UInterchangeBaseNode
 {
@@ -101,10 +134,10 @@ public:
 	bool SetMorphTargetName(const FString& MorphTargetName);
 
 	/** Mesh node Interface Begin */
-	virtual const TOptional<FString> GetPayLoadKey() const;
+	virtual const TOptional<FInterchangeMeshPayLoadKey> GetPayLoadKey() const;
 
 	UFUNCTION(BlueprintCallable, Category = "Interchange | Node | Mesh")
-	virtual void SetPayLoadKey(const FString& PayloadKey);
+	virtual void SetPayLoadKey(const FString& PayLoadKey, const EInterchangeMeshPayLoadType& PayLoadType);
 	
 	/** Query this mesh vertices count. Return false if the attribute was not set.*/
 	UFUNCTION(BlueprintCallable, Category = "Interchange | Node | Mesh")
diff --git a/Engine/Plugins/Interchange/Runtime/Source/Parsers/Fbx/Private/FbxMesh.cpp b/Engine/Plugins/Interchange/Runtime/Source/Parsers/Fbx/Private/FbxMesh.cpp
--- a/Engine/Plugins/Interchange/Runtime/Source/Parsers/Fbx/Private/FbxMesh.cpp
+++ b/Engine/Plugins/Interchange/Runtime/Source/Parsers/Fbx/Private/FbxMesh.cpp
@@ -1392,7 +1392,7 @@ void FFbxMesh::AddAllMeshes(FbxScene* SDKScene, FbxGeometryConverter* SDKGeometr
 			GeoPayload->bIsSkinnedMesh = MeshNode->IsSkinnedMesh();
 			PayloadContexts.Add(PayLoadKey, GeoPayload);
 		}
-		MeshNode->SetPayLoadKey(PayLoadKey);
+		MeshNode->SetPayLoadKey(PayLoadKey, MeshNode->IsSkinnedMesh() ? EInterchangeMeshPayLoadType::SKELETAL : EInterchangeMeshPayLoadType::STATIC);
 
 		int32 NumAnimations = SDKScene->GetSrcObjectCount<FbxAnimStack>();
 
@@ -1467,7 +1467,7 @@ void FFbxMesh::AddAllMeshes(FbxScene* SDKScene, FbxGeometryConverter* SDKGeometr
 							GeoPayload->SDKGeometryConverter = SDKGeometryConverter;
 							PayloadContexts.Add(MorphTargetPayLoadKey, GeoPayload);
 						}
-						MorphTargetNode->SetPayLoadKey(MorphTargetPayLoadKey);
+						MorphTargetNode->SetPayLoadKey(MorphTargetPayLoadKey, EInterchangeMeshPayLoadType::MORPHTARGET);
 
 						for (int32 AnimationIndex = 0; AnimationIndex < NumAnimations; AnimationIndex++)
 						{
