MaterialX Translator:
- Pre-loading an array of packages in the Game Thread instead of just one asset.
- fixing the 2nd input still present after converting math nodes with neutral element to <dot> node, method: ConvertNeutralNodesToDot
- Adding a method to recursively update a document and all of its subgraphs
- Adding the support of <surface_unlit> surface shader node
- Adding a method to connect nodes that are not part of a node graph, see file: MaterialX\MaterialX-1.38.5\resources\Materials\TestSuite\stdlib\shader\surface.mtlx
Generic Material Pipeline:
- Adding the support of MaterialX surface_unlit
- Adding the data driven feature for surface_unlit

Adding the Material Function MX_SurfaceUnlit, that allows us to blend between surface_unlit's transmission and opacity handled respectively as Opacity with Translucent blend mode and Opacity Mask with Mask blend mode in UE

diff --git a/Engine/Plugins/Interchange/Editor/Source/Pipelines/Private/InterchangeMaterialXPipelineSettingsCustomization.cpp b/Engine/Plugins/Interchange/Editor/Source/Pipelines/Private/InterchangeMaterialXPipelineSettingsCustomization.cpp
--- a/Engine/Plugins/Interchange/Editor/Source/Pipelines/Private/InterchangeMaterialXPipelineSettingsCustomization.cpp
+++ b/Engine/Plugins/Interchange/Editor/Source/Pipelines/Private/InterchangeMaterialXPipelineSettingsCustomization.cpp
@@ -44,6 +44,10 @@ void FInterchangeMaterialXPipelineSettingsCustomization::CustomizeChildren(TShar
 				{
 					OnShouldFilterAsset = &FInterchangeMaterialXPipelineSettingsCustomization::OnShouldFilterAssetStandardSurfaceTransmission;
 				}
+				else if(Str == TEXT("Surface Unlit"))
+				{
+					OnShouldFilterAsset = &FInterchangeMaterialXPipelineSettingsCustomization::OnShouldFilterAssetSurfaceUnlit;
+				}
 			}
 
 			IDetailPropertyRow& SurfaceShaderTypeRow = ChildBuilder.AddProperty(ChildPropertyHandle.ToSharedRef());
@@ -88,4 +92,9 @@ bool FInterchangeMaterialXPipelineSettingsCustomization::OnShouldFilterAssetStan
 bool FInterchangeMaterialXPipelineSettingsCustomization::OnShouldFilterAssetStandardSurfaceTransmission(const FAssetData& InAssetData)
 {
 	return FMaterialXPipelineSettings::ShouldFilterAssets(Cast<UMaterialFunction>(InAssetData.GetAsset()), FMaterialXPipelineSettings::TransmissionSurfaceInputs, FMaterialXPipelineSettings::TransmissionSurfaceOutputs);
-}
\ No newline at end of file
+}
+
+bool FInterchangeMaterialXPipelineSettingsCustomization::OnShouldFilterAssetSurfaceUnlit(const FAssetData& InAssetData)
+{
+	return FMaterialXPipelineSettings::ShouldFilterAssets(Cast<UMaterialFunction>(InAssetData.GetAsset()), FMaterialXPipelineSettings::SurfaceUnlitInputs, FMaterialXPipelineSettings::SurfaceUnlitOutputs);
+}
diff --git a/Engine/Plugins/Interchange/Editor/Source/Pipelines/Private/InterchangeMaterialXPipelineSettingsCustomization.h b/Engine/Plugins/Interchange/Editor/Source/Pipelines/Private/InterchangeMaterialXPipelineSettingsCustomization.h
--- a/Engine/Plugins/Interchange/Editor/Source/Pipelines/Private/InterchangeMaterialXPipelineSettingsCustomization.h
+++ b/Engine/Plugins/Interchange/Editor/Source/Pipelines/Private/InterchangeMaterialXPipelineSettingsCustomization.h
@@ -21,4 +21,5 @@ protected:
 
 	bool OnShouldFilterAssetStandardSurface(const FAssetData & InAssetData);
 	bool OnShouldFilterAssetStandardSurfaceTransmission(const FAssetData& InAssetData);
+	bool OnShouldFilterAssetSurfaceUnlit(const FAssetData& InAssetData);
 };
\ No newline at end of file
diff --git a/Engine/Plugins/Interchange/Runtime/Source/Import/Private/MaterialX/InterchangeMaterialXDefinitions.h b/Engine/Plugins/Interchange/Runtime/Source/Import/Private/MaterialX/InterchangeMaterialXDefinitions.h
--- a/Engine/Plugins/Interchange/Runtime/Source/Import/Private/MaterialX/InterchangeMaterialXDefinitions.h
+++ b/Engine/Plugins/Interchange/Runtime/Source/Import/Private/MaterialX/InterchangeMaterialXDefinitions.h
@@ -129,6 +129,22 @@ MATERIALX_NAMESPACE_BEGIN
 			static constexpr const char* TransmissionColor = "transmission_color";
 			static constexpr const char* Opacity = "opacity";
 		}
+
+		namespace DefaultValue
+		{
+			namespace Float
+			{
+				constexpr float Emission = 1.f;
+				constexpr float Transmission = 0.f;
+				constexpr float Opacity = 1.f;
+			}
+
+			namespace Color3
+			{
+				constexpr FLinearColor EmissionColor{ 1, 1, 1 };
+				constexpr FLinearColor TransmissionColor{ 1, 1, 1 };
+			}
+		}
 	}
 
 	namespace Lights
diff --git a/Engine/Plugins/Interchange/Runtime/Source/Import/Private/MaterialX/InterchangeMaterialXTranslator.cpp b/Engine/Plugins/Interchange/Runtime/Source/Import/Private/MaterialX/InterchangeMaterialXTranslator.cpp
--- a/Engine/Plugins/Interchange/Runtime/Source/Import/Private/MaterialX/InterchangeMaterialXTranslator.cpp
+++ b/Engine/Plugins/Interchange/Runtime/Source/Import/Private/MaterialX/InterchangeMaterialXTranslator.cpp
@@ -38,34 +38,36 @@ namespace mx = MaterialX;
 
 namespace UE::Interchange::MaterialX
 {
-	bool IsMaterialFunctionPackageLoaded()
+	bool AreMaterialFunctionPackagesLoaded()
 	{
-		auto IsPackageLoaded = [](const FString& TextPath) -> bool
+		auto ArePackagesLoaded = [](const TArray<FString> & TextPaths) -> bool
 		{
-			const FString FunctionPath{ FPackageName::ExportTextPathToObjectPath(TextPath) };
-			if(FPackageName::DoesPackageExist(FunctionPath))
+			bool bAllLoaded = true;
+			for(const FString & TextPath : TextPaths)
 			{
-				if(FSoftObjectPath(FunctionPath).TryLoad())
+				const FString FunctionPath{ FPackageName::ExportTextPathToObjectPath(TextPath) };
+				if(FPackageName::DoesPackageExist(FunctionPath))
 				{
-					return true;
+					if(!FSoftObjectPath(FunctionPath).TryLoad())
+					{
+						UE_LOG(LogInterchangeImport, Warning, TEXT("Couldn't load %s"), *FunctionPath);
+						bAllLoaded = false;
+					}
 				}
 				else
 				{
-					UE_LOG(LogInterchangeImport, Warning, TEXT("Couldn't load %s"), *FunctionPath);
+					UE_LOG(LogInterchangeImport, Warning, TEXT("Couldn't find %s"), *FunctionPath);
+					bAllLoaded = false;
 				}
 			}
-			else
-			{
-				UE_LOG(LogInterchangeImport, Warning, TEXT("Couldn't find %s"), *FunctionPath);
-			}
 
-			return false;
+			return bAllLoaded;
 		};
 
-		static const bool bNormalFromHeightMapPackageLoaded = 
-			IsPackageLoaded(TEXT("MaterialFunction'/Engine/Functions/Engine_MaterialFunctions03/Procedurals/NormalFromHeightmap.NormalFromHeightmap'"));
+		static const bool bPackagesLoaded = 
+			ArePackagesLoaded({ TEXT("MaterialFunction'/Engine/Functions/Engine_MaterialFunctions03/Procedurals/NormalFromHeightmap.NormalFromHeightmap'") });
 
-		return bNormalFromHeightMapPackageLoaded;
+		return bPackagesLoaded;
 	}
 }
 
@@ -256,7 +258,7 @@ TArray<FString> UInterchangeMaterialXTranslator::GetSupportedFormats() const
 		return TArray<FString>{};
 	}
 
-	return UE::Interchange::MaterialX::IsMaterialFunctionPackageLoaded() ? TArray<FString>{ TEXT("mtlx;MaterialX File Format") } : TArray<FString>{};
+	return UE::Interchange::MaterialX::AreMaterialFunctionPackagesLoaded() ? TArray<FString>{ TEXT("mtlx;MaterialX File Format") } : TArray<FString>{};
 }
 
 bool UInterchangeMaterialXTranslator::Translate(UInterchangeBaseNodeContainer& BaseNodeContainer) const
@@ -304,6 +306,9 @@ bool UInterchangeMaterialXTranslator::Translate(UInterchangeBaseNodeContainer& B
 			return false;
 		}
 
+		//Update the document by initializing and reorganizing the different nodes and subgraphs
+		UpdateDocumentRecursively(Document);
+
 		for(mx::ElementPtr Elem : Document->traverseTree())
 		{
 			//make sure to read only the current file otherwise we'll process the entire library
@@ -333,7 +338,8 @@ bool UInterchangeMaterialXTranslator::Translate(UInterchangeBaseNodeContainer& B
 					//The entry point for materials is only on a surfacematerial node
 					if(bIsMaterialShader && Node->getCategory() == mx::SURFACE_MATERIAL_NODE_STRING)
 					{
-						bool bHasStandardSurface = false;
+						//standard_surface and surface_unlit are the only surface shaders supported for now
+						bool bHasSupportedSurfaceShader = false;
 
 						for(mx::InputPtr Input : Node->getInputs())
 						{
@@ -343,14 +349,19 @@ bool UInterchangeMaterialXTranslator::Translate(UInterchangeBaseNodeContainer& B
 							if(ConnectedNode && ConnectedNode->getCategory() == mx::Category::StandardSurface)
 							{
 								ProcessStandardSurface(BaseNodeContainer, ConnectedNode, Document);
-								bHasStandardSurface = true;
+								bHasSupportedSurfaceShader = true;
+							}
+							else if(ConnectedNode && ConnectedNode->getCategory() == mx::Category::SurfaceUnlit)
+							{
+								ProcessSurfaceUnlit(BaseNodeContainer, ConnectedNode, Document);
+								bHasSupportedSurfaceShader = true;
 							}
 						}
 
-						if(!bHasStandardSurface)
+						if(!bHasSupportedSurfaceShader)
 						{
 							UInterchangeResultWarning_Generic* Message = AddMessage<UInterchangeResultWarning_Generic>();
-							Message->Text = FText::Format(LOCTEXT("StandardSurfaceNotFound", "<{0}> has no standard_surface inputs"),
+							Message->Text = FText::Format(LOCTEXT("SurfaceShaderNotFound", "the surfaceshader of <{0}> is not supported"),
 														  FText::FromString(Node->getName().c_str()));
 						}
 					}
@@ -414,349 +425,167 @@ void UInterchangeMaterialXTranslator::ProcessStandardSurface(UInterchangeBaseNod
 
 	UInterchangeShaderGraphNode* ShaderGraphNode = CreateShaderNode<UInterchangeShaderGraphNode>(StandardSurfaceNode->getName().c_str(), StandardSurface::Name.ToString(), NamesToShaderNodes, NodeContainer);
 
-	//Base
+	auto ConnectNodeOutputToInput = [&](const char* InputName, UInterchangeShaderNode * ShaderGraphNode, const FString & InputShaderName, auto DefaultValue)
 	{
-		//Weight
-		{
-			mx::InputPtr InputBase = GetStandardSurfaceInput(StandardSurfaceNode, mx::StandardSurface::Input::Base, Document);
-
-			if(!ConnectNodeGraphOutputToInput(InputBase, ShaderGraphNode, StandardSurface::Parameters::Base.ToString(), NamesToShaderNodes, NodeContainer))
-			{
-				AddFloatAttribute(InputBase, StandardSurface::Parameters::Base.ToString(), ShaderGraphNode, mx::StandardSurface::DefaultValue::Float::Base);
-			}
-		}
+		mx::InputPtr Input = GetStandardSurfaceInput(StandardSurfaceNode, InputName, Document);
 
-		//Color
+		if(!ConnectNodeGraphOutputToInput(Input, ShaderGraphNode, InputShaderName, NamesToShaderNodes, NodeContainer))
 		{
-			mx::InputPtr InputBaseColor = GetStandardSurfaceInput(StandardSurfaceNode, mx::StandardSurface::Input::BaseColor, Document);
-
-			if(!ConnectNodeGraphOutputToInput(InputBaseColor, ShaderGraphNode, StandardSurface::Parameters::BaseColor.ToString(), NamesToShaderNodes, NodeContainer))
+			if(!ConnectNodeNameOutputToInput(Input, ShaderGraphNode, InputShaderName, NamesToShaderNodes, NodeContainer))
 			{
-				AddLinearColorAttribute(InputBaseColor, StandardSurface::Parameters::BaseColor.ToString(), ShaderGraphNode, mx::StandardSurface::DefaultValue::Color3::BaseColor);
+				//the different inputs from a standard surface are either float, linear color, or no attribute at all
+				if constexpr(std::is_same_v<decltype(DefaultValue), float>)
+				{
+					AddFloatAttribute(Input, InputShaderName, ShaderGraphNode, DefaultValue);
+				}
+				else if constexpr(std::is_same_v<decltype(DefaultValue), FLinearColor>)
+				{
+					AddLinearColorAttribute(Input, InputShaderName, ShaderGraphNode, DefaultValue);
+				}
 			}
 		}
-	}
+	};
 
-	//DiffuseRoughness
-	{
-		mx::InputPtr InputDiffuseRoughness = GetStandardSurfaceInput(StandardSurfaceNode, mx::StandardSurface::Input::DiffuseRoughness, Document);
+	//Base
+	ConnectNodeOutputToInput(mx::StandardSurface::Input::Base, ShaderGraphNode, StandardSurface::Parameters::Base.ToString(), mx::StandardSurface::DefaultValue::Float::Base);
+	
+	//Base Color
+	ConnectNodeOutputToInput(mx::StandardSurface::Input::BaseColor, ShaderGraphNode, StandardSurface::Parameters::BaseColor.ToString(), mx::StandardSurface::DefaultValue::Color3::BaseColor);
 
-		if(!ConnectNodeGraphOutputToInput(InputDiffuseRoughness, ShaderGraphNode, StandardSurface::Parameters::DiffuseRoughness.ToString(), NamesToShaderNodes, NodeContainer))
-		{
-			AddFloatAttribute(InputDiffuseRoughness, StandardSurface::Parameters::DiffuseRoughness.ToString(), ShaderGraphNode, mx::StandardSurface::DefaultValue::Float::DiffuseRoughness);
-		}
-	}
+	//Diffuse Roughness
+	ConnectNodeOutputToInput(mx::StandardSurface::Input::DiffuseRoughness, ShaderGraphNode, StandardSurface::Parameters::DiffuseRoughness.ToString(), mx::StandardSurface::DefaultValue::Float::DiffuseRoughness);
 
 	//Specular
-	{
-		//Weight
-		{
-			mx::InputPtr InputSpecular = GetStandardSurfaceInput(StandardSurfaceNode, mx::StandardSurface::Input::Specular, Document);
-
-			if(!ConnectNodeGraphOutputToInput(InputSpecular, ShaderGraphNode, StandardSurface::Parameters::Specular.ToString(), NamesToShaderNodes, NodeContainer))
-			{
-				AddFloatAttribute(InputSpecular, StandardSurface::Parameters::Specular.ToString(), ShaderGraphNode, mx::StandardSurface::DefaultValue::Float::Specular);
-			}
-		}
-
-		//Roughness
-		{
-			mx::InputPtr InputSpecularRoughness = GetStandardSurfaceInput(StandardSurfaceNode, mx::StandardSurface::Input::SpecularRoughness, Document);
-
-			if(!ConnectNodeGraphOutputToInput(InputSpecularRoughness, ShaderGraphNode, StandardSurface::Parameters::SpecularRoughness.ToString(), NamesToShaderNodes, NodeContainer))
-			{
-				AddFloatAttribute(InputSpecularRoughness, StandardSurface::Parameters::SpecularRoughness.ToString(), ShaderGraphNode, mx::StandardSurface::DefaultValue::Float::SpecularRoughness);
-			}
-		}
-
-		//IOR
-		{
-			mx::InputPtr InputSpecularIOR = GetStandardSurfaceInput(StandardSurfaceNode, mx::StandardSurface::Input::SpecularIOR, Document);
-
-			if(!ConnectNodeGraphOutputToInput(InputSpecularIOR, ShaderGraphNode, StandardSurface::Parameters::SpecularIOR.ToString(), NamesToShaderNodes, NodeContainer))
-			{
-				AddFloatAttribute(InputSpecularIOR, StandardSurface::Parameters::SpecularIOR.ToString(), ShaderGraphNode, mx::StandardSurface::DefaultValue::Float::SpecularIOR);
-			}
-		}
-
-		//Anisotropy
-		{
-			mx::InputPtr InputSpecularAnisotropy = GetStandardSurfaceInput(StandardSurfaceNode, mx::StandardSurface::Input::SpecularAnisotropy, Document);
-
-			if(!ConnectNodeGraphOutputToInput(InputSpecularAnisotropy, ShaderGraphNode, StandardSurface::Parameters::SpecularAnisotropy.ToString(), NamesToShaderNodes, NodeContainer))
-			{
-				AddFloatAttribute(InputSpecularAnisotropy, StandardSurface::Parameters::SpecularAnisotropy.ToString(), ShaderGraphNode, mx::StandardSurface::DefaultValue::Float::SpecularAnisotropy);
-			}
-		}
-
-		//Rotation
-		{
-			mx::InputPtr InputSpecularRotation = GetStandardSurfaceInput(StandardSurfaceNode, mx::StandardSurface::Input::SpecularRotation, Document);
-
-			if(!ConnectNodeGraphOutputToInput(InputSpecularRotation, ShaderGraphNode, StandardSurface::Parameters::SpecularRotation.ToString(), NamesToShaderNodes, NodeContainer))
-			{
-				AddFloatAttribute(InputSpecularRotation, StandardSurface::Parameters::SpecularRotation.ToString(), ShaderGraphNode, mx::StandardSurface::DefaultValue::Float::SpecularRotation);
-			}
-		}
-	}
+	ConnectNodeOutputToInput(mx::StandardSurface::Input::Specular, ShaderGraphNode, StandardSurface::Parameters::Specular.ToString(), mx::StandardSurface::DefaultValue::Float::Specular);
+	
+	//Specular Roughness
+	ConnectNodeOutputToInput(mx::StandardSurface::Input::SpecularRoughness, ShaderGraphNode, StandardSurface::Parameters::SpecularRoughness.ToString(), mx::StandardSurface::DefaultValue::Float::SpecularRoughness);
+	
+	//Specular IOR
+	ConnectNodeOutputToInput(mx::StandardSurface::Input::SpecularIOR, ShaderGraphNode, StandardSurface::Parameters::SpecularIOR.ToString(), mx::StandardSurface::DefaultValue::Float::SpecularIOR);
+	
+	//Specular Anisotropy
+	ConnectNodeOutputToInput(mx::StandardSurface::Input::SpecularAnisotropy, ShaderGraphNode, StandardSurface::Parameters::SpecularAnisotropy.ToString(), mx::StandardSurface::DefaultValue::Float::SpecularAnisotropy);
+	
+	//Specular Rotation
+	ConnectNodeOutputToInput(mx::StandardSurface::Input::SpecularRotation, ShaderGraphNode, StandardSurface::Parameters::SpecularRotation.ToString(), mx::StandardSurface::DefaultValue::Float::SpecularRotation);
 
 	//Metallic
-	{
-		mx::InputPtr InputMetalness = GetStandardSurfaceInput(StandardSurfaceNode, mx::StandardSurface::Input::Metalness, Document);
-
-		if(!ConnectNodeGraphOutputToInput(InputMetalness, ShaderGraphNode, StandardSurface::Parameters::Metalness.ToString(), NamesToShaderNodes, NodeContainer))
-		{
-			AddFloatAttribute(InputMetalness, StandardSurface::Parameters::Metalness.ToString(), ShaderGraphNode, mx::StandardSurface::DefaultValue::Float::Metalness);
-		}
-	}
+	ConnectNodeOutputToInput(mx::StandardSurface::Input::Metalness, ShaderGraphNode, StandardSurface::Parameters::Metalness.ToString(), mx::StandardSurface::DefaultValue::Float::Metalness);
 
 	//Subsurface
-	{
-		//Weight
-		{
-			mx::InputPtr InputSubsurface = GetStandardSurfaceInput(StandardSurfaceNode, mx::StandardSurface::Input::Subsurface, Document);
-
-			if(!ConnectNodeGraphOutputToInput(InputSubsurface, ShaderGraphNode, StandardSurface::Parameters::Subsurface.ToString(), NamesToShaderNodes, NodeContainer))
-			{
-				AddFloatAttribute(InputSubsurface, StandardSurface::Parameters::Subsurface.ToString(), ShaderGraphNode, mx::StandardSurface::DefaultValue::Float::Subsurface);
-			}
-		}
-
-		//Color
-		{
-			mx::InputPtr InputSubsurfaceColor = GetStandardSurfaceInput(StandardSurfaceNode, mx::StandardSurface::Input::SubsurfaceColor, Document);
-
-			if(!ConnectNodeGraphOutputToInput(InputSubsurfaceColor, ShaderGraphNode, StandardSurface::Parameters::SubsurfaceColor.ToString(), NamesToShaderNodes, NodeContainer))
-			{
-				AddLinearColorAttribute(InputSubsurfaceColor, StandardSurface::Parameters::SubsurfaceColor.ToString(), ShaderGraphNode, mx::StandardSurface::DefaultValue::Color3::SubsurfaceColor);
-			}
-		}
-
-		//Radius
-		{
-			mx::InputPtr InputSubsurfaceRadius = GetStandardSurfaceInput(StandardSurfaceNode, mx::StandardSurface::Input::SubsurfaceRadius, Document);
-
-			if(!ConnectNodeGraphOutputToInput(InputSubsurfaceRadius, ShaderGraphNode, StandardSurface::Parameters::SubsurfaceRadius.ToString(), NamesToShaderNodes, NodeContainer))
-			{
-				AddLinearColorAttribute(InputSubsurfaceRadius, StandardSurface::Parameters::SubsurfaceRadius.ToString(), ShaderGraphNode, mx::StandardSurface::DefaultValue::Color3::SubsurfaceRadius);
-			}
-		}
-
-		//Scale
-		{
-			mx::InputPtr InputSubsurfaceScale = GetStandardSurfaceInput(StandardSurfaceNode, mx::StandardSurface::Input::SubsurfaceScale, Document);
-
-			if(!ConnectNodeGraphOutputToInput(InputSubsurfaceScale, ShaderGraphNode, StandardSurface::Parameters::SubsurfaceScale.ToString(), NamesToShaderNodes, NodeContainer))
-			{
-				AddFloatAttribute(InputSubsurfaceScale, StandardSurface::Parameters::SubsurfaceScale.ToString(), ShaderGraphNode, mx::StandardSurface::DefaultValue::Float::SubsurfaceScale);
-			}
-		}
-	}
+	ConnectNodeOutputToInput(mx::StandardSurface::Input::Subsurface, ShaderGraphNode, StandardSurface::Parameters::Subsurface.ToString(), mx::StandardSurface::DefaultValue::Float::Subsurface);
+	
+	//Subsurface Color
+	ConnectNodeOutputToInput(mx::StandardSurface::Input::SubsurfaceColor, ShaderGraphNode, StandardSurface::Parameters::SubsurfaceColor.ToString(), mx::StandardSurface::DefaultValue::Color3::SubsurfaceColor);
+	
+	//Subsurface Radius
+	ConnectNodeOutputToInput(mx::StandardSurface::Input::SubsurfaceRadius, ShaderGraphNode, StandardSurface::Parameters::SubsurfaceRadius.ToString(), mx::StandardSurface::DefaultValue::Color3::SubsurfaceRadius);
+	
+	//Subsurface Scale
+	ConnectNodeOutputToInput(mx::StandardSurface::Input::SubsurfaceScale, ShaderGraphNode, StandardSurface::Parameters::SubsurfaceScale.ToString(), mx::StandardSurface::DefaultValue::Float::SubsurfaceScale);
 
 	//Sheen
-	{
-		//Weight
-		{
-			mx::InputPtr InputSheen = GetStandardSurfaceInput(StandardSurfaceNode, mx::StandardSurface::Input::Sheen, Document);
-
-			if(!ConnectNodeGraphOutputToInput(InputSheen, ShaderGraphNode, StandardSurface::Parameters::Sheen.ToString(), NamesToShaderNodes, NodeContainer))
-			{
-				AddFloatAttribute(InputSheen, StandardSurface::Parameters::Sheen.ToString(), ShaderGraphNode, mx::StandardSurface::DefaultValue::Float::Sheen);
-			}
-		}
-
-		//Color
-		{
-			mx::InputPtr InputSheenColor = GetStandardSurfaceInput(StandardSurfaceNode, mx::StandardSurface::Input::SheenColor, Document);
-
-			if(!ConnectNodeGraphOutputToInput(InputSheenColor, ShaderGraphNode, StandardSurface::Parameters::SubsurfaceColor.ToString(), NamesToShaderNodes, NodeContainer))
-			{
-				AddLinearColorAttribute(InputSheenColor, StandardSurface::Parameters::SheenColor.ToString(), ShaderGraphNode, mx::StandardSurface::DefaultValue::Color3::SheenColor);
-			}
-		}
-
-		//Roughness
-		{
-			mx::InputPtr InputSheenRoughness = GetStandardSurfaceInput(StandardSurfaceNode, mx::StandardSurface::Input::SheenRoughness, Document);
-
-			if(!ConnectNodeGraphOutputToInput(InputSheenRoughness, ShaderGraphNode, StandardSurface::Parameters::SheenRoughness.ToString(), NamesToShaderNodes, NodeContainer))
-			{
-				AddFloatAttribute(InputSheenRoughness, StandardSurface::Parameters::SheenRoughness.ToString(), ShaderGraphNode, mx::StandardSurface::DefaultValue::Float::SheenRoughness);
-			}
-		}
-	}
+	ConnectNodeOutputToInput(mx::StandardSurface::Input::Sheen, ShaderGraphNode, StandardSurface::Parameters::Sheen.ToString(), mx::StandardSurface::DefaultValue::Float::Sheen);
+	
+	//Sheen Color
+	ConnectNodeOutputToInput(mx::StandardSurface::Input::SheenColor, ShaderGraphNode, StandardSurface::Parameters::SheenColor.ToString(), mx::StandardSurface::DefaultValue::Color3::SheenColor);
+	
+	//Sheen Roughness
+	ConnectNodeOutputToInput(mx::StandardSurface::Input::SheenRoughness, ShaderGraphNode, StandardSurface::Parameters::SheenRoughness.ToString(), mx::StandardSurface::DefaultValue::Float::SheenRoughness);
 
 	//Coat
-	{
-		//Weight
-		{
-			mx::InputPtr InputCoat = GetStandardSurfaceInput(StandardSurfaceNode, mx::StandardSurface::Input::Coat, Document);
-
-			if(!ConnectNodeGraphOutputToInput(InputCoat, ShaderGraphNode, StandardSurface::Parameters::Coat.ToString(), NamesToShaderNodes, NodeContainer))
-			{
-				AddFloatAttribute(InputCoat, StandardSurface::Parameters::Coat.ToString(), ShaderGraphNode, mx::StandardSurface::DefaultValue::Float::Coat);
-			}
-		}
-
-		//Color
-		{
-			mx::InputPtr InputCoatColor = GetStandardSurfaceInput(StandardSurfaceNode, mx::StandardSurface::Input::CoatColor, Document);
-
-			if(!ConnectNodeGraphOutputToInput(InputCoatColor, ShaderGraphNode, StandardSurface::Parameters::CoatColor.ToString(), NamesToShaderNodes, NodeContainer))
-			{
-				AddLinearColorAttribute(InputCoatColor, StandardSurface::Parameters::CoatColor.ToString(), ShaderGraphNode, mx::StandardSurface::DefaultValue::Color3::CoatColor);
-			}
-		}
-
-		//Roughness
-		{
-			mx::InputPtr InputCoatRoughness = GetStandardSurfaceInput(StandardSurfaceNode, mx::StandardSurface::Input::CoatRoughness, Document);
-
-			if(!ConnectNodeGraphOutputToInput(InputCoatRoughness, ShaderGraphNode, StandardSurface::Parameters::CoatRoughness.ToString(), NamesToShaderNodes, NodeContainer))
-			{
-				AddFloatAttribute(InputCoatRoughness, StandardSurface::Parameters::CoatRoughness.ToString(), ShaderGraphNode, mx::StandardSurface::DefaultValue::Float::CoatRoughness);
-			}
-		}
-
-		//Normal
-		{
-			//No need to take the default input if there is no Normal input
-			mx::InputPtr InputCoatNormal = StandardSurfaceNode->getInput(mx::StandardSurface::Input::CoatNormal);
-
-			if(InputCoatNormal)
-			{
-				ConnectNodeGraphOutputToInput(InputCoatNormal, ShaderGraphNode, StandardSurface::Parameters::CoatNormal.ToString(), NamesToShaderNodes, NodeContainer);
-			}
-		}
-	}
-
-	//ThinFilmThickness
-	{
-		mx::InputPtr InputThinFilmThickness = GetStandardSurfaceInput(StandardSurfaceNode, mx::StandardSurface::Input::ThinFilmThickness, Document);
+	ConnectNodeOutputToInput(mx::StandardSurface::Input::Coat, ShaderGraphNode, StandardSurface::Parameters::Coat.ToString(), mx::StandardSurface::DefaultValue::Float::Coat);
+	
+	//Coat Color
+	ConnectNodeOutputToInput(mx::StandardSurface::Input::CoatColor, ShaderGraphNode, StandardSurface::Parameters::CoatColor.ToString(), mx::StandardSurface::DefaultValue::Color3::CoatColor);
+	
+	//Coat Roughness
+	ConnectNodeOutputToInput(mx::StandardSurface::Input::CoatRoughness, ShaderGraphNode, StandardSurface::Parameters::CoatRoughness.ToString(), mx::StandardSurface::DefaultValue::Float::CoatRoughness);
+	
+	//Coat Normal: No need to take the default input if there is no CoatNormal input
+	ConnectNodeOutputToInput(mx::StandardSurface::Input::CoatNormal, ShaderGraphNode, StandardSurface::Parameters::CoatNormal.ToString(), nullptr);
 
-		if(!ConnectNodeGraphOutputToInput(InputThinFilmThickness, ShaderGraphNode, StandardSurface::Parameters::ThinFilmThickness.ToString(), NamesToShaderNodes, NodeContainer))
-		{
-			AddFloatAttribute(InputThinFilmThickness, StandardSurface::Parameters::ThinFilmThickness.ToString(), ShaderGraphNode, mx::StandardSurface::DefaultValue::Float::ThinFilmThickness);
-		}
-	}
+	//Thin Film Thickness
+	ConnectNodeOutputToInput(mx::StandardSurface::Input::ThinFilmThickness, ShaderGraphNode, StandardSurface::Parameters::ThinFilmThickness.ToString(), mx::StandardSurface::DefaultValue::Float::ThinFilmThickness);
 
 	//Emission
-	{
-		//Weight
-		{
-			mx::InputPtr InputEmission = GetStandardSurfaceInput(StandardSurfaceNode, mx::StandardSurface::Input::Emission, Document);
-
-			if(!ConnectNodeGraphOutputToInput(InputEmission, ShaderGraphNode, StandardSurface::Parameters::Emission.ToString(), NamesToShaderNodes, NodeContainer))
-			{
-				AddFloatAttribute(InputEmission, StandardSurface::Parameters::Emission.ToString(), ShaderGraphNode, mx::StandardSurface::DefaultValue::Float::Emission);
-			}
-		}
-
-		//Color
-		{
-			mx::InputPtr InputEmissionColor = GetStandardSurfaceInput(StandardSurfaceNode, mx::StandardSurface::Input::EmissionColor, Document);
-
-			if(!ConnectNodeGraphOutputToInput(InputEmissionColor, ShaderGraphNode, StandardSurface::Parameters::EmissionColor.ToString(), NamesToShaderNodes, NodeContainer))
-			{
-				AddLinearColorAttribute(InputEmissionColor, StandardSurface::Parameters::EmissionColor.ToString(), ShaderGraphNode, mx::StandardSurface::DefaultValue::Color3::EmissionColor);
-			}
-		}
-	}
-
-	//Normal
-	{
-		//No need to take the default input if there is no Normal input
-		mx::InputPtr InputNormal = StandardSurfaceNode->getInput(mx::StandardSurface::Input::Normal);
-
-		if(InputNormal)
-		{
-			ConnectNodeGraphOutputToInput(InputNormal, ShaderGraphNode, StandardSurface::Parameters::Normal.ToString(), NamesToShaderNodes, NodeContainer);
-		}
-	}
+	ConnectNodeOutputToInput(mx::StandardSurface::Input::Emission, ShaderGraphNode, StandardSurface::Parameters::Emission.ToString(), mx::StandardSurface::DefaultValue::Float::Emission);
+	
+	//Emission Color
+	ConnectNodeOutputToInput(mx::StandardSurface::Input::EmissionColor, ShaderGraphNode, StandardSurface::Parameters::EmissionColor.ToString(), mx::StandardSurface::DefaultValue::Color3::EmissionColor);
 
-	//Tangent
-	{
-		//No need to take the default input if there is no Tangent input
-		mx::InputPtr InputTangent = StandardSurfaceNode->getInput(mx::StandardSurface::Input::Tangent);
+	//Normal: No need to take the default input if there is no Normal input
+	ConnectNodeOutputToInput(mx::StandardSurface::Input::Normal, ShaderGraphNode, StandardSurface::Parameters::Normal.ToString(), nullptr);
 
-		if(InputTangent)
-		{
-			ConnectNodeGraphOutputToInput(InputTangent, ShaderGraphNode, StandardSurface::Parameters::Tangent.ToString(), NamesToShaderNodes, NodeContainer);
-		}
-	}
+	//Tangent: No need to take the default input if there is no Tangent input
+	ConnectNodeOutputToInput(mx::StandardSurface::Input::Tangent, ShaderGraphNode, StandardSurface::Parameters::Tangent.ToString(), nullptr);
 
 	//Transmission
-	{
-		//Weight
-		{
-			mx::InputPtr InputTransmission = GetStandardSurfaceInput(StandardSurfaceNode, mx::StandardSurface::Input::Transmission, Document);
+	ConnectNodeOutputToInput(mx::StandardSurface::Input::Transmission, ShaderGraphNode, StandardSurface::Parameters::Transmission.ToString(), mx::StandardSurface::DefaultValue::Float::Transmission);
+	
+	//Transmission Color
+	ConnectNodeOutputToInput(mx::StandardSurface::Input::TransmissionColor, ShaderGraphNode, StandardSurface::Parameters::TransmissionColor.ToString(), mx::StandardSurface::DefaultValue::Color3::TransmissionColor);
+	
+	//Transmission Depth
+	ConnectNodeOutputToInput(mx::StandardSurface::Input::TransmissionDepth, ShaderGraphNode, StandardSurface::Parameters::TransmissionDepth.ToString(), mx::StandardSurface::DefaultValue::Float::TransmissionDepth);
+	
+	//Transmission Scatter
+	ConnectNodeOutputToInput(mx::StandardSurface::Input::TransmissionScatter, ShaderGraphNode, StandardSurface::Parameters::TransmissionScatter.ToString(), mx::StandardSurface::DefaultValue::Color3::TransmissionScatter);
+	
+	//Transmission Scatter Anisotropy
+	ConnectNodeOutputToInput(mx::StandardSurface::Input::TransmissionScatterAnisotropy, ShaderGraphNode, StandardSurface::Parameters::TransmissionScatterAnisotropy.ToString(), mx::StandardSurface::DefaultValue::Float::TransmissionScatterAnisotropy);
 
-			if (!ConnectNodeGraphOutputToInput(InputTransmission, ShaderGraphNode, StandardSurface::Parameters::Transmission.ToString(), NamesToShaderNodes, NodeContainer))
-			{
-				AddFloatAttribute(InputTransmission, StandardSurface::Parameters::Transmission.ToString(), ShaderGraphNode, mx::StandardSurface::DefaultValue::Float::Transmission);
-			}
-		}
+	//Transmission Dispersion
+	ConnectNodeOutputToInput(mx::StandardSurface::Input::TransmissionDispersion, ShaderGraphNode, StandardSurface::Parameters::TransmissionDispersion.ToString(), mx::StandardSurface::DefaultValue::Float::TransmissionDispersion);
+	
+	//Transmission Extra Roughness
+	ConnectNodeOutputToInput(mx::StandardSurface::Input::TransmissionExtraRoughness, ShaderGraphNode, StandardSurface::Parameters::TransmissionExtraRoughness.ToString(), mx::StandardSurface::DefaultValue::Float::TransmissionExtraRoughness);
+}
 
-		//Color
-		{
-			mx::InputPtr InputTransmissionColor = GetStandardSurfaceInput(StandardSurfaceNode, mx::StandardSurface::Input::TransmissionColor, Document);
+void UInterchangeMaterialXTranslator::ProcessSurfaceUnlit(UInterchangeBaseNodeContainer& NodeContainer, MaterialX::NodePtr SurfaceUnlitNode, MaterialX::DocumentPtr Document) const
+{
+	using namespace UE::Interchange::Materials;
 
-			if (!ConnectNodeGraphOutputToInput(InputTransmissionColor, ShaderGraphNode, StandardSurface::Parameters::TransmissionColor.ToString(), NamesToShaderNodes, NodeContainer))
-			{
-				AddLinearColorAttribute(InputTransmissionColor, StandardSurface::Parameters::TransmissionColor.ToString(), ShaderGraphNode, mx::StandardSurface::DefaultValue::Color3::TransmissionColor);
-			}
-		}
+	TMap<FString, UInterchangeShaderNode*> NamesToShaderNodes;
 
-		//Depth
-		{
-			mx::InputPtr InputTransmissionDepth = GetStandardSurfaceInput(StandardSurfaceNode, mx::StandardSurface::Input::TransmissionDepth, Document);
+	UInterchangeShaderGraphNode* ShaderGraphNode = CreateShaderNode<UInterchangeShaderGraphNode>(SurfaceUnlitNode->getName().c_str(), SurfaceUnlit::Name.ToString(), NamesToShaderNodes, NodeContainer);
 
-			if (!ConnectNodeGraphOutputToInput(InputTransmissionDepth, ShaderGraphNode, StandardSurface::Parameters::TransmissionDepth.ToString(), NamesToShaderNodes, NodeContainer))
-			{
-				AddFloatAttribute(InputTransmissionDepth, StandardSurface::Parameters::TransmissionDepth.ToString(), ShaderGraphNode, mx::StandardSurface::DefaultValue::Float::TransmissionDepth);
-			}
-		}
+	auto ConnectNodeOutputToInput = [&](const char* InputName, UInterchangeShaderNode* ShaderGraphNode, const FString& InputShaderName, auto DefaultValue)
+	{
+		mx::InputPtr Input = GetSurfaceUnlitInput(SurfaceUnlitNode, InputName, Document);
 
-		//Scatter
+		if(!ConnectNodeGraphOutputToInput(Input, ShaderGraphNode, InputShaderName, NamesToShaderNodes, NodeContainer))
 		{
-			mx::InputPtr InputTransmissionScatter = GetStandardSurfaceInput(StandardSurfaceNode, mx::StandardSurface::Input::TransmissionScatter, Document);
-
-			if (!ConnectNodeGraphOutputToInput(InputTransmissionScatter, ShaderGraphNode, StandardSurface::Parameters::TransmissionScatter.ToString(), NamesToShaderNodes, NodeContainer))
+			if(!ConnectNodeNameOutputToInput(Input, ShaderGraphNode, InputShaderName, NamesToShaderNodes, NodeContainer))
 			{
-				AddLinearColorAttribute(InputTransmissionScatter, StandardSurface::Parameters::TransmissionScatter.ToString(), ShaderGraphNode, mx::StandardSurface::DefaultValue::Color3::TransmissionScatter);
+				//the different inputs from a surface unlit are either float or linear color
+				if constexpr(std::is_same_v<decltype(DefaultValue), float>)
+				{
+					AddFloatAttribute(Input, InputShaderName, ShaderGraphNode, DefaultValue);
+				}
+				else if constexpr(std::is_same_v<decltype(DefaultValue), FLinearColor>)
+				{
+					AddLinearColorAttribute(Input, InputShaderName, ShaderGraphNode, DefaultValue);
+				}
 			}
 		}
+	};
 
-		//Scatter Anisotropy
-		{
-			mx::InputPtr InputTransmissionScatterAnisotropy = GetStandardSurfaceInput(StandardSurfaceNode, mx::StandardSurface::Input::TransmissionScatterAnisotropy, Document);
-
-			if (!ConnectNodeGraphOutputToInput(InputTransmissionScatterAnisotropy, ShaderGraphNode, StandardSurface::Parameters::TransmissionScatterAnisotropy.ToString(), NamesToShaderNodes, NodeContainer))
-			{
-				AddFloatAttribute(InputTransmissionScatterAnisotropy, StandardSurface::Parameters::TransmissionScatterAnisotropy.ToString(), ShaderGraphNode, mx::StandardSurface::DefaultValue::Float::TransmissionScatterAnisotropy);
-			}
-		}
+	//Emission
+	ConnectNodeOutputToInput(mx::SurfaceUnlit::Input::Emission, ShaderGraphNode, SurfaceUnlit::Parameters::Emission.ToString(), mx::SurfaceUnlit::DefaultValue::Float::Emission);
 
-		//Dispersion
-		{
-			mx::InputPtr InputTransmissionDispersion = GetStandardSurfaceInput(StandardSurfaceNode, mx::StandardSurface::Input::TransmissionDispersion, Document);
+	//Emission Color
+	ConnectNodeOutputToInput(mx::SurfaceUnlit::Input::EmissionColor, ShaderGraphNode, SurfaceUnlit::Parameters::EmissionColor.ToString(), mx::SurfaceUnlit::DefaultValue::Color3::EmissionColor);
 
-			if (!ConnectNodeGraphOutputToInput(InputTransmissionDispersion, ShaderGraphNode, StandardSurface::Parameters::TransmissionDispersion.ToString(), NamesToShaderNodes, NodeContainer))
-			{
-				AddFloatAttribute(InputTransmissionDispersion, StandardSurface::Parameters::TransmissionDispersion.ToString(), ShaderGraphNode, mx::StandardSurface::DefaultValue::Float::TransmissionDispersion);
-			}
-		}
+	//Opacity
+	ConnectNodeOutputToInput(mx::SurfaceUnlit::Input::Opacity, ShaderGraphNode, SurfaceUnlit::Parameters::Opacity.ToString(), mx::SurfaceUnlit::DefaultValue::Float::Opacity);
 
-		//Extra Roughness
-		{
-			mx::InputPtr InputTransmissionExtraRoughness = GetStandardSurfaceInput(StandardSurfaceNode, mx::StandardSurface::Input::TransmissionExtraRoughness, Document);
+	//Transmission
+	ConnectNodeOutputToInput(mx::SurfaceUnlit::Input::Transmission, ShaderGraphNode, SurfaceUnlit::Parameters::Transmission.ToString(), mx::SurfaceUnlit::DefaultValue::Float::Transmission);
 
-			if (!ConnectNodeGraphOutputToInput(InputTransmissionExtraRoughness, ShaderGraphNode, StandardSurface::Parameters::TransmissionExtraRoughness.ToString(), NamesToShaderNodes, NodeContainer))
-			{
-				AddFloatAttribute(InputTransmissionExtraRoughness, StandardSurface::Parameters::TransmissionExtraRoughness.ToString(), ShaderGraphNode, mx::StandardSurface::DefaultValue::Float::TransmissionExtraRoughness);
-			}
-		}
-	}
+	//Transmission Color
+	ConnectNodeOutputToInput(mx::SurfaceUnlit::Input::TransmissionColor, ShaderGraphNode, SurfaceUnlit::Parameters::TransmissionColor.ToString(), mx::SurfaceUnlit::DefaultValue::Color3::TransmissionColor);
 }
 
 void UInterchangeMaterialXTranslator::ProcessLightShader(UInterchangeBaseNodeContainer& NodeContainer, MaterialX::NodePtr LightShaderNode, MaterialX::DocumentPtr Document) const
@@ -929,157 +758,10 @@ bool UInterchangeMaterialXTranslator::ConnectNodeGraphOutputToInput(MaterialX::I
 										  FText::FromString(GetInputName(InputToNodeGraph)));
 			return false;
 		}
-	
-		//extract and separate nodes will be handled by ourselves, since they are defined with a swizzle node, which is not present in UE
-		auto FilterNodes = [](mx::NodePtr Node) ->bool
-		{
-			const std::string& Category= Node->getCategory();
-			return
-				Category != mx::Category::Extract &&
-				Category != mx::Category::Saturate &&
-				Category != mx::Category::Separate2 &&
-				Category != mx::Category::Separate3 &&
-				Category != mx::Category::Separate4;
-		};
-
-		mx::NodeGraphPtr NodeGraph = Output->getParent()->asA<mx::NodeGraph>();
-
-		AddTexCoordToTiledImageNodes(NodeGraph);
-
-		//This function will replace all the nodes that are defined with a nodegraph with their equivalent node network
-		//This allows us to handle nodes defined with a nodegraph, with nodes that we already support while traversing the current node graph
-		//For example the node tiledimage is just a succession of the nodes image, multiply, add, etc.
-		NodeGraph->flattenSubgraphs(mx::EMPTY_STRING, FilterNodes);
-
-		ConvertNeutralNodesToDot(NodeGraph);
 
 		for(mx::Edge Edge : Output->traverseGraph())
 		{
-			if(mx::NodePtr UpstreamNode = Edge.getUpstreamElement()->asA<mx::Node>())
-			{
-				UInterchangeShaderNode* ParentShaderNode = ShaderNode;
-				FString InputChannelName = ParentInputName;
-
-				//Replace the input's name by the one used in UE
-				RenameNodeInputs(UpstreamNode);
-
-				if(mx::NodePtr DownstreamNode = Edge.getDownstreamElement()->asA<mx::Node>())
-				{
-					if(UInterchangeShaderNode** FoundNode = NamesToShaderNodes.Find(GetAttributeParentName(DownstreamNode)))// DownstreamNode->getName().c_str()))
-					{
-						ParentShaderNode = *FoundNode;
-					}
-
-					if(mx::InputPtr ConnectedInput = Edge.getConnectingElement()->asA<mx::Input>())
-					{
-						InputChannelName = GetInputName(ConnectedInput);
-					}
-				}
-
-				if(ConnectNodeOutputToInput(UpstreamNode, ParentShaderNode, InputChannelName, NamesToShaderNodes, NodeContainer))
-				{
-					continue;
-				}
-				else if(UpstreamNode->getCategory() == mx::Category::Constant)
-				{
-					ConnectConstantInputToOutput(UpstreamNode, ParentShaderNode, InputChannelName, NamesToShaderNodes, NodeContainer);
-				}
-				else if(UpstreamNode->getCategory() == mx::Category::Extract)
-				{
-					ConnectExtractInputToOutput(UpstreamNode, ParentShaderNode, InputChannelName, NamesToShaderNodes, NodeContainer);
-				}
-				//dot means identity, input == output
-				else if(UpstreamNode->getCategory() == mx::Category::Dot || UpstreamNode->getCategory() == mx::Category::NormalMap)
-				{
-					ConnectDotInputToOutput(UpstreamNode, ParentShaderNode, InputChannelName, NamesToShaderNodes, NodeContainer);
-				}
-				else if(UpstreamNode->getCategory() == mx::Category::TransformPoint)
-				{
-					ConnectTransformPositionInputToOutput(UpstreamNode, ParentShaderNode, InputChannelName, NamesToShaderNodes, NodeContainer);
-				}
-				else if(UpstreamNode->getCategory() == mx::Category::TransformVector || UpstreamNode->getCategory() == mx::Category::TransformNormal)
-				{
-					ConnectTransformVectorInputToOutput(UpstreamNode, ParentShaderNode, InputChannelName, NamesToShaderNodes, NodeContainer);
-				}
-				else if(UpstreamNode->getCategory() == mx::Category::Rotate2D)
-				{
-					ConnectRotate2DInputToOutput(UpstreamNode, ParentShaderNode, InputChannelName, NamesToShaderNodes, NodeContainer);
-				}
-				else if(UpstreamNode->getCategory() == mx::Category::Rotate3D)
-				{
-					ConnectRotate3DInputToOutput(UpstreamNode, ParentShaderNode, InputChannelName, NamesToShaderNodes, NodeContainer);
-				}
-				else if(UpstreamNode->getCategory() == mx::Category::Image)
-				{
-					ConnectImageInputToOutput(UpstreamNode, ParentShaderNode, InputChannelName, NamesToShaderNodes, NodeContainer);
-				}
-				else if(UpstreamNode->getCategory() == mx::Category::Convert)
-				{
-					ConnectConvertInputToOutput(UpstreamNode, ParentShaderNode, InputChannelName, NamesToShaderNodes, NodeContainer);
-				}
-				else if(UpstreamNode->getCategory() == mx::Category::IfGreater)
-				{
-					ConnectIfGreaterInputToOutput(UpstreamNode, ParentShaderNode, InputChannelName, NamesToShaderNodes, NodeContainer);
-				}
-				else if(UpstreamNode->getCategory() == mx::Category::IfGreaterEq)
-				{
-					ConnectIfGreaterEqInputToOutput(UpstreamNode, ParentShaderNode, InputChannelName, NamesToShaderNodes, NodeContainer);
-				}
-				else if(UpstreamNode->getCategory() == mx::Category::IfEqual)
-				{
-					ConnectIfEqualInputToOutput(UpstreamNode, ParentShaderNode, InputChannelName, NamesToShaderNodes, NodeContainer);
-				}
-				else if(UpstreamNode->getCategory() == mx::Category::Outside)
-				{
-					ConnectOutsideInputToOutput(UpstreamNode, ParentShaderNode, InputChannelName, NamesToShaderNodes, NodeContainer);
-				}
-				else if(UpstreamNode->getCategory() == mx::Category::Position)
-				{
-					ConnectPositionInputToOutput(UpstreamNode, ParentShaderNode, InputChannelName, NamesToShaderNodes, NodeContainer);
-				}
-				else if(UpstreamNode->getCategory() == mx::Category::Normal)
-				{
-					ConnectNormalInputToOutput(UpstreamNode, ParentShaderNode, InputChannelName, NamesToShaderNodes, NodeContainer);
-				}
-				else if(UpstreamNode->getCategory() == mx::Category::Tangent)
-				{
-					ConnectTangentInputToOutput(UpstreamNode, ParentShaderNode, InputChannelName, NamesToShaderNodes, NodeContainer);
-				}
-				else if(UpstreamNode->getCategory() == mx::Category::Bitangent)
-				{
-					ConnectBitangentInputToOutput(UpstreamNode, ParentShaderNode, InputChannelName, NamesToShaderNodes, NodeContainer);
-				}
-				else if(UpstreamNode->getCategory() == mx::Category::Time)
-				{
-					ConnectTimeInputToOutput(UpstreamNode, ParentShaderNode, InputChannelName, NamesToShaderNodes, NodeContainer);
-				}
-				else if(UpstreamNode->getCategory() == mx::Category::Noise3D)
-				{
-					ConnectNoise3DInputToOutput(UpstreamNode, ParentShaderNode, InputChannelName, NamesToShaderNodes, NodeContainer);
-				}
-				else if(UpstreamNode->getCategory() == mx::Category::CellNoise3D)
-				{
-					ConnectCellNoise3DInputToOutput(UpstreamNode, ParentShaderNode, InputChannelName, NamesToShaderNodes, NodeContainer);
-				}
-				else if(UpstreamNode->getCategory() == mx::Category::WorleyNoise3D)
-				{
-					ConnectWorleyNoise3DInputToOutput(UpstreamNode, ParentShaderNode, InputChannelName, NamesToShaderNodes, NodeContainer);
-				}
-				else if(UpstreamNode->getCategory() == mx::Category::Blur)
-				{
-					ConnectBlurInputToOutput(UpstreamNode, ParentShaderNode, InputChannelName, NamesToShaderNodes, NodeContainer);
-				}
-				else if(UpstreamNode->getCategory() == mx::Category::HeightToNormal)
-				{
-					ConnectHeightToNormalInputToOutput(UpstreamNode, ParentShaderNode, InputChannelName, NamesToShaderNodes, NodeContainer);
-				}
-				else
-				{
-					UInterchangeResultWarning_Generic* Message = AddMessage<UInterchangeResultWarning_Generic>();
-					Message->Text = FText::Format(LOCTEXT("NodeCategoryNotSupported", "<{0}> is not supported yet"),
-												  FText::FromString(UpstreamNode->getCategory().c_str()));
-				}
-			}
+			ConnectNodeCategoryOutputToInput(Edge, ShaderNode, ParentInputName, NamesToShaderNodes, NodeContainer);
 		}
 	}
 
@@ -1110,6 +792,172 @@ bool UInterchangeMaterialXTranslator::ConnectNodeOutputToInput(MaterialX::NodePt
 	return bIsConnected;
 }
 
+void UInterchangeMaterialXTranslator::ConnectNodeCategoryOutputToInput(const MaterialX::Edge& Edge, UInterchangeShaderNode* ShaderNode, const FString& ParentInputName, TMap<FString, UInterchangeShaderNode*>& NamesToShaderNodes, UInterchangeBaseNodeContainer& NodeContainer) const
+{
+	if(mx::NodePtr UpstreamNode = Edge.getUpstreamElement()->asA<mx::Node>())
+	{
+		UInterchangeShaderNode* ParentShaderNode = ShaderNode;
+		FString InputChannelName = ParentInputName;
+
+		//Replace the input's name by the one used in UE
+		RenameNodeInputs(UpstreamNode);
+
+		if(mx::ElementPtr DownstreamElement = Edge.getDownstreamElement())
+		{
+			if(mx::NodePtr DownstreamNode = DownstreamElement->asA<mx::Node>())
+			{
+				if(UInterchangeShaderNode** FoundNode = NamesToShaderNodes.Find(GetAttributeParentName(DownstreamNode)))// DownstreamNode->getName().c_str()))
+				{
+					ParentShaderNode = *FoundNode;
+				}
+
+				if(mx::InputPtr ConnectedInput = Edge.getConnectingElement()->asA<mx::Input>())
+				{
+					InputChannelName = GetInputName(ConnectedInput);
+				}
+			}
+		}
+
+		if(ConnectNodeOutputToInput(UpstreamNode, ParentShaderNode, InputChannelName, NamesToShaderNodes, NodeContainer))
+		{
+			return;
+		}
+		else if(UpstreamNode->getCategory() == mx::Category::Constant)
+		{
+			ConnectConstantInputToOutput(UpstreamNode, ParentShaderNode, InputChannelName, NamesToShaderNodes, NodeContainer);
+		}
+		else if(UpstreamNode->getCategory() == mx::Category::Extract)
+		{
+			ConnectExtractInputToOutput(UpstreamNode, ParentShaderNode, InputChannelName, NamesToShaderNodes, NodeContainer);
+		}
+		//dot means identity, input == output
+		else if(UpstreamNode->getCategory() == mx::Category::Dot || UpstreamNode->getCategory() == mx::Category::NormalMap)
+		{
+			ConnectDotInputToOutput(UpstreamNode, ParentShaderNode, InputChannelName, NamesToShaderNodes, NodeContainer);
+		}
+		else if(UpstreamNode->getCategory() == mx::Category::TransformPoint)
+		{
+			ConnectTransformPositionInputToOutput(UpstreamNode, ParentShaderNode, InputChannelName, NamesToShaderNodes, NodeContainer);
+		}
+		else if(UpstreamNode->getCategory() == mx::Category::TransformVector || UpstreamNode->getCategory() == mx::Category::TransformNormal)
+		{
+			ConnectTransformVectorInputToOutput(UpstreamNode, ParentShaderNode, InputChannelName, NamesToShaderNodes, NodeContainer);
+		}
+		else if(UpstreamNode->getCategory() == mx::Category::Rotate2D)
+		{
+			ConnectRotate2DInputToOutput(UpstreamNode, ParentShaderNode, InputChannelName, NamesToShaderNodes, NodeContainer);
+		}
+		else if(UpstreamNode->getCategory() == mx::Category::Rotate3D)
+		{
+			ConnectRotate3DInputToOutput(UpstreamNode, ParentShaderNode, InputChannelName, NamesToShaderNodes, NodeContainer);
+		}
+		else if(UpstreamNode->getCategory() == mx::Category::Image)
+		{
+			ConnectImageInputToOutput(UpstreamNode, ParentShaderNode, InputChannelName, NamesToShaderNodes, NodeContainer);
+		}
+		else if(UpstreamNode->getCategory() == mx::Category::Convert)
+		{
+			ConnectConvertInputToOutput(UpstreamNode, ParentShaderNode, InputChannelName, NamesToShaderNodes, NodeContainer);
+		}
+		else if(UpstreamNode->getCategory() == mx::Category::IfGreater)
+		{
+			ConnectIfGreaterInputToOutput(UpstreamNode, ParentShaderNode, InputChannelName, NamesToShaderNodes, NodeContainer);
+		}
+		else if(UpstreamNode->getCategory() == mx::Category::IfGreaterEq)
+		{
+			ConnectIfGreaterEqInputToOutput(UpstreamNode, ParentShaderNode, InputChannelName, NamesToShaderNodes, NodeContainer);
+		}
+		else if(UpstreamNode->getCategory() == mx::Category::IfEqual)
+		{
+			ConnectIfEqualInputToOutput(UpstreamNode, ParentShaderNode, InputChannelName, NamesToShaderNodes, NodeContainer);
+		}
+		else if(UpstreamNode->getCategory() == mx::Category::Outside)
+		{
+			ConnectOutsideInputToOutput(UpstreamNode, ParentShaderNode, InputChannelName, NamesToShaderNodes, NodeContainer);
+		}
+		else if(UpstreamNode->getCategory() == mx::Category::Position)
+		{
+			ConnectPositionInputToOutput(UpstreamNode, ParentShaderNode, InputChannelName, NamesToShaderNodes, NodeContainer);
+		}
+		else if(UpstreamNode->getCategory() == mx::Category::Normal)
+		{
+			ConnectNormalInputToOutput(UpstreamNode, ParentShaderNode, InputChannelName, NamesToShaderNodes, NodeContainer);
+		}
+		else if(UpstreamNode->getCategory() == mx::Category::Tangent)
+		{
+			ConnectTangentInputToOutput(UpstreamNode, ParentShaderNode, InputChannelName, NamesToShaderNodes, NodeContainer);
+		}
+		else if(UpstreamNode->getCategory() == mx::Category::Bitangent)
+		{
+			ConnectBitangentInputToOutput(UpstreamNode, ParentShaderNode, InputChannelName, NamesToShaderNodes, NodeContainer);
+		}
+		else if(UpstreamNode->getCategory() == mx::Category::Time)
+		{
+			ConnectTimeInputToOutput(UpstreamNode, ParentShaderNode, InputChannelName, NamesToShaderNodes, NodeContainer);
+		}
+		else if(UpstreamNode->getCategory() == mx::Category::Noise3D)
+		{
+			ConnectNoise3DInputToOutput(UpstreamNode, ParentShaderNode, InputChannelName, NamesToShaderNodes, NodeContainer);
+		}
+		else if(UpstreamNode->getCategory() == mx::Category::CellNoise3D)
+		{
+			ConnectCellNoise3DInputToOutput(UpstreamNode, ParentShaderNode, InputChannelName, NamesToShaderNodes, NodeContainer);
+		}
+		else if(UpstreamNode->getCategory() == mx::Category::WorleyNoise3D)
+		{
+			ConnectWorleyNoise3DInputToOutput(UpstreamNode, ParentShaderNode, InputChannelName, NamesToShaderNodes, NodeContainer);
+		}
+		else if(UpstreamNode->getCategory() == mx::Category::Blur)
+		{
+			ConnectBlurInputToOutput(UpstreamNode, ParentShaderNode, InputChannelName, NamesToShaderNodes, NodeContainer);
+		}
+		else if(UpstreamNode->getCategory() == mx::Category::HeightToNormal)
+		{
+			ConnectHeightToNormalInputToOutput(UpstreamNode, ParentShaderNode, InputChannelName, NamesToShaderNodes, NodeContainer);
+		}
+		else
+		{
+			UInterchangeResultWarning_Generic* Message = AddMessage<UInterchangeResultWarning_Generic>();
+			Message->Text = FText::Format(LOCTEXT("NodeCategoryNotSupported", "<{0}> is not supported yet"),
+										  FText::FromString(UpstreamNode->getCategory().c_str()));
+		}
+	}
+}
+
+bool UInterchangeMaterialXTranslator::ConnectNodeNameOutputToInput(MaterialX::InputPtr InputToConnectedNode, UInterchangeShaderNode* ShaderNode, const FString& ParentInputName, TMap<FString, UInterchangeShaderNode*>& NamesToShaderNodes, UInterchangeBaseNodeContainer& NodeContainer) const
+{
+	mx::NodePtr ConnectedNode = InputToConnectedNode->getConnectedNode();
+
+	if(!ConnectedNode)
+	{
+		return false;
+	}
+
+	UInterchangeShaderNode* ParentShaderNode = ShaderNode;
+	FString InputChannelName = ParentInputName;
+
+	mx::Edge Edge(nullptr, InputToConnectedNode, ConnectedNode);
+
+	TArray<mx::Edge> Stack{ Edge };
+
+	while(!Stack.IsEmpty())
+	{
+		Edge = Stack.Pop();
+		
+		if(Edge.getUpstreamElement())
+		{
+			ConnectNodeCategoryOutputToInput(Edge, ShaderNode, ParentInputName, NamesToShaderNodes, NodeContainer);
+			ConnectedNode = Edge.getUpstreamElement()->asA<mx::Node>();
+			for(mx::InputPtr Input : ConnectedNode->getInputs())
+			{
+				Stack.Emplace(ConnectedNode, Input, Input->getConnectedNode());
+			}
+		}
+	}
+
+	return true;
+}
+
 const FString& UInterchangeMaterialXTranslator::GetMatchedInputName(MaterialX::NodePtr Node, MaterialX::InputPtr Input) const
 {
 	static FString EmptyString;
@@ -1180,6 +1028,18 @@ MaterialX::InputPtr UInterchangeMaterialXTranslator::GetStandardSurfaceInput(Mat
 	return Input;
 }
 
+MaterialX::InputPtr UInterchangeMaterialXTranslator::GetSurfaceUnlitInput(MaterialX::NodePtr SurfaceUnlit, const char* InputName, MaterialX::DocumentPtr Document) const
+{
+	mx::InputPtr Input = SurfaceUnlit->getInput(InputName);
+
+	if(!Input)
+	{
+		Input = Document->getNodeDef(mx::NodeDefinition::SurfaceUnlit)->getActiveInput(InputName);
+	}
+
+	return Input;
+}
+
 MaterialX::InputPtr UInterchangeMaterialXTranslator::GetPointLightInput(MaterialX::NodePtr PointLight, const char* InputName, MaterialX::DocumentPtr Document) const
 {
 	mx::InputPtr Input = PointLight->getInput(InputName);
@@ -2216,11 +2076,11 @@ FLinearColor UInterchangeMaterialXTranslator::GetVector(MaterialX::InputPtr Inpu
 	return LinearColor;
 }
 
-MaterialX::NodePtr UInterchangeMaterialXTranslator::CreateNode(MaterialX::NodeGraphPtr NodeGraph, const char* NodeName, const char* Category, TArray<FInputToCopy> InputsToCopy, TArray<FInputToCreate> InputsToCreate) const
+MaterialX::NodePtr UInterchangeMaterialXTranslator::CreateNode(MaterialX::ElementPtr NodeGraph, const char* NodeName, const char* Category, TArray<FInputToCopy> InputsToCopy, TArray<FInputToCreate> InputsToCreate) const
 {
 	std::string UniqueNodeName{ NodeName + std::string{"_"} + Category };
 
-	mx::NodePtr Node = NodeGraph->getNode(UniqueNodeName);
+	mx::NodePtr Node = NodeGraph->getChildOfType<mx::Node>(UniqueNodeName);
 
 	if(!Node)
 	{
@@ -2261,9 +2121,9 @@ MaterialX::NodePtr UInterchangeMaterialXTranslator::CreateNode(MaterialX::NodeGr
 	return Node;
 }
 
-void UInterchangeMaterialXTranslator::ConvertNeutralNodesToDot(MaterialX::NodeGraphPtr NodeGraph) const
+void UInterchangeMaterialXTranslator::ConvertNeutralNodesToDot(MaterialX::ElementPtr NodeGraph) const
 {
-	const std::vector<mx::NodePtr> Nodes = NodeGraph->getNodes();
+	const std::vector<mx::NodePtr> Nodes = NodeGraph->getChildrenOfType<mx::Node>();
 	for(mx::NodePtr Node : Nodes)
 	{
 		const std::string& Category = Node->getCategory();
@@ -2326,6 +2186,15 @@ void UInterchangeMaterialXTranslator::ConvertNeutralNodesToDot(MaterialX::NodeGr
 					{
 						IsValueStringNeutral(mx::Color4());
 					}
+
+					if(bIsNeutral)
+					{
+						Node->removeInput("in2");
+						if(mx::InputPtr Input1 = Node->getInput("in1"))
+						{
+							Input1->setName("in");
+						}
+					}
 				}
 			}
 
@@ -2337,22 +2206,59 @@ void UInterchangeMaterialXTranslator::ConvertNeutralNodesToDot(MaterialX::NodeGr
 	}
 }
 
-void UInterchangeMaterialXTranslator::AddTexCoordToTiledImageNodes(MaterialX::NodeGraphPtr NodeGraph) const
+void UInterchangeMaterialXTranslator::AddTexCoordToTiledImageNodes(MaterialX::ElementPtr Graph) const
 {
-	for(mx::NodePtr Node : NodeGraph->getNodes(mx::Category::TiledImage))
+	for(mx::NodePtr Node : Graph->getChildrenOfType<mx::Node>(mx::Category::TiledImage))
 	{
-		mx::InputPtr InputTexCoord = Node->getInput("texcoord");
-
-		if(!InputTexCoord)
+		if(Node->getCategory() == mx::Category::TiledImage)
 		{
-			mx::NodePtr NodeTexCoord = CreateNode(Node->getParent()->asA<mx::NodeGraph>(),
-												  "node_ue_texcoord",
-												  mx::Category::TexCoord,
-												  {},
-												  { {"index", FAttributeValueArray{{"type", "integer"}, {"value", "0"}}} });
+			mx::InputPtr InputTexCoord = Node->getInput("texcoord");
 
-			InputTexCoord = Node->addInput(mx::NodeGroup::Texture2D::Inputs::TexCoord, mx::Type::Vector2);
-			InputTexCoord->setNodeName(NodeTexCoord->getName());
+			if(!InputTexCoord)
+			{
+				mx::NodePtr NodeTexCoord = CreateNode(Node->getParent(),
+													  "node_ue_texcoord",
+													  mx::Category::TexCoord,
+													  {},
+													  { {"index", FAttributeValueArray{{"type", "integer"}, {"value", "0"}}} });
+
+				InputTexCoord = Node->addInput(mx::NodeGroup::Texture2D::Inputs::TexCoord, mx::Type::Vector2);
+				InputTexCoord->setNodeName(NodeTexCoord->getName());
+			}
+		}
+	}
+}
+
+void UInterchangeMaterialXTranslator::UpdateDocumentRecursively(MaterialX::GraphElementPtr Graph) const
+{
+	AddTexCoordToTiledImageNodes(Graph);
+
+	auto FilterNode = [](mx::NodePtr Node)
+	{
+		const std::string& Category = Node->getCategory();
+		return
+			Category != mx::Category::StandardSurface &&
+			Category != mx::Category::Extract &&
+			Category != mx::Category::Saturate &&
+			Category != mx::Category::Separate2 &&
+			Category != mx::Category::Separate3 &&
+			Category != mx::Category::Separate4;
+	};
+
+	//This function will replace all the nodes that are defined with a nodegraph with their equivalent node network
+	//This allows us to handle nodes defined with a nodegraph, with nodes that we already support while traversing the current node graph
+	//For example the node tiledimage is just a succession of the nodes image, multiply, add, etc.
+	Graph->flattenSubgraphs(mx::EMPTY_STRING, FilterNode);
+
+	ConvertNeutralNodesToDot(Graph);
+
+	mx::DocumentPtr Document = Graph->getDocument();
+	//we also need need to flatten the subgraphs
+	for(mx::NodeGraphPtr SubGraph : Graph->getChildrenOfType<mx::NodeGraph>())
+	{
+		if(SubGraph->getActiveSourceUri() == Document->getSourceUri())
+		{
+			UpdateDocumentRecursively(SubGraph);
 		}
 	}
 }
diff --git a/Engine/Plugins/Interchange/Runtime/Source/Import/Public/MaterialX/InterchangeMaterialXTranslator.h b/Engine/Plugins/Interchange/Runtime/Source/Import/Public/MaterialX/InterchangeMaterialXTranslator.h
--- a/Engine/Plugins/Interchange/Runtime/Source/Import/Public/MaterialX/InterchangeMaterialXTranslator.h
+++ b/Engine/Plugins/Interchange/Runtime/Source/Import/Public/MaterialX/InterchangeMaterialXTranslator.h
@@ -62,6 +62,15 @@ protected:
 	 */
 	void ProcessStandardSurface(UInterchangeBaseNodeContainer & NodeContainer, MaterialX::NodePtr StandardSurfaceNode, MaterialX::DocumentPtr Document) const;
 
+	/**
+	 * Process surface_unlit shader
+	 *
+	 * @param NodeContainer - The container where to add the translated Interchange nodes.
+	 * @param SurfaceUnlitNode - The <surface_unlit> in the MaterialX file
+	 * @param Document - The MaterialX Document that contains all the definitions of the loaded libraries
+	 */
+	void ProcessSurfaceUnlit(UInterchangeBaseNodeContainer& NodeContainer, MaterialX::NodePtr SurfaceUnlitNode, MaterialX::DocumentPtr Document) const;
+
 	/**
 	 * Process light shader, MaterialX doesn't standardized lights, but defines the 3 common ones, directional, point and spot
 	 * 
@@ -114,7 +123,7 @@ protected:
 	bool ConnectNodeGraphOutputToInput(MaterialX::InputPtr InputToNodeGraph, UInterchangeShaderNode * ShaderNode, const FString& ParentInputName, TMap<FString, UInterchangeShaderNode*> & NamesToShaderNodes, UInterchangeBaseNodeContainer & NodeContainer) const;
 
 	/**
-	 * Create and Connect the output of a node to a shader node
+	 * Create and Connect the output of a MaterialX node that has already a matching in UE to a shader node
 	 * 
 	 * @param Node - The MaterialX node of a given type used to create the appropriate shader node
 	 * @param ParentShaderNode - The shader node to connect to
@@ -126,6 +135,32 @@ protected:
 	 */
 	bool ConnectNodeOutputToInput(MaterialX::NodePtr Node, UInterchangeShaderNode* ParentShaderNode, const FString& InputChannelName, TMap<FString, UInterchangeShaderNode*>& NamesToShaderNodes, UInterchangeBaseNodeContainer& NodeContainer) const;
 
+	/**
+	 * Create and Connect manually the output of a MaterialX node to a shader node
+	 *
+	 * @param Edge - The MaterialX edge that has the current node, its parent and the input bridge between the two
+	 * @param ParentShaderNode - The shader node to connect to
+	 * @param InputChannelName - The input of the ParentShaderNode to connect to
+	 * @param NamesToShaderNodes - Map of the shader nodes already created
+	 * @param The container where to add the translated Interchange nodes.
+	 *
+	 * @return true if a shader node has been successfully created and is connected to the given input
+	 */
+	void ConnectNodeCategoryOutputToInput(const MaterialX::Edge & Edge, UInterchangeShaderNode* ParentShaderNode, const FString& InputChannelName, TMap<FString, UInterchangeShaderNode*>& NamesToShaderNodes, UInterchangeBaseNodeContainer& NodeContainer) const;
+
+	/**
+     * Create and Connect a node name directly connected from an input to a shader node
+     *
+     * @param Node - The MaterialX node of a given type used to create the appropriate shader node
+     * @param ParentShaderNode - The shader node to connect to
+     * @param InputChannelName - The input of the ParentShaderNode to connect to
+     * @param NamesToShaderNodes - Map of the shader nodes already created
+     * @param The container where to add the translated Interchange nodes.
+     *
+     * @return true if a shader node has been successfully created and is connected to the given input
+     */
+	bool ConnectNodeNameOutputToInput(MaterialX::InputPtr Input, UInterchangeShaderNode* ShaderNode, const FString& ParentInputName, TMap<FString, UInterchangeShaderNode*>& NamesToShaderNodes, UInterchangeBaseNodeContainer& NodeContainer) const;
+	
 	/**
 	 * Helper template function to create an InterchangeShaderNode
 	 * 
@@ -261,6 +296,19 @@ protected:
 	 */
 	MaterialX::InputPtr GetStandardSurfaceInput(MaterialX::NodePtr StandardSurface, const char* InputName, MaterialX::DocumentPtr Document) const;
 
+	/**
+	 * Retrieve the input from a surface_unlit node, or take the default input from the library,
+	 * this function should only be called after testing the MaterialX libraries have been successfully imported, meaning the node definition of the surface_unlit
+	 * should always be valid
+	 *
+	 * @param SurfaceUnlit - the <surface_unlit> node
+	 * @param InputName - the input name to retrieve
+	 * @param Document - the MaterialX library which has the node definition of the surface_unlit
+	 *
+	 * @return the input from the given name
+	 */
+	MaterialX::InputPtr GetSurfaceUnlitInput(MaterialX::NodePtr SurfaceUnlit, const char* InputName, MaterialX::DocumentPtr Document) const;
+
 	/**
 	 * Retrieve the input from a point_light node, or take the default input from the library,
 	 * this function should only be called after testing the MaterialX libraries have been successfully imported, meaning the node definition of the point_light
@@ -400,7 +448,7 @@ protected:
 	
 	/** <image> */
 	void ConnectImageInputToOutput(MaterialX::NodePtr UpstreamNode, UInterchangeShaderNode* ParentShaderNode, const FString& InputChannelName, TMap<FString, UInterchangeShaderNode*>& NamesToShaderNodes, UInterchangeBaseNodeContainer& NodeContainer) const;
-	
+
 	/** <convert> */
 	void ConnectConvertInputToOutput(MaterialX::NodePtr UpstreamNode, UInterchangeShaderNode* ParentShaderNode, const FString& InputChannelName, TMap<FString, UInterchangeShaderNode*>& NamesToShaderNodes, UInterchangeBaseNodeContainer& NodeContainer) const;
 	
@@ -493,7 +541,7 @@ protected:
 	 * 
 	 * @return The newly created node
 	 */
-	MaterialX::NodePtr CreateNode(MaterialX::NodeGraphPtr NodeGraph, const char * NodeName, const char * Category, TArray<FInputToCopy> InputsToCopy, TArray<FInputToCreate> InputsToCreate) const;
+	MaterialX::NodePtr CreateNode(MaterialX::ElementPtr NodeGraph, const char * NodeName, const char * Category, TArray<FInputToCopy> InputsToCopy, TArray<FInputToCreate> InputsToCreate) const;
 
 	/**
 	 * Convert math nodes (add, sub, div, mul) that have a neutral input, to the MaterialX <dot> node which is simply the identity or a noop
@@ -502,7 +550,7 @@ protected:
 	 * 
 	 * @param NodeGraph - The node graph to retrieve the different math nodes
 	 */
-	void ConvertNeutralNodesToDot(MaterialX::NodeGraphPtr NodeGraph) const;
+	void ConvertNeutralNodesToDot(MaterialX::ElementPtr NodeGraph) const;
 
 	/**
 	 * Add a texcoord input to the tiledimage nodes if the input is missing, and connect it to a texcoord node in the nodegraph
@@ -511,7 +559,14 @@ protected:
 	 * 
 	 * @param Node - The nodegraph to retrieve the different tiledimage nodes
 	 */
-	void AddTexCoordToTiledImageNodes(MaterialX::NodeGraphPtr NodeGraph) const;
+	void AddTexCoordToTiledImageNodes(MaterialX::ElementPtr Graph) const;
+
+	/**
+	 * Update a MaterialX Document recursively by either initializing nodes (like <tiledimage>) or updating some nodes (e.g: math nodes with neutral elements)
+	 * 
+	 * @param Graph - The graph to update and all of its sub graphs
+	 */
+	void UpdateDocumentRecursively(MaterialX::GraphElementPtr Graph) const;
 
 	/**
 	 * Retrieve the interchange parent name of a MaterialX node, useful when a node is a combination of several nodes connected to different inputs (look for example Noise3D)
diff --git a/Engine/Plugins/Interchange/Runtime/Source/Nodes/Public/InterchangeMaterialDefinitions.h b/Engine/Plugins/Interchange/Runtime/Source/Nodes/Public/InterchangeMaterialDefinitions.h
--- a/Engine/Plugins/Interchange/Runtime/Source/Nodes/Public/InterchangeMaterialDefinitions.h
+++ b/Engine/Plugins/Interchange/Runtime/Source/Nodes/Public/InterchangeMaterialDefinitions.h
@@ -406,6 +406,20 @@ namespace Interchange
 			}
 		}
 
+		namespace SurfaceUnlit
+		{
+			const FName Name = TEXT("surface_unlit");
+
+			namespace Parameters
+			{
+				const FName Emission = TEXT("emission");
+				const FName EmissionColor = TEXT("emission_color");
+				const FName Transmission = TEXT("transmission");
+				const FName TransmissionColor = TEXT("transmission_color");
+				const FName Opacity = TEXT("opacity");
+			}
+		}
+
 		namespace Unlit
 		{
 			namespace Parameters
diff --git a/Engine/Plugins/Interchange/Runtime/Source/Pipelines/Private/InterchangeGenericMaterialPipeline.cpp b/Engine/Plugins/Interchange/Runtime/Source/Pipelines/Private/InterchangeGenericMaterialPipeline.cpp
--- a/Engine/Plugins/Interchange/Runtime/Source/Pipelines/Private/InterchangeGenericMaterialPipeline.cpp
+++ b/Engine/Plugins/Interchange/Runtime/Source/Pipelines/Private/InterchangeGenericMaterialPipeline.cpp
@@ -81,6 +81,7 @@ FMaterialXPipelineSettings::FMaterialXPipelineSettings()
 {
 	SurfaceShader.FindOrAdd(EInterchangeMaterialXShaders::StandardSurface, FSoftObjectPath{ TEXT("MaterialFunction'/Interchange/Functions/MX_StandardSurface.MX_StandardSurface'") });
 	SurfaceShader.FindOrAdd(EInterchangeMaterialXShaders::StandardSurfaceTransmission, FSoftObjectPath{ TEXT("MaterialFunction'/Interchange/Functions/MX_TransmissionSurface.MX_TransmissionSurface'") });
+	SurfaceShader.FindOrAdd(EInterchangeMaterialXShaders::SurfaceUnlit, FSoftObjectPath{ TEXT("MaterialFunction'/Interchange/Functions/MX_SurfaceUnlit.MX_SurfaceUnlit'") });
 }
 
 bool FMaterialXPipelineSettings::AreRequiredPackagesLoaded()
@@ -115,6 +116,10 @@ bool FMaterialXPipelineSettings::AreRequiredPackagesLoaded()
 						{
 							bAllLoaded = !ShouldFilterAssets(Cast<UMaterialFunction>(Asset), TransmissionSurfaceInputs, TransmissionSurfaceOutputs);
 						}
+						else if(Pair.Get<EInterchangeMaterialXShaders>() == EInterchangeMaterialXShaders::SurfaceUnlit)
+						{
+							bAllLoaded = !ShouldFilterAssets(Cast<UMaterialFunction>(Asset), SurfaceUnlitInputs, SurfaceUnlitOutputs);
+						}
 
 					}
 #endif // WITH_EDITOR
@@ -243,6 +248,21 @@ TSet<FName> FMaterialXPipelineSettings::TransmissionSurfaceOutputs
 	UE::Interchange::Materials::ThinTranslucent::Parameters::TransmissionColor
 };
 
+TSet<FName> FMaterialXPipelineSettings::SurfaceUnlitInputs
+{
+	UE::Interchange::Materials::SurfaceUnlit::Parameters::Emission,
+	UE::Interchange::Materials::SurfaceUnlit::Parameters::EmissionColor,
+	UE::Interchange::Materials::SurfaceUnlit::Parameters::Transmission,
+	UE::Interchange::Materials::SurfaceUnlit::Parameters::TransmissionColor,
+	UE::Interchange::Materials::SurfaceUnlit::Parameters::Opacity
+};
+
+TSet<FName> FMaterialXPipelineSettings::SurfaceUnlitOutputs
+{
+	UE::Interchange::Materials::Common::Parameters::EmissiveColor,
+	UE::Interchange::Materials::Common::Parameters::Opacity,
+	TEXT("OpacityMask")
+};
 
 bool FMaterialXPipelineSettings::ShouldFilterAssets(UMaterialFunction* Asset, const TSet<FName>& Inputs, const TSet<FName>& Outputs)
 {
@@ -659,6 +679,20 @@ bool UInterchangeGenericMaterialPipeline::IsStandardSurfaceModel(const UIntercha
 	return false;
 }
 
+bool UInterchangeGenericMaterialPipeline::IsSurfaceUnlitModel(const UInterchangeShaderGraphNode* ShaderGraphNode) const
+{
+	using namespace UE::Interchange::Materials;
+	FString ShaderType;
+	ShaderGraphNode->GetCustomShaderType(ShaderType);
+
+	if(ShaderType == SurfaceUnlit::Name.ToString())
+	{
+		return true;
+	}
+
+	return false;
+}
+
 bool UInterchangeGenericMaterialPipeline::HandlePhongModel(const UInterchangeShaderGraphNode* ShaderGraphNode, UInterchangeMaterialFactoryNode* MaterialFactoryNode)
 {
 	using namespace UE::Interchange::Materials::Phong;
@@ -858,91 +892,104 @@ bool UInterchangeGenericMaterialPipeline::HandlePBRModel(const UInterchangeShade
 	return bShadingModelHandled;
 }
 
-bool UInterchangeGenericMaterialPipeline::HandleStandardSurfaceModel(const UInterchangeShaderGraphNode* ShaderGraphNode, UInterchangeMaterialFactoryNode* MaterialFactoryNode)
+FString UInterchangeGenericMaterialPipeline::HandleSurfaceShaderModel(const UInterchangeShaderGraphNode* ShaderGraphNode, UInterchangeMaterialFactoryNode* MaterialFactoryNode, const FSurfaceShader& SurfaceShader)
 {
-	if(IsStandardSurfaceModel(ShaderGraphNode))
-	{
-		UInterchangeMaterialExpressionFactoryNode* FunctionCallExpression = NewObject<UInterchangeMaterialExpressionFactoryNode>(BaseNodeContainer, NAME_None);
-		FunctionCallExpression->SetCustomExpressionClassName(UMaterialExpressionMaterialFunctionCall::StaticClass()->GetName());
-		FString FunctionCallExpressionUid = MaterialFactoryNode->GetUniqueID() + TEXT("StandardSurface");
-		FunctionCallExpression->InitializeNode(FunctionCallExpressionUid, TEXT("StandardSurface"), EInterchangeNodeContainerType::FactoryData);
+	UInterchangeMaterialExpressionFactoryNode* FunctionCallExpression = NewObject<UInterchangeMaterialExpressionFactoryNode>(BaseNodeContainer, NAME_None);
+	FunctionCallExpression->SetCustomExpressionClassName(UMaterialExpressionMaterialFunctionCall::StaticClass()->GetName());
+	FString FunctionCallExpressionUid = MaterialFactoryNode->GetUniqueID() + SurfaceShader.UniqueID;
+	FunctionCallExpression->InitializeNode(FunctionCallExpressionUid, SurfaceShader.UniqueID, EInterchangeNodeContainerType::FactoryData);
 
-		BaseNodeContainer->AddNode(FunctionCallExpression);
-		BaseNodeContainer->SetNodeParentUid(FunctionCallExpressionUid, MaterialFactoryNode->GetUniqueID());
+	BaseNodeContainer->AddNode(FunctionCallExpression);
+	BaseNodeContainer->SetNodeParentUid(FunctionCallExpressionUid, MaterialFactoryNode->GetUniqueID());
 
-		const FName MaterialFunctionMemberName = GET_MEMBER_NAME_CHECKED(UMaterialExpressionMaterialFunctionCall, MaterialFunction);
+	const FName MaterialFunctionMemberName = GET_MEMBER_NAME_CHECKED(UMaterialExpressionMaterialFunctionCall, MaterialFunction);
 
-		using namespace UE::Interchange::Materials;
-		FString StandardSurfaceOrTransmissionPath = MaterialXSettings.GetAssetPathString(EInterchangeMaterialXShaders::StandardSurface);
 
-		if(UInterchangeShaderPortsAPI::HasInput(ShaderGraphNode, StandardSurface::Parameters::Transmission))
+	FunctionCallExpression->AddStringAttribute(
+		MaterialFunctionMemberName,
+		SurfaceShader.Path);
+	FunctionCallExpression->AddApplyAndFillDelegates<FString>(MaterialFunctionMemberName, UMaterialExpressionMaterialFunctionCall::StaticClass(), MaterialFunctionMemberName);
+
+	auto ConnectMaterialExpressionOutputToInput = [&](const FString& InputName, EMaterialInputType InputType)
+	{
+		TGuardValue<EMaterialInputType> InputTypeBeingProcessedGuard(MaterialCreationContext.InputTypeBeingProcessed, InputType);
+
+		TTuple<UInterchangeMaterialExpressionFactoryNode*, FString> ExpressionFactoryNode =
+			CreateMaterialExpressionForInput(MaterialFactoryNode, ShaderGraphNode, InputName, FunctionCallExpression->GetUniqueID());
+
+		if(ExpressionFactoryNode.Get<0>())
 		{
-			StandardSurfaceOrTransmissionPath = MaterialXSettings.GetAssetPathString(EInterchangeMaterialXShaders::StandardSurfaceTransmission);
+			UInterchangeShaderPortsAPI::ConnectOuputToInput(FunctionCallExpression, InputName,
+															ExpressionFactoryNode.Get<0>()->GetUniqueID(), ExpressionFactoryNode.Get<1>());
 		}
+	};
 
-		FunctionCallExpression->AddStringAttribute(
-			MaterialFunctionMemberName,
-			StandardSurfaceOrTransmissionPath);
-		FunctionCallExpression->AddApplyAndFillDelegates<FString>(MaterialFunctionMemberName, UMaterialExpressionMaterialFunctionCall::StaticClass(), MaterialFunctionMemberName);
+	for(auto& Arguments : SurfaceShader.Parameters)
+	{
+		const FString& InputName = Arguments.Get<0>();
+		EMaterialInputType InputType = Arguments.Get<1>();
+		ConnectMaterialExpressionOutputToInput(InputName, InputType);
+	}
 
-		auto ConnectMaterialExpressionOutputToInput = [&](const FString & InputName, EMaterialInputType InputType)
-		{
-			TGuardValue<EMaterialInputType> InputTypeBeingProcessedGuard(MaterialCreationContext.InputTypeBeingProcessed, InputType);
+	return FunctionCallExpressionUid;
+}
 
-			TTuple<UInterchangeMaterialExpressionFactoryNode*, FString> ExpressionFactoryNode =
-				CreateMaterialExpressionForInput(MaterialFactoryNode, ShaderGraphNode, InputName, FunctionCallExpression->GetUniqueID());
+bool UInterchangeGenericMaterialPipeline::HandleStandardSurfaceModel(const UInterchangeShaderGraphNode* ShaderGraphNode, UInterchangeMaterialFactoryNode* MaterialFactoryNode)
+{
+	using namespace UE::Interchange::Materials;
 
-			if(ExpressionFactoryNode.Get<0>())
+	if(IsStandardSurfaceModel(ShaderGraphNode))
+	{
+		static FSurfaceShader StandardSurfaceShader
+		{
+			// Parameters
 			{
-				UInterchangeShaderPortsAPI::ConnectOuputToInput(FunctionCallExpression, InputName,
-																ExpressionFactoryNode.Get<0>()->GetUniqueID(), ExpressionFactoryNode.Get<1>());
-			}
+			{StandardSurface::Parameters::Base.ToString(), EMaterialInputType::Scalar},                          // Base
+			{StandardSurface::Parameters::BaseColor.ToString(), EMaterialInputType::Color},                      // BaseColor
+			{StandardSurface::Parameters::DiffuseRoughness.ToString(),  EMaterialInputType::Scalar},             // DiffuseRoughness
+			{StandardSurface::Parameters::Metalness.ToString(), EMaterialInputType::Scalar},                     // Metalness
+			{StandardSurface::Parameters::Specular.ToString(), EMaterialInputType::Scalar},                      // Specular
+			{StandardSurface::Parameters::SpecularRoughness.ToString(), EMaterialInputType::Scalar},             // SpecularRoughness
+			{StandardSurface::Parameters::SpecularIOR.ToString(), EMaterialInputType::Scalar},                   // SpecularIOR
+			{StandardSurface::Parameters::SpecularAnisotropy.ToString(), EMaterialInputType::Scalar},            // SpecularAnisotropy
+			{StandardSurface::Parameters::SpecularRotation.ToString(), EMaterialInputType::Scalar},              // SpecularRotation
+			{StandardSurface::Parameters::Subsurface.ToString(), EMaterialInputType::Scalar},                    // Subsurface
+			{StandardSurface::Parameters::SubsurfaceColor.ToString(), EMaterialInputType::Scalar},               // SubsurfaceColor
+			{StandardSurface::Parameters::SubsurfaceRadius.ToString(), EMaterialInputType::Scalar},              // SubsurfaceRadius
+			{StandardSurface::Parameters::SubsurfaceScale.ToString(), EMaterialInputType::Scalar},               // SubsurfaceScale
+			{StandardSurface::Parameters::Sheen.ToString(), EMaterialInputType::Scalar},                         // Sheen
+			{StandardSurface::Parameters::SheenColor.ToString(), EMaterialInputType::Color},                     // SheenColor
+			{StandardSurface::Parameters::SheenRoughness.ToString(), EMaterialInputType::Scalar},                // SheenRoughness
+			{StandardSurface::Parameters::Coat.ToString(), EMaterialInputType::Scalar},                          // Coat
+			{StandardSurface::Parameters::CoatColor.ToString(),  EMaterialInputType::Color},                     // CoatColor
+			{StandardSurface::Parameters::CoatNormal.ToString(),  EMaterialInputType::Vector},                   // CoatNormal
+			{StandardSurface::Parameters::CoatRoughness.ToString(), EMaterialInputType::Scalar},                 // CoatRoughness
+			{StandardSurface::Parameters::ThinFilmThickness.ToString(), EMaterialInputType::Scalar},             // ThinFilmThickness
+			{StandardSurface::Parameters::Emission.ToString(), EMaterialInputType::Scalar},                      // Emission
+			{StandardSurface::Parameters::EmissionColor.ToString(), EMaterialInputType::Color},                  // EmissionColor
+			{StandardSurface::Parameters::Normal.ToString(), EMaterialInputType::Vector},                        // Normal
+			{StandardSurface::Parameters::Tangent.ToString(), EMaterialInputType::Vector},                       // Tangent
+			{StandardSurface::Parameters::Transmission.ToString(),  EMaterialInputType::Scalar},                 // Transmission
+			{StandardSurface::Parameters::TransmissionColor.ToString(), EMaterialInputType::Color},              // TransmissionColor
+			{StandardSurface::Parameters::TransmissionDepth.ToString(), EMaterialInputType::Scalar},             // TransmissionDepth
+			{StandardSurface::Parameters::TransmissionScatter.ToString(), EMaterialInputType::Color},            // TransmissionScatter
+			{StandardSurface::Parameters::TransmissionScatterAnisotropy.ToString(), EMaterialInputType::Scalar}, // TransmissionScatterAnisotropy
+			{StandardSurface::Parameters::TransmissionDispersion.ToString(), EMaterialInputType::Scalar},        // TransmissionDispersion
+			{StandardSurface::Parameters::TransmissionExtraRoughness.ToString(), EMaterialInputType::Scalar},    // TransmissionExtraRoughness
+			},
+
+			// UniqueID
+			TEXT("StandardSurface")
 		};
 
-		using ArgumentsType = TTuple<FString, EMaterialInputType>;
-		
-		static ArgumentsType ArgumentsPerInputs[]
-		{
-			ArgumentsType{StandardSurface::Parameters::Base.ToString(), EMaterialInputType::Scalar},                          // Base
-			ArgumentsType{StandardSurface::Parameters::BaseColor.ToString(), EMaterialInputType::Color},                      // BaseColor
-			ArgumentsType{StandardSurface::Parameters::DiffuseRoughness.ToString(),  EMaterialInputType::Scalar},             // DiffuseRoughness
-			ArgumentsType{StandardSurface::Parameters::Metalness.ToString(), EMaterialInputType::Scalar},                     // Metalness
-			ArgumentsType{StandardSurface::Parameters::Specular.ToString(), EMaterialInputType::Scalar},                      // Specular
-			ArgumentsType{StandardSurface::Parameters::SpecularRoughness.ToString(), EMaterialInputType::Scalar},             // SpecularRoughness
-			ArgumentsType{StandardSurface::Parameters::SpecularIOR.ToString(), EMaterialInputType::Scalar},                   // SpecularIOR
-			ArgumentsType{StandardSurface::Parameters::SpecularAnisotropy.ToString(), EMaterialInputType::Scalar},            // SpecularAnisotropy
-			ArgumentsType{StandardSurface::Parameters::SpecularRotation.ToString(), EMaterialInputType::Scalar},              // SpecularRotation
-			ArgumentsType{StandardSurface::Parameters::Subsurface.ToString(), EMaterialInputType::Scalar},                    // Subsurface
-			ArgumentsType{StandardSurface::Parameters::SubsurfaceColor.ToString(), EMaterialInputType::Scalar},               // SubsurfaceColor
-			ArgumentsType{StandardSurface::Parameters::SubsurfaceRadius.ToString(), EMaterialInputType::Scalar},              // SubsurfaceRadius
-			ArgumentsType{StandardSurface::Parameters::SubsurfaceScale.ToString(), EMaterialInputType::Scalar},               // SubsurfaceScale
-			ArgumentsType{StandardSurface::Parameters::Sheen.ToString(), EMaterialInputType::Scalar},                         // Sheen
-			ArgumentsType{StandardSurface::Parameters::SheenColor.ToString(), EMaterialInputType::Color},                     // SheenColor
-			ArgumentsType{StandardSurface::Parameters::SheenRoughness.ToString(), EMaterialInputType::Scalar},                // SheenRoughness
-			ArgumentsType{StandardSurface::Parameters::Coat.ToString(), EMaterialInputType::Scalar},                          // Coat
-			ArgumentsType{StandardSurface::Parameters::CoatColor.ToString(),  EMaterialInputType::Color},                     // CoatColor
-			ArgumentsType{StandardSurface::Parameters::CoatNormal.ToString(),  EMaterialInputType::Vector},                   // CoatNormal
-			ArgumentsType{StandardSurface::Parameters::CoatRoughness.ToString(), EMaterialInputType::Scalar},                 // CoatRoughness
-			ArgumentsType{StandardSurface::Parameters::ThinFilmThickness.ToString(), EMaterialInputType::Scalar},             // ThinFilmThickness
-			ArgumentsType{StandardSurface::Parameters::Emission.ToString(), EMaterialInputType::Scalar},                      // Emission
-			ArgumentsType{StandardSurface::Parameters::EmissionColor.ToString(), EMaterialInputType::Color},                  // EmissionColor
-			ArgumentsType{StandardSurface::Parameters::Normal.ToString(), EMaterialInputType::Vector},                        // Normal
-			ArgumentsType{StandardSurface::Parameters::Tangent.ToString(), EMaterialInputType::Vector},                       // Tangent
-			ArgumentsType{StandardSurface::Parameters::Transmission.ToString(),  EMaterialInputType::Scalar},                 // Transmission
-			ArgumentsType{StandardSurface::Parameters::TransmissionColor.ToString(), EMaterialInputType::Color},              // TransmissionColor
-			ArgumentsType{StandardSurface::Parameters::TransmissionDepth.ToString(), EMaterialInputType::Scalar},             // TransmissionDepth
-			ArgumentsType{StandardSurface::Parameters::TransmissionScatter.ToString(), EMaterialInputType::Color},            // TransmissionScatter
-			ArgumentsType{StandardSurface::Parameters::TransmissionScatterAnisotropy.ToString(), EMaterialInputType::Scalar}, // TransmissionScatterAnisotropy
-			ArgumentsType{StandardSurface::Parameters::TransmissionDispersion.ToString(), EMaterialInputType::Scalar},        // TransmissionDispersion
-			ArgumentsType{StandardSurface::Parameters::TransmissionExtraRoughness.ToString(), EMaterialInputType::Scalar},    // TransmissionExtraRoughness
-		};
+		StandardSurfaceShader.Path = MaterialXSettings.GetAssetPathString(EInterchangeMaterialXShaders::StandardSurface);
 
-		for(const ArgumentsType & Arguments : ArgumentsPerInputs)
+		if(UInterchangeShaderPortsAPI::HasInput(ShaderGraphNode, StandardSurface::Parameters::Transmission))
 		{
-			const FString& InputName = Arguments.Get<0>();
-			EMaterialInputType InputType = Arguments.Get<1>();
-			ConnectMaterialExpressionOutputToInput(InputName, InputType);
+			StandardSurfaceShader.Path = MaterialXSettings.GetAssetPathString(EInterchangeMaterialXShaders::StandardSurfaceTransmission);
 		}
+
+		const FString FunctionCallExpressionUid = HandleSurfaceShaderModel(ShaderGraphNode, MaterialFactoryNode, StandardSurfaceShader);
 		
 		MaterialFactoryNode->ConnectOutputToBaseColor(FunctionCallExpressionUid,  PBR::Parameters::BaseColor.ToString());
 		MaterialFactoryNode->ConnectOutputToMetallic(FunctionCallExpressionUid, PBR::Parameters::Metallic.ToString());
@@ -994,6 +1041,62 @@ bool UInterchangeGenericMaterialPipeline::HandleStandardSurfaceModel(const UInte
 	return false;
 }
 
+bool UInterchangeGenericMaterialPipeline::HandleSurfaceUnlitModel(const UInterchangeShaderGraphNode* ShaderGraphNode, UInterchangeMaterialFactoryNode* MaterialFactoryNode)
+{
+	using namespace UE::Interchange::Materials;
+	using namespace UE::Interchange::InterchangeGenericMaterialPipeline;
+
+	if(IsSurfaceUnlitModel(ShaderGraphNode))
+	{
+		static FSurfaceShader StandardSurfaceShader
+		{
+			// Parameters
+			{
+			{SurfaceUnlit::Parameters::Emission.ToString(), EMaterialInputType::Scalar},                      // Emission
+			{SurfaceUnlit::Parameters::EmissionColor.ToString(), EMaterialInputType::Color},                  // EmissionColor
+			{SurfaceUnlit::Parameters::Opacity.ToString(), EMaterialInputType::Scalar},                       // Opacity
+			{SurfaceUnlit::Parameters::Transmission.ToString(),  EMaterialInputType::Scalar},                 // Transmission
+			{SurfaceUnlit::Parameters::TransmissionColor.ToString(), EMaterialInputType::Color},              // TransmissionColor
+			},
+
+			// UniqueID
+			TEXT("SurfaceUnlit"),
+
+			// Path
+			MaterialXSettings.GetAssetPathString(EInterchangeMaterialXShaders::SurfaceUnlit)
+		};
+
+		const FString FunctionCallExpressionUid = HandleSurfaceShaderModel(ShaderGraphNode, MaterialFactoryNode, StandardSurfaceShader);
+
+		MaterialFactoryNode->ConnectOutputToEmissiveColor(FunctionCallExpressionUid, PBR::Parameters::EmissiveColor.ToString());
+		
+		//We can't have both Opacity and Opacity Mask we need to make a choice
+		bool bHasOpacityOrTransmission = false;
+		if(UInterchangeShaderPortsAPI::HasInput(ShaderGraphNode, SurfaceUnlit::Parameters::Transmission))
+		{
+			MaterialFactoryNode->ConnectOutputToOpacity(FunctionCallExpressionUid, PBR::Parameters::Opacity.ToString());
+			bHasOpacityOrTransmission = true;
+		}
+		else if(UInterchangeShaderPortsAPI::HasInput(ShaderGraphNode, SurfaceUnlit::Parameters::Opacity))
+		{
+			MaterialFactoryNode->ConnectOutputToOpacity(FunctionCallExpressionUid, "OpacityMask");
+			MaterialFactoryNode->SetCustomOpacityMaskClipValue(0.f); // just to connect to the opacity mask
+			bHasOpacityOrTransmission = true;
+		}
+
+		if(bHasOpacityOrTransmission)
+		{
+			Private::UpdateBlendModeBasedOnOpacityAttributes(ShaderGraphNode, MaterialFactoryNode);
+		}
+
+		MaterialFactoryNode->SetCustomShadingModel(EMaterialShadingModel::MSM_Unlit);
+
+		return true;
+	}
+
+	return false;
+}
+
 bool UInterchangeGenericMaterialPipeline::HandleClearCoat(const UInterchangeShaderGraphNode* ShaderGraphNode, UInterchangeMaterialFactoryNode* MaterialFactoryNode)
 {
 	using namespace UE::Interchange::Materials::ClearCoat;
@@ -2397,6 +2500,11 @@ UInterchangeMaterialFactoryNode* UInterchangeGenericMaterialPipeline::CreateMate
 		return MaterialFactoryNode;
 	}
 
+	if(HandleSurfaceUnlitModel(ShaderGraphNode, MaterialFactoryNode))
+	{
+		return MaterialFactoryNode;
+	}
+
 	// Handle the case where the material will be connected through the material attributes input
 	if (HandleBxDFInput(ShaderGraphNode, MaterialFactoryNode))
 	{
diff --git a/Engine/Plugins/Interchange/Runtime/Source/Pipelines/Public/InterchangeGenericMaterialPipeline.h b/Engine/Plugins/Interchange/Runtime/Source/Pipelines/Public/InterchangeGenericMaterialPipeline.h
--- a/Engine/Plugins/Interchange/Runtime/Source/Pipelines/Public/InterchangeGenericMaterialPipeline.h
+++ b/Engine/Plugins/Interchange/Runtime/Source/Pipelines/Public/InterchangeGenericMaterialPipeline.h
@@ -38,8 +38,11 @@ enum class EInterchangeMaterialXShaders : uint8
 	/** Default settings for Autodesk's Standard Surface shader	*/
 	StandardSurface,
 
-	/** Standard Surface shader use for translucency	*/
-	StandardSurfaceTransmission
+	/** Standard Surface shader used for translucency	*/
+	StandardSurfaceTransmission,
+
+	/** Shader used for unlit surface*/
+	SurfaceUnlit
 };
 
 USTRUCT(BlueprintType)
@@ -68,6 +71,8 @@ private:
 	static TSet<FName> StandardSurfaceOutputs;
 	static TSet<FName> TransmissionSurfaceInputs;
 	static TSet<FName> TransmissionSurfaceOutputs;
+	static TSet<FName> SurfaceUnlitInputs;
+	static TSet<FName> SurfaceUnlitOutputs;
 #endif // WITH_EDITOR
 };
 
@@ -149,6 +154,9 @@ private:
 	/** True if the shader graph has the standard surface's shader type name */
 	bool IsStandardSurfaceModel(const UInterchangeShaderGraphNode* ShaderGraphNode) const;
 
+	/** True if the shader graph has the surface unlit's shader type name */
+	bool IsSurfaceUnlitModel(const UInterchangeShaderGraphNode* ShaderGraphNode) const;
+
 	/** True if the shader graph has an unlit color input. */
 	bool IsUnlitModel(const UInterchangeShaderGraphNode* ShaderGraphNode) const;
 
@@ -156,6 +164,7 @@ private:
 	bool HandleLambertModel(const UInterchangeShaderGraphNode* ShaderGraphNode, UInterchangeMaterialFactoryNode* MaterialFactoryNode);
 	bool HandlePBRModel(const UInterchangeShaderGraphNode* ShaderGraphNode, UInterchangeMaterialFactoryNode* MaterialFactoryNode);
 	bool HandleStandardSurfaceModel(const UInterchangeShaderGraphNode* ShaderGraphNode, UInterchangeMaterialFactoryNode* MaterialFactoryNode);
+	bool HandleSurfaceUnlitModel(const UInterchangeShaderGraphNode* ShaderGraphNode, UInterchangeMaterialFactoryNode* MaterialFactoryNode);
 	bool HandleClearCoat(const UInterchangeShaderGraphNode* ShaderGraphNode, UInterchangeMaterialFactoryNode* MaterialFactoryNode);
 	bool HandleSubsurface(const UInterchangeShaderGraphNode* ShaderGraphNode, UInterchangeMaterialFactoryNode* MaterialFactoryNode);
 	bool HandleSheen(const UInterchangeShaderGraphNode* ShaderGraphNode, UInterchangeMaterialFactoryNode* MaterialFactoryNode);
@@ -227,4 +236,15 @@ private:
 	};
 
 	TArray<FMaterialExpressionCreationContext> MaterialExpressionCreationContextStack;
+
+	using FParameterMaterialInputType = TTuple<FString, EMaterialInputType>;
+
+	struct FSurfaceShader
+	{
+		TArray<FParameterMaterialInputType> Parameters;
+		FString UniqueID;
+		FString Path;
+	};
+
+	FString HandleSurfaceShaderModel(const UInterchangeShaderGraphNode* ShaderGraphNode, UInterchangeMaterialFactoryNode* MaterialFactoryNode, const FSurfaceShader& SurfaceShader);
 };
