[Interchange] Output log contains errors for Module : InterchangeNodes fix

diff --git a/Engine/Plugins/Interchange/Runtime/Source/Import/Private/Fbx/InterchangeFbxTranslator.cpp b/Engine/Plugins/Interchange/Runtime/Source/Import/Private/Fbx/InterchangeFbxTranslator.cpp
--- a/Engine/Plugins/Interchange/Runtime/Source/Import/Private/Fbx/InterchangeFbxTranslator.cpp
+++ b/Engine/Plugins/Interchange/Runtime/Source/Import/Private/Fbx/InterchangeFbxTranslator.cpp
@@ -260,12 +260,13 @@ TFuture<TOptional<UE::Interchange::FMeshPayloadData>> UInterchangeFbxTranslator:
 			}
 			break;
 		case EInterchangeMeshPayLoadType::MORPHTARGET:
-		{
-			//Buffer keep the ownership of the data, the large memory reader is use to serialize the TMap
-			FLargeMemoryReader Ar(FileData, FileDataSize);
-			MeshPayloadData.MeshDescription.Serialize(Ar);
-		}
+			{
+				//Buffer keep the ownership of the data, the large memory reader is use to serialize the TMap
+				FLargeMemoryReader Ar(FileData, FileDataSize);
+				MeshPayloadData.MeshDescription.Serialize(Ar);
+			}
 			break;
+		case EInterchangeMeshPayLoadType::NONE:
 		default:
 			break;
 		}
diff --git a/Engine/Plugins/Interchange/Runtime/Source/Import/Private/Gltf/InterchangeGltfTranslator.cpp b/Engine/Plugins/Interchange/Runtime/Source/Import/Private/Gltf/InterchangeGltfTranslator.cpp
--- a/Engine/Plugins/Interchange/Runtime/Source/Import/Private/Gltf/InterchangeGltfTranslator.cpp
+++ b/Engine/Plugins/Interchange/Runtime/Source/Import/Private/Gltf/InterchangeGltfTranslator.cpp
@@ -1910,6 +1910,7 @@ TFuture< TOptional< UE::Interchange::FMeshPayloadData > > UInterchangeGltfTransl
 				//GLTF handles morph targets as simple Meshes
 				bSuccessfullAcquisition = UE::Interchange::Gltf::Private::GetStaticMeshPayloadDataForPayLoadKey(GltfAsset, PayLoadKey.UniqueId, MeshPayLoadData.MeshDescription);
 				break;
+			case EInterchangeMeshPayLoadType::NONE:
 			default:
 				break;
 			}
diff --git a/Engine/Plugins/Interchange/Runtime/Source/Nodes/Public/InterchangeMeshNode.h b/Engine/Plugins/Interchange/Runtime/Source/Nodes/Public/InterchangeMeshNode.h
--- a/Engine/Plugins/Interchange/Runtime/Source/Nodes/Public/InterchangeMeshNode.h
+++ b/Engine/Plugins/Interchange/Runtime/Source/Nodes/Public/InterchangeMeshNode.h
@@ -32,7 +32,8 @@ namespace UE
 UENUM(BlueprintType)
 enum class EInterchangeMeshPayLoadType : uint8
 {
-	STATIC = 0,
+	NONE = 0,
+	STATIC,
 	SKELETAL,
 	MORPHTARGET
 };
@@ -46,10 +47,10 @@ struct INTERCHANGENODES_API FInterchangeMeshPayLoadKey
 public:
 
 	UPROPERTY(EditAnywhere, BlueprintReadWrite, Category = "Interchange | Mesh")
-	FString UniqueId;
+	FString UniqueId = "";
 
 	UPROPERTY(EditAnywhere, BlueprintReadWrite, Category = "Interchange | Mesh")
-	EInterchangeMeshPayLoadType Type;
+	EInterchangeMeshPayLoadType Type = EInterchangeMeshPayLoadType::NONE;
 
 	FInterchangeMeshPayLoadKey() {}
 
