[Interchange] Asset import of glTF that contains node animation will create an empty level sequence

diff --git a/Engine/Plugins/Interchange/Runtime/Source/Import/Private/Animation/InterchangeAnimationTrackSetFactory.cpp b/Engine/Plugins/Interchange/Runtime/Source/Import/Private/Animation/InterchangeAnimationTrackSetFactory.cpp
--- a/Engine/Plugins/Interchange/Runtime/Source/Import/Private/Animation/InterchangeAnimationTrackSetFactory.cpp
+++ b/Engine/Plugins/Interchange/Runtime/Source/Import/Private/Animation/InterchangeAnimationTrackSetFactory.cpp
@@ -40,6 +40,99 @@
 #if WITH_EDITOR
 namespace UE::Interchange::Private
 {
+	AActor* GetActor(const UInterchangeBaseNodeContainer* NodeContainer, const UInterchangeAnimationTrackNode* TrackNode)
+	{
+		AActor* Actor = nullptr;
+
+		FString ActorNodeUid;
+		if (TrackNode->GetCustomActorDependencyUid(ActorNodeUid))
+		{
+			const FString ActorFactoryNodeUid = TEXT("Factory_") + ActorNodeUid;
+			const UInterchangeFactoryBaseNode* ActorFactoryNode = Cast<UInterchangeFactoryBaseNode>(NodeContainer->GetNode(ActorFactoryNodeUid));
+
+			if (ActorFactoryNode)
+			{
+				FSoftObjectPath ReferenceObject;
+				ActorFactoryNode->GetCustomReferenceObject(ReferenceObject);
+				if (ReferenceObject.IsValid())
+				{
+					Actor = Cast<AActor>(ReferenceObject.TryLoad());
+				}
+			}
+		}
+
+		return Actor;
+	}
+
+	bool HasActorToUse(const UInterchangeBaseNodeContainer* NodeContainer, const UInterchangeAnimationTrackSetFactoryNode* FactoryNode)
+	{
+
+		TArray<FString> AnimationTrackUids;
+		FactoryNode->GetCustomAnimationTrackUids(AnimationTrackUids);
+
+		for (const FString& AnimationTrackUid : AnimationTrackUids)
+		{
+			if (const UInterchangeBaseNode* TranslatedNode = NodeContainer->GetNode(AnimationTrackUid))
+			{
+				if (const UInterchangeTransformAnimationTrackNode* TransformTrackNode = Cast<UInterchangeTransformAnimationTrackNode>(TranslatedNode))
+				{
+					AActor* Actor = GetActor(NodeContainer, TransformTrackNode);
+					if (Actor)
+					{
+						return true;
+					}
+				}
+				else if (const UInterchangeAnimationTrackSetInstanceNode* InstanceTrackNode = Cast<UInterchangeAnimationTrackSetInstanceNode>(TranslatedNode))
+				{
+					FString TrackSetNodeUid;
+					if (!InstanceTrackNode->GetCustomTrackSetDependencyUid(TrackSetNodeUid))
+					{
+						continue;
+					}
+
+					const FString TrackSetFactoryNodeUid = TEXT("Factory_") + TrackSetNodeUid;
+					const UInterchangeAnimationTrackSetFactoryNode* InstanceFactoryNode = Cast<UInterchangeAnimationTrackSetFactoryNode>(NodeContainer->GetNode(TrackSetFactoryNodeUid));
+
+					if (!InstanceFactoryNode)
+					{
+						continue;
+					}
+					FSoftObjectPath ReferenceObject;
+					InstanceFactoryNode->GetCustomReferenceObject(ReferenceObject);
+					if (!ReferenceObject.IsValid())
+					{
+						continue;
+					}
+
+					return true;
+				}
+				else if (const UInterchangeAnimationTrackNode* TrackNode = Cast<UInterchangeAnimationTrackNode>(TranslatedNode))
+				{
+					int32 TargetedProperty;
+					if (!TrackNode->GetCustomTargetedProperty(TargetedProperty))
+					{
+						continue;
+					}
+
+					// Only visibility is supported for the time being
+					if (TargetedProperty != (int32)EInterchangeAnimatedProperty::Visibility)
+					{
+						continue;
+					}
+
+					// Get targeted actor exists
+					AActor* Actor = GetActor(NodeContainer, TrackNode);
+					if (Actor)
+					{
+						return true;
+					}
+				}
+			}
+		}
+
+		return false;
+	}
+
 	class FAnimationTrackSetHelper
 	{
 	public:
@@ -590,6 +683,12 @@ UObject* UInterchangeAnimationTrackSetFactory::ImportAssetObject_GameThread(cons
 		return nullptr;
 	}
 
+	if (!UE::Interchange::Private::HasActorToUse(Arguments.NodeContainer, FactoryNode))
+	{
+		UE_LOG(LogInterchangeImport, Warning, TEXT("Level sequence asset, %s, not imported, because all referenced actors are missing."), *FactoryNode->GetDisplayLabel());
+		return nullptr;
+	}
+
 	// create an asset if it doesn't exist
 	UObject* ExistingAsset = StaticFindObject(nullptr, Arguments.Parent, *Arguments.AssetName);
 
