Moving MaterialX material expressions in a subfolder

diff --git a/Engine/Plugins/Interchange/Runtime/Source/Import/Private/MaterialX/MaterialExpressionAppend3Vector.cpp b/Engine/Plugins/Interchange/Runtime/Source/Import/Private/MaterialX/MaterialExpressionAppend3Vector.cpp
deleted file mode 100644
diff --git a/Engine/Plugins/Interchange/Runtime/Source/Import/Private/MaterialX/MaterialExpressionAppend3Vector.h b/Engine/Plugins/Interchange/Runtime/Source/Import/Private/MaterialX/MaterialExpressionAppend3Vector.h
deleted file mode 100644
diff --git a/Engine/Plugins/Interchange/Runtime/Source/Import/Private/MaterialX/MaterialExpressionAppend4Vector.cpp b/Engine/Plugins/Interchange/Runtime/Source/Import/Private/MaterialX/MaterialExpressionAppend4Vector.cpp
deleted file mode 100644
diff --git a/Engine/Plugins/Interchange/Runtime/Source/Import/Private/MaterialX/MaterialExpressionAppend4Vector.h b/Engine/Plugins/Interchange/Runtime/Source/Import/Private/MaterialX/MaterialExpressionAppend4Vector.h
deleted file mode 100644
diff --git a/Engine/Plugins/Interchange/Runtime/Source/Import/Private/MaterialX/MaterialExpressionBurn.cpp b/Engine/Plugins/Interchange/Runtime/Source/Import/Private/MaterialX/MaterialExpressionBurn.cpp
deleted file mode 100644
diff --git a/Engine/Plugins/Interchange/Runtime/Source/Import/Private/MaterialX/MaterialExpressionBurn.h b/Engine/Plugins/Interchange/Runtime/Source/Import/Private/MaterialX/MaterialExpressionBurn.h
deleted file mode 100644
diff --git a/Engine/Plugins/Interchange/Runtime/Source/Import/Private/MaterialX/MaterialExpressionDifference.cpp b/Engine/Plugins/Interchange/Runtime/Source/Import/Private/MaterialX/MaterialExpressionDifference.cpp
deleted file mode 100644
diff --git a/Engine/Plugins/Interchange/Runtime/Source/Import/Private/MaterialX/MaterialExpressionDifference.h b/Engine/Plugins/Interchange/Runtime/Source/Import/Private/MaterialX/MaterialExpressionDifference.h
deleted file mode 100644
diff --git a/Engine/Plugins/Interchange/Runtime/Source/Import/Private/MaterialX/MaterialExpressionDisjointOver.cpp b/Engine/Plugins/Interchange/Runtime/Source/Import/Private/MaterialX/MaterialExpressionDisjointOver.cpp
deleted file mode 100644
diff --git a/Engine/Plugins/Interchange/Runtime/Source/Import/Private/MaterialX/MaterialExpressionDisjointOver.h b/Engine/Plugins/Interchange/Runtime/Source/Import/Private/MaterialX/MaterialExpressionDisjointOver.h
deleted file mode 100644
diff --git a/Engine/Plugins/Interchange/Runtime/Source/Import/Private/MaterialX/MaterialExpressionDodge.cpp b/Engine/Plugins/Interchange/Runtime/Source/Import/Private/MaterialX/MaterialExpressionDodge.cpp
deleted file mode 100644
diff --git a/Engine/Plugins/Interchange/Runtime/Source/Import/Private/MaterialX/MaterialExpressionDodge.h b/Engine/Plugins/Interchange/Runtime/Source/Import/Private/MaterialX/MaterialExpressionDodge.h
deleted file mode 100644
diff --git a/Engine/Plugins/Interchange/Runtime/Source/Import/Private/MaterialX/MaterialExpressionExponential.cpp b/Engine/Plugins/Interchange/Runtime/Source/Import/Private/MaterialX/MaterialExpressionExponential.cpp
deleted file mode 100644
diff --git a/Engine/Plugins/Interchange/Runtime/Source/Import/Private/MaterialX/MaterialExpressionExponential.h b/Engine/Plugins/Interchange/Runtime/Source/Import/Private/MaterialX/MaterialExpressionExponential.h
deleted file mode 100644
diff --git a/Engine/Plugins/Interchange/Runtime/Source/Import/Private/MaterialX/MaterialExpressionFractal3D.cpp b/Engine/Plugins/Interchange/Runtime/Source/Import/Private/MaterialX/MaterialExpressionFractal3D.cpp
deleted file mode 100644
diff --git a/Engine/Plugins/Interchange/Runtime/Source/Import/Private/MaterialX/MaterialExpressionFractal3D.h b/Engine/Plugins/Interchange/Runtime/Source/Import/Private/MaterialX/MaterialExpressionFractal3D.h
deleted file mode 100644
diff --git a/Engine/Plugins/Interchange/Runtime/Source/Import/Private/MaterialX/MaterialExpressionHsvToRgb.cpp b/Engine/Plugins/Interchange/Runtime/Source/Import/Private/MaterialX/MaterialExpressionHsvToRgb.cpp
deleted file mode 100644
diff --git a/Engine/Plugins/Interchange/Runtime/Source/Import/Private/MaterialX/MaterialExpressionHsvToRgb.h b/Engine/Plugins/Interchange/Runtime/Source/Import/Private/MaterialX/MaterialExpressionHsvToRgb.h
deleted file mode 100644
diff --git a/Engine/Plugins/Interchange/Runtime/Source/Import/Private/MaterialX/MaterialExpressionIn.cpp b/Engine/Plugins/Interchange/Runtime/Source/Import/Private/MaterialX/MaterialExpressionIn.cpp
deleted file mode 100644
diff --git a/Engine/Plugins/Interchange/Runtime/Source/Import/Private/MaterialX/MaterialExpressionIn.h b/Engine/Plugins/Interchange/Runtime/Source/Import/Private/MaterialX/MaterialExpressionIn.h
deleted file mode 100644
diff --git a/Engine/Plugins/Interchange/Runtime/Source/Import/Private/MaterialX/MaterialExpressionLength.cpp b/Engine/Plugins/Interchange/Runtime/Source/Import/Private/MaterialX/MaterialExpressionLength.cpp
deleted file mode 100644
diff --git a/Engine/Plugins/Interchange/Runtime/Source/Import/Private/MaterialX/MaterialExpressionLength.h b/Engine/Plugins/Interchange/Runtime/Source/Import/Private/MaterialX/MaterialExpressionLength.h
deleted file mode 100644
diff --git a/Engine/Plugins/Interchange/Runtime/Source/Import/Private/MaterialX/MaterialExpressionLogarithm.cpp b/Engine/Plugins/Interchange/Runtime/Source/Import/Private/MaterialX/MaterialExpressionLogarithm.cpp
deleted file mode 100644
diff --git a/Engine/Plugins/Interchange/Runtime/Source/Import/Private/MaterialX/MaterialExpressionLogarithm.h b/Engine/Plugins/Interchange/Runtime/Source/Import/Private/MaterialX/MaterialExpressionLogarithm.h
deleted file mode 100644
diff --git a/Engine/Plugins/Interchange/Runtime/Source/Import/Private/MaterialX/MaterialExpressionLuminance.cpp b/Engine/Plugins/Interchange/Runtime/Source/Import/Private/MaterialX/MaterialExpressionLuminance.cpp
deleted file mode 100644
diff --git a/Engine/Plugins/Interchange/Runtime/Source/Import/Private/MaterialX/MaterialExpressionLuminance.h b/Engine/Plugins/Interchange/Runtime/Source/Import/Private/MaterialX/MaterialExpressionLuminance.h
deleted file mode 100644
diff --git a/Engine/Plugins/Interchange/Runtime/Source/Import/Private/MaterialX/MaterialExpressionMask.cpp b/Engine/Plugins/Interchange/Runtime/Source/Import/Private/MaterialX/MaterialExpressionMask.cpp
deleted file mode 100644
diff --git a/Engine/Plugins/Interchange/Runtime/Source/Import/Private/MaterialX/MaterialExpressionMask.h b/Engine/Plugins/Interchange/Runtime/Source/Import/Private/MaterialX/MaterialExpressionMask.h
deleted file mode 100644
diff --git a/Engine/Plugins/Interchange/Runtime/Source/Import/Private/MaterialX/MaterialExpressionMatte.cpp b/Engine/Plugins/Interchange/Runtime/Source/Import/Private/MaterialX/MaterialExpressionMatte.cpp
deleted file mode 100644
diff --git a/Engine/Plugins/Interchange/Runtime/Source/Import/Private/MaterialX/MaterialExpressionMatte.h b/Engine/Plugins/Interchange/Runtime/Source/Import/Private/MaterialX/MaterialExpressionMatte.h
deleted file mode 100644
diff --git a/Engine/Plugins/Interchange/Runtime/Source/Import/Private/MaterialX/MaterialExpressionMinus.cpp b/Engine/Plugins/Interchange/Runtime/Source/Import/Private/MaterialX/MaterialExpressionMinus.cpp
deleted file mode 100644
diff --git a/Engine/Plugins/Interchange/Runtime/Source/Import/Private/MaterialX/MaterialExpressionMinus.h b/Engine/Plugins/Interchange/Runtime/Source/Import/Private/MaterialX/MaterialExpressionMinus.h
deleted file mode 100644
diff --git a/Engine/Plugins/Interchange/Runtime/Source/Import/Private/MaterialX/MaterialExpressionOut.cpp b/Engine/Plugins/Interchange/Runtime/Source/Import/Private/MaterialX/MaterialExpressionOut.cpp
deleted file mode 100644
diff --git a/Engine/Plugins/Interchange/Runtime/Source/Import/Private/MaterialX/MaterialExpressionOut.h b/Engine/Plugins/Interchange/Runtime/Source/Import/Private/MaterialX/MaterialExpressionOut.h
deleted file mode 100644
diff --git a/Engine/Plugins/Interchange/Runtime/Source/Import/Private/MaterialX/MaterialExpressionOver.cpp b/Engine/Plugins/Interchange/Runtime/Source/Import/Private/MaterialX/MaterialExpressionOver.cpp
deleted file mode 100644
diff --git a/Engine/Plugins/Interchange/Runtime/Source/Import/Private/MaterialX/MaterialExpressionOver.h b/Engine/Plugins/Interchange/Runtime/Source/Import/Private/MaterialX/MaterialExpressionOver.h
deleted file mode 100644
diff --git a/Engine/Plugins/Interchange/Runtime/Source/Import/Private/MaterialX/MaterialExpressionOverlay.cpp b/Engine/Plugins/Interchange/Runtime/Source/Import/Private/MaterialX/MaterialExpressionOverlay.cpp
deleted file mode 100644
diff --git a/Engine/Plugins/Interchange/Runtime/Source/Import/Private/MaterialX/MaterialExpressionOverlay.h b/Engine/Plugins/Interchange/Runtime/Source/Import/Private/MaterialX/MaterialExpressionOverlay.h
deleted file mode 100644
diff --git a/Engine/Plugins/Interchange/Runtime/Source/Import/Private/MaterialX/MaterialExpressionPlus.cpp b/Engine/Plugins/Interchange/Runtime/Source/Import/Private/MaterialX/MaterialExpressionPlus.cpp
deleted file mode 100644
diff --git a/Engine/Plugins/Interchange/Runtime/Source/Import/Private/MaterialX/MaterialExpressionPlus.h b/Engine/Plugins/Interchange/Runtime/Source/Import/Private/MaterialX/MaterialExpressionPlus.h
deleted file mode 100644
diff --git a/Engine/Plugins/Interchange/Runtime/Source/Import/Private/MaterialX/MaterialExpressionPremult.cpp b/Engine/Plugins/Interchange/Runtime/Source/Import/Private/MaterialX/MaterialExpressionPremult.cpp
deleted file mode 100644
diff --git a/Engine/Plugins/Interchange/Runtime/Source/Import/Private/MaterialX/MaterialExpressionPremult.h b/Engine/Plugins/Interchange/Runtime/Source/Import/Private/MaterialX/MaterialExpressionPremult.h
deleted file mode 100644
diff --git a/Engine/Plugins/Interchange/Runtime/Source/Import/Private/MaterialX/MaterialExpressionRampLeftRight.cpp b/Engine/Plugins/Interchange/Runtime/Source/Import/Private/MaterialX/MaterialExpressionRampLeftRight.cpp
deleted file mode 100644
diff --git a/Engine/Plugins/Interchange/Runtime/Source/Import/Private/MaterialX/MaterialExpressionRampLeftRight.h b/Engine/Plugins/Interchange/Runtime/Source/Import/Private/MaterialX/MaterialExpressionRampLeftRight.h
deleted file mode 100644
diff --git a/Engine/Plugins/Interchange/Runtime/Source/Import/Private/MaterialX/MaterialExpressionRampTopBottom.cpp b/Engine/Plugins/Interchange/Runtime/Source/Import/Private/MaterialX/MaterialExpressionRampTopBottom.cpp
deleted file mode 100644
diff --git a/Engine/Plugins/Interchange/Runtime/Source/Import/Private/MaterialX/MaterialExpressionRampTopBottom.h b/Engine/Plugins/Interchange/Runtime/Source/Import/Private/MaterialX/MaterialExpressionRampTopBottom.h
deleted file mode 100644
diff --git a/Engine/Plugins/Interchange/Runtime/Source/Import/Private/MaterialX/MaterialExpressionRemap.cpp b/Engine/Plugins/Interchange/Runtime/Source/Import/Private/MaterialX/MaterialExpressionRemap.cpp
deleted file mode 100644
diff --git a/Engine/Plugins/Interchange/Runtime/Source/Import/Private/MaterialX/MaterialExpressionRemap.h b/Engine/Plugins/Interchange/Runtime/Source/Import/Private/MaterialX/MaterialExpressionRemap.h
deleted file mode 100644
diff --git a/Engine/Plugins/Interchange/Runtime/Source/Import/Private/MaterialX/MaterialExpressionRgbToHsv.cpp b/Engine/Plugins/Interchange/Runtime/Source/Import/Private/MaterialX/MaterialExpressionRgbToHsv.cpp
deleted file mode 100644
diff --git a/Engine/Plugins/Interchange/Runtime/Source/Import/Private/MaterialX/MaterialExpressionRgbToHsv.h b/Engine/Plugins/Interchange/Runtime/Source/Import/Private/MaterialX/MaterialExpressionRgbToHsv.h
deleted file mode 100644
diff --git a/Engine/Plugins/Interchange/Runtime/Source/Import/Private/MaterialX/MaterialExpressionScreen.cpp b/Engine/Plugins/Interchange/Runtime/Source/Import/Private/MaterialX/MaterialExpressionScreen.cpp
deleted file mode 100644
diff --git a/Engine/Plugins/Interchange/Runtime/Source/Import/Private/MaterialX/MaterialExpressionScreen.h b/Engine/Plugins/Interchange/Runtime/Source/Import/Private/MaterialX/MaterialExpressionScreen.h
deleted file mode 100644
diff --git a/Engine/Plugins/Interchange/Runtime/Source/Import/Private/MaterialX/MaterialExpressionSplitLeftRight.cpp b/Engine/Plugins/Interchange/Runtime/Source/Import/Private/MaterialX/MaterialExpressionSplitLeftRight.cpp
deleted file mode 100644
diff --git a/Engine/Plugins/Interchange/Runtime/Source/Import/Private/MaterialX/MaterialExpressionSplitLeftRight.h b/Engine/Plugins/Interchange/Runtime/Source/Import/Private/MaterialX/MaterialExpressionSplitLeftRight.h
deleted file mode 100644
diff --git a/Engine/Plugins/Interchange/Runtime/Source/Import/Private/MaterialX/MaterialExpressionSplitTopBottom.cpp b/Engine/Plugins/Interchange/Runtime/Source/Import/Private/MaterialX/MaterialExpressionSplitTopBottom.cpp
deleted file mode 100644
diff --git a/Engine/Plugins/Interchange/Runtime/Source/Import/Private/MaterialX/MaterialExpressionSplitTopBottom.h b/Engine/Plugins/Interchange/Runtime/Source/Import/Private/MaterialX/MaterialExpressionSplitTopBottom.h
deleted file mode 100644
diff --git a/Engine/Plugins/Interchange/Runtime/Source/Import/Private/MaterialX/MaterialExpressionTextureSampleParameterBlur.cpp b/Engine/Plugins/Interchange/Runtime/Source/Import/Private/MaterialX/MaterialExpressionTextureSampleParameterBlur.cpp
deleted file mode 100644
diff --git a/Engine/Plugins/Interchange/Runtime/Source/Import/Private/MaterialX/MaterialExpressionUnpremult.cpp b/Engine/Plugins/Interchange/Runtime/Source/Import/Private/MaterialX/MaterialExpressionUnpremult.cpp
deleted file mode 100644
diff --git a/Engine/Plugins/Interchange/Runtime/Source/Import/Private/MaterialX/MaterialExpressionUnpremult.h b/Engine/Plugins/Interchange/Runtime/Source/Import/Private/MaterialX/MaterialExpressionUnpremult.h
deleted file mode 100644
diff --git a/Engine/Plugins/Interchange/Runtime/Source/Import/Private/MaterialX/MaterialExpressions/MaterialExpressionAppend3Vector.cpp b/Engine/Plugins/Interchange/Runtime/Source/Import/Private/MaterialX/MaterialExpressions/MaterialExpressionAppend3Vector.cpp
new file mode 100644
--- /dev/null
+++ b/Engine/Plugins/Interchange/Runtime/Source/Import/Private/MaterialX/MaterialExpressions/MaterialExpressionAppend3Vector.cpp
@@ -0,0 +1,65 @@
+// Copyright Epic Games, Inc. All Rights Reserved.
+
+#include "MaterialExpressionAppend3Vector.h"
+#include "MaterialCompiler.h"
+
+#include UE_INLINE_GENERATED_CPP_BY_NAME(MaterialExpressionAppend3Vector)
+
+#define LOCTEXT_NAMESPACE "MaterialExpressionAppend3Vector"
+
+UMaterialExpressionAppend3Vector::UMaterialExpressionAppend3Vector(const FObjectInitializer& ObjectInitializer)
+	: Super(ObjectInitializer)
+{
+	// Structure to hold one-time initialization
+	struct FConstructorStatics
+	{
+		FText NAME_MaterialX;
+		FText NAME_Math;
+		FText NAME_VectorOps;
+		FConstructorStatics()
+			: NAME_MaterialX(LOCTEXT("MaterialX", "MaterialX"))
+			, NAME_Math(LOCTEXT("Math", "Math"))
+			, NAME_VectorOps(LOCTEXT("VectorOps", "VectorOps"))
+		{}
+	};
+	static FConstructorStatics ConstructorStatics;
+
+#if WITH_EDITORONLY_DATA
+	MenuCategories.Add(ConstructorStatics.NAME_MaterialX);
+	MenuCategories.Add(ConstructorStatics.NAME_Math);
+	MenuCategories.Add(ConstructorStatics.NAME_VectorOps);
+#endif
+}
+
+#if WITH_EDITOR
+int32 UMaterialExpressionAppend3Vector::Compile(FMaterialCompiler* Compiler, int32 OutputIndex)
+{
+	if(!A.GetTracedInput().Expression)
+	{
+		return Compiler->Errorf(TEXT("Missing AppendVector input A"));
+	}
+	else if(!B.GetTracedInput().Expression)
+	{
+		return Compiler->Errorf(TEXT("Missing AppendVector input B"));
+	}
+	else if(!C.GetTracedInput().Expression)
+	{
+		return Compiler->Errorf(TEXT("Missing AppendVector input C"));
+	}
+	else
+	{
+		int32 Arg1 = A.Compile(Compiler);
+		int32 Arg2 = B.Compile(Compiler);
+		int32 Arg3 = C.Compile(Compiler);
+		return Compiler->AppendVector(Arg1,
+			   Compiler->AppendVector(Arg2, Arg3));
+	}
+}
+
+void UMaterialExpressionAppend3Vector::GetCaption(TArray<FString>& OutCaptions) const
+{
+	OutCaptions.Add(TEXT("Append3"));
+}
+#endif
+
+#undef LOCTEXT_NAMESPACE 
\ No newline at end of file
diff --git a/Engine/Plugins/Interchange/Runtime/Source/Import/Private/MaterialX/MaterialExpressions/MaterialExpressionAppend3Vector.h b/Engine/Plugins/Interchange/Runtime/Source/Import/Private/MaterialX/MaterialExpressions/MaterialExpressionAppend3Vector.h
new file mode 100644
--- /dev/null
+++ b/Engine/Plugins/Interchange/Runtime/Source/Import/Private/MaterialX/MaterialExpressions/MaterialExpressionAppend3Vector.h
@@ -0,0 +1,33 @@
+// Copyright Epic Games, Inc. All Rights Reserved.
+#pragma once
+
+#include "CoreMinimal.h"
+#include "Materials/MaterialExpression.h"
+
+#include "MaterialExpressionAppend3Vector.generated.h"
+
+/**
+ * A material expression that allows combining 3 channels together to create a vector with more channel than the original
+ */
+UCLASS(MinimalAPI, collapsecategories, hidecategories = Object)
+class UMaterialExpressionAppend3Vector : public UMaterialExpression
+{
+	GENERATED_UCLASS_BODY()
+
+	UPROPERTY()
+	FExpressionInput A;
+
+	UPROPERTY()
+	FExpressionInput B;
+
+	UPROPERTY()
+	FExpressionInput C;
+
+	//~ Begin UMaterialExpression Interface
+#if WITH_EDITOR
+	virtual int32 Compile(class FMaterialCompiler* Compiler, int32 OutputIndex) override;
+	virtual void GetCaption(TArray<FString>& OutCaptions) const override;
+#endif
+	//~ End UMaterialExpression Interface
+};
+
diff --git a/Engine/Plugins/Interchange/Runtime/Source/Import/Private/MaterialX/MaterialExpressions/MaterialExpressionAppend4Vector.cpp b/Engine/Plugins/Interchange/Runtime/Source/Import/Private/MaterialX/MaterialExpressions/MaterialExpressionAppend4Vector.cpp
new file mode 100644
--- /dev/null
+++ b/Engine/Plugins/Interchange/Runtime/Source/Import/Private/MaterialX/MaterialExpressions/MaterialExpressionAppend4Vector.cpp
@@ -0,0 +1,71 @@
+// Copyright Epic Games, Inc. All Rights Reserved.
+
+#include "MaterialExpressionAppend4Vector.h"
+#include "MaterialCompiler.h"
+
+#include UE_INLINE_GENERATED_CPP_BY_NAME(MaterialExpressionAppend4Vector)
+
+#define LOCTEXT_NAMESPACE "MaterialExpressionAppend4Vector"
+
+UMaterialExpressionAppend4Vector::UMaterialExpressionAppend4Vector(const FObjectInitializer& ObjectInitializer)
+	: Super(ObjectInitializer)
+{
+	// Structure to hold one-time initialization
+	struct FConstructorStatics
+	{
+		FText NAME_MaterialX;
+		FText NAME_Math;
+		FText NAME_VectorOps;
+		FConstructorStatics()
+			: NAME_MaterialX(LOCTEXT("MaterialX", "MaterialX"))
+			, NAME_Math(LOCTEXT("Math", "Math"))
+			, NAME_VectorOps(LOCTEXT("VectorOps", "VectorOps"))
+		{}
+	};
+	static FConstructorStatics ConstructorStatics;
+
+#if WITH_EDITORONLY_DATA
+	MenuCategories.Add(ConstructorStatics.NAME_MaterialX);
+	MenuCategories.Add(ConstructorStatics.NAME_Math);
+	MenuCategories.Add(ConstructorStatics.NAME_VectorOps);
+#endif
+}
+
+#if WITH_EDITOR
+int32 UMaterialExpressionAppend4Vector::Compile(FMaterialCompiler* Compiler, int32 OutputIndex)
+{
+	if(!A.GetTracedInput().Expression)
+	{
+		return Compiler->Errorf(TEXT("Missing AppendVector input A"));
+	}
+	else if(!B.GetTracedInput().Expression)
+	{
+		return Compiler->Errorf(TEXT("Missing AppendVector input B"));
+	}
+	else if(!C.GetTracedInput().Expression)
+	{
+		return Compiler->Errorf(TEXT("Missing AppendVector input C"));
+	}
+	else if(!D.GetTracedInput().Expression)
+	{
+		return Compiler->Errorf(TEXT("Missing AppendVector input D"));
+	}
+	else
+	{
+		int32 Arg1 = A.Compile(Compiler);
+		int32 Arg2 = B.Compile(Compiler);
+		int32 Arg3 = C.Compile(Compiler);
+		int32 Arg4 = D.Compile(Compiler);
+		return Compiler->AppendVector(Arg1,
+			   Compiler->AppendVector(Arg2,
+			   Compiler->AppendVector(Arg3, Arg4)));
+	}
+}
+
+void UMaterialExpressionAppend4Vector::GetCaption(TArray<FString>& OutCaptions) const
+{
+	OutCaptions.Add(TEXT("Append4"));
+}
+#endif
+
+#undef LOCTEXT_NAMESPACE
\ No newline at end of file
diff --git a/Engine/Plugins/Interchange/Runtime/Source/Import/Private/MaterialX/MaterialExpressions/MaterialExpressionAppend4Vector.h b/Engine/Plugins/Interchange/Runtime/Source/Import/Private/MaterialX/MaterialExpressions/MaterialExpressionAppend4Vector.h
new file mode 100644
--- /dev/null
+++ b/Engine/Plugins/Interchange/Runtime/Source/Import/Private/MaterialX/MaterialExpressions/MaterialExpressionAppend4Vector.h
@@ -0,0 +1,38 @@
+// Copyright Epic Games, Inc. All Rights Reserved.
+#pragma once
+
+#include "CoreMinimal.h"
+#include "UObject/ObjectMacros.h"
+#include "MaterialExpressionIO.h"
+#include "Materials/MaterialExpression.h"
+
+#include "MaterialExpressionAppend4Vector.generated.h"
+
+/**
+ * A material expression that allows combining 4 channels together to create a vector with more channel than the original
+ */
+UCLASS(MinimalAPI, collapsecategories, hidecategories = Object)
+class UMaterialExpressionAppend4Vector : public UMaterialExpression
+{
+	GENERATED_UCLASS_BODY()
+
+	UPROPERTY()
+	FExpressionInput A;
+	
+	UPROPERTY()
+	FExpressionInput B;
+
+	UPROPERTY()
+	FExpressionInput C;
+	
+	UPROPERTY()
+	FExpressionInput D;
+
+	//~ Begin UMaterialExpression Interface
+#if WITH_EDITOR
+	virtual int32 Compile(class FMaterialCompiler* Compiler, int32 OutputIndex) override;
+	virtual void GetCaption(TArray<FString>& OutCaptions) const override;
+#endif
+	//~ End UMaterialExpression Interface
+};
+
diff --git a/Engine/Plugins/Interchange/Runtime/Source/Import/Private/MaterialX/MaterialExpressions/MaterialExpressionBurn.cpp b/Engine/Plugins/Interchange/Runtime/Source/Import/Private/MaterialX/MaterialExpressions/MaterialExpressionBurn.cpp
new file mode 100644
--- /dev/null
+++ b/Engine/Plugins/Interchange/Runtime/Source/Import/Private/MaterialX/MaterialExpressions/MaterialExpressionBurn.cpp
@@ -0,0 +1,63 @@
+// Copyright Epic Games, Inc. All Rights Reserved.
+#include "MaterialExpressionBurn.h"
+#include "MaterialCompiler.h"
+
+#include UE_INLINE_GENERATED_CPP_BY_NAME(MaterialExpressionBurn)
+
+#define LOCTEXT_NAMESPACE "MaterialExpressionBurn"
+
+UMaterialExpressionBurn::UMaterialExpressionBurn(const FObjectInitializer& ObjectInitializer)
+	: Super(ObjectInitializer)
+{
+	// Structure to hold one-time initialization
+	struct FConstructorStatics
+	{
+		FText NAME_MaterialX;
+		FText NAME_Compositing;
+		FConstructorStatics() 
+			: NAME_MaterialX(LOCTEXT("MaterialX", "MaterialX"))
+			, NAME_Compositing(LOCTEXT("Compositing", "Compositing"))
+		{}
+	};
+	static FConstructorStatics ConstructorStatics;
+
+#if WITH_EDITORONLY_DATA
+	MenuCategories.Add(ConstructorStatics.NAME_MaterialX);
+	MenuCategories.Add(ConstructorStatics.NAME_Compositing);
+#endif
+}
+
+#if WITH_EDITOR
+int32 UMaterialExpressionBurn::Compile(FMaterialCompiler* Compiler, int32 OutputIndex)
+{
+	if(!A.GetTracedInput().Expression)
+	{
+		return Compiler->Errorf(TEXT("Missing A input"));
+	}
+
+	if(!B.GetTracedInput().Expression)
+	{
+		return Compiler->Errorf(TEXT("Missing B input"));
+	}
+
+	int32 IndexAlpha = Alpha.GetTracedInput().Expression ? Alpha.Compile(Compiler) : Compiler->Constant(ConstAlpha);
+
+	int32 OneIndex = Compiler->Constant(1.f);
+	auto OneMinus = [&](int32 index)
+	{
+		return Compiler->Sub(OneIndex, index);
+	};
+
+	int32 IndexB = B.Compile(Compiler);
+	int32 Burn = OneMinus(Compiler->Div(OneMinus(IndexB), A.Compile(Compiler)));
+
+	return Compiler->Lerp(IndexB, Burn, IndexAlpha);
+}
+
+void UMaterialExpressionBurn::GetCaption(TArray<FString>& OutCaptions) const
+{
+	OutCaptions.Add(TEXT("Burn"));
+}
+#endif
+
+#undef LOCTEXT_NAMESPACE 
\ No newline at end of file
diff --git a/Engine/Plugins/Interchange/Runtime/Source/Import/Private/MaterialX/MaterialExpressions/MaterialExpressionBurn.h b/Engine/Plugins/Interchange/Runtime/Source/Import/Private/MaterialX/MaterialExpressions/MaterialExpressionBurn.h
new file mode 100644
--- /dev/null
+++ b/Engine/Plugins/Interchange/Runtime/Source/Import/Private/MaterialX/MaterialExpressions/MaterialExpressionBurn.h
@@ -0,0 +1,41 @@
+// Copyright Epic Games, Inc. All Rights Reserved.
+#pragma once
+
+#include "CoreMinimal.h"
+#include "Materials/MaterialExpression.h"
+
+#include "MaterialExpressionBurn.generated.h"
+
+/**
+ * Blend nodes take two 1-4 channel inputs and apply the same operator to all channels.
+ * Blend nodes support an optional float input mix , which can be used
+ * to mix the original B value with the result of the blend operation.
+ * Operation: 1-(1-B)/A
+ * Result: Lerp(B, 1-(1-B)/A, Alpha)
+ */
+UCLASS(MinimalAPI, collapsecategories, hidecategories = Object)
+class UMaterialExpressionBurn : public UMaterialExpression
+{
+	GENERATED_UCLASS_BODY()
+
+	UPROPERTY()
+	FExpressionInput A;
+
+	UPROPERTY()
+	FExpressionInput B;
+
+	UPROPERTY(meta = (RequiredInput = "false", ToolTip = "Defaults to 'ConstAlpha' if not specified"))
+	FExpressionInput Alpha;
+
+	/** only used if Alpha is not hooked up */
+	UPROPERTY(EditAnywhere, Category = MaterialExpressionLinearInterpolate, meta = (OverridingInputProperty = "Alpha"))
+	float ConstAlpha = 1.f;
+
+	//~ Begin UMaterialExpression Interface
+#if WITH_EDITOR
+	virtual int32 Compile(class FMaterialCompiler* Compiler, int32 OutputIndex) override;
+	virtual void GetCaption(TArray<FString>& OutCaptions) const override;
+#endif
+	//~ End UMaterialExpression Interface
+};
+
diff --git a/Engine/Plugins/Interchange/Runtime/Source/Import/Private/MaterialX/MaterialExpressions/MaterialExpressionDifference.cpp b/Engine/Plugins/Interchange/Runtime/Source/Import/Private/MaterialX/MaterialExpressions/MaterialExpressionDifference.cpp
new file mode 100644
--- /dev/null
+++ b/Engine/Plugins/Interchange/Runtime/Source/Import/Private/MaterialX/MaterialExpressions/MaterialExpressionDifference.cpp
@@ -0,0 +1,56 @@
+// Copyright Epic Games, Inc. All Rights Reserved.
+#include "MaterialExpressionDifference.h"
+#include "MaterialCompiler.h"
+
+#include UE_INLINE_GENERATED_CPP_BY_NAME(MaterialExpressionDifference)
+
+#define LOCTEXT_NAMESPACE "MaterialExpressionDifference"
+
+UMaterialExpressionDifference::UMaterialExpressionDifference(const FObjectInitializer& ObjectInitializer)
+	: Super(ObjectInitializer)
+{
+	// Structure to hold one-time initialization
+	struct FConstructorStatics
+	{
+		FText NAME_MaterialX;
+		FText NAME_Compositing;
+		FConstructorStatics()
+			: NAME_MaterialX(LOCTEXT("MaterialX", "MaterialX"))
+			, NAME_Compositing(LOCTEXT("Compositing", "Compositing"))
+		{}
+	};
+	static FConstructorStatics ConstructorStatics;
+
+#if WITH_EDITORONLY_DATA
+	MenuCategories.Add(ConstructorStatics.NAME_MaterialX);
+	MenuCategories.Add(ConstructorStatics.NAME_Compositing);
+#endif
+}
+
+#if WITH_EDITOR
+int32 UMaterialExpressionDifference::Compile(FMaterialCompiler* Compiler, int32 OutputIndex)
+{
+	if(!A.GetTracedInput().Expression)
+	{
+		return Compiler->Errorf(TEXT("Missing A input"));
+	}
+
+	if(!B.GetTracedInput().Expression)
+	{
+		return Compiler->Errorf(TEXT("Missing B input"));
+	}
+
+	int32 IndexAlpha = Alpha.GetTracedInput().Expression ? Alpha.Compile(Compiler) : Compiler->Constant(ConstAlpha);
+
+	int32 IndexB = B.Compile(Compiler);
+	int32 Diff = Compiler->Abs(Compiler->Sub(A.Compile(Compiler), IndexB));
+	return Compiler->Lerp(IndexB, Diff, IndexAlpha);
+}
+
+void UMaterialExpressionDifference::GetCaption(TArray<FString>& OutCaptions) const
+{
+	OutCaptions.Add(TEXT("Difference"));
+}
+#endif
+
+#undef LOCTEXT_NAMESPACE 
\ No newline at end of file
diff --git a/Engine/Plugins/Interchange/Runtime/Source/Import/Private/MaterialX/MaterialExpressions/MaterialExpressionDifference.h b/Engine/Plugins/Interchange/Runtime/Source/Import/Private/MaterialX/MaterialExpressions/MaterialExpressionDifference.h
new file mode 100644
--- /dev/null
+++ b/Engine/Plugins/Interchange/Runtime/Source/Import/Private/MaterialX/MaterialExpressions/MaterialExpressionDifference.h
@@ -0,0 +1,41 @@
+// Copyright Epic Games, Inc. All Rights Reserved.
+#pragma once
+
+#include "CoreMinimal.h"
+#include "Materials/MaterialExpression.h"
+
+#include "MaterialExpressionDifference.generated.h"
+
+/**
+ * Blend nodes take two 1-4 channel inputs and apply the same operator to all channels.
+ * Blend nodes support an optional float input mix , which can be used
+ * to mix the original B value with the result of the blend operation.
+ * Operation: abs(A - B)
+ * Result: Lerp(B, abs(A - B), Alpha)
+ */
+UCLASS(MinimalAPI, collapsecategories, hidecategories = Object)
+class UMaterialExpressionDifference : public UMaterialExpression
+{
+	GENERATED_UCLASS_BODY()
+
+	UPROPERTY()
+	FExpressionInput A;
+
+	UPROPERTY()
+	FExpressionInput B;
+
+	UPROPERTY(meta = (RequiredInput = "false", ToolTip = "Defaults to 'ConstAlpha' if not specified"))
+	FExpressionInput Alpha;
+
+	/** only used if Alpha is not hooked up */
+	UPROPERTY(EditAnywhere, Category = MaterialExpressionLinearInterpolate, meta = (OverridingInputProperty = "Alpha"))
+	float ConstAlpha = 1.f;
+
+	//~ Begin UMaterialExpression Interface
+#if WITH_EDITOR
+	virtual int32 Compile(class FMaterialCompiler* Compiler, int32 OutputIndex) override;
+	virtual void GetCaption(TArray<FString>& OutCaptions) const override;
+#endif
+	//~ End UMaterialExpression Interface
+};
+
diff --git a/Engine/Plugins/Interchange/Runtime/Source/Import/Private/MaterialX/MaterialExpressions/MaterialExpressionDisjointOver.cpp b/Engine/Plugins/Interchange/Runtime/Source/Import/Private/MaterialX/MaterialExpressions/MaterialExpressionDisjointOver.cpp
new file mode 100644
--- /dev/null
+++ b/Engine/Plugins/Interchange/Runtime/Source/Import/Private/MaterialX/MaterialExpressions/MaterialExpressionDisjointOver.cpp
@@ -0,0 +1,73 @@
+// Copyright Epic Games, Inc. All Rights Reserved.
+#include "MaterialExpressionDisjointOver.h"
+#include "MaterialCompiler.h"
+
+#include UE_INLINE_GENERATED_CPP_BY_NAME(MaterialExpressionDisjointOver)
+
+#define LOCTEXT_NAMESPACE "MaterialExpressionDisjointOver"
+
+UMaterialExpressionDisjointOver::UMaterialExpressionDisjointOver(const FObjectInitializer& ObjectInitializer)
+	: Super(ObjectInitializer)
+{
+	// Structure to hold one-time initialization
+	struct FConstructorStatics
+	{
+		FText NAME_MaterialX;
+		FText NAME_Compositing;
+		FConstructorStatics()
+			: NAME_MaterialX(LOCTEXT("MaterialX", "MaterialX"))
+			, NAME_Compositing(LOCTEXT("Compositing", "Compositing"))
+		{}
+	};
+	static FConstructorStatics ConstructorStatics;
+
+#if WITH_EDITORONLY_DATA
+	MenuCategories.Add(ConstructorStatics.NAME_MaterialX);
+	MenuCategories.Add(ConstructorStatics.NAME_Compositing);
+#endif
+}
+
+#if WITH_EDITOR
+int32 UMaterialExpressionDisjointOver::Compile(FMaterialCompiler* Compiler, int32 OutputIndex)
+{
+	if(!A.GetTracedInput().Expression)
+	{
+		return Compiler->Errorf(TEXT("Missing A input"));
+	}
+
+	if(!B.GetTracedInput().Expression)
+	{
+		return Compiler->Errorf(TEXT("Missing B input"));
+	}
+
+	int32 IndexAlpha = Alpha.GetTracedInput().Expression ? Alpha.Compile(Compiler) : Compiler->Constant(ConstAlpha);
+	
+	int32 IndexA = A.Compile(Compiler);
+	int32 IndexB = B.Compile(Compiler);
+
+	int32 IndexAlphaA = Compiler->ComponentMask(IndexA, false, false, false, true);
+	int32 IndexAlphaB = Compiler->ComponentMask(IndexB, false, false, false, true);
+	int32 IndexSumAlpha = Compiler->Add(IndexAlphaA, IndexAlphaB);
+
+	int32 IndexRgbA = Compiler->ComponentMask(IndexA, true, true, true, false);
+	int32 IndexRgbB = Compiler->ComponentMask(IndexB, true, true, true, false);
+	int32 IndexSumRgb = Compiler->Add(IndexRgbA, IndexRgbB);
+	
+	int32 IndexOne = Compiler->Constant(1.f);
+	int32 X = Compiler->Div(Compiler->Sub(IndexOne, IndexAlphaA), IndexAlphaB);
+	int32 Disjoint = Compiler->Add(IndexRgbA, Compiler->Mul(IndexRgbB, X));
+
+	int32 Result = Compiler->If(IndexSumAlpha, IndexOne, Disjoint, IndexSumRgb, IndexSumRgb, Compiler->Constant(0.00001f));
+	int32 ResultAlpha = Compiler->Min(IndexSumAlpha, IndexOne);
+
+	int32 IndexAppend = Compiler->AppendVector(Result, ResultAlpha);
+	return Compiler->Lerp(IndexB, IndexAppend, IndexAlpha);
+}
+
+void UMaterialExpressionDisjointOver::GetCaption(TArray<FString>& OutCaptions) const
+{
+	OutCaptions.Add(TEXT("DisjointOver"));
+}
+#endif
+
+#undef LOCTEXT_NAMESPACE 
\ No newline at end of file
diff --git a/Engine/Plugins/Interchange/Runtime/Source/Import/Private/MaterialX/MaterialExpressions/MaterialExpressionDisjointOver.h b/Engine/Plugins/Interchange/Runtime/Source/Import/Private/MaterialX/MaterialExpressions/MaterialExpressionDisjointOver.h
new file mode 100644
--- /dev/null
+++ b/Engine/Plugins/Interchange/Runtime/Source/Import/Private/MaterialX/MaterialExpressions/MaterialExpressionDisjointOver.h
@@ -0,0 +1,46 @@
+// Copyright Epic Games, Inc. All Rights Reserved.
+#pragma once
+
+#include "CoreMinimal.h"
+#include "Materials/MaterialExpression.h"
+
+#include "MaterialExpressionDisjointOver.generated.h"
+
+/**
+ * Merge nodes take two 4-channel (color4) inputs and use the built-in alpha channel(s) to control the
+ * compositing of the A and B inputs. "A" and "B" refer to the non-alpha channels of the A and B inputs respectively, 
+ * and "a" and "b" refer to the alpha channels of the A and B inputs.
+ * Merge nodes are only defined for float4 inputs
+ * Merge nodes support an optional float input Alpha , which can be used to mix the
+ * original B value with the result of the blend operation.
+ * 
+ * Operation: A + B        if a + b <= 1;
+ *            A + B(1-a)/b if a + b > 1
+ * Result: Lerp(B, Op, Alpha)
+ */
+UCLASS(MinimalAPI, collapsecategories, hidecategories = Object)
+class UMaterialExpressionDisjointOver : public UMaterialExpression
+{
+	GENERATED_UCLASS_BODY()
+
+	UPROPERTY()
+	FExpressionInput A;
+
+	UPROPERTY()
+	FExpressionInput B;
+
+	UPROPERTY(meta = (RequiredInput = "false", ToolTip = "Defaults to 'ConstAlpha' if not specified"))
+	FExpressionInput Alpha;
+
+	/** only used if Alpha is not hooked up */
+	UPROPERTY(EditAnywhere, Category = MaterialExpressionLinearInterpolate, meta = (OverridingInputProperty = "Alpha"))
+	float ConstAlpha = 1.f;
+
+	//~ Begin UMaterialExpression Interface
+#if WITH_EDITOR
+	virtual int32 Compile(class FMaterialCompiler* Compiler, int32 OutputIndex) override;
+	virtual void GetCaption(TArray<FString>& OutCaptions) const override;
+#endif
+	//~ End UMaterialExpression Interface
+};
+
diff --git a/Engine/Plugins/Interchange/Runtime/Source/Import/Private/MaterialX/MaterialExpressions/MaterialExpressionDodge.cpp b/Engine/Plugins/Interchange/Runtime/Source/Import/Private/MaterialX/MaterialExpressions/MaterialExpressionDodge.cpp
new file mode 100644
--- /dev/null
+++ b/Engine/Plugins/Interchange/Runtime/Source/Import/Private/MaterialX/MaterialExpressions/MaterialExpressionDodge.cpp
@@ -0,0 +1,61 @@
+// Copyright Epic Games, Inc. All Rights Reserved.
+#include "MaterialExpressionDodge.h"
+#include "MaterialCompiler.h"
+
+#include UE_INLINE_GENERATED_CPP_BY_NAME(MaterialExpressionDodge)
+
+#define LOCTEXT_NAMESPACE "MaterialExpressionDodge"
+
+UMaterialExpressionDodge::UMaterialExpressionDodge(const FObjectInitializer& ObjectInitializer)
+	: Super(ObjectInitializer)
+{
+	// Structure to hold one-time initialization
+	struct FConstructorStatics
+	{
+		FText NAME_MaterialX;
+		FText NAME_Compositing;
+		FConstructorStatics()
+			: NAME_MaterialX(LOCTEXT("MaterialX", "MaterialX"))
+			, NAME_Compositing(LOCTEXT("Compositing", "Compositing"))
+		{}
+	};
+	static FConstructorStatics ConstructorStatics;
+
+#if WITH_EDITORONLY_DATA
+	MenuCategories.Add(ConstructorStatics.NAME_MaterialX);
+	MenuCategories.Add(ConstructorStatics.NAME_Compositing);
+#endif
+}
+
+#if WITH_EDITOR
+int32 UMaterialExpressionDodge::Compile(FMaterialCompiler* Compiler, int32 OutputIndex)
+{
+	if(!A.GetTracedInput().Expression)
+	{
+		return Compiler->Errorf(TEXT("Missing A input"));
+	}
+
+	if(!B.GetTracedInput().Expression)
+	{
+		return Compiler->Errorf(TEXT("Missing B input"));
+	}
+
+	int32 IndexAlpha = Alpha.GetTracedInput().Expression ? Alpha.Compile(Compiler) : Compiler->Constant(ConstAlpha);
+	int32 IndexB = B.Compile(Compiler);
+	int32 Sub = Compiler->Sub(Compiler->Constant(1.f), A.Compile(Compiler));
+
+	int32 Zero = Compiler->Constant(0.f);
+	int32 Dodge = Compiler->Div(IndexB, Sub);
+
+	int32 Result = Compiler->Lerp(IndexB, Dodge, IndexAlpha);
+
+	return Compiler->If(Compiler->Abs(Sub), Zero, Result, Result, Zero, Compiler->Constant(0.00001f));
+}
+
+void UMaterialExpressionDodge::GetCaption(TArray<FString>& OutCaptions) const
+{
+	OutCaptions.Add(TEXT("Dodge"));
+}
+#endif
+
+#undef LOCTEXT_NAMESPACE 
\ No newline at end of file
diff --git a/Engine/Plugins/Interchange/Runtime/Source/Import/Private/MaterialX/MaterialExpressions/MaterialExpressionDodge.h b/Engine/Plugins/Interchange/Runtime/Source/Import/Private/MaterialX/MaterialExpressions/MaterialExpressionDodge.h
new file mode 100644
--- /dev/null
+++ b/Engine/Plugins/Interchange/Runtime/Source/Import/Private/MaterialX/MaterialExpressions/MaterialExpressionDodge.h
@@ -0,0 +1,41 @@
+// Copyright Epic Games, Inc. All Rights Reserved.
+#pragma once
+
+#include "CoreMinimal.h"
+#include "Materials/MaterialExpression.h"
+
+#include "MaterialExpressionDodge.generated.h"
+
+/**
+ * Blend nodes take two 1-4 channel inputs and apply the same operator to all channels.
+ * Blend nodes support an optional float input mix , which can be used
+ * to mix the original B value with the result of the blend operation.
+ * Operation: B/(1-A)
+ * Result: Lerp(B, B/(1-A), Alpha)
+ */
+UCLASS(MinimalAPI, collapsecategories, hidecategories = Object)
+class UMaterialExpressionDodge: public UMaterialExpression
+{
+	GENERATED_UCLASS_BODY()
+
+	UPROPERTY()
+	FExpressionInput A;
+
+	UPROPERTY()
+	FExpressionInput B;
+
+	UPROPERTY(meta = (RequiredInput = "false", ToolTip = "Defaults to 'ConstAlpha' if not specified"))
+	FExpressionInput Alpha;
+
+	/** only used if Alpha is not hooked up */
+	UPROPERTY(EditAnywhere, Category = MaterialExpressionLinearInterpolate, meta = (OverridingInputProperty = "Alpha"))
+	float ConstAlpha = 1.f;
+
+	//~ Begin UMaterialExpression Interface
+#if WITH_EDITOR
+	virtual int32 Compile(class FMaterialCompiler* Compiler, int32 OutputIndex) override;
+	virtual void GetCaption(TArray<FString>& OutCaptions) const override;
+#endif
+	//~ End UMaterialExpression Interface
+};
+
diff --git a/Engine/Plugins/Interchange/Runtime/Source/Import/Private/MaterialX/MaterialExpressions/MaterialExpressionExponential.cpp b/Engine/Plugins/Interchange/Runtime/Source/Import/Private/MaterialX/MaterialExpressions/MaterialExpressionExponential.cpp
new file mode 100644
--- /dev/null
+++ b/Engine/Plugins/Interchange/Runtime/Source/Import/Private/MaterialX/MaterialExpressions/MaterialExpressionExponential.cpp
@@ -0,0 +1,50 @@
+// Copyright Epic Games, Inc. All Rights Reserved.
+#include "MaterialExpressionExponential.h"
+#include "MaterialCompiler.h"
+#include "MaterialHLSLGenerator.h"
+
+#include UE_INLINE_GENERATED_CPP_BY_NAME(MaterialExpressionExponential)
+
+#define LOCTEXT_NAMESPACE "MaterialExpressionExponential"
+
+UMaterialExpressionExponential::UMaterialExpressionExponential(const FObjectInitializer& ObjectInitializer)
+	: Super(ObjectInitializer)
+{
+	// Structure to hold one-time initialization
+	struct FConstructorStatics
+	{
+		FText NAME_MaterialX;
+		FText NAME_Math;
+		FConstructorStatics()
+			: NAME_MaterialX(LOCTEXT("MaterialX", "MaterialX"))
+			, NAME_Math(LOCTEXT("Math", "Math"))
+		{}
+	};
+	static FConstructorStatics ConstructorStatics;
+
+#if WITH_EDITORONLY_DATA
+	MenuCategories.Add(ConstructorStatics.NAME_MaterialX);
+	MenuCategories.Add(ConstructorStatics.NAME_Math);
+#endif
+}
+
+#if WITH_EDITOR
+int32 UMaterialExpressionExponential::Compile(FMaterialCompiler* Compiler, int32 OutputIndex)
+{
+	if(!Input.GetTracedInput().Expression)
+	{
+		return Compiler->Errorf(TEXT("Missing Exponential input"));
+	}
+
+	int32 Euler = Compiler->Constant(2.7182818);
+
+	return Compiler->Power(Euler, Input.Compile(Compiler));
+}
+
+void UMaterialExpressionExponential::GetCaption(TArray<FString>& OutCaptions) const
+{
+	OutCaptions.Add(TEXT("Exponential"));
+}
+#endif
+
+#undef LOCTEXT_NAMESPACE 
\ No newline at end of file
diff --git a/Engine/Plugins/Interchange/Runtime/Source/Import/Private/MaterialX/MaterialExpressions/MaterialExpressionExponential.h b/Engine/Plugins/Interchange/Runtime/Source/Import/Private/MaterialX/MaterialExpressions/MaterialExpressionExponential.h
new file mode 100644
--- /dev/null
+++ b/Engine/Plugins/Interchange/Runtime/Source/Import/Private/MaterialX/MaterialExpressions/MaterialExpressionExponential.h
@@ -0,0 +1,27 @@
+// Copyright Epic Games, Inc. All Rights Reserved.
+#pragma once
+
+#include "CoreMinimal.h"
+#include "Materials/MaterialExpression.h"
+
+#include "MaterialExpressionExponential.generated.h"
+
+/**
+ * A material expression that computes the exponential of an input
+ */
+UCLASS(collapsecategories, hidecategories = Object, MinimalAPI)
+class UMaterialExpressionExponential : public UMaterialExpression
+{
+	GENERATED_UCLASS_BODY()
+
+	UPROPERTY()
+	FExpressionInput Input;
+
+	//~ Begin UMaterialExpression Interface
+#if WITH_EDITOR
+	virtual int32 Compile(class FMaterialCompiler* Compiler, int32 OutputIndex) override;
+	virtual void GetCaption(TArray<FString>& OutCaptions) const override;
+#endif
+	//~ End UMaterialExpression Interface
+};
+
diff --git a/Engine/Plugins/Interchange/Runtime/Source/Import/Private/MaterialX/MaterialExpressions/MaterialExpressionFractal3D.cpp b/Engine/Plugins/Interchange/Runtime/Source/Import/Private/MaterialX/MaterialExpressions/MaterialExpressionFractal3D.cpp
new file mode 100644
--- /dev/null
+++ b/Engine/Plugins/Interchange/Runtime/Source/Import/Private/MaterialX/MaterialExpressions/MaterialExpressionFractal3D.cpp
@@ -0,0 +1,99 @@
+// Copyright Epic Games, Inc. All Rights Reserved.
+
+#include "MaterialExpressionFractal3D.h"
+#include "MaterialCompiler.h"
+#include "Materials/MaterialExpressionCustom.h"
+
+#include UE_INLINE_GENERATED_CPP_BY_NAME(MaterialExpressionFractal3D)
+
+#define LOCTEXT_NAMESPACE "MaterialExpressionFractal3D"
+
+UMaterialExpressionFractal3D::UMaterialExpressionFractal3D(const FObjectInitializer& ObjectInitializer)
+	: Super(ObjectInitializer)
+{
+	// Structure to hold one-time initialization
+	struct FConstructorStatics
+	{
+		FText NAME_MaterialX;
+		FText NAME_Procedural;
+		FConstructorStatics()
+			: NAME_MaterialX(LOCTEXT("MaterialX", "MaterialX"))
+			, NAME_Procedural(LOCTEXT("Procedural", "Procedural"))
+		{}
+	};
+	static FConstructorStatics ConstructorStatics;
+
+#if WITH_EDITORONLY_DATA
+	MenuCategories.Add(ConstructorStatics.NAME_MaterialX);
+	MenuCategories.Add(ConstructorStatics.NAME_Procedural);
+#endif
+}
+
+#if WITH_EDITOR
+int32 UMaterialExpressionFractal3D::Compile(FMaterialCompiler* Compiler, int32 OutputIndex)
+{
+
+	UMaterialExpressionCustom* MaterialExpressionCustom = NewObject<UMaterialExpressionCustom>();
+	MaterialExpressionCustom->Inputs[0].InputName = TEXT("Position");
+	MaterialExpressionCustom->Inputs.Add({ TEXT("Octaves") });
+	MaterialExpressionCustom->Inputs.Add({ TEXT("Lacunarity") });
+	MaterialExpressionCustom->Inputs.Add({ TEXT("Diminish") });
+	
+	MaterialExpressionCustom->OutputType = ECustomMaterialOutputType::CMOT_Float3;
+	MaterialExpressionCustom->Code = 
+		TEXT("float Scale =") + FString::SanitizeFloat(Scale) + TEXT(";") +
+		TEXT(R"(int Quality = 1;
+             int Function = 2; // Fast Gradient
+             bool bTurbulence =)") + FString::FromInt(bTurbulence) + TEXT(";") +
+		TEXT("uint Levels =") + FString::FromInt(Levels) + TEXT(";") +
+		TEXT("float OutputMin =") + FString::SanitizeFloat(OutputMin) + TEXT(";") +
+		TEXT("float OutputMax =") + FString::SanitizeFloat(OutputMax) + TEXT(";") +
+		TEXT(R"(float LevelScale = 2;
+             float FilterWidth = 0;
+             bool bTiling = false;
+             float RepeatSize = 512;
+             
+             float3 result = float3(0.0,0.0,0.0);
+             float amplitude = 1.0;
+             for (int i = 0;  i < Octaves; ++i)
+             {
+                 result += amplitude * 
+                 MaterialExpressionNoise(
+                     Position,
+                     Scale,
+                     Quality,
+                     Function,
+                     bTurbulence,
+                     Levels,
+                     OutputMin,
+                     OutputMax,
+                     LevelScale,
+                     FilterWidth,
+                     bTiling,
+                     RepeatSize);
+             
+                 amplitude *= Diminish;
+                 Position *= Lacunarity;
+             }
+             return result;)");
+
+	int32 IndexPosition = Position.GetTracedInput().Expression ? Position.Compile(Compiler) : Compiler->TransformPosition(EMaterialCommonBasis::MCB_World, EMaterialCommonBasis::MCB_Local, Compiler->WorldPosition(WPT_Default));
+	int32 IndexOctaves = Octaves.GetTracedInput().Expression ? Octaves.Compile(Compiler) : Compiler->Constant(ConstOctaves);
+	int32 IndexLacunarity = Lacunarity.GetTracedInput().Expression ? Lacunarity.Compile(Compiler) : Compiler->Constant(ConstLacunarity);
+	int32 IndexDiminish = Diminish.GetTracedInput().Expression ? Diminish.Compile(Compiler) : Compiler->Constant(ConstDiminish);
+
+	TArray<int32> Inputs{ IndexPosition, IndexOctaves, IndexLacunarity, IndexDiminish };
+	int32 IndexFractal = Compiler->CustomExpression(MaterialExpressionCustom, 0, Inputs);
+
+	int32 IndexAmplitude = Amplitude.GetTracedInput().Expression ? Amplitude.Compile(Compiler) : Compiler->Constant(ConstAmplitude);
+
+	return Compiler->Mul(IndexFractal, IndexAmplitude);
+}
+
+void UMaterialExpressionFractal3D::GetCaption(TArray<FString>& OutCaptions) const
+{
+	OutCaptions.Add(TEXT("Fractal3D"));
+}
+#endif
+
+#undef LOCTEXT_NAMESPACE 
\ No newline at end of file
diff --git a/Engine/Plugins/Interchange/Runtime/Source/Import/Private/MaterialX/MaterialExpressions/MaterialExpressionFractal3D.h b/Engine/Plugins/Interchange/Runtime/Source/Import/Private/MaterialX/MaterialExpressions/MaterialExpressionFractal3D.h
new file mode 100644
--- /dev/null
+++ b/Engine/Plugins/Interchange/Runtime/Source/Import/Private/MaterialX/MaterialExpressions/MaterialExpressionFractal3D.h
@@ -0,0 +1,79 @@
+// Copyright Epic Games, Inc. All Rights Reserved.
+#pragma once
+
+#include "Materials/MaterialExpression.h"
+
+#include "MaterialExpressionFractal3D.generated.h"
+
+/**
+ * Zero-centered 3D Fractal noise in 1, 2, 3 or 4 channels, created by summing several
+ * octaves of 3D Perlin noise, increasing the frequency and decreasing the amplitude at each octave.
+ */
+
+UCLASS(collapsecategories, hidecategories = Object, MinimalAPI)
+class UMaterialExpressionFractal3D : public UMaterialExpression
+{
+	GENERATED_UCLASS_BODY()
+
+	/** The name of a vector3-type node specifying the 3D position at which the noise is evaluated. By default the vector is in local space*/
+	UPROPERTY()
+	FExpressionInput Position;	
+
+	/** Center-to-peak amplitude of the noise (peak-to-peak amplitude is 2x this value).*/
+	UPROPERTY(meta = (RequiredInput = "false", ToolTip = "Defaults to 'ConstAmplitude' if not specified"))
+	FExpressionInput Amplitude;
+
+	/** only used if Amplitude is not hooked up */
+	UPROPERTY(EditAnywhere, Category = MaterialExpressionFractal, meta = (OverridingInputProperty = "Amplitude"))
+	float ConstAmplitude = 1;
+
+	/** The number of octaves of noise to be summed. */
+	UPROPERTY(meta = (RequiredInput = "false", ToolTip = "Defaults to 'ConstOctaves' if not specified"))
+	FExpressionInput Octaves;
+
+	/** only used if Octaves is not hooked up */
+	UPROPERTY(EditAnywhere, Category = MaterialExpressionFractal, meta = (OverridingInputProperty = "Octaves"))
+	int32 ConstOctaves = 3;
+
+	/** The exponential scale between successive octaves of noise. */
+	UPROPERTY(meta = (RequiredInput = "false", ToolTip = "Defaults to 'ConstLacunarity' if not specified"))
+	FExpressionInput Lacunarity;
+
+	/** only used if Lacunarity is not hooked up */
+	UPROPERTY(EditAnywhere, Category = MaterialExpressionFractal, meta = (OverridingInputProperty = "Lacunarity"))
+	float ConstLacunarity = 2;
+
+	/** The rate at which noise amplitude is diminished for each octave. Should be between 0.0 and 1.0 */
+	UPROPERTY(meta = (RequiredInput = "false", ToolTip = "Defaults to 'ConstDiminish' if not specified"))
+	FExpressionInput Diminish;
+
+	/** only used if Diminish is not hooked up */
+	UPROPERTY(EditAnywhere, Category = MaterialExpressionFractal, meta = (OverridingInputProperty = "Diminish"))
+	float ConstDiminish = 0.5;
+
+	/** can also be done with a multiply on the Position */
+	UPROPERTY(EditAnywhere, Category = MaterialExpressionNoise)
+	float Scale = 1.f;
+
+	/** How multiple frequencies are getting combined */
+	UPROPERTY(EditAnywhere, Category = MaterialExpressionNoise)
+	bool bTurbulence = false;
+
+	/** 1 = fast but little detail, .. larger numbers cost more performance */
+	UPROPERTY(EditAnywhere, Category = MaterialExpressionNoise, meta = (UIMin = "1", UIMax = "10"))
+	int32 Levels = 6;
+
+	UPROPERTY(EditAnywhere, Category = MaterialExpressionNoise)
+	float OutputMin = 0.f;
+
+	UPROPERTY(EditAnywhere, Category = MaterialExpressionNoise)
+	float OutputMax = 1.f;
+
+	//~ Begin UMaterialExpression Interface
+#if WITH_EDITOR
+	virtual int32 Compile(class FMaterialCompiler* Compiler, int32 OutputIndex) override;
+	virtual void GetCaption(TArray<FString>& OutCaptions) const override;
+#endif
+	//~ End UMaterialExpression Interface
+};
+
diff --git a/Engine/Plugins/Interchange/Runtime/Source/Import/Private/MaterialX/MaterialExpressions/MaterialExpressionHsvToRgb.cpp b/Engine/Plugins/Interchange/Runtime/Source/Import/Private/MaterialX/MaterialExpressions/MaterialExpressionHsvToRgb.cpp
new file mode 100644
--- /dev/null
+++ b/Engine/Plugins/Interchange/Runtime/Source/Import/Private/MaterialX/MaterialExpressions/MaterialExpressionHsvToRgb.cpp
@@ -0,0 +1,56 @@
+// Copyright Epic Games, Inc. All Rights Reserved.
+
+#include "MaterialExpressionHsvToRgb.h"
+#include "MaterialCompiler.h"
+#include "Materials/MaterialExpressionCustom.h"
+
+#define LOCTEXT_NAMESPACE "MaterialExpressionHsvToRgb"
+
+UMaterialExpressionHsvToRgb::UMaterialExpressionHsvToRgb(const FObjectInitializer& ObjectInitializer)
+	: Super(ObjectInitializer)
+{
+	// Structure to hold one-time initialization
+	struct FConstructorStatics
+	{
+		FText NAME_MaterialX;
+		FText NAME_ImageAdjustment;
+		FConstructorStatics()
+			: NAME_MaterialX(LOCTEXT("MaterialX", "MaterialX"))
+			, NAME_ImageAdjustment(LOCTEXT("Image Adjustment", "Image Adjustment"))
+		{}
+	};
+	static FConstructorStatics ConstructorStatics;
+
+#if WITH_EDITORONLY_DATA
+	MenuCategories.Add(ConstructorStatics.NAME_MaterialX);
+	MenuCategories.Add(ConstructorStatics.NAME_ImageAdjustment);
+#endif
+}
+
+#if WITH_EDITOR
+int32 UMaterialExpressionHsvToRgb::Compile(FMaterialCompiler* Compiler, int32 OutputIndex)
+{
+	if(!Input.GetTracedInput().Expression)
+	{
+		return Compiler->Errorf(TEXT("Missing RGBtoHSV input"));
+	}
+
+	UMaterialExpressionCustom* MaterialExpressionCustom = NewObject<UMaterialExpressionCustom>();
+	MaterialExpressionCustom->Inputs[0].InputName = TEXT("c");
+	MaterialExpressionCustom->Inputs[0].Input = Input;
+	MaterialExpressionCustom->OutputType = ECustomMaterialOutputType::CMOT_Float3;
+	MaterialExpressionCustom->Code = TEXT(R"(
+    float4 K = float4(1.0, 2.0 / 3.0, 1.0 / 3.0, 3.0);
+    float3 p = abs(frac(c.xxx + K.xyz) * 6.0 - K.www);
+    return c.z * lerp(K.xxx, saturate(p - K.xxx), c.y);)");
+
+	return MaterialExpressionCustom->Compile(Compiler, OutputIndex);
+}
+
+void UMaterialExpressionHsvToRgb::GetCaption(TArray<FString>& OutCaptions) const
+{
+	OutCaptions.Add(TEXT("HSVToRGB"));
+}
+#endif
+
+#undef LOCTEXT_NAMESPACE 
diff --git a/Engine/Plugins/Interchange/Runtime/Source/Import/Private/MaterialX/MaterialExpressions/MaterialExpressionHsvToRgb.h b/Engine/Plugins/Interchange/Runtime/Source/Import/Private/MaterialX/MaterialExpressions/MaterialExpressionHsvToRgb.h
new file mode 100644
--- /dev/null
+++ b/Engine/Plugins/Interchange/Runtime/Source/Import/Private/MaterialX/MaterialExpressions/MaterialExpressionHsvToRgb.h
@@ -0,0 +1,23 @@
+// Copyright Epic Games, Inc. All Rights Reserved.
+#pragma once
+
+#include "Materials/MaterialExpression.h"
+
+#include "MaterialExpressionHsvToRgb.generated.h"
+
+UCLASS(collapsecategories, hidecategories = Object, MinimalAPI)
+class UMaterialExpressionHsvToRgb : public UMaterialExpression
+{
+	GENERATED_UCLASS_BODY()
+
+	UPROPERTY()
+	FExpressionInput Input;
+
+	//~ Begin UMaterialExpression Interface
+#if WITH_EDITOR
+	virtual int32 Compile(class FMaterialCompiler* Compiler, int32 OutputIndex) override;
+	virtual void GetCaption(TArray<FString>& OutCaptions) const override;
+#endif
+	//~ End UMaterialExpression Interface
+};
+
diff --git a/Engine/Plugins/Interchange/Runtime/Source/Import/Private/MaterialX/MaterialExpressions/MaterialExpressionIn.cpp b/Engine/Plugins/Interchange/Runtime/Source/Import/Private/MaterialX/MaterialExpressions/MaterialExpressionIn.cpp
new file mode 100644
--- /dev/null
+++ b/Engine/Plugins/Interchange/Runtime/Source/Import/Private/MaterialX/MaterialExpressions/MaterialExpressionIn.cpp
@@ -0,0 +1,59 @@
+// Copyright Epic Games, Inc. All Rights Reserved.
+#include "MaterialExpressionIn.h"
+#include "MaterialCompiler.h"
+
+#include UE_INLINE_GENERATED_CPP_BY_NAME(MaterialExpressionIn)
+
+#define LOCTEXT_NAMESPACE "MaterialExpressionIn"
+
+UMaterialExpressionIn::UMaterialExpressionIn(const FObjectInitializer& ObjectInitializer)
+	: Super(ObjectInitializer)
+{
+	// Structure to hold one-time initialization
+	struct FConstructorStatics
+	{
+		FText NAME_MaterialX;
+		FText NAME_Compositing;
+		FConstructorStatics()
+			: NAME_MaterialX(LOCTEXT("MaterialX", "MaterialX"))
+			, NAME_Compositing(LOCTEXT("Compositing", "Compositing"))
+		{}
+	};
+	static FConstructorStatics ConstructorStatics;
+
+#if WITH_EDITORONLY_DATA
+	MenuCategories.Add(ConstructorStatics.NAME_MaterialX);
+	MenuCategories.Add(ConstructorStatics.NAME_Compositing);
+#endif
+}
+
+#if WITH_EDITOR
+int32 UMaterialExpressionIn::Compile(FMaterialCompiler* Compiler, int32 OutputIndex)
+{
+	if(!A.GetTracedInput().Expression)
+	{
+		return Compiler->Errorf(TEXT("Missing A input"));
+	}
+
+	if(!B.GetTracedInput().Expression)
+	{
+		return Compiler->Errorf(TEXT("Missing B input"));
+	}
+
+	int32 IndexAlpha = Alpha.GetTracedInput().Expression ? Alpha.Compile(Compiler) : Compiler->Constant(ConstAlpha);
+
+	int32 IndexA = A.Compile(Compiler);
+	int32 IndexB = B.Compile(Compiler);
+
+	int32 IndexIn = Compiler->Mul(IndexA, Compiler->ComponentMask(IndexB, false, false, false, true));
+
+	return Compiler->Lerp(IndexB, IndexIn, IndexAlpha);
+}
+
+void UMaterialExpressionIn::GetCaption(TArray<FString>& OutCaptions) const
+{
+	OutCaptions.Add(TEXT("In"));
+}
+#endif
+
+#undef LOCTEXT_NAMESPACE 
\ No newline at end of file
diff --git a/Engine/Plugins/Interchange/Runtime/Source/Import/Private/MaterialX/MaterialExpressions/MaterialExpressionIn.h b/Engine/Plugins/Interchange/Runtime/Source/Import/Private/MaterialX/MaterialExpressions/MaterialExpressionIn.h
new file mode 100644
--- /dev/null
+++ b/Engine/Plugins/Interchange/Runtime/Source/Import/Private/MaterialX/MaterialExpressions/MaterialExpressionIn.h
@@ -0,0 +1,45 @@
+// Copyright Epic Games, Inc. All Rights Reserved.
+#pragma once
+
+#include "CoreMinimal.h"
+#include "Materials/MaterialExpression.h"
+
+#include "MaterialExpressionIn.generated.h"
+
+/**
+ * Merge nodes take two 4-channel (color4) inputs and use the built-in alpha channel(s) to control the
+ * compositing of the A and B inputs. "A" and "B" refer to the non-alpha channels of the A and B inputs respectively,
+ * and "a" and "b" refer to the alpha channels of the A and B inputs.
+ * Merge nodes are only defined for float4 inputs
+ * Merge nodes support an optional float input Alpha , which can be used to mix the
+ * original B value with the result of the blend operation.
+ *
+ * Operation: Ab
+ * Result: Lerp(B, Ab, Alpha)
+ */
+UCLASS(MinimalAPI, collapsecategories, hidecategories = Object)
+class UMaterialExpressionIn: public UMaterialExpression
+{
+	GENERATED_UCLASS_BODY()
+
+	UPROPERTY()
+	FExpressionInput A;
+
+	UPROPERTY()
+	FExpressionInput B;
+
+	UPROPERTY(meta = (RequiredInput = "false", ToolTip = "Defaults to 'ConstAlpha' if not specified"))
+	FExpressionInput Alpha;
+
+	/** only used if Alpha is not hooked up */
+	UPROPERTY(EditAnywhere, Category = MaterialExpressionLinearInterpolate, meta = (OverridingInputProperty = "Alpha"))
+	float ConstAlpha = 1.f;
+
+	//~ Begin UMaterialExpression Interface
+#if WITH_EDITOR
+	virtual int32 Compile(class FMaterialCompiler* Compiler, int32 OutputIndex) override;
+	virtual void GetCaption(TArray<FString>& OutCaptions) const override;
+#endif
+	//~ End UMaterialExpression Interface
+};
+
diff --git a/Engine/Plugins/Interchange/Runtime/Source/Import/Private/MaterialX/MaterialExpressions/MaterialExpressionLength.cpp b/Engine/Plugins/Interchange/Runtime/Source/Import/Private/MaterialX/MaterialExpressions/MaterialExpressionLength.cpp
new file mode 100644
--- /dev/null
+++ b/Engine/Plugins/Interchange/Runtime/Source/Import/Private/MaterialX/MaterialExpressions/MaterialExpressionLength.cpp
@@ -0,0 +1,50 @@
+// Copyright Epic Games, Inc. All Rights Reserved.
+#include "MaterialExpressionLength.h"
+#include "MaterialCompiler.h"
+
+#include UE_INLINE_GENERATED_CPP_BY_NAME(MaterialExpressionLength)
+
+#define LOCTEXT_NAMESPACE "MaterialExpressionLength"
+
+UMaterialExpressionLength::UMaterialExpressionLength(const FObjectInitializer& ObjectInitializer)
+	: Super(ObjectInitializer)
+{
+	// Structure to hold one-time initialization
+	struct FConstructorStatics
+	{
+		FText NAME_MaterialX;
+		FText NAME_Math;
+		FText NAME_VectorOps;
+		FConstructorStatics()
+			: NAME_MaterialX(LOCTEXT("MaterialX", "MaterialX"))
+			, NAME_Math(LOCTEXT("Math", "Math"))
+			, NAME_VectorOps(LOCTEXT("VectorOps", "VectorOps"))
+		{}
+	};
+	static FConstructorStatics ConstructorStatics;
+
+#if WITH_EDITORONLY_DATA
+	MenuCategories.Add(ConstructorStatics.NAME_MaterialX);
+	MenuCategories.Add(ConstructorStatics.NAME_Math);
+	MenuCategories.Add(ConstructorStatics.NAME_VectorOps);
+#endif
+}
+
+#if WITH_EDITOR
+int32 UMaterialExpressionLength::Compile(FMaterialCompiler* Compiler, int32 OutputIndex)
+{
+	if(!Input.GetTracedInput().Expression)
+	{
+		return Compiler->Errorf(TEXT("Missing Length input"));
+	}
+
+	return Compiler->Length(Input.Compile(Compiler));
+}
+
+void UMaterialExpressionLength::GetCaption(TArray<FString>& OutCaptions) const
+{
+	OutCaptions.Add(TEXT("Length"));
+}
+#endif
+
+#undef LOCTEXT_NAMESPACE 
\ No newline at end of file
diff --git a/Engine/Plugins/Interchange/Runtime/Source/Import/Private/MaterialX/MaterialExpressions/MaterialExpressionLength.h b/Engine/Plugins/Interchange/Runtime/Source/Import/Private/MaterialX/MaterialExpressions/MaterialExpressionLength.h
new file mode 100644
--- /dev/null
+++ b/Engine/Plugins/Interchange/Runtime/Source/Import/Private/MaterialX/MaterialExpressions/MaterialExpressionLength.h
@@ -0,0 +1,27 @@
+// Copyright Epic Games, Inc. All Rights Reserved.
+#pragma once
+
+#include "CoreMinimal.h"
+#include "Materials/MaterialExpression.h"
+
+#include "MaterialExpressionLength.generated.h"
+
+/**
+ * A material expression that computes the length (magnitude) of a vector
+ */
+UCLASS(collapsecategories, hidecategories = Object, MinimalAPI)
+class UMaterialExpressionLength : public UMaterialExpression
+{
+	GENERATED_UCLASS_BODY()
+
+	UPROPERTY()
+	FExpressionInput Input;
+
+	//~ Begin UMaterialExpression Interface
+#if WITH_EDITOR
+	virtual int32 Compile(class FMaterialCompiler* Compiler, int32 OutputIndex) override;
+	virtual void GetCaption(TArray<FString>& OutCaptions) const override;
+#endif
+	//~ End UMaterialExpression Interface
+};
+
diff --git a/Engine/Plugins/Interchange/Runtime/Source/Import/Private/MaterialX/MaterialExpressions/MaterialExpressionLogarithm.cpp b/Engine/Plugins/Interchange/Runtime/Source/Import/Private/MaterialX/MaterialExpressions/MaterialExpressionLogarithm.cpp
new file mode 100644
--- /dev/null
+++ b/Engine/Plugins/Interchange/Runtime/Source/Import/Private/MaterialX/MaterialExpressions/MaterialExpressionLogarithm.cpp
@@ -0,0 +1,52 @@
+// Copyright Epic Games, Inc. All Rights Reserved.
+
+#include "MaterialExpressionLogarithm.h"
+#include "MaterialCompiler.h"
+#include "MaterialHLSLGenerator.h"
+
+#include UE_INLINE_GENERATED_CPP_BY_NAME(MaterialExpressionLogarithm)
+
+#define LOCTEXT_NAMESPACE "MaterialExpressionLogarithm"
+
+UMaterialExpressionLogarithm::UMaterialExpressionLogarithm(const FObjectInitializer& ObjectInitializer)
+	: Super(ObjectInitializer)
+{
+	// Structure to hold one-time initialization
+	struct FConstructorStatics
+	{
+
+		FText NAME_MaterialX;
+		FText NAME_Math;
+		FConstructorStatics()
+			: NAME_MaterialX(LOCTEXT("MaterialX", "MaterialX"))
+			, NAME_Math(LOCTEXT("Math", "Math"))
+		{}
+	};
+	static FConstructorStatics ConstructorStatics;
+
+#if WITH_EDITORONLY_DATA
+	MenuCategories.Add(ConstructorStatics.NAME_MaterialX);
+	MenuCategories.Add(ConstructorStatics.NAME_Math);
+#endif
+}
+
+#if WITH_EDITOR
+int32 UMaterialExpressionLogarithm::Compile(FMaterialCompiler* Compiler, int32 OutputIndex)
+{
+	if(!Input.GetTracedInput().Expression)
+	{
+		return Compiler->Errorf(TEXT("Missing Ln input"));
+	}
+
+	const int32 Log10Euler = Compiler->Constant(0.4342945);
+
+	return Compiler->Div(Compiler->Logarithm10(Input.Compile(Compiler)), Log10Euler);
+}
+
+void UMaterialExpressionLogarithm::GetCaption(TArray<FString>& OutCaptions) const
+{
+	OutCaptions.Add(TEXT("Ln"));
+}
+#endif
+
+#undef LOCTEXT_NAMESPACE 
\ No newline at end of file
diff --git a/Engine/Plugins/Interchange/Runtime/Source/Import/Private/MaterialX/MaterialExpressions/MaterialExpressionLogarithm.h b/Engine/Plugins/Interchange/Runtime/Source/Import/Private/MaterialX/MaterialExpressions/MaterialExpressionLogarithm.h
new file mode 100644
--- /dev/null
+++ b/Engine/Plugins/Interchange/Runtime/Source/Import/Private/MaterialX/MaterialExpressions/MaterialExpressionLogarithm.h
@@ -0,0 +1,27 @@
+// Copyright Epic Games, Inc. All Rights Reserved.
+#pragma once
+
+#include "CoreMinimal.h"
+#include "Materials/MaterialExpression.h"
+
+#include "MaterialExpressionLogarithm.generated.h"
+
+/**
+ * A material expression that computes the natural logarithm
+ */
+UCLASS(collapsecategories, hidecategories = Object, MinimalAPI)
+class UMaterialExpressionLogarithm : public UMaterialExpression
+{
+	GENERATED_UCLASS_BODY()
+
+	UPROPERTY()
+	FExpressionInput Input;
+
+	//~ Begin UMaterialExpression Interface
+#if WITH_EDITOR
+	virtual int32 Compile(class FMaterialCompiler* Compiler, int32 OutputIndex) override;
+	virtual void GetCaption(TArray<FString>& OutCaptions) const override;
+#endif
+	//~ End UMaterialExpression Interface
+};
+
diff --git a/Engine/Plugins/Interchange/Runtime/Source/Import/Private/MaterialX/MaterialExpressions/MaterialExpressionLuminance.cpp b/Engine/Plugins/Interchange/Runtime/Source/Import/Private/MaterialX/MaterialExpressions/MaterialExpressionLuminance.cpp
new file mode 100644
--- /dev/null
+++ b/Engine/Plugins/Interchange/Runtime/Source/Import/Private/MaterialX/MaterialExpressions/MaterialExpressionLuminance.cpp
@@ -0,0 +1,84 @@
+// Copyright Epic Games, Inc. All Rights Reserved.
+
+#include "MaterialExpressionLuminance.h"
+#include "MaterialCompiler.h"
+
+#include UE_INLINE_GENERATED_CPP_BY_NAME(MaterialExpressionLuminance)
+
+#define LOCTEXT_NAMESPACE "MaterialExpressionLuminance"
+
+UMaterialExpressionLuminance::UMaterialExpressionLuminance(const FObjectInitializer& ObjectInitializer)
+	: Super(ObjectInitializer),
+	LuminanceFactors(0.2722287, 0.6740818, 0.0536895),
+	LuminanceMode(ELuminanceMode::ACEScg)
+{
+	// Structure to hold one-time initialization
+	struct FConstructorStatics
+	{
+		FText NAME_MaterialX;
+		FText NAME_Adjustment;
+		FConstructorStatics()
+			: NAME_MaterialX(LOCTEXT("MaterialX", "MaterialX"))
+			, NAME_Adjustment(LOCTEXT("Image Adjustment", "Image Adjustment"))
+		{}
+	};
+	static FConstructorStatics ConstructorStatics;
+
+#if WITH_EDITORONLY_DATA
+	MenuCategories.Add(ConstructorStatics.NAME_MaterialX);
+	MenuCategories.Add(ConstructorStatics.NAME_Adjustment);
+#endif
+}
+
+#if WITH_EDITOR
+int32 UMaterialExpressionLuminance::Compile(FMaterialCompiler* Compiler, int32 OutputIndex)
+{
+	if(!Input.GetTracedInput().Expression)
+	{
+		return Compiler->Errorf(TEXT("Missing Luminance input"));
+	}
+
+	const int32 LuminanceFactorsIndex = Compiler->Constant3(LuminanceFactors.R, LuminanceFactors.G, LuminanceFactors.B);
+	auto Type = Compiler->GetParameterType(Input.Compile(Compiler));
+
+	return Compiler->Dot(Compiler->ComponentMask(Input.Compile(Compiler), true, true, true, false), LuminanceFactorsIndex);
+}
+
+void UMaterialExpressionLuminance::GetCaption(TArray<FString>& OutCaptions) const
+{
+	OutCaptions.Add(TEXT("Luminance"));
+}
+
+void UMaterialExpressionLuminance::PostEditChangeProperty(FPropertyChangedEvent& PropertyChangedEvent)
+{
+	FProperty* PropertyThatChanged = PropertyChangedEvent.MemberProperty;
+	const FString PropertyName = PropertyThatChanged ? PropertyThatChanged->GetName() : TEXT("");
+
+	if(PropertyName == GET_MEMBER_NAME_STRING_CHECKED(UMaterialExpressionLuminance, LuminanceMode))
+	{
+		switch(LuminanceMode)
+		{
+		case ELuminanceMode::ACEScg:
+			LuminanceFactors = FLinearColor(0.2722287, 0.6740818, 0.0536895);
+			break;
+
+		case ELuminanceMode::Rec709:
+			LuminanceFactors = FLinearColor(0.2126, 0.7152, 0.0722);
+			break;
+
+		case ELuminanceMode::Rec2020:
+			LuminanceFactors = FLinearColor(0.2627, 0.6780, 0.0593);
+			break;
+		}
+	}
+	else if(PropertyName == GET_MEMBER_NAME_STRING_CHECKED(UMaterialExpressionLuminance, LuminanceFactors))
+	{
+		LuminanceMode = ELuminanceMode::Custom;
+	}
+		
+	Super::PostEditChangeProperty(PropertyChangedEvent);
+}
+
+#endif
+
+#undef LOCTEXT_NAMESPACE 
\ No newline at end of file
diff --git a/Engine/Plugins/Interchange/Runtime/Source/Import/Private/MaterialX/MaterialExpressions/MaterialExpressionLuminance.h b/Engine/Plugins/Interchange/Runtime/Source/Import/Private/MaterialX/MaterialExpressions/MaterialExpressionLuminance.h
new file mode 100644
--- /dev/null
+++ b/Engine/Plugins/Interchange/Runtime/Source/Import/Private/MaterialX/MaterialExpressions/MaterialExpressionLuminance.h
@@ -0,0 +1,46 @@
+// Copyright Epic Games, Inc. All Rights Reserved.
+#pragma once
+
+#include "CoreMinimal.h"
+#include "Materials/MaterialExpression.h"
+
+#include "MaterialExpressionLuminance.generated.h"
+
+/**
+ * A material expression that outputs a grayscale image containing the luminance of the incoming RGB color in all color channels;
+    the alpha channel is left unchanged if present.
+ */
+
+UENUM()
+enum class ELuminanceMode : uint8
+{
+	ACEScg,
+	Rec709,
+	Rec2020,
+	Rec2100 = Rec2020,
+	Custom
+};
+
+UCLASS(collapsecategories, hidecategories = Object, MinimalAPI)
+class UMaterialExpressionLuminance : public UMaterialExpression
+{
+	GENERATED_UCLASS_BODY()
+
+	UPROPERTY()
+	FExpressionInput Input;
+
+	UPROPERTY(EditAnywhere, Category = MaterialExpressionLuminance, Meta = (ShowAsInputPin = "Advanced"))
+	FLinearColor LuminanceFactors;    // Color component factors for converting a color to greyscale.
+	
+	UPROPERTY(EditAnywhere, Category = MaterialExpressionLuminance, Meta = (ShowAsInputPin = "Advanced"))
+	ELuminanceMode LuminanceMode;
+	
+	//~ Begin UMaterialExpression Interface
+#if WITH_EDITOR
+	virtual int32 Compile(class FMaterialCompiler* Compiler, int32 OutputIndex) override;
+	virtual void GetCaption(TArray<FString>& OutCaptions) const override;
+	virtual void PostEditChangeProperty(FPropertyChangedEvent& PropertyChangedEvent) override;
+#endif
+	//~ End UMaterialExpression Interface
+};
+
diff --git a/Engine/Plugins/Interchange/Runtime/Source/Import/Private/MaterialX/MaterialExpressions/MaterialExpressionMask.cpp b/Engine/Plugins/Interchange/Runtime/Source/Import/Private/MaterialX/MaterialExpressions/MaterialExpressionMask.cpp
new file mode 100644
--- /dev/null
+++ b/Engine/Plugins/Interchange/Runtime/Source/Import/Private/MaterialX/MaterialExpressions/MaterialExpressionMask.cpp
@@ -0,0 +1,59 @@
+// Copyright Epic Games, Inc. All Rights Reserved.
+#include "MaterialExpressionMask.h"
+#include "MaterialCompiler.h"
+
+#include UE_INLINE_GENERATED_CPP_BY_NAME(MaterialExpressionMask)
+
+#define LOCTEXT_NAMESPACE "MaterialExpressionMask"
+
+UMaterialExpressionMask::UMaterialExpressionMask(const FObjectInitializer& ObjectInitializer)
+	: Super(ObjectInitializer)
+{
+	// Structure to hold one-time initialization
+	struct FConstructorStatics
+	{
+		FText NAME_MaterialX;
+		FText NAME_Compositing;
+		FConstructorStatics()
+			: NAME_MaterialX(LOCTEXT("MaterialX", "MaterialX"))
+			, NAME_Compositing(LOCTEXT("Compositing", "Compositing"))
+		{}
+	};
+	static FConstructorStatics ConstructorStatics;
+
+#if WITH_EDITORONLY_DATA
+	MenuCategories.Add(ConstructorStatics.NAME_MaterialX);
+	MenuCategories.Add(ConstructorStatics.NAME_Compositing);
+#endif
+}
+
+#if WITH_EDITOR
+int32 UMaterialExpressionMask::Compile(FMaterialCompiler* Compiler, int32 OutputIndex)
+{
+	if(!A.GetTracedInput().Expression)
+	{
+		return Compiler->Errorf(TEXT("Missing A input"));
+	}
+
+	if(!B.GetTracedInput().Expression)
+	{
+		return Compiler->Errorf(TEXT("Missing B input"));
+	}
+
+	int32 IndexAlpha = Alpha.GetTracedInput().Expression ? Alpha.Compile(Compiler) : Compiler->Constant(ConstAlpha);
+
+	int32 IndexA = A.Compile(Compiler);
+	int32 IndexB = B.Compile(Compiler);
+
+	int32 Mask = Compiler->Mul(IndexB, Compiler->ComponentMask(IndexA, false, false, false, true));
+
+	return Compiler->Lerp(IndexB, Mask, IndexAlpha);
+}
+
+void UMaterialExpressionMask::GetCaption(TArray<FString>& OutCaptions) const
+{
+	OutCaptions.Add(TEXT("Mask"));
+}
+#endif
+
+#undef LOCTEXT_NAMESPACE 
\ No newline at end of file
diff --git a/Engine/Plugins/Interchange/Runtime/Source/Import/Private/MaterialX/MaterialExpressions/MaterialExpressionMask.h b/Engine/Plugins/Interchange/Runtime/Source/Import/Private/MaterialX/MaterialExpressions/MaterialExpressionMask.h
new file mode 100644
--- /dev/null
+++ b/Engine/Plugins/Interchange/Runtime/Source/Import/Private/MaterialX/MaterialExpressions/MaterialExpressionMask.h
@@ -0,0 +1,45 @@
+// Copyright Epic Games, Inc. All Rights Reserved.
+#pragma once
+
+#include "CoreMinimal.h"
+#include "Materials/MaterialExpression.h"
+
+#include "MaterialExpressionMask.generated.h"
+
+/**
+ * Merge nodes take two 4-channel (color4) inputs and use the built-in alpha channel(s) to control the
+ * compositing of the A and B inputs. "A" and "B" refer to the non-alpha channels of the A and B inputs respectively,
+ * and "a" and "b" refer to the alpha channels of the A and B inputs.
+ * Merge nodes are only defined for float4 inputs
+ * Merge nodes support an optional float input Alpha , which can be used to mix the
+ * original B value with the result of the blend operation.
+ *
+ * Operation: B*a
+ * Result: Lerp(B, B*a, Alpha)
+ */
+UCLASS(MinimalAPI, collapsecategories, hidecategories = Object)
+class UMaterialExpressionMask : public UMaterialExpression
+{
+	GENERATED_UCLASS_BODY()
+
+	UPROPERTY()
+	FExpressionInput A;
+
+	UPROPERTY()
+	FExpressionInput B;
+
+	UPROPERTY(meta = (RequiredInput = "false", ToolTip = "Defaults to 'ConstAlpha' if not specified"))
+	FExpressionInput Alpha;
+
+	/** only used if Alpha is not hooked up */
+	UPROPERTY(EditAnywhere, Category = MaterialExpressionLinearInterpolate, meta = (OverridingInputProperty = "Alpha"))
+	float ConstAlpha = 1.f;
+
+	//~ Begin UMaterialExpression Interface
+#if WITH_EDITOR
+	virtual int32 Compile(class FMaterialCompiler* Compiler, int32 OutputIndex) override;
+	virtual void GetCaption(TArray<FString>& OutCaptions) const override;
+#endif
+	//~ End UMaterialExpression Interface
+};
+
diff --git a/Engine/Plugins/Interchange/Runtime/Source/Import/Private/MaterialX/MaterialExpressions/MaterialExpressionMatte.cpp b/Engine/Plugins/Interchange/Runtime/Source/Import/Private/MaterialX/MaterialExpressions/MaterialExpressionMatte.cpp
new file mode 100644
--- /dev/null
+++ b/Engine/Plugins/Interchange/Runtime/Source/Import/Private/MaterialX/MaterialExpressions/MaterialExpressionMatte.cpp
@@ -0,0 +1,68 @@
+// Copyright Epic Games, Inc. All Rights Reserved.
+#include "MaterialExpressionMatte.h"
+#include "MaterialCompiler.h"
+
+#include UE_INLINE_GENERATED_CPP_BY_NAME(MaterialExpressionMatte)
+
+#define LOCTEXT_NAMESPACE "MaterialExpressionMatte"
+
+UMaterialExpressionMatte::UMaterialExpressionMatte(const FObjectInitializer& ObjectInitializer)
+	: Super(ObjectInitializer)
+{
+	// Structure to hold one-time initialization
+	struct FConstructorStatics
+	{
+		FText NAME_MaterialX;
+		FText NAME_Compositing;
+		FConstructorStatics()
+			: NAME_MaterialX(LOCTEXT("MaterialX", "MaterialX"))
+			, NAME_Compositing(LOCTEXT("Compositing", "Compositing"))
+		{}
+	};
+	static FConstructorStatics ConstructorStatics;
+
+#if WITH_EDITORONLY_DATA
+	MenuCategories.Add(ConstructorStatics.NAME_MaterialX);
+	MenuCategories.Add(ConstructorStatics.NAME_Compositing);
+#endif
+}
+
+#if WITH_EDITOR
+int32 UMaterialExpressionMatte::Compile(FMaterialCompiler* Compiler, int32 OutputIndex)
+{
+	if(!A.GetTracedInput().Expression)
+	{
+		return Compiler->Errorf(TEXT("Missing A input"));
+	}
+
+	if(!B.GetTracedInput().Expression)
+	{
+		return Compiler->Errorf(TEXT("Missing B input"));
+	}
+
+	int32 IndexAlpha = Alpha.GetTracedInput().Expression ? Alpha.Compile(Compiler) : Compiler->Constant(ConstAlpha);
+
+	int32 IndexA = A.Compile(Compiler);
+	int32 IndexB = B.Compile(Compiler);
+
+	int32 IndexAlphaA = Compiler->ComponentMask(IndexA, false, false, false, true);
+	int32 IndexRgbA = Compiler->ComponentMask(IndexA, true, true, true, false);
+	int32 IndexAlphaB = Compiler->ComponentMask(IndexB, false, false, false, true);
+	int32 IndexRgbB = Compiler->ComponentMask(IndexB, true, true, true, false);
+
+	int32 IndexOneMinusAlphaA = Compiler->Sub(Compiler->Constant(1.f), IndexAlphaA);
+	int32 XYZ = Compiler->Add(Compiler->Mul(IndexRgbA, IndexAlphaA), Compiler->Mul(IndexRgbB, IndexOneMinusAlphaA));
+	int32 W = Compiler->Add(IndexAlphaA, Compiler->Mul(IndexAlphaB, IndexOneMinusAlphaA));
+
+	int32 Matte = Compiler->AppendVector(XYZ, W);
+
+	return Compiler->Lerp(IndexB, Matte, IndexAlpha);
+}
+
+void UMaterialExpressionMatte::GetCaption(TArray<FString>& OutCaptions) const
+{
+	OutCaptions.Add(TEXT("Matte"));
+}
+#endif
+
+#undef LOCTEXT_NAMESPACE 
\ No newline at end of file
diff --git a/Engine/Plugins/Interchange/Runtime/Source/Import/Private/MaterialX/MaterialExpressions/MaterialExpressionMatte.h b/Engine/Plugins/Interchange/Runtime/Source/Import/Private/MaterialX/MaterialExpressions/MaterialExpressionMatte.h
new file mode 100644
--- /dev/null
+++ b/Engine/Plugins/Interchange/Runtime/Source/Import/Private/MaterialX/MaterialExpressions/MaterialExpressionMatte.h
@@ -0,0 +1,45 @@
+// Copyright Epic Games, Inc. All Rights Reserved.
+#pragma once
+
+#include "CoreMinimal.h"
+#include "Materials/MaterialExpression.h"
+
+#include "MaterialExpressionMatte.generated.h"
+
+/**
+ * Merge nodes take two 4-channel (color4) inputs and use the built-in alpha channel(s) to control the
+ * compositing of the A and B inputs. "A" and "B" refer to the non-alpha channels of the A and B inputs respectively,
+ * and "a" and "b" refer to the alpha channels of the A and B inputs.
+ * Merge nodes are only defined for float4 inputs
+ * Merge nodes support an optional float input Alpha , which can be used to mix the
+ * original B value with the result of the blend operation.
+ *
+ * Operation: A*a + B*(1-a)
+ * Result: Lerp(B, A*a + B*(1-a), Alpha)
+ */
+UCLASS(MinimalAPI, collapsecategories, hidecategories = Object)
+class UMaterialExpressionMatte : public UMaterialExpression
+{
+	GENERATED_UCLASS_BODY()
+
+	UPROPERTY()
+	FExpressionInput A;
+
+	UPROPERTY()
+	FExpressionInput B;
+
+	UPROPERTY(meta = (RequiredInput = "false", ToolTip = "Defaults to 'ConstAlpha' if not specified"))
+	FExpressionInput Alpha;
+
+	/** only used if Alpha is not hooked up */
+	UPROPERTY(EditAnywhere, Category = MaterialExpressionLinearInterpolate, meta = (OverridingInputProperty = "Alpha"))
+	float ConstAlpha = 1.f;
+
+	//~ Begin UMaterialExpression Interface
+#if WITH_EDITOR
+	virtual int32 Compile(class FMaterialCompiler* Compiler, int32 OutputIndex) override;
+	virtual void GetCaption(TArray<FString>& OutCaptions) const override;
+#endif
+	//~ End UMaterialExpression Interface
+};
+
diff --git a/Engine/Plugins/Interchange/Runtime/Source/Import/Private/MaterialX/MaterialExpressions/MaterialExpressionMinus.cpp b/Engine/Plugins/Interchange/Runtime/Source/Import/Private/MaterialX/MaterialExpressions/MaterialExpressionMinus.cpp
new file mode 100644
--- /dev/null
+++ b/Engine/Plugins/Interchange/Runtime/Source/Import/Private/MaterialX/MaterialExpressions/MaterialExpressionMinus.cpp
@@ -0,0 +1,56 @@
+// Copyright Epic Games, Inc. All Rights Reserved.
+#include "MaterialExpressionMinus.h"
+#include "MaterialCompiler.h"
+
+#include UE_INLINE_GENERATED_CPP_BY_NAME(MaterialExpressionMinus)
+
+#define LOCTEXT_NAMESPACE "MaterialExpressionMinus"
+
+UMaterialExpressionMinus::UMaterialExpressionMinus(const FObjectInitializer& ObjectInitializer)
+	: Super(ObjectInitializer)
+{
+	// Structure to hold one-time initialization
+	struct FConstructorStatics
+	{
+		FText NAME_MaterialX;
+		FText NAME_Compositing;
+		FConstructorStatics()
+			: NAME_MaterialX(LOCTEXT("MaterialX", "MaterialX"))
+			, NAME_Compositing(LOCTEXT("Compositing", "Compositing"))
+		{}
+	};
+	static FConstructorStatics ConstructorStatics;
+
+#if WITH_EDITORONLY_DATA
+	MenuCategories.Add(ConstructorStatics.NAME_MaterialX);
+	MenuCategories.Add(ConstructorStatics.NAME_Compositing);
+#endif
+}
+
+#if WITH_EDITOR
+int32 UMaterialExpressionMinus::Compile(FMaterialCompiler* Compiler, int32 OutputIndex)
+{
+	if(!A.GetTracedInput().Expression)
+	{
+		return Compiler->Errorf(TEXT("Missing A input"));
+	}
+
+	if(!B.GetTracedInput().Expression)
+	{
+		return Compiler->Errorf(TEXT("Missing B input"));
+	}
+
+	int32 IndexAlpha = Alpha.GetTracedInput().Expression ? Alpha.Compile(Compiler) : Compiler->Constant(ConstAlpha);
+
+	int32 IndexB = B.Compile(Compiler);
+	int32 Sub = Compiler->Sub(A.Compile(Compiler), IndexB);
+	return Compiler->Lerp(IndexB, Sub, IndexAlpha);
+}
+
+void UMaterialExpressionMinus::GetCaption(TArray<FString>& OutCaptions) const
+{
+	OutCaptions.Add(TEXT("Minus"));
+}
+#endif
+
+#undef LOCTEXT_NAMESPACE 
\ No newline at end of file
diff --git a/Engine/Plugins/Interchange/Runtime/Source/Import/Private/MaterialX/MaterialExpressions/MaterialExpressionMinus.h b/Engine/Plugins/Interchange/Runtime/Source/Import/Private/MaterialX/MaterialExpressions/MaterialExpressionMinus.h
new file mode 100644
--- /dev/null
+++ b/Engine/Plugins/Interchange/Runtime/Source/Import/Private/MaterialX/MaterialExpressions/MaterialExpressionMinus.h
@@ -0,0 +1,41 @@
+// Copyright Epic Games, Inc. All Rights Reserved.
+#pragma once
+
+#include "CoreMinimal.h"
+#include "Materials/MaterialExpression.h"
+
+#include "MaterialExpressionMinus.generated.h"
+
+/**
+ * Blend nodes take two 1-4 channel inputs and apply the same operator to all channels.
+ * Blend nodes support an optional float input mix , which can be used
+ * to mix the original B value with the result of the blend operation.
+ * Operation: A - B
+ * Result: Lerp(B, A - B, Alpha)
+ */
+UCLASS(MinimalAPI, collapsecategories, hidecategories = Object)
+class UMaterialExpressionMinus : public UMaterialExpression
+{
+	GENERATED_UCLASS_BODY()
+
+	UPROPERTY()
+	FExpressionInput A;
+
+	UPROPERTY()
+	FExpressionInput B;
+
+	UPROPERTY(meta = (RequiredInput = "false", ToolTip = "Defaults to 'ConstAlpha' if not specified"))
+	FExpressionInput Alpha;
+
+	/** only used if Alpha is not hooked up */
+	UPROPERTY(EditAnywhere, Category = MaterialExpressionLinearInterpolate, meta = (OverridingInputProperty = "Alpha"))
+	float ConstAlpha = 1.f;
+
+	//~ Begin UMaterialExpression Interface
+#if WITH_EDITOR
+	virtual int32 Compile(class FMaterialCompiler* Compiler, int32 OutputIndex) override;
+	virtual void GetCaption(TArray<FString>& OutCaptions) const override;
+#endif
+	//~ End UMaterialExpression Interface
+};
+
diff --git a/Engine/Plugins/Interchange/Runtime/Source/Import/Private/MaterialX/MaterialExpressions/MaterialExpressionOut.cpp b/Engine/Plugins/Interchange/Runtime/Source/Import/Private/MaterialX/MaterialExpressions/MaterialExpressionOut.cpp
new file mode 100644
--- /dev/null
+++ b/Engine/Plugins/Interchange/Runtime/Source/Import/Private/MaterialX/MaterialExpressions/MaterialExpressionOut.cpp
@@ -0,0 +1,62 @@
+// Copyright Epic Games, Inc. All Rights Reserved.
+#include "MaterialExpressionOut.h"
+#include "MaterialCompiler.h"
+
+#include UE_INLINE_GENERATED_CPP_BY_NAME(MaterialExpressionOut)
+
+#define LOCTEXT_NAMESPACE "MaterialExpressionOut"
+
+UMaterialExpressionOut::UMaterialExpressionOut(const FObjectInitializer& ObjectInitializer)
+	: Super(ObjectInitializer)
+{
+	// Structure to hold one-time initialization
+	struct FConstructorStatics
+	{
+		FText NAME_MaterialX;
+		FText NAME_Compositing;
+		FConstructorStatics()
+			: NAME_MaterialX(LOCTEXT("MaterialX", "MaterialX"))
+			, NAME_Compositing(LOCTEXT("Compositing", "Compositing"))
+		{}
+	};
+	static FConstructorStatics ConstructorStatics;
+
+#if WITH_EDITORONLY_DATA
+	MenuCategories.Add(ConstructorStatics.NAME_MaterialX);
+	MenuCategories.Add(ConstructorStatics.NAME_Compositing);
+#endif
+}
+
+#if WITH_EDITOR
+int32 UMaterialExpressionOut::Compile(FMaterialCompiler* Compiler, int32 OutputIndex)
+{
+	if(!A.GetTracedInput().Expression)
+	{
+		return Compiler->Errorf(TEXT("Missing A input"));
+	}
+
+	if(!B.GetTracedInput().Expression)
+	{
+		return Compiler->Errorf(TEXT("Missing B input"));
+	}
+
+	int32 IndexAlpha = Alpha.GetTracedInput().Expression ? Alpha.Compile(Compiler) : Compiler->Constant(ConstAlpha);
+
+	int32 IndexA = A.Compile(Compiler);
+	int32 IndexB = B.Compile(Compiler);
+
+	int32 IndexAlphaB = Compiler->ComponentMask(IndexB, false, false, false, true);
+
+	int32 IndexOneMinusAlphaB = Compiler->Sub(Compiler->Constant(1.f), IndexAlphaB);
+	int32 Out = Compiler->Mul(IndexA, IndexOneMinusAlphaB);
+
+	return Compiler->Lerp(IndexB, Out, IndexAlpha);
+}
+
+void UMaterialExpressionOut::GetCaption(TArray<FString>& OutCaptions) const
+{
+	OutCaptions.Add(TEXT("Out"));
+}
+#endif
+
+#undef LOCTEXT_NAMESPACE 
\ No newline at end of file
diff --git a/Engine/Plugins/Interchange/Runtime/Source/Import/Private/MaterialX/MaterialExpressions/MaterialExpressionOut.h b/Engine/Plugins/Interchange/Runtime/Source/Import/Private/MaterialX/MaterialExpressions/MaterialExpressionOut.h
new file mode 100644
--- /dev/null
+++ b/Engine/Plugins/Interchange/Runtime/Source/Import/Private/MaterialX/MaterialExpressions/MaterialExpressionOut.h
@@ -0,0 +1,45 @@
+// Copyright Epic Games, Inc. All Rights Reserved.
+#pragma once
+
+#include "CoreMinimal.h"
+#include "Materials/MaterialExpression.h"
+
+#include "MaterialExpressionOut.generated.h"
+
+/**
+ * Merge nodes take two 4-channel (color4) inputs and use the built-in alpha channel(s) to control the
+ * compositing of the A and B inputs. "A" and "B" refer to the non-alpha channels of the A and B inputs respectively,
+ * and "a" and "b" refer to the alpha channels of the A and B inputs.
+ * Merge nodes are only defined for float4 inputs
+ * Merge nodes support an optional float input Alpha , which can be used to mix the
+ * original B value with the result of the blend operation.
+ *
+ * Operation: A(1-b)
+ * Result: Lerp(B, A(1-b), Alpha)
+ */
+UCLASS(MinimalAPI, collapsecategories, hidecategories = Object)
+class UMaterialExpressionOut : public UMaterialExpression
+{
+	GENERATED_UCLASS_BODY()
+
+	UPROPERTY()
+	FExpressionInput A;
+
+	UPROPERTY()
+	FExpressionInput B;
+
+	UPROPERTY(meta = (RequiredInput = "false", ToolTip = "Defaults to 'ConstAlpha' if not specified"))
+	FExpressionInput Alpha;
+
+	/** only used if Alpha is not hooked up */
+	UPROPERTY(EditAnywhere, Category = MaterialExpressionLinearInterpolate, meta = (OverridingInputProperty = "Alpha"))
+	float ConstAlpha = 1.f;
+
+	//~ Begin UMaterialExpression Interface
+#if WITH_EDITOR
+	virtual int32 Compile(class FMaterialCompiler* Compiler, int32 OutputIndex) override;
+	virtual void GetCaption(TArray<FString>& OutCaptions) const override;
+#endif
+	//~ End UMaterialExpression Interface
+};
+
diff --git a/Engine/Plugins/Interchange/Runtime/Source/Import/Private/MaterialX/MaterialExpressions/MaterialExpressionOver.cpp b/Engine/Plugins/Interchange/Runtime/Source/Import/Private/MaterialX/MaterialExpressions/MaterialExpressionOver.cpp
new file mode 100644
--- /dev/null
+++ b/Engine/Plugins/Interchange/Runtime/Source/Import/Private/MaterialX/MaterialExpressions/MaterialExpressionOver.cpp
@@ -0,0 +1,62 @@
+// Copyright Epic Games, Inc. All Rights Reserved.
+#include "MaterialExpressionOver.h"
+#include "MaterialCompiler.h"
+
+#include UE_INLINE_GENERATED_CPP_BY_NAME(MaterialExpressionOver)
+
+#define LOCTEXT_NAMESPACE "MaterialExpressionOver"
+
+UMaterialExpressionOver::UMaterialExpressionOver(const FObjectInitializer& ObjectInitializer)
+	: Super(ObjectInitializer)
+{
+	// Structure to hold one-time initialization
+	struct FConstructorStatics
+	{
+		FText NAME_MaterialX;
+		FText NAME_Compositing;
+		FConstructorStatics()
+			: NAME_MaterialX(LOCTEXT("MaterialX", "MaterialX"))
+			, NAME_Compositing(LOCTEXT("Compositing", "Compositing"))
+		{}
+	};
+	static FConstructorStatics ConstructorStatics;
+
+#if WITH_EDITORONLY_DATA
+	MenuCategories.Add(ConstructorStatics.NAME_MaterialX);
+	MenuCategories.Add(ConstructorStatics.NAME_Compositing);
+#endif
+}
+
+#if WITH_EDITOR
+int32 UMaterialExpressionOver::Compile(FMaterialCompiler* Compiler, int32 OutputIndex)
+{
+	if(!A.GetTracedInput().Expression)
+	{
+		return Compiler->Errorf(TEXT("Missing A input"));
+	}
+
+	if(!B.GetTracedInput().Expression)
+	{
+		return Compiler->Errorf(TEXT("Missing B input"));
+	}
+
+	int32 IndexAlpha = Alpha.GetTracedInput().Expression ? Alpha.Compile(Compiler) : Compiler->Constant(ConstAlpha);
+
+	int32 IndexA = A.Compile(Compiler);
+	int32 IndexB = B.Compile(Compiler);
+
+	int32 IndexAlphaA = Compiler->ComponentMask(IndexA, false, false, false, true);
+
+	int32 IndexOneMinusAlphaA = Compiler->Sub(Compiler->Constant(1.f), IndexAlphaA);
+	int32 Over = Compiler->Add(IndexA, Compiler->Mul(IndexB, IndexOneMinusAlphaA));
+
+	return Compiler->Lerp(IndexB, Over, IndexAlpha);
+}
+
+void UMaterialExpressionOver::GetCaption(TArray<FString>& OutCaptions) const
+{
+	OutCaptions.Add(TEXT("Over"));
+}
+#endif
+
+#undef LOCTEXT_NAMESPACE 
\ No newline at end of file
diff --git a/Engine/Plugins/Interchange/Runtime/Source/Import/Private/MaterialX/MaterialExpressions/MaterialExpressionOver.h b/Engine/Plugins/Interchange/Runtime/Source/Import/Private/MaterialX/MaterialExpressions/MaterialExpressionOver.h
new file mode 100644
--- /dev/null
+++ b/Engine/Plugins/Interchange/Runtime/Source/Import/Private/MaterialX/MaterialExpressions/MaterialExpressionOver.h
@@ -0,0 +1,45 @@
+// Copyright Epic Games, Inc. All Rights Reserved.
+#pragma once
+
+#include "CoreMinimal.h"
+#include "Materials/MaterialExpression.h"
+
+#include "MaterialExpressionOver.generated.h"
+
+/**
+ * Merge nodes take two 4-channel (color4) inputs and use the built-in alpha channel(s) to control the
+ * compositing of the A and B inputs. "A" and "B" refer to the non-alpha channels of the A and B inputs respectively,
+ * and "a" and "b" refer to the alpha channels of the A and B inputs.
+ * Merge nodes are only defined for float4 inputs
+ * Merge nodes support an optional float input Alpha , which can be used to mix the
+ * original B value with the result of the blend operation.
+ *
+ * Operation: A + B(1-a)
+ * Result: Lerp(B, A + B(1-a), Alpha)
+ */
+UCLASS(MinimalAPI, collapsecategories, hidecategories = Object)
+class UMaterialExpressionOver : public UMaterialExpression
+{
+	GENERATED_UCLASS_BODY()
+
+	UPROPERTY()
+	FExpressionInput A;
+
+	UPROPERTY()
+	FExpressionInput B;
+
+	UPROPERTY(meta = (RequiredInput = "false", ToolTip = "Defaults to 'ConstAlpha' if not specified"))
+	FExpressionInput Alpha;
+
+	/** only used if Alpha is not hooked up */
+	UPROPERTY(EditAnywhere, Category = MaterialExpressionLinearInterpolate, meta = (OverridingInputProperty = "Alpha"))
+	float ConstAlpha = 1.f;
+
+	//~ Begin UMaterialExpression Interface
+#if WITH_EDITOR
+	virtual int32 Compile(class FMaterialCompiler* Compiler, int32 OutputIndex) override;
+	virtual void GetCaption(TArray<FString>& OutCaptions) const override;
+#endif
+	//~ End UMaterialExpression Interface
+};
+
diff --git a/Engine/Plugins/Interchange/Runtime/Source/Import/Private/MaterialX/MaterialExpressions/MaterialExpressionOverlay.cpp b/Engine/Plugins/Interchange/Runtime/Source/Import/Private/MaterialX/MaterialExpressions/MaterialExpressionOverlay.cpp
new file mode 100644
--- /dev/null
+++ b/Engine/Plugins/Interchange/Runtime/Source/Import/Private/MaterialX/MaterialExpressions/MaterialExpressionOverlay.cpp
@@ -0,0 +1,71 @@
+// Copyright Epic Games, Inc. All Rights Reserved.
+#include "MaterialExpressionOverlay.h"
+#include "MaterialCompiler.h"
+
+#include UE_INLINE_GENERATED_CPP_BY_NAME(MaterialExpressionOverlay)
+
+#define LOCTEXT_NAMESPACE "MaterialExpressionOverlay"
+
+UMaterialExpressionOverlay::UMaterialExpressionOverlay(const FObjectInitializer& ObjectInitializer)
+	: Super(ObjectInitializer)
+{
+	// Structure to hold one-time initialization
+	struct FConstructorStatics
+	{
+		FText NAME_MaterialX;
+		FText NAME_Compositing;
+		FConstructorStatics()
+			: NAME_MaterialX(LOCTEXT("MaterialX", "MaterialX"))
+			, NAME_Compositing(LOCTEXT("Compositing", "Compositing"))
+		{}
+	};
+	static FConstructorStatics ConstructorStatics;
+
+#if WITH_EDITORONLY_DATA
+	MenuCategories.Add(ConstructorStatics.NAME_MaterialX);
+	MenuCategories.Add(ConstructorStatics.NAME_Compositing);
+#endif
+}
+
+#if WITH_EDITOR
+int32 UMaterialExpressionOverlay::Compile(FMaterialCompiler* Compiler, int32 OutputIndex)
+{
+	if(!A.GetTracedInput().Expression)
+	{
+		return Compiler->Errorf(TEXT("Missing A input"));
+	}
+
+	if(!B.GetTracedInput().Expression)
+	{
+		return Compiler->Errorf(TEXT("Missing B input"));
+	}
+
+	int32 IndexAlpha = Alpha.GetTracedInput().Expression ? Alpha.Compile(Compiler) : Compiler->Constant(ConstAlpha);
+
+	int32 OneIndex = Compiler->Constant(1.f);
+	auto OneMinus = [&](int32 index)
+	{
+		return Compiler->Sub(OneIndex, index);
+	};
+
+	int32 IndexA = A.Compile(Compiler);
+	int32 IndexB = B.Compile(Compiler);
+	int32 Screen = OneMinus(Compiler->Mul(OneMinus(IndexA), OneMinus(IndexB)));
+
+	int32 Overlay = Compiler->If(IndexA,
+								 Compiler->Constant(0.5f),
+								 Screen,
+								 Screen,
+								 Compiler->Mul(Compiler->Mul(Compiler->Constant(2.f), IndexA), IndexB),
+								 Compiler->Constant(0.00001f));
+
+	return Compiler->Lerp(IndexB, Overlay, IndexAlpha);
+}
+
+void UMaterialExpressionOverlay::GetCaption(TArray<FString>& OutCaptions) const
+{
+	OutCaptions.Add(TEXT("Overlay"));
+}
+#endif
+
+#undef LOCTEXT_NAMESPACE 
\ No newline at end of file
diff --git a/Engine/Plugins/Interchange/Runtime/Source/Import/Private/MaterialX/MaterialExpressions/MaterialExpressionOverlay.h b/Engine/Plugins/Interchange/Runtime/Source/Import/Private/MaterialX/MaterialExpressions/MaterialExpressionOverlay.h
new file mode 100644
--- /dev/null
+++ b/Engine/Plugins/Interchange/Runtime/Source/Import/Private/MaterialX/MaterialExpressions/MaterialExpressionOverlay.h
@@ -0,0 +1,42 @@
+// Copyright Epic Games, Inc. All Rights Reserved.
+#pragma once
+
+#include "CoreMinimal.h"
+#include "Materials/MaterialExpression.h"
+
+#include "MaterialExpressionOverlay.generated.h"
+
+/**
+ * Blend nodes take two 1-4 channel inputs and apply the same operator to all channels.
+ * Blend nodes support an optional float input mix , which can be used
+ * to mix the original B value with the result of the blend operation.
+ * Operation: 2*A*B          if A < 0.5;
+ *            1-(1-A)(1-B) if A >= 0.5
+ * Result: Lerp(B, Op, Alpha)
+ */
+UCLASS(MinimalAPI, collapsecategories, hidecategories = Object)
+class UMaterialExpressionOverlay : public UMaterialExpression
+{
+	GENERATED_UCLASS_BODY()
+
+	UPROPERTY()
+	FExpressionInput A;
+
+	UPROPERTY()
+	FExpressionInput B;
+
+	UPROPERTY(meta = (RequiredInput = "false", ToolTip = "Defaults to 'ConstAlpha' if not specified"))
+	FExpressionInput Alpha;
+
+	/** only used if Alpha is not hooked up */
+	UPROPERTY(EditAnywhere, Category = MaterialExpressionLinearInterpolate, meta = (OverridingInputProperty = "Alpha"))
+	float ConstAlpha = 1.f;
+
+	//~ Begin UMaterialExpression Interface
+#if WITH_EDITOR
+	virtual int32 Compile(class FMaterialCompiler* Compiler, int32 OutputIndex) override;
+	virtual void GetCaption(TArray<FString>& OutCaptions) const override;
+#endif
+	//~ End UMaterialExpression Interface
+};
+
diff --git a/Engine/Plugins/Interchange/Runtime/Source/Import/Private/MaterialX/MaterialExpressions/MaterialExpressionPlus.cpp b/Engine/Plugins/Interchange/Runtime/Source/Import/Private/MaterialX/MaterialExpressions/MaterialExpressionPlus.cpp
new file mode 100644
--- /dev/null
+++ b/Engine/Plugins/Interchange/Runtime/Source/Import/Private/MaterialX/MaterialExpressions/MaterialExpressionPlus.cpp
@@ -0,0 +1,56 @@
+// Copyright Epic Games, Inc. All Rights Reserved.
+#include "MaterialExpressionPlus.h"
+#include "MaterialCompiler.h"
+
+#include UE_INLINE_GENERATED_CPP_BY_NAME(MaterialExpressionPlus)
+
+#define LOCTEXT_NAMESPACE "MaterialExpressionPlus"
+
+UMaterialExpressionPlus::UMaterialExpressionPlus(const FObjectInitializer& ObjectInitializer)
+	: Super(ObjectInitializer)
+{
+	// Structure to hold one-time initialization
+	struct FConstructorStatics
+	{
+		FText NAME_MaterialX;
+		FText NAME_Compositing;
+		FConstructorStatics()
+			: NAME_MaterialX(LOCTEXT("MaterialX", "MaterialX"))
+			, NAME_Compositing(LOCTEXT("Compositing", "Compositing"))
+		{}
+	};
+	static FConstructorStatics ConstructorStatics;
+
+#if WITH_EDITORONLY_DATA
+	MenuCategories.Add(ConstructorStatics.NAME_MaterialX);
+	MenuCategories.Add(ConstructorStatics.NAME_Compositing);
+#endif
+}
+
+#if WITH_EDITOR
+int32 UMaterialExpressionPlus::Compile(FMaterialCompiler* Compiler, int32 OutputIndex)
+{
+	if(!A.GetTracedInput().Expression)
+	{
+		return Compiler->Errorf(TEXT("Missing A input"));
+	}
+
+	if(!B.GetTracedInput().Expression)
+	{
+		return Compiler->Errorf(TEXT("Missing B input"));
+	}
+
+	int32 IndexAlpha = Alpha.GetTracedInput().Expression ? Alpha.Compile(Compiler) : Compiler->Constant(ConstAlpha);
+
+	int32 IndexB = B.Compile(Compiler);
+	int32 Add = Compiler->Add(A.Compile(Compiler), IndexB);
+	return Compiler->Lerp(IndexB, Add, IndexAlpha);
+}
+
+void UMaterialExpressionPlus::GetCaption(TArray<FString>& OutCaptions) const
+{
+	OutCaptions.Add(TEXT("Plus"));
+}
+#endif
+
+#undef LOCTEXT_NAMESPACE 
\ No newline at end of file
diff --git a/Engine/Plugins/Interchange/Runtime/Source/Import/Private/MaterialX/MaterialExpressions/MaterialExpressionPlus.h b/Engine/Plugins/Interchange/Runtime/Source/Import/Private/MaterialX/MaterialExpressions/MaterialExpressionPlus.h
new file mode 100644
--- /dev/null
+++ b/Engine/Plugins/Interchange/Runtime/Source/Import/Private/MaterialX/MaterialExpressions/MaterialExpressionPlus.h
@@ -0,0 +1,41 @@
+// Copyright Epic Games, Inc. All Rights Reserved.
+#pragma once
+
+#include "CoreMinimal.h"
+#include "Materials/MaterialExpression.h"
+
+#include "MaterialExpressionPlus.generated.h"
+
+/**
+ * Blend nodes take two 1-4 channel inputs and apply the same operator to all channels.
+ * Blend nodes support an optional float input mix , which can be used
+ * to mix the original B value with the result of the blend operation.
+ * Operation: A + B
+ * Result: Lerp(B, A + B, Alpha)
+ */
+UCLASS(MinimalAPI, collapsecategories, hidecategories = Object)
+class UMaterialExpressionPlus : public UMaterialExpression
+{
+	GENERATED_UCLASS_BODY()
+
+	UPROPERTY()
+	FExpressionInput A;
+
+	UPROPERTY()
+	FExpressionInput B;
+
+	UPROPERTY(meta = (RequiredInput = "false", ToolTip = "Defaults to 'ConstAlpha' if not specified"))
+	FExpressionInput Alpha;
+
+	/** only used if Alpha is not hooked up */
+	UPROPERTY(EditAnywhere, Category = MaterialExpressionLinearInterpolate, meta = (OverridingInputProperty = "Alpha"))
+	float ConstAlpha = 1.f;
+
+	//~ Begin UMaterialExpression Interface
+#if WITH_EDITOR
+	virtual int32 Compile(class FMaterialCompiler* Compiler, int32 OutputIndex) override;
+	virtual void GetCaption(TArray<FString>& OutCaptions) const override;
+#endif
+	//~ End UMaterialExpression Interface
+};
+
diff --git a/Engine/Plugins/Interchange/Runtime/Source/Import/Private/MaterialX/MaterialExpressions/MaterialExpressionPremult.cpp b/Engine/Plugins/Interchange/Runtime/Source/Import/Private/MaterialX/MaterialExpressions/MaterialExpressionPremult.cpp
new file mode 100644
--- /dev/null
+++ b/Engine/Plugins/Interchange/Runtime/Source/Import/Private/MaterialX/MaterialExpressions/MaterialExpressionPremult.cpp
@@ -0,0 +1,50 @@
+// Copyright Epic Games, Inc. All Rights Reserved.
+#include "MaterialExpressionPremult.h"
+#include "MaterialCompiler.h"
+
+#include UE_INLINE_GENERATED_CPP_BY_NAME(MaterialExpressionPremult)
+
+#define LOCTEXT_NAMESPACE "MaterialExpressionPremult"
+
+UMaterialExpressionPremult::UMaterialExpressionPremult(const FObjectInitializer& ObjectInitializer)
+	: Super(ObjectInitializer)
+{
+	// Structure to hold one-time initialization
+	struct FConstructorStatics
+	{
+		FText NAME_MaterialX;
+		FText NAME_Compositing;
+		FConstructorStatics()
+			: NAME_MaterialX(LOCTEXT("MaterialX", "MaterialX"))
+			, NAME_Compositing(LOCTEXT("Compositing", "Compositing"))
+		{}
+	};
+	static FConstructorStatics ConstructorStatics;
+
+#if WITH_EDITORONLY_DATA
+	MenuCategories.Add(ConstructorStatics.NAME_MaterialX);
+	MenuCategories.Add(ConstructorStatics.NAME_Compositing);
+#endif
+}
+
+#if WITH_EDITOR
+int32 UMaterialExpressionPremult::Compile(FMaterialCompiler* Compiler, int32 OutputIndex)
+{
+	if(!Input.GetTracedInput().Expression)
+	{
+		return Compiler->Errorf(TEXT("Missing Premult input"));
+	}
+
+	int32 IndexInput = Input.Compile(Compiler);
+	int32 RGB = Compiler->ComponentMask(IndexInput, true, true, true, false);
+	int32 Alpha = Compiler->ComponentMask(IndexInput, false, false, false, true);
+	return Compiler->AppendVector(Compiler->Mul(RGB, Alpha), Alpha);
+}
+
+void UMaterialExpressionPremult::GetCaption(TArray<FString>& OutCaptions) const
+{
+	OutCaptions.Add(TEXT("Premult"));
+}
+#endif
+
+#undef LOCTEXT_NAMESPACE 
\ No newline at end of file
diff --git a/Engine/Plugins/Interchange/Runtime/Source/Import/Private/MaterialX/MaterialExpressions/MaterialExpressionPremult.h b/Engine/Plugins/Interchange/Runtime/Source/Import/Private/MaterialX/MaterialExpressions/MaterialExpressionPremult.h
new file mode 100644
--- /dev/null
+++ b/Engine/Plugins/Interchange/Runtime/Source/Import/Private/MaterialX/MaterialExpressions/MaterialExpressionPremult.h
@@ -0,0 +1,28 @@
+// Copyright Epic Games, Inc. All Rights Reserved.
+#pragma once
+
+#include "CoreMinimal.h"
+#include "Materials/MaterialExpression.h"
+
+#include "MaterialExpressionPremult.generated.h"
+
+/**
+ * Multiply the RGB channels of the input by the Alpha channel of the input.
+ * Input must be of type float4
+ */
+UCLASS(collapsecategories, hidecategories = Object, MinimalAPI)
+class UMaterialExpressionPremult : public UMaterialExpression
+{
+	GENERATED_UCLASS_BODY()
+
+	UPROPERTY()
+	FExpressionInput Input;
+
+	//~ Begin UMaterialExpression Interface
+#if WITH_EDITOR
+	virtual int32 Compile(class FMaterialCompiler* Compiler, int32 OutputIndex) override;
+	virtual void GetCaption(TArray<FString>& OutCaptions) const override;
+#endif
+	//~ End UMaterialExpression Interface
+};
+
diff --git a/Engine/Plugins/Interchange/Runtime/Source/Import/Private/MaterialX/MaterialExpressions/MaterialExpressionRampLeftRight.cpp b/Engine/Plugins/Interchange/Runtime/Source/Import/Private/MaterialX/MaterialExpressions/MaterialExpressionRampLeftRight.cpp
new file mode 100644
--- /dev/null
+++ b/Engine/Plugins/Interchange/Runtime/Source/Import/Private/MaterialX/MaterialExpressions/MaterialExpressionRampLeftRight.cpp
@@ -0,0 +1,54 @@
+// Copyright Epic Games, Inc. All Rights Reserved.
+#include "MaterialExpressionRampLeftRight.h"
+#include "MaterialCompiler.h"
+
+#include UE_INLINE_GENERATED_CPP_BY_NAME(MaterialExpressionRampLeftRight)
+
+#define LOCTEXT_NAMESPACE "MaterialExpressionRampLeftRight"
+
+UMaterialExpressionRampLeftRight::UMaterialExpressionRampLeftRight(const FObjectInitializer& ObjectInitializer)
+	: Super(ObjectInitializer)
+{
+	// Structure to hold one-time initialization
+	struct FConstructorStatics
+	{
+		FText NAME_MaterialX;
+		FText NAME_Procedural;
+		FConstructorStatics()
+			: NAME_MaterialX(LOCTEXT("MaterialX", "MaterialX"))
+			, NAME_Procedural(LOCTEXT("Procedural", "Procedural"))
+		{}
+	};
+	static FConstructorStatics ConstructorStatics;
+
+#if WITH_EDITORONLY_DATA
+	MenuCategories.Add(ConstructorStatics.NAME_MaterialX);
+	MenuCategories.Add(ConstructorStatics.NAME_Procedural);
+#endif
+}
+
+#if WITH_EDITOR
+int32 UMaterialExpressionRampLeftRight::Compile(FMaterialCompiler* Compiler, int32 OutputIndex)
+{
+	if(!A.GetTracedInput().Expression)
+	{
+		return Compiler->Errorf(TEXT("Missing A input"));
+	}
+
+	if(!B.GetTracedInput().Expression)
+	{
+		return Compiler->Errorf(TEXT("Missing B input"));
+	}
+
+	int32 CoordinateIndex = Coordinates.GetTracedInput().Expression ? Coordinates.Compile(Compiler) : Compiler->TextureCoordinate(ConstCoordinate, false, false);
+
+	return Compiler->Lerp(A.Compile(Compiler), B.Compile(Compiler), Compiler->Saturate(Compiler->ComponentMask(CoordinateIndex, true, false, false, false)));
+}
+
+void UMaterialExpressionRampLeftRight::GetCaption(TArray<FString>& OutCaptions) const
+{
+	OutCaptions.Add(TEXT("RampLR"));
+}
+#endif
+
+#undef LOCTEXT_NAMESPACE 
\ No newline at end of file
diff --git a/Engine/Plugins/Interchange/Runtime/Source/Import/Private/MaterialX/MaterialExpressions/MaterialExpressionRampLeftRight.h b/Engine/Plugins/Interchange/Runtime/Source/Import/Private/MaterialX/MaterialExpressions/MaterialExpressionRampLeftRight.h
new file mode 100644
--- /dev/null
+++ b/Engine/Plugins/Interchange/Runtime/Source/Import/Private/MaterialX/MaterialExpressions/MaterialExpressionRampLeftRight.h
@@ -0,0 +1,37 @@
+// Copyright Epic Games, Inc. All Rights Reserved.
+#pragma once
+
+#include "CoreMinimal.h"
+#include "Materials/MaterialExpression.h"
+
+#include "MaterialExpressionRampLeftRight.generated.h"
+
+/**
+ * A material expression that computes a left-to-right bilinear value ramp.
+ */
+UCLASS(collapsecategories, hidecategories = Object, MinimalAPI)
+class UMaterialExpressionRampLeftRight : public UMaterialExpression
+{
+	GENERATED_UCLASS_BODY()
+
+	UPROPERTY(meta = (RequiredInput = "false", ToolTip = "Defaults to 'ConstCoordinate' if not specified"))
+	FExpressionInput Coordinates;
+
+	UPROPERTY()
+	FExpressionInput A;
+
+	UPROPERTY()
+	FExpressionInput B;
+
+	/** only used if Coordinates is not hooked up */
+	UPROPERTY(EditAnywhere, Category = MaterialExpressionTextureSample, meta = (OverridingInputProperty = "Coordinates"))
+	uint8 ConstCoordinate;
+
+	//~ Begin UMaterialExpression Interface
+#if WITH_EDITOR
+	virtual int32 Compile(class FMaterialCompiler* Compiler, int32 OutputIndex) override;
+	virtual void GetCaption(TArray<FString>& OutCaptions) const override;
+#endif
+	//~ End UMaterialExpression Interface
+};
+
diff --git a/Engine/Plugins/Interchange/Runtime/Source/Import/Private/MaterialX/MaterialExpressions/MaterialExpressionRampTopBottom.cpp b/Engine/Plugins/Interchange/Runtime/Source/Import/Private/MaterialX/MaterialExpressions/MaterialExpressionRampTopBottom.cpp
new file mode 100644
--- /dev/null
+++ b/Engine/Plugins/Interchange/Runtime/Source/Import/Private/MaterialX/MaterialExpressions/MaterialExpressionRampTopBottom.cpp
@@ -0,0 +1,54 @@
+// Copyright Epic Games, Inc. All Rights Reserved.
+#include "MaterialExpressionRampTopBottom.h"
+#include "MaterialCompiler.h"
+
+#include UE_INLINE_GENERATED_CPP_BY_NAME(MaterialExpressionRampTopBottom)
+
+#define LOCTEXT_NAMESPACE "MaterialExpressionRampLeftRight"
+
+UMaterialExpressionRampTopBottom::UMaterialExpressionRampTopBottom(const FObjectInitializer& ObjectInitializer)
+	: Super(ObjectInitializer)
+{
+	// Structure to hold one-time initialization
+	struct FConstructorStatics
+	{
+		FText NAME_MaterialX;
+		FText NAME_Procedural;
+		FConstructorStatics()
+			: NAME_MaterialX(LOCTEXT("MaterialX", "MaterialX"))
+			, NAME_Procedural(LOCTEXT("Procedural", "Procedural"))
+		{}
+	};
+	static FConstructorStatics ConstructorStatics;
+
+#if WITH_EDITORONLY_DATA
+	MenuCategories.Add(ConstructorStatics.NAME_MaterialX);
+	MenuCategories.Add(ConstructorStatics.NAME_Procedural);
+#endif
+}
+
+#if WITH_EDITOR
+int32 UMaterialExpressionRampTopBottom::Compile(FMaterialCompiler* Compiler, int32 OutputIndex)
+{
+	if(!A.GetTracedInput().Expression)
+	{
+		return Compiler->Errorf(TEXT("Missing A input"));
+	}
+
+	if(!B.GetTracedInput().Expression)
+	{
+		return Compiler->Errorf(TEXT("Missing B input"));
+	}
+
+	int32 CoordinateIndex = Coordinates.GetTracedInput().Expression ? Coordinates.Compile(Compiler) : Compiler->TextureCoordinate(ConstCoordinate, false, false);
+
+	return Compiler->Lerp(A.Compile(Compiler), B.Compile(Compiler), Compiler->Saturate(Compiler->ComponentMask(CoordinateIndex, false, true, false, false)));
+}
+
+void UMaterialExpressionRampTopBottom::GetCaption(TArray<FString>& OutCaptions) const
+{
+	OutCaptions.Add(TEXT("RampTB"));
+}
+#endif
+
+#undef LOCTEXT_NAMESPACE 
\ No newline at end of file
diff --git a/Engine/Plugins/Interchange/Runtime/Source/Import/Private/MaterialX/MaterialExpressions/MaterialExpressionRampTopBottom.h b/Engine/Plugins/Interchange/Runtime/Source/Import/Private/MaterialX/MaterialExpressions/MaterialExpressionRampTopBottom.h
new file mode 100644
--- /dev/null
+++ b/Engine/Plugins/Interchange/Runtime/Source/Import/Private/MaterialX/MaterialExpressions/MaterialExpressionRampTopBottom.h
@@ -0,0 +1,37 @@
+// Copyright Epic Games, Inc. All Rights Reserved.
+#pragma once
+
+#include "CoreMinimal.h"
+#include "Materials/MaterialExpression.h"
+
+#include "MaterialExpressionRampTopBottom.generated.h"
+
+/**
+ * A material expression that computes a top-to-bottom bilinear value ramp.
+ */
+UCLASS(collapsecategories, hidecategories = Object, MinimalAPI)
+class UMaterialExpressionRampTopBottom : public UMaterialExpression
+{
+	GENERATED_UCLASS_BODY()
+
+		UPROPERTY(meta = (RequiredInput = "false", ToolTip = "Defaults to 'ConstCoordinate' if not specified"))
+		FExpressionInput Coordinates;
+
+	UPROPERTY()
+		FExpressionInput A;
+
+	UPROPERTY()
+		FExpressionInput B;
+
+	/** only used if Coordinates is not hooked up */
+	UPROPERTY(EditAnywhere, Category = MaterialExpressionTextureSample, meta = (OverridingInputProperty = "Coordinates"))
+		uint8 ConstCoordinate;
+
+	//~ Begin UMaterialExpression Interface
+#if WITH_EDITOR
+	virtual int32 Compile(class FMaterialCompiler* Compiler, int32 OutputIndex) override;
+	virtual void GetCaption(TArray<FString>& OutCaptions) const override;
+#endif
+	//~ End UMaterialExpression Interface
+};
+
diff --git a/Engine/Plugins/Interchange/Runtime/Source/Import/Private/MaterialX/MaterialExpressions/MaterialExpressionRemap.cpp b/Engine/Plugins/Interchange/Runtime/Source/Import/Private/MaterialX/MaterialExpressions/MaterialExpressionRemap.cpp
new file mode 100644
--- /dev/null
+++ b/Engine/Plugins/Interchange/Runtime/Source/Import/Private/MaterialX/MaterialExpressions/MaterialExpressionRemap.cpp
@@ -0,0 +1,63 @@
+// Copyright Epic Games, Inc. All Rights Reserved.
+#include "MaterialExpressionRemap.h"
+#include "MaterialCompiler.h"
+
+#include UE_INLINE_GENERATED_CPP_BY_NAME(MaterialExpressionRemap)
+
+#define LOCTEXT_NAMESPACE "MaterialExpressionRemap"
+
+UMaterialExpressionRemap::UMaterialExpressionRemap(const FObjectInitializer& ObjectInitializer)
+	: Super(ObjectInitializer)
+{
+	// Structure to hold one-time initialization
+	struct FConstructorStatics
+	{
+		FText NAME_MaterialX;
+		FText NAME_Adjustment;
+		FConstructorStatics()
+			: NAME_MaterialX(LOCTEXT("MaterialX", "MaterialX"))
+			, NAME_Adjustment(LOCTEXT("Image Adjustment", "Image Adjustment"))
+		{}
+	};
+	static FConstructorStatics ConstructorStatics;
+
+	InputLowDefault = 0.0f;
+	InputHighDefault = 1.0f;
+	TargetLowDefault = 0.0f;
+	TargetHighDefault = 1.0f;
+
+#if WITH_EDITORONLY_DATA
+	MenuCategories.Add(ConstructorStatics.NAME_MaterialX);
+	MenuCategories.Add(ConstructorStatics.NAME_Adjustment);
+#endif
+}
+
+#if WITH_EDITOR
+int32 UMaterialExpressionRemap::Compile(FMaterialCompiler* Compiler, int32 OutputIndex)
+{
+	if(!Input.GetTracedInput().Expression)
+	{
+		return Compiler->Errorf(TEXT("Missing Remap input"));
+	}
+
+	// Min2 + (Input - Min1) x (Max2 - Min2) / (Max1 - Min1)
+	const int32 InputLowIndex = InputLow.GetTracedInput().Expression ? InputLow.Compile(Compiler) : Compiler->Constant(InputLowDefault);
+	const int32 InputHighIndex = InputHigh.GetTracedInput().Expression ? InputHigh.Compile(Compiler) : Compiler->Constant(InputHighDefault);
+
+	const int32 TargetLowIndex = TargetLow.GetTracedInput().Expression ? TargetLow.Compile(Compiler) : Compiler->Constant(TargetLowDefault);
+	const int32 TargetHighIndex = TargetHigh.GetTracedInput().Expression ? TargetHigh.Compile(Compiler) : Compiler->Constant(TargetHighDefault);
+
+	const int32 InputDelta = Compiler->Sub(InputHighIndex, InputLowIndex);
+	const int32 TargetDelta = Compiler->Sub(TargetHighIndex, TargetLowIndex);
+	const int32 Delta = Compiler->Div(TargetDelta, InputDelta);
+
+	return Compiler->Add(TargetLowIndex, Compiler->Mul(Compiler->Sub(Input.Compile(Compiler), InputLowIndex), Delta));
+}
+
+void UMaterialExpressionRemap::GetCaption(TArray<FString>& OutCaptions) const
+{
+	OutCaptions.Add(TEXT("Remap"));
+}
+#endif
+
+#undef LOCTEXT_NAMESPACE 
\ No newline at end of file
diff --git a/Engine/Plugins/Interchange/Runtime/Source/Import/Private/MaterialX/MaterialExpressions/MaterialExpressionRemap.h b/Engine/Plugins/Interchange/Runtime/Source/Import/Private/MaterialX/MaterialExpressions/MaterialExpressionRemap.h
new file mode 100644
--- /dev/null
+++ b/Engine/Plugins/Interchange/Runtime/Source/Import/Private/MaterialX/MaterialExpressions/MaterialExpressionRemap.h
@@ -0,0 +1,51 @@
+// Copyright Epic Games, Inc. All Rights Reserved.
+#pragma once
+
+#include "CoreMinimal.h"
+#include "Materials/MaterialExpression.h"
+
+#include "MaterialExpressionRemap.generated.h"
+
+/**
+ * A material expression that Remap a value from one range to another.
+ */
+UCLASS(collapsecategories, hidecategories = Object, MinimalAPI)
+class UMaterialExpressionRemap : public UMaterialExpression
+{
+	GENERATED_UCLASS_BODY()
+
+	UPROPERTY()
+	FExpressionInput Input;
+
+	UPROPERTY(meta = (RequiredInput = "false", ToolTip = "Defaults to 'MinDefault1' if not specified"))
+	FExpressionInput InputLow;
+
+	UPROPERTY(meta = (RequiredInput = "false", ToolTip = "Defaults to 'MaxDefault1' if not specified"))
+	FExpressionInput InputHigh;
+
+	UPROPERTY(meta = (RequiredInput = "false", ToolTip = "Defaults to 'MinDefault2' if not specified"))
+	FExpressionInput TargetLow;
+
+	UPROPERTY(meta = (RequiredInput = "false", ToolTip = "Defaults to 'MaxDefault1' if not specified"))
+	FExpressionInput TargetHigh;
+
+	UPROPERTY(EditAnywhere, Category = MaterialExpressionClamp, meta = (OverridingInputProperty = "InputLow"))
+	float InputLowDefault;
+
+	UPROPERTY(EditAnywhere, Category = MaterialExpressionClamp, meta = (OverridingInputProperty = "InputHigh"))
+	float InputHighDefault;
+
+	UPROPERTY(EditAnywhere, Category = MaterialExpressionClamp, meta = (OverridingInputProperty = "TargetLow"))
+	float TargetLowDefault;
+
+	UPROPERTY(EditAnywhere, Category = MaterialExpressionClamp, meta = (OverridingInputProperty = "TargetHigh"))
+	float TargetHighDefault;
+
+	//~ Begin UMaterialExpression Interface
+#if WITH_EDITOR
+	virtual int32 Compile(class FMaterialCompiler* Compiler, int32 OutputIndex) override;
+	virtual void GetCaption(TArray<FString>& OutCaptions) const override;
+#endif
+	//~ End UMaterialExpression Interface
+};
+
diff --git a/Engine/Plugins/Interchange/Runtime/Source/Import/Private/MaterialX/MaterialExpressions/MaterialExpressionRgbToHsv.cpp b/Engine/Plugins/Interchange/Runtime/Source/Import/Private/MaterialX/MaterialExpressions/MaterialExpressionRgbToHsv.cpp
new file mode 100644
--- /dev/null
+++ b/Engine/Plugins/Interchange/Runtime/Source/Import/Private/MaterialX/MaterialExpressions/MaterialExpressionRgbToHsv.cpp
@@ -0,0 +1,60 @@
+// Copyright Epic Games, Inc. All Rights Reserved.
+
+#include "MaterialExpressionRgbToHsv.h"
+#include "MaterialCompiler.h"
+#include "Materials/MaterialExpressionCustom.h"
+
+#define LOCTEXT_NAMESPACE "MaterialExpressionRgbToHsv"
+
+UMaterialExpressionRgbToHsv::UMaterialExpressionRgbToHsv(const FObjectInitializer& ObjectInitializer)
+	: Super(ObjectInitializer)
+{
+	// Structure to hold one-time initialization
+	struct FConstructorStatics
+	{
+		FText NAME_ImageAdjustment;
+		FText NAME_MaterialX;
+		FConstructorStatics()
+			: NAME_ImageAdjustment(LOCTEXT("Image Adjustment", "Image Adjustment"))
+			, NAME_MaterialX(LOCTEXT("MaterialX", "MaterialX"))
+		{}
+	};
+	static FConstructorStatics ConstructorStatics;
+
+#if WITH_EDITORONLY_DATA
+	MenuCategories.Add(ConstructorStatics.NAME_MaterialX);
+	MenuCategories.Add(ConstructorStatics.NAME_ImageAdjustment);
+#endif
+}
+
+#if WITH_EDITOR
+int32 UMaterialExpressionRgbToHsv::Compile(FMaterialCompiler* Compiler, int32 OutputIndex)
+{
+	if(!Input.GetTracedInput().Expression)
+	{
+		return Compiler->Errorf(TEXT("Missing RGBtoHSV input"));
+	}
+
+	UMaterialExpressionCustom* MaterialExpressionCustom = NewObject<UMaterialExpressionCustom>();
+	MaterialExpressionCustom->Inputs[0].InputName = TEXT("c");
+	MaterialExpressionCustom->Inputs[0].Input = Input;
+	MaterialExpressionCustom->OutputType = ECustomMaterialOutputType::CMOT_Float3;
+	MaterialExpressionCustom->Code = TEXT(R"(
+    float4 K = float4(0.0, -1.0 / 3.0, 2.0 / 3.0, -1.0);
+    float4 p = lerp(float4(c.bg, K.wz), float4(c.gb, K.xy), step(c.b, c.g));
+    float4 q = lerp(float4(p.xyw, c.r), float4(c.r, p.yzx), step(p.x, c.r));
+
+    float d = q.x - min(q.w, q.y);
+    float e = 1.0e-10;
+    return float3(abs(q.z + (q.w - q.y) / (6.0 * d + e)), d / (q.x + e), q.x);)");
+		
+	return MaterialExpressionCustom->Compile(Compiler, OutputIndex);
+}
+
+void UMaterialExpressionRgbToHsv::GetCaption(TArray<FString>& OutCaptions) const
+{
+	OutCaptions.Add(TEXT("RGBToHSV"));
+}
+#endif
+
+#undef LOCTEXT_NAMESPACE 
\ No newline at end of file
diff --git a/Engine/Plugins/Interchange/Runtime/Source/Import/Private/MaterialX/MaterialExpressions/MaterialExpressionRgbToHsv.h b/Engine/Plugins/Interchange/Runtime/Source/Import/Private/MaterialX/MaterialExpressions/MaterialExpressionRgbToHsv.h
new file mode 100644
--- /dev/null
+++ b/Engine/Plugins/Interchange/Runtime/Source/Import/Private/MaterialX/MaterialExpressions/MaterialExpressionRgbToHsv.h
@@ -0,0 +1,23 @@
+// Copyright Epic Games, Inc. All Rights Reserved.
+#pragma once
+
+#include "Materials/MaterialExpression.h"
+
+#include "MaterialExpressionRgbToHsv.generated.h"
+
+UCLASS(collapsecategories, hidecategories = Object, MinimalAPI)
+class UMaterialExpressionRgbToHsv : public UMaterialExpression
+{
+	GENERATED_UCLASS_BODY()
+
+	UPROPERTY()
+	FExpressionInput Input;
+
+	//~ Begin UMaterialExpression Interface
+#if WITH_EDITOR
+	virtual int32 Compile(class FMaterialCompiler* Compiler, int32 OutputIndex) override;
+	virtual void GetCaption(TArray<FString>& OutCaptions) const override;
+#endif
+	//~ End UMaterialExpression Interface
+};
+
diff --git a/Engine/Plugins/Interchange/Runtime/Source/Import/Private/MaterialX/MaterialExpressions/MaterialExpressionScreen.cpp b/Engine/Plugins/Interchange/Runtime/Source/Import/Private/MaterialX/MaterialExpressions/MaterialExpressionScreen.cpp
new file mode 100644
--- /dev/null
+++ b/Engine/Plugins/Interchange/Runtime/Source/Import/Private/MaterialX/MaterialExpressions/MaterialExpressionScreen.cpp
@@ -0,0 +1,64 @@
+// Copyright Epic Games, Inc. All Rights Reserved.
+#include "MaterialExpressionScreen.h"
+#include "MaterialCompiler.h"
+
+#include UE_INLINE_GENERATED_CPP_BY_NAME(MaterialExpressionScreen)
+
+#define LOCTEXT_NAMESPACE "MaterialExpressionScreen"
+
+UMaterialExpressionScreen::UMaterialExpressionScreen(const FObjectInitializer& ObjectInitializer)
+	: Super(ObjectInitializer)
+{
+	// Structure to hold one-time initialization
+	struct FConstructorStatics
+	{
+		FText NAME_MaterialX;
+		FText NAME_Compositing;
+		FConstructorStatics()
+			: NAME_MaterialX(LOCTEXT("MaterialX", "MaterialX"))
+			, NAME_Compositing(LOCTEXT("Compositing", "Compositing"))
+		{}
+	};
+	static FConstructorStatics ConstructorStatics;
+
+#if WITH_EDITORONLY_DATA
+	MenuCategories.Add(ConstructorStatics.NAME_MaterialX);
+	MenuCategories.Add(ConstructorStatics.NAME_Compositing);
+#endif
+}
+
+#if WITH_EDITOR
+int32 UMaterialExpressionScreen::Compile(FMaterialCompiler* Compiler, int32 OutputIndex)
+{
+	if(!A.GetTracedInput().Expression)
+	{
+		return Compiler->Errorf(TEXT("Missing A input"));
+	}
+
+	if(!B.GetTracedInput().Expression)
+	{
+		return Compiler->Errorf(TEXT("Missing B input"));
+	}
+
+	int32 IndexAlpha = Alpha.GetTracedInput().Expression ? Alpha.Compile(Compiler) : Compiler->Constant(ConstAlpha);
+
+	int32 OneIndex = Compiler->Constant(1.f);
+	auto OneMinus = [&](int32 index)
+	{
+		return Compiler->Sub(OneIndex, index);
+	};
+
+	int32 IndexA = A.Compile(Compiler);
+	int32 IndexB = B.Compile(Compiler);
+	int32 Screen = OneMinus(Compiler->Mul(OneMinus(IndexA), OneMinus(IndexB)));
+
+	return Compiler->Lerp(IndexB, Screen, IndexAlpha);
+}
+
+void UMaterialExpressionScreen::GetCaption(TArray<FString>& OutCaptions) const
+{
+	OutCaptions.Add(TEXT("Screen"));
+}
+#endif
+
+#undef LOCTEXT_NAMESPACE 
\ No newline at end of file
diff --git a/Engine/Plugins/Interchange/Runtime/Source/Import/Private/MaterialX/MaterialExpressions/MaterialExpressionScreen.h b/Engine/Plugins/Interchange/Runtime/Source/Import/Private/MaterialX/MaterialExpressions/MaterialExpressionScreen.h
new file mode 100644
--- /dev/null
+++ b/Engine/Plugins/Interchange/Runtime/Source/Import/Private/MaterialX/MaterialExpressions/MaterialExpressionScreen.h
@@ -0,0 +1,41 @@
+// Copyright Epic Games, Inc. All Rights Reserved.
+#pragma once
+
+#include "CoreMinimal.h"
+#include "Materials/MaterialExpression.h"
+
+#include "MaterialExpressionScreen.generated.h"
+
+/**
+ * Blend nodes take two 1-4 channel inputs and apply the same operator to all channels.
+ * Blend nodes support an optional float input mix , which can be used
+ * to mix the original B value with the result of the blend operation.
+ * Operation: 1-(1-A)(1-B)
+ * Result: Lerp(B, 1-(1-F)(1-B), Alpha)
+ */
+UCLASS(MinimalAPI, collapsecategories, hidecategories = Object)
+class UMaterialExpressionScreen : public UMaterialExpression
+{
+	GENERATED_UCLASS_BODY()
+
+	UPROPERTY()
+	FExpressionInput A;
+
+	UPROPERTY()
+	FExpressionInput B;
+
+	UPROPERTY(meta = (RequiredInput = "false", ToolTip = "Defaults to 'ConstAlpha' if not specified"))
+	FExpressionInput Alpha;
+
+	/** only used if Alpha is not hooked up */
+	UPROPERTY(EditAnywhere, Category = MaterialExpressionLinearInterpolate, meta = (OverridingInputProperty = "Alpha"))
+	float ConstAlpha = 1.f;
+
+	//~ Begin UMaterialExpression Interface
+#if WITH_EDITOR
+	virtual int32 Compile(class FMaterialCompiler* Compiler, int32 OutputIndex) override;
+	virtual void GetCaption(TArray<FString>& OutCaptions) const override;
+#endif
+	//~ End UMaterialExpression Interface
+};
+
diff --git a/Engine/Plugins/Interchange/Runtime/Source/Import/Private/MaterialX/MaterialExpressions/MaterialExpressionSplitLeftRight.cpp b/Engine/Plugins/Interchange/Runtime/Source/Import/Private/MaterialX/MaterialExpressions/MaterialExpressionSplitLeftRight.cpp
new file mode 100644
--- /dev/null
+++ b/Engine/Plugins/Interchange/Runtime/Source/Import/Private/MaterialX/MaterialExpressions/MaterialExpressionSplitLeftRight.cpp
@@ -0,0 +1,62 @@
+// Copyright Epic Games, Inc. All Rights Reserved.
+#include "MaterialExpressionSplitLeftRight.h"
+#include "MaterialCompiler.h"
+
+#include UE_INLINE_GENERATED_CPP_BY_NAME(MaterialExpressionSplitLeftRight)
+
+#define LOCTEXT_NAMESPACE "MaterialExpressionSplitLeftRight"
+
+UMaterialExpressionSplitLeftRight::UMaterialExpressionSplitLeftRight(const FObjectInitializer& ObjectInitializer)
+	: Super(ObjectInitializer),
+	ConstCenter(0.5f),
+	ConstCoordinate(0)
+{
+	// Structure to hold one-time initialization
+	struct FConstructorStatics
+	{
+		FText NAME_MaterialX;
+		FText NAME_Procedural;
+		FConstructorStatics()
+			: NAME_MaterialX(LOCTEXT("MaterialX", "MaterialX"))
+			, NAME_Procedural(LOCTEXT("Procedural", "Procedural"))
+		{}
+	};
+	static FConstructorStatics ConstructorStatics;
+
+#if WITH_EDITORONLY_DATA
+	MenuCategories.Add(ConstructorStatics.NAME_MaterialX);
+	MenuCategories.Add(ConstructorStatics.NAME_Procedural);
+#endif
+}
+
+#if WITH_EDITOR
+int32 UMaterialExpressionSplitLeftRight::Compile(FMaterialCompiler* Compiler, int32 OutputIndex)
+{
+	if(!A.GetTracedInput().Expression)
+	{
+		return Compiler->Errorf(TEXT("Missing A input"));
+	}
+
+	if(!B.GetTracedInput().Expression)
+	{
+		return Compiler->Errorf(TEXT("Missing B input"));
+	}
+
+	int32 CoordinateIndex = Coordinates.GetTracedInput().Expression ? Coordinates.Compile(Compiler) : Compiler->TextureCoordinate(ConstCoordinate, false, false);
+	int32 CenterIndex = Center.GetTracedInput().Expression ? Center.Compile(Compiler) : Compiler->Constant(ConstCenter);
+
+	int32 TexcoordXIndex = Compiler->ComponentMask(CoordinateIndex, true, false, false, false);
+	int32 InvSqrt2 = Compiler->Constant(0.70710678118654757);
+	int32 AFWidth = Compiler->Mul(Compiler->Length(Compiler->AppendVector(Compiler->DDX(TexcoordXIndex), Compiler->DDY(TexcoordXIndex))), InvSqrt2);
+	int32 AAStep = Compiler->SmoothStep(Compiler->Sub(CenterIndex, AFWidth), Compiler->Add(CenterIndex, AFWidth), TexcoordXIndex);
+
+	return Compiler->Lerp(A.Compile(Compiler), B.Compile(Compiler), AAStep);
+}
+
+void UMaterialExpressionSplitLeftRight::GetCaption(TArray<FString>& OutCaptions) const
+{
+	OutCaptions.Add(TEXT("SplitLR"));
+}
+#endif
+
+#undef LOCTEXT_NAMESPACE 
\ No newline at end of file
diff --git a/Engine/Plugins/Interchange/Runtime/Source/Import/Private/MaterialX/MaterialExpressions/MaterialExpressionSplitLeftRight.h b/Engine/Plugins/Interchange/Runtime/Source/Import/Private/MaterialX/MaterialExpressions/MaterialExpressionSplitLeftRight.h
new file mode 100644
--- /dev/null
+++ b/Engine/Plugins/Interchange/Runtime/Source/Import/Private/MaterialX/MaterialExpressions/MaterialExpressionSplitLeftRight.h
@@ -0,0 +1,44 @@
+// Copyright Epic Games, Inc. All Rights Reserved.
+#pragma once
+
+#include "CoreMinimal.h"
+#include "Materials/MaterialExpression.h"
+
+#include "MaterialExpressionSplitLeftRight.generated.h"
+
+/**
+ * A material expression that computes a left-right split matte, split at a specified u value.
+ */
+UCLASS(collapsecategories, hidecategories = Object, MinimalAPI)
+class UMaterialExpressionSplitLeftRight : public UMaterialExpression
+{
+	GENERATED_UCLASS_BODY()
+
+	UPROPERTY(meta = (RequiredInput = "false", ToolTip = "Defaults to 'ConstCoordinate' if not specified"))
+	FExpressionInput Coordinates;
+
+	UPROPERTY()
+	FExpressionInput A;
+
+	UPROPERTY()
+	FExpressionInput B;
+
+	UPROPERTY(meta = (RequiredInput = "false", ToolTip = "Defaults to 'ConstCenter' if not specified"))
+	FExpressionInput Center;
+
+	/** only used if A is not hooked up */
+	UPROPERTY(EditAnywhere, Category = MaterialExpressionMultiply, meta = (OverridingInputProperty = "Center"))
+	float ConstCenter;
+
+	/** only used if Coordinates is not hooked up */
+	UPROPERTY(EditAnywhere, Category = MaterialExpressionTextureSample, meta = (OverridingInputProperty = "Coordinates"))
+	uint8 ConstCoordinate;
+
+	//~ Begin UMaterialExpression Interface
+#if WITH_EDITOR
+	virtual int32 Compile(class FMaterialCompiler* Compiler, int32 OutputIndex) override;
+	virtual void GetCaption(TArray<FString>& OutCaptions) const override;
+#endif
+	//~ End UMaterialExpression Interface
+};
+
diff --git a/Engine/Plugins/Interchange/Runtime/Source/Import/Private/MaterialX/MaterialExpressions/MaterialExpressionSplitTopBottom.cpp b/Engine/Plugins/Interchange/Runtime/Source/Import/Private/MaterialX/MaterialExpressions/MaterialExpressionSplitTopBottom.cpp
new file mode 100644
--- /dev/null
+++ b/Engine/Plugins/Interchange/Runtime/Source/Import/Private/MaterialX/MaterialExpressions/MaterialExpressionSplitTopBottom.cpp
@@ -0,0 +1,62 @@
+// Copyright Epic Games, Inc. All Rights Reserved.
+#include "MaterialExpressionSplitTopBottom.h"
+#include "MaterialCompiler.h"
+
+#include UE_INLINE_GENERATED_CPP_BY_NAME(MaterialExpressionSplitTopBottom)
+
+#define LOCTEXT_NAMESPACE "MaterialExpressionSplitLeftRight"
+
+UMaterialExpressionSplitTopBottom::UMaterialExpressionSplitTopBottom(const FObjectInitializer& ObjectInitializer)
+	: Super(ObjectInitializer),
+	ConstCenter(0.5f),
+	ConstCoordinate(0)
+{
+	// Structure to hold one-time initialization
+	struct FConstructorStatics
+	{
+		FText NAME_MaterialX;
+		FText NAME_Procedural;
+		FConstructorStatics()
+			: NAME_MaterialX(LOCTEXT("MaterialX", "MaterialX"))
+			, NAME_Procedural(LOCTEXT("Procedural", "Procedural"))
+		{}
+	};
+	static FConstructorStatics ConstructorStatics;
+
+#if WITH_EDITORONLY_DATA
+	MenuCategories.Add(ConstructorStatics.NAME_MaterialX);
+	MenuCategories.Add(ConstructorStatics.NAME_Procedural);
+#endif
+}
+
+#if WITH_EDITOR
+int32 UMaterialExpressionSplitTopBottom::Compile(FMaterialCompiler* Compiler, int32 OutputIndex)
+{
+	if(!A.GetTracedInput().Expression)
+	{
+		return Compiler->Errorf(TEXT("Missing A input"));
+	}
+
+	if(!B.GetTracedInput().Expression)
+	{
+		return Compiler->Errorf(TEXT("Missing B input"));
+	}
+
+	int32 CoordinateIndex = Coordinates.GetTracedInput().Expression ? Coordinates.Compile(Compiler) : Compiler->TextureCoordinate(ConstCoordinate, false, false);
+	int32 CenterIndex = Center.GetTracedInput().Expression ? Center.Compile(Compiler) : Compiler->Constant(ConstCenter);
+
+	int32 TexcoordYIndex = Compiler->ComponentMask(CoordinateIndex, false, true, false, false);
+	int32 InvSqrt2 = Compiler->Constant(0.70710678118654757);
+	int32 AFWidth = Compiler->Mul(Compiler->Length(Compiler->AppendVector(Compiler->DDX(TexcoordYIndex), Compiler->DDY(TexcoordYIndex))), InvSqrt2);
+	int32 AAStep = Compiler->SmoothStep(Compiler->Sub(CenterIndex, AFWidth), Compiler->Add(CenterIndex, AFWidth), TexcoordYIndex);
+
+	return Compiler->Lerp(A.Compile(Compiler), B.Compile(Compiler), AAStep);
+}
+
+void UMaterialExpressionSplitTopBottom::GetCaption(TArray<FString>& OutCaptions) const
+{
+	OutCaptions.Add(TEXT("SplitTB"));
+}
+#endif
+
+#undef LOCTEXT_NAMESPACE 
\ No newline at end of file
diff --git a/Engine/Plugins/Interchange/Runtime/Source/Import/Private/MaterialX/MaterialExpressions/MaterialExpressionSplitTopBottom.h b/Engine/Plugins/Interchange/Runtime/Source/Import/Private/MaterialX/MaterialExpressions/MaterialExpressionSplitTopBottom.h
new file mode 100644
--- /dev/null
+++ b/Engine/Plugins/Interchange/Runtime/Source/Import/Private/MaterialX/MaterialExpressions/MaterialExpressionSplitTopBottom.h
@@ -0,0 +1,44 @@
+// Copyright Epic Games, Inc. All Rights Reserved.
+#pragma once
+
+#include "CoreMinimal.h"
+#include "Materials/MaterialExpression.h"
+
+#include "MaterialExpressionSplitTopBottom.generated.h"
+
+/**
+ * A material expression that computes a top-bottom split matte, split at a specified v value.
+ */
+UCLASS(collapsecategories, hidecategories = Object, MinimalAPI)
+class UMaterialExpressionSplitTopBottom : public UMaterialExpression
+{
+	GENERATED_UCLASS_BODY()
+
+	UPROPERTY(meta = (RequiredInput = "false", ToolTip = "Defaults to 'ConstCoordinate' if not specified"))
+	FExpressionInput Coordinates;
+
+	UPROPERTY()
+	FExpressionInput A;
+
+	UPROPERTY()
+	FExpressionInput B;
+
+	UPROPERTY(meta = (RequiredInput = "false", ToolTip = "Defaults to 'ConstCenter' if not specified"))
+	FExpressionInput Center;
+
+	/** only used if A is not hooked up */
+	UPROPERTY(EditAnywhere, Category = MaterialExpressionMultiply, meta = (OverridingInputProperty = "Center"))
+	float ConstCenter;
+
+	/** only used if Coordinates is not hooked up */
+	UPROPERTY(EditAnywhere, Category = MaterialExpressionTextureSample, meta = (OverridingInputProperty = "Coordinates"))
+	uint8 ConstCoordinate;
+
+	//~ Begin UMaterialExpression Interface
+#if WITH_EDITOR
+	virtual int32 Compile(class FMaterialCompiler* Compiler, int32 OutputIndex) override;
+	virtual void GetCaption(TArray<FString>& OutCaptions) const override;
+#endif
+	//~ End UMaterialExpression Interface
+};
+
diff --git a/Engine/Plugins/Interchange/Runtime/Source/Import/Private/MaterialX/MaterialExpressions/MaterialExpressionTextureSampleParameterBlur.cpp b/Engine/Plugins/Interchange/Runtime/Source/Import/Private/MaterialX/MaterialExpressions/MaterialExpressionTextureSampleParameterBlur.cpp
new file mode 100644
--- /dev/null
+++ b/Engine/Plugins/Interchange/Runtime/Source/Import/Private/MaterialX/MaterialExpressions/MaterialExpressionTextureSampleParameterBlur.cpp
@@ -0,0 +1,246 @@
+// Copyright Epic Games, Inc. All Rights Reserved.
+
+#include "MaterialX/MaterialExpressions/MaterialExpressionTextureSampleParameterBlur.h"
+#include "Materials/MaterialExpressionAdd.h"
+#include "Materials/MaterialExpressionConstant.h"
+#include "Materials/MaterialExpressionConstant2Vector.h"
+#include "Materials/MaterialExpressionCustom.h"
+#include "Materials/MaterialExpressionMultiply.h"
+#include "Materials/MaterialExpressionTextureCoordinate.h"
+
+#include UE_INLINE_GENERATED_CPP_BY_NAME(MaterialExpressionTextureSampleParameterBlur)
+
+#define LOCTEXT_NAMESPACE "MaterialExpressionTextureSampleBlur"
+
+namespace Gauss
+{
+	namespace
+	{
+		static constexpr float KernelWeights3x3[9] =
+		{
+			0.077847, 0.123317, 0.077847,
+			0.123317, 0.195346, 0.123317,
+			0.077847, 0.123317, 0.077847
+		};
+
+		static constexpr float KernelWeights5x5[25] =
+		{
+			0.003765, 0.015019, 0.023792, 0.015019, 0.003765,
+			0.015019, 0.059912, 0.094907, 0.059912, 0.015019,
+			0.023792, 0.094907, 0.150342, 0.094907, 0.023792,
+			0.015019, 0.059912, 0.094907, 0.059912, 0.015019,
+			0.003765, 0.015019, 0.023792, 0.015019, 0.003765
+		};
+
+		static constexpr float KernelWeights7x7[49] =
+		{
+			0.000036, 0.000363, 0.001446, 0.002291, 0.001446, 0.000363, 0.000036,
+			0.000363, 0.003676, 0.014662, 0.023226, 0.014662, 0.003676, 0.000363,
+			0.001446, 0.014662, 0.058488, 0.092651, 0.058488, 0.014662, 0.001446,
+			0.002291, 0.023226, 0.092651, 0.146768, 0.092651, 0.023226, 0.002291,
+			0.001446, 0.014662, 0.058488, 0.092651, 0.058488, 0.014662, 0.001446,
+			0.000363, 0.003676, 0.014662, 0.023226, 0.014662, 0.003676, 0.000363,
+			0.000036, 0.000363, 0.001446, 0.002291, 0.001446, 0.000363, 0.000036
+		};
+	}
+}
+
+namespace Box
+{
+	namespace
+	{
+		static constexpr float KernelWeights3x3[9] =
+		{
+			1.f / 9.f, 1.f / 9.f, 1.f / 9.f,
+			1.f / 9.f, 1.f / 9.f, 1.f / 9.f,
+			1.f / 9.f, 1.f / 9.f, 1.f / 9.f,
+		};
+
+		static constexpr float KernelWeights5x5[25] =
+		{
+			1.f / 25.f, 1.f / 25.f, 1.f / 25.f, 1.f / 25.f, 1.f / 25.f,
+			1.f / 25.f, 1.f / 25.f, 1.f / 25.f, 1.f / 25.f, 1.f / 25.f,
+			1.f / 25.f, 1.f / 25.f, 1.f / 25.f, 1.f / 25.f, 1.f / 25.f,
+			1.f / 25.f, 1.f / 25.f, 1.f / 25.f, 1.f / 25.f, 1.f / 25.f,
+			1.f / 25.f, 1.f / 25.f, 1.f / 25.f, 1.f / 25.f, 1.f / 25.f,
+		};
+
+		static constexpr float KernelWeights7x7[49] =
+		{
+			1.f / 49.f, 1.f / 49.f, 1.f / 49.f, 1.f / 49.f, 1.f / 49.f, 1.f / 49.f, 1.f / 49.f,
+			1.f / 49.f, 1.f / 49.f, 1.f / 49.f, 1.f / 49.f, 1.f / 49.f, 1.f / 49.f, 1.f / 49.f,
+			1.f / 49.f, 1.f / 49.f, 1.f / 49.f, 1.f / 49.f, 1.f / 49.f, 1.f / 49.f, 1.f / 49.f,
+			1.f / 49.f, 1.f / 49.f, 1.f / 49.f, 1.f / 49.f, 1.f / 49.f, 1.f / 49.f, 1.f / 49.f,
+			1.f / 49.f, 1.f / 49.f, 1.f / 49.f, 1.f / 49.f, 1.f / 49.f, 1.f / 49.f, 1.f / 49.f,
+			1.f / 49.f, 1.f / 49.f, 1.f / 49.f, 1.f / 49.f, 1.f / 49.f, 1.f / 49.f, 1.f / 49.f,
+			1.f / 49.f, 1.f / 49.f, 1.f / 49.f, 1.f / 49.f, 1.f / 49.f, 1.f / 49.f, 1.f / 49.f,
+		};
+	}
+}
+
+UMaterialExpressionTextureSampleParameterBlur::UMaterialExpressionTextureSampleParameterBlur(const FObjectInitializer& ObjectInitializer)
+	: Super(ObjectInitializer)
+{
+	// Structure to hold one-time initialization
+	struct FConstructorStatics
+	{
+		FText NAME_MaterialX;
+		FText NAME_Convolution;
+		FConstructorStatics()
+			: NAME_MaterialX(LOCTEXT("MaterialX", "MaterialX"))
+			, NAME_Convolution(LOCTEXT("Convolution", "Convolution"))
+		{}
+	};
+	static FConstructorStatics ConstructorStatics;
+
+#if WITH_EDITORONLY_DATA
+	MenuCategories.Add(ConstructorStatics.NAME_MaterialX);
+	MenuCategories.Add(ConstructorStatics.NAME_Convolution);
+#endif
+}
+
+#if WITH_EDITOR
+int32 UMaterialExpressionTextureSampleParameterBlur::Compile(FMaterialCompiler* Compiler, int32 OutputIndex)
+{
+	if(KernelSize > 0)
+	{
+		FExpressionInput TexCoordInput;
+		if(Coordinates.GetTracedInput().Expression)
+		{
+			TexCoordInput = Coordinates;
+		}
+		else
+		{
+			UMaterialExpressionTextureCoordinate* DefaultTexCoord = NewObject<UMaterialExpressionTextureCoordinate>();
+			DefaultTexCoord->CoordinateIndex = ConstCoordinate;
+			TexCoordInput.Connect(0, DefaultTexCoord);
+		}
+
+		UMaterialExpressionCustom* ComputeSampleSize = NewObject<UMaterialExpressionCustom>();
+
+		ComputeSampleSize->Inputs[0].InputName = TEXT("uv");
+		ComputeSampleSize->Inputs[0].Input = TexCoordInput;
+
+		auto ComputeSampleSizeInput = [&](float Value, FName InputName)
+		{
+			UMaterialExpressionConstant* FilterSizeConstant = NewObject<UMaterialExpressionConstant>();
+			FilterSizeConstant->R = Value;
+			FCustomInput & Input = ComputeSampleSize->Inputs.Add_GetRef({});
+			Input.InputName = InputName;
+			Input.Input.Expression = FilterSizeConstant;
+		};
+
+		ComputeSampleSizeInput(FilterSize, TEXT("filterSize"));
+		ComputeSampleSizeInput(FilterOffset, TEXT("filterOffset"));
+
+		ComputeSampleSize->OutputType = ECustomMaterialOutputType::CMOT_Float2;
+		ComputeSampleSize->Code = TEXT(
+			R"(float2 derivUVx = ddx(uv) * 0.5f;
+   float2 derivUVy = ddy(uv) * 0.5f;
+   float derivX = abs(derivUVx.x) + abs(derivUVy.x);
+   float derivY = abs(derivUVx.y) + abs(derivUVy.y);
+   float sampleSizeU = 2.0f * filterSize * derivX + filterOffset;
+   if (sampleSizeU < 1.0E-05f)
+       sampleSizeU = 1.0E-05f;
+   float sampleSizeV = 2.0f * filterSize * derivY + filterOffset;
+   if (sampleSizeV < 1.0E-05f)
+       sampleSizeV = 1.0E-05f;
+   return float2(sampleSizeU, sampleSizeV);)"
+		);
+
+		const float* Kernel;
+		int32 FilterWidth = 1;
+
+		if(KernelSize <= (1.f / 3.f))
+		{
+			FilterWidth = 3 / 2;
+			if(Filter == ETextureSampleBlurFilter::Box)
+			{
+				Kernel = Box::KernelWeights3x3;
+			}
+			else
+			{
+				Kernel = Gauss::KernelWeights3x3;
+			}
+		}
+		else if(KernelSize <= (2.f / 3.f))
+		{
+			FilterWidth = 5 / 2;
+			if(Filter == ETextureSampleBlurFilter::Box)
+			{
+				Kernel = Box::KernelWeights5x5;
+			}
+			else
+			{
+				Kernel = Gauss::KernelWeights5x5;
+			}
+		}
+		else
+		{
+			FilterWidth = 7 / 2;
+			if(Filter == ETextureSampleBlurFilter::Box)
+			{
+				Kernel = Box::KernelWeights7x7;
+			}
+			else
+			{
+				Kernel = Gauss::KernelWeights7x7;
+			}
+		}
+
+		UMaterialExpressionCustom* Convolution = NewObject<UMaterialExpressionCustom>();
+		Convolution->OutputType = ECustomMaterialOutputType::CMOT_Float4;
+		Convolution->Inputs.Empty();
+		Convolution->Inputs.Reserve(25);
+		Convolution->Code = TEXT("float4 result = float4(0,0,0,0);\n");
+
+		int32 i = 0;
+
+		for(int Row = -FilterWidth; Row <= FilterWidth; ++Row)
+		{
+			for(int Col = -FilterWidth; Col <= FilterWidth; ++Col)
+			{
+				UMaterialExpressionConstant2Vector* Constant = NewObject<UMaterialExpressionConstant2Vector>();
+				Constant->R = Col;
+				Constant->G = Row;
+				UMaterialExpressionMultiply* MultiplyCoord = NewObject<UMaterialExpressionMultiply>();
+				MultiplyCoord->A.Expression = ComputeSampleSize;
+				MultiplyCoord->B.Expression = Constant;
+				UMaterialExpressionAdd* AddCoord = NewObject<UMaterialExpressionAdd>();
+				AddCoord->A = TexCoordInput;
+				AddCoord->B.Expression = MultiplyCoord;
+
+				UMaterialExpressionTextureSampleParameter2D* TextureSample = NewObject<UMaterialExpressionTextureSampleParameter2D>();
+				TextureSample->ConstCoordinate = ConstCoordinate;
+				TextureSample->Coordinates.Expression = AddCoord;
+				TextureSample->TextureObject = TextureObject;
+				TextureSample->Texture = Texture;
+				TextureSample->SamplerSource = SamplerSource;
+				TextureSample->SamplerType = SamplerType;
+
+				Convolution->Inputs.Add({});
+				FString InputName{ TEXT("S") + FString::FromInt(i) };
+				Convolution->Inputs[i].InputName = FName(InputName);// FName(TEXT("S") + Index);
+				Convolution->Inputs[i].Input.Expression = TextureSample;
+				Convolution->Code += (TEXT("result += ") + InputName + TEXT("*") + FString::SanitizeFloat(Kernel[i]) + TEXT(";\n"));
+				++i;
+			}
+		}
+		Convolution->Code += TEXT("return result;");
+
+		return Convolution->Compile(Compiler, 0);
+	}
+	else
+	{
+		return Super::Compile(Compiler, OutputIndex);
+	}
+}
+
+void UMaterialExpressionTextureSampleParameterBlur::GetCaption(TArray<FString>& OutCaptions) const
+{
+	OutCaptions.Add(TEXT("ParamBlur"));
+	OutCaptions.Add(FString::Printf(TEXT("'%s'"), *ParameterName.ToString()));
+}
+#endif // WITH_EDITOR
+
+#undef LOCTEXT_NAMESPACE
\ No newline at end of file
diff --git a/Engine/Plugins/Interchange/Runtime/Source/Import/Private/MaterialX/MaterialExpressions/MaterialExpressionUnpremult.cpp b/Engine/Plugins/Interchange/Runtime/Source/Import/Private/MaterialX/MaterialExpressions/MaterialExpressionUnpremult.cpp
new file mode 100644
--- /dev/null
+++ b/Engine/Plugins/Interchange/Runtime/Source/Import/Private/MaterialX/MaterialExpressions/MaterialExpressionUnpremult.cpp
@@ -0,0 +1,52 @@
+// Copyright Epic Games, Inc. All Rights Reserved.
+#include "MaterialExpressionUnpremult.h"
+#include "MaterialCompiler.h"
+
+#include UE_INLINE_GENERATED_CPP_BY_NAME(MaterialExpressionUnpremult)
+
+#define LOCTEXT_NAMESPACE "MaterialExpressionUnpremult"
+
+UMaterialExpressionUnpremult::UMaterialExpressionUnpremult(const FObjectInitializer& ObjectInitializer)
+	: Super(ObjectInitializer)
+{
+	// Structure to hold one-time initialization
+	struct FConstructorStatics
+	{
+		FText NAME_MaterialX;
+		FText NAME_Compositing;
+		FConstructorStatics()
+			: NAME_MaterialX(LOCTEXT("MaterialX", "MaterialX"))
+			, NAME_Compositing(LOCTEXT("Compositing", "Compositing"))
+		{}
+	};
+	static FConstructorStatics ConstructorStatics;
+
+#if WITH_EDITORONLY_DATA
+	MenuCategories.Add(ConstructorStatics.NAME_MaterialX);
+	MenuCategories.Add(ConstructorStatics.NAME_Compositing);
+#endif
+}
+
+#if WITH_EDITOR
+int32 UMaterialExpressionUnpremult::Compile(FMaterialCompiler* Compiler, int32 OutputIndex)
+{
+	if(!Input.GetTracedInput().Expression)
+	{
+		return Compiler->Errorf(TEXT("Missing Unpremult input"));
+	}
+
+	int32 IndexInput = Input.Compile(Compiler);
+	int32 RGB = Compiler->ComponentMask(IndexInput, true, true, true, false);
+	int32 Alpha = Compiler->ComponentMask(IndexInput, false, false, false, true);
+
+	int32 Unpremult = Compiler->AppendVector(Compiler->Div(RGB, Alpha), Alpha);
+	return Compiler->If(Alpha, Compiler->Constant(0.f), Unpremult, IndexInput, Unpremult, Compiler->Constant(0.00001f));
+}
+
+void UMaterialExpressionUnpremult::GetCaption(TArray<FString>& OutCaptions) const
+{
+	OutCaptions.Add(TEXT("Unpremult"));
+}
+#endif
+
+#undef LOCTEXT_NAMESPACE 
\ No newline at end of file
diff --git a/Engine/Plugins/Interchange/Runtime/Source/Import/Private/MaterialX/MaterialExpressions/MaterialExpressionUnpremult.h b/Engine/Plugins/Interchange/Runtime/Source/Import/Private/MaterialX/MaterialExpressions/MaterialExpressionUnpremult.h
new file mode 100644
--- /dev/null
+++ b/Engine/Plugins/Interchange/Runtime/Source/Import/Private/MaterialX/MaterialExpressions/MaterialExpressionUnpremult.h
@@ -0,0 +1,29 @@
+// Copyright Epic Games, Inc. All Rights Reserved.
+#pragma once
+
+#include "CoreMinimal.h"
+#include "Materials/MaterialExpression.h"
+
+#include "MaterialExpressionUnpremult.generated.h"
+
+/**
+ * Divide the RGB channels of the input by the Alpha channel of the input. If the
+ * Alpha value is zero, the original color4 input value is passed through unchanged.
+ * Input must be of type float4
+ */
+UCLASS(collapsecategories, hidecategories = Object, MinimalAPI)
+class UMaterialExpressionUnpremult : public UMaterialExpression
+{
+	GENERATED_UCLASS_BODY()
+
+	UPROPERTY()
+	FExpressionInput Input;
+
+	//~ Begin UMaterialExpression Interface
+#if WITH_EDITOR
+	virtual int32 Compile(class FMaterialCompiler* Compiler, int32 OutputIndex) override;
+	virtual void GetCaption(TArray<FString>& OutCaptions) const override;
+#endif
+	//~ End UMaterialExpression Interface
+};
+
diff --git a/Engine/Plugins/Interchange/Runtime/Source/Import/Public/MaterialX/MaterialExpressionTextureSampleParameterBlur.h b/Engine/Plugins/Interchange/Runtime/Source/Import/Public/MaterialX/MaterialExpressionTextureSampleParameterBlur.h
deleted file mode 100644
diff --git a/Engine/Plugins/Interchange/Runtime/Source/Import/Public/MaterialX/MaterialExpressions/MaterialExpressionTextureSampleParameterBlur.h b/Engine/Plugins/Interchange/Runtime/Source/Import/Public/MaterialX/MaterialExpressions/MaterialExpressionTextureSampleParameterBlur.h
new file mode 100644
--- /dev/null
+++ b/Engine/Plugins/Interchange/Runtime/Source/Import/Public/MaterialX/MaterialExpressions/MaterialExpressionTextureSampleParameterBlur.h
@@ -0,0 +1,47 @@
+// Copyright Epic Games, Inc. All Rights Reserved.
+#pragma once
+
+#include "Materials/MaterialExpressionTextureSampleParameter2D.h"
+#include "MaterialExpressionTextureSampleParameterBlur.generated.h"
+
+UENUM()
+enum class ETextureSampleBlurFilter : uint8
+{
+	Box,
+	Gaussian
+};
+
+UCLASS(MinimalAPI, collapsecategories, hidecategories = Object)
+class UMaterialExpressionTextureSampleParameterBlur: public UMaterialExpressionTextureSampleParameter2D
+{
+	GENERATED_UCLASS_BODY()
+
+	/** The size of the blur kernel, relative to 0-1 UV space, only 3x3, 5x5 and 7x7 kernel are supported.
+	  * Basically depending on the KernelSize value:
+	  * KernelSize == 0   -> 1x1 Kernel (Meaning no Blur)
+	  * KernelSize <= 1/3 -> 3x3 Kernel
+	  * KernelSize <= 2/3 -> 5x5 Kernel
+	  * Otherwise         -> 7x7 Kernel
+	  */
+	UPROPERTY(EditAnywhere, Category = Filtering, Meta = (ShowAsInputPin = "Advanced"))
+	float KernelSize = 0.f;
+
+	/** Size of the filter when we sample a texture coordinate */
+	UPROPERTY(EditAnywhere, Category = Filtering, Meta = (ShowAsInputPin = "Advanced"))
+	float FilterSize = 1.f;
+
+	/** Offset of the filter when we sample a texture coordinate */
+	UPROPERTY(EditAnywhere, Category = Filtering, Meta = (ShowAsInputPin = "Advanced"))
+	float FilterOffset = 0.f;
+
+	/** Filter to use when we blur a Texture: Gaussian or Box Linear filter*/
+	UPROPERTY(EditAnywhere, Category = Filtering, Meta = (ShowAsInputPin = "Advanced"))
+	ETextureSampleBlurFilter Filter = ETextureSampleBlurFilter::Gaussian;
+
+	//~ Begin UMaterialExpression Interface
+#if WITH_EDITOR
+	virtual int32 Compile(class FMaterialCompiler* Compiler, int32 OutputIndex) override;
+	virtual void GetCaption(TArray<FString>& OutCaptions) const override;
+#endif
+	//~ End UMaterialExpression Interface
+};
\ No newline at end of file
diff --git a/Engine/Plugins/Interchange/Runtime/Source/Pipelines/Private/InterchangeGenericMaterialPipeline.cpp b/Engine/Plugins/Interchange/Runtime/Source/Pipelines/Private/InterchangeGenericMaterialPipeline.cpp
--- a/Engine/Plugins/Interchange/Runtime/Source/Pipelines/Private/InterchangeGenericMaterialPipeline.cpp
+++ b/Engine/Plugins/Interchange/Runtime/Source/Pipelines/Private/InterchangeGenericMaterialPipeline.cpp
@@ -44,7 +44,7 @@
 #include "Materials/MaterialFunction.h"
 #include "Materials/MaterialInstanceConstant.h"
 #include "Materials/MaterialInstanceDynamic.h"
-#include "MaterialX/MaterialExpressionTextureSampleParameterBlur.h"
+#include "MaterialX/MaterialExpressions/MaterialExpressionTextureSampleParameterBlur.h"
 #include "Misc/CoreMisc.h"
 #include "Misc/Paths.h"
 #include "Nodes/InterchangeBaseNode.h"
