[Interchange] Refactor GetAnimationPayloadData interface.
Also fixing: When forcing a Skeletal mesh on glTF asset import the animation sequence is not created

diff --git a/Engine/Plugins/Experimental/Enterprise/DatasmithInterchange/Source/DatasmithInterchange/Private/InterchangeDatasmithAnimationUtils.cpp b/Engine/Plugins/Experimental/Enterprise/DatasmithInterchange/Source/DatasmithInterchange/Private/InterchangeDatasmithAnimationUtils.cpp
--- a/Engine/Plugins/Experimental/Enterprise/DatasmithInterchange/Source/DatasmithInterchange/Private/InterchangeDatasmithAnimationUtils.cpp
+++ b/Engine/Plugins/Experimental/Enterprise/DatasmithInterchange/Source/DatasmithInterchange/Private/InterchangeDatasmithAnimationUtils.cpp
@@ -5,7 +5,7 @@
 
 #include "DatasmithAnimationElements.h"
 
-#include "Animation/InterchangeAnimationPayload.h"
+#include "InterchangeCommonAnimationPayload.h"
 #include "InterchangeAnimationTrackSetNode.h"
 #include "InterchangeAnimSequenceFactoryNode.h"
 
@@ -138,7 +138,7 @@ namespace UE::DatasmithInterchange::AnimUtils
 		TrackNode->SetCustomActorDependencyUid(ActorNodeUid);
 
 		TrackNode->SetCustomTargetedProperty((int32)EInterchangeAnimatedProperty::Visibility);
-		TrackNode->SetCustomAnimationPayloadKey(TrackNode->GetUniqueID());
+		TrackNode->SetCustomAnimationPayloadKey(TrackNode->GetUniqueID(), EInterchangeAnimationPayLoadType::STEPCURVE);
 
 		LevelSequenceNode->AddCustomAnimationTrackUid(TrackNode->GetUniqueID());
 
@@ -159,7 +159,7 @@ namespace UE::DatasmithInterchange::AnimUtils
 
 		TrackNode->SetCustomCompletionMode((int32)TransformAnimation.GetCompletionMode());
 		TrackNode->SetCustomUsedChannels((int32)TransformAnimation.GetEnabledTransformChannels());
-		TrackNode->SetCustomAnimationPayloadKey(TrackNode->GetUniqueID());
+		TrackNode->SetCustomAnimationPayloadKey(TrackNode->GetUniqueID(), EInterchangeAnimationPayLoadType::CURVE);
 
 		LevelSequenceNode->AddCustomAnimationTrackUid(TrackNode->GetUniqueID());
 
@@ -334,7 +334,7 @@ namespace UE::DatasmithInterchange::AnimUtils
 		TArray<FRichCurve>& Curves;
 	};
 
-	bool GetAnimationPayloadData(const IDatasmithBaseAnimationElement& AnimationElement, float FrameRate, UE::Interchange::FAnimationCurvePayloadData& PayLoadData)
+	bool GetAnimationPayloadData(const IDatasmithBaseAnimationElement& AnimationElement, float FrameRate, TArray<FRichCurve>& Curves)
 	{
 		using namespace UE::Interchange::Animation;
 
@@ -342,7 +342,6 @@ namespace UE::DatasmithInterchange::AnimUtils
 		{
 			const IDatasmithTransformAnimationElement& TransformAnimation = static_cast<const IDatasmithTransformAnimationElement&>(AnimationElement);
 
-			TArray<FRichCurve>& Curves = PayLoadData.Curves;
 			Curves.SetNum(9);
 
 			FAnimationCurvesHelper Helper(TransformAnimation, ConvertSampleRatetoFrameRate(FrameRate), Curves);
@@ -357,7 +356,7 @@ namespace UE::DatasmithInterchange::AnimUtils
 		return false;
 	}
 
-	bool GetAnimationPayloadData(const IDatasmithBaseAnimationElement& AnimationElement, float FloatFrameRate, UE::Interchange::FAnimationStepCurvePayloadData& PayLoadData)
+	bool GetAnimationPayloadData(const IDatasmithBaseAnimationElement& AnimationElement, float FloatFrameRate, TArray<FInterchangeStepCurve>& StepCurves)
 	{
 		using namespace UE::Interchange::Animation;
 
@@ -366,7 +365,7 @@ namespace UE::DatasmithInterchange::AnimUtils
 			const IDatasmithVisibilityAnimationElement& VisibilityAnimation = static_cast<const IDatasmithVisibilityAnimationElement&>(AnimationElement);
 			const FFrameRate FrameRate = ConvertSampleRatetoFrameRate(FloatFrameRate);
 
-			FInterchangeStepCurve& Curve = PayLoadData.StepCurves.AddDefaulted_GetRef();
+			FInterchangeStepCurve& Curve = StepCurves.AddDefaulted_GetRef();
 
 			const int32 NumFrames = VisibilityAnimation.GetFramesCount();
 
diff --git a/Engine/Plugins/Experimental/Enterprise/DatasmithInterchange/Source/DatasmithInterchange/Private/InterchangeDatasmithTranslator.cpp b/Engine/Plugins/Experimental/Enterprise/DatasmithInterchange/Source/DatasmithInterchange/Private/InterchangeDatasmithTranslator.cpp
--- a/Engine/Plugins/Experimental/Enterprise/DatasmithInterchange/Source/DatasmithInterchange/Private/InterchangeDatasmithTranslator.cpp
+++ b/Engine/Plugins/Experimental/Enterprise/DatasmithInterchange/Source/DatasmithInterchange/Private/InterchangeDatasmithTranslator.cpp
@@ -593,20 +593,75 @@ TFuture<TOptional<UE::Interchange::FMeshPayloadData>> UInterchangeDatasmithTrans
 	);
 }
 
-TFuture<TOptional<UE::Interchange::FAnimationCurvePayloadData>> UInterchangeDatasmithTranslator::GetAnimationCurvePayloadData(const FString& PayLoadKey) const
+TFuture<TOptional<UE::Interchange::FAnimationPayloadData>> UInterchangeDatasmithTranslator::GetAnimationPayloadData(const FInterchangeAnimationPayLoadKey& PayLoadKey, const double BakeFrequency, const double RangeStartSecond, const double RangeStopSecond) const
 {
-	return GetAnimationPayloadDataAsCurve<UE::Interchange::FAnimationCurvePayloadData>(PayLoadKey);
-}
+	UE::Interchange::FAnimationPayloadData TransformPayloadData(PayLoadKey.Type);
 
-TFuture<TOptional<UE::Interchange::FAnimationStepCurvePayloadData>> UInterchangeDatasmithTranslator::GetAnimationStepCurvePayloadData(const FString& PayLoadKey) const
-{
-	return GetAnimationPayloadDataAsCurve<UE::Interchange::FAnimationStepCurvePayloadData>(PayLoadKey);
-}
+	TPromise<TOptional<UE::Interchange::FAnimationPayloadData>> EmptyPromise;
+	EmptyPromise.SetValue(TOptional<UE::Interchange::FAnimationPayloadData>());
+	
+
+	if (!LoadedExternalSource || !LoadedExternalSource->GetDatasmithScene())
+	{
+		return EmptyPromise.GetFuture();
+	}
+
+	TSharedPtr<IDatasmithBaseAnimationElement> AnimationElement;
+	float FrameRate = 0.f;
+	if (UE::DatasmithInterchange::AnimUtils::FAnimationPayloadDesc* PayloadDescPtr = AnimationPayLoadMapping.Find(PayLoadKey.UniqueId))
+	{
+		AnimationElement = PayloadDescPtr->Value;
+		if (!ensure(AnimationElement))
+		{
+			// #ueent_logwarning:
+			return EmptyPromise.GetFuture();
+		}
+
+		FrameRate = PayloadDescPtr->Key;
+	}
+
+	switch (PayLoadKey.Type)
+	{
+	case EInterchangeAnimationPayLoadType::CURVE:
+	case EInterchangeAnimationPayLoadType::MORPHTARGETCURVE:
+		{
+			return Async(EAsyncExecution::TaskGraph, [this, AnimationElement = MoveTemp(AnimationElement), FrameRate, &TransformPayloadData]
+				{
+					
+					TOptional<UE::Interchange::FAnimationPayloadData> Result;
+
+					if (UE::DatasmithInterchange::AnimUtils::GetAnimationPayloadData(*AnimationElement, FrameRate, TransformPayloadData.Curves))
+					{
+						Result.Emplace(MoveTemp(TransformPayloadData));
+					}
+
+					return Result;
+					}
+				);
+		}
+		break;
+	case EInterchangeAnimationPayLoadType::STEPCURVE:
+		{
+			return Async(EAsyncExecution::TaskGraph, [this, AnimationElement = MoveTemp(AnimationElement), FrameRate, &TransformPayloadData]
+				{
+					TOptional<UE::Interchange::FAnimationPayloadData> Result;
+
+					if (UE::DatasmithInterchange::AnimUtils::GetAnimationPayloadData(*AnimationElement, FrameRate, TransformPayloadData.StepCurves))
+					{
+						Result.Emplace(MoveTemp(TransformPayloadData));
+					}
+
+					return Result;
+				}
+			);
+		}
+		break;
+	case EInterchangeAnimationPayLoadType::BAKED:
+	case EInterchangeAnimationPayLoadType::NONE:
+	default:
+		break;
+	}
 
-TFuture<TOptional<UE::Interchange::FAnimationBakeTransformPayloadData>> UInterchangeDatasmithTranslator::GetAnimationBakeTransformPayloadData(const FString& PayLoadKey, const double BakeFrequency, const double RangeStartSecond, const double RangeStopSecond) const
-{
-	TPromise<TOptional<UE::Interchange::FAnimationBakeTransformPayloadData>> EmptyPromise;
-	EmptyPromise.SetValue(TOptional<UE::Interchange::FAnimationBakeTransformPayloadData>());
 	return EmptyPromise.GetFuture();
 }
 
diff --git a/Engine/Plugins/Experimental/Enterprise/DatasmithInterchange/Source/DatasmithInterchange/Public/InterchangeDatasmithTranslator.h b/Engine/Plugins/Experimental/Enterprise/DatasmithInterchange/Source/DatasmithInterchange/Public/InterchangeDatasmithTranslator.h
--- a/Engine/Plugins/Experimental/Enterprise/DatasmithInterchange/Source/DatasmithInterchange/Public/InterchangeDatasmithTranslator.h
+++ b/Engine/Plugins/Experimental/Enterprise/DatasmithInterchange/Source/DatasmithInterchange/Public/InterchangeDatasmithTranslator.h
@@ -45,8 +45,8 @@ namespace UE::DatasmithImporter
 namespace UE::DatasmithInterchange::AnimUtils
 {
 	typedef TPair<float, TSharedPtr<IDatasmithBaseAnimationElement>> FAnimationPayloadDesc;
-	extern bool GetAnimationPayloadData(const IDatasmithBaseAnimationElement& AnimationElement, float FrameRate, UE::Interchange::FAnimationCurvePayloadData& PayLoadData);
-	extern bool GetAnimationPayloadData(const IDatasmithBaseAnimationElement& AnimationElement, float FrameRate, UE::Interchange::FAnimationStepCurvePayloadData& PayLoadData);
+	extern bool GetAnimationPayloadData(const IDatasmithBaseAnimationElement& AnimationElement, float FrameRate, TArray<FRichCurve>& Curves);
+	extern bool GetAnimationPayloadData(const IDatasmithBaseAnimationElement& AnimationElement, float FrameRate, TArray<FInterchangeStepCurve>& StepCurves);
 }
 
 UCLASS(BlueprintType, Experimental)
@@ -83,9 +83,7 @@ public:
 	/* IInterchangeStaticMeshPayloadInterface End */
 
 	/* IInterchangeAnimationPayloadInterface Begin */
-	virtual TFuture<TOptional<UE::Interchange::FAnimationCurvePayloadData>> GetAnimationCurvePayloadData(const FString& PayLoadKey) const override;
-	virtual TFuture<TOptional<UE::Interchange::FAnimationStepCurvePayloadData>> GetAnimationStepCurvePayloadData(const FString& PayLoadKey) const override;
-	virtual TFuture<TOptional<UE::Interchange::FAnimationBakeTransformPayloadData>> GetAnimationBakeTransformPayloadData(const FString& PayLoadKey, const double BakeFrequency, const double RangeStartSecond, const double RangeStopSecond) const override;
+	virtual TFuture<TOptional<UE::Interchange::FAnimationPayloadData>> GetAnimationPayloadData(const FInterchangeAnimationPayLoadKey& PayLoadKey, const double BakeFrequency = 0, const double RangeStartSecond = 0, const double RangeStopSecond = 0) const override;
 	/* IInterchangeAnimationPayloadInterface End */
 
 	/* IInterchangeVariantSetPayloadInterface Begin */
@@ -105,48 +103,6 @@ private:
 
 	TSharedPtr<UE::DatasmithImporter::FExternalSource> LoadedExternalSource;
 
-	template<typename T>
-	TFuture<TOptional<T>> GetAnimationPayloadDataAsCurve(const FString& PayLoadKey) const
-	{
-		using namespace UE::DatasmithInterchange::AnimUtils;
-
-		TPromise<TOptional<T>> EmptyPromise;
-		EmptyPromise.SetValue(TOptional<T>());
-
-		if (!LoadedExternalSource || !LoadedExternalSource->GetDatasmithScene())
-		{
-			return EmptyPromise.GetFuture();
-		}
-
-		TSharedPtr<IDatasmithBaseAnimationElement> AnimationElement;
-		float FrameRate = 0.f;
-		if (FAnimationPayloadDesc* PayloadDescPtr = AnimationPayLoadMapping.Find(PayLoadKey))
-		{
-			AnimationElement = PayloadDescPtr->Value;
-			if (!ensure(AnimationElement))
-			{
-				// #ueent_logwarning:
-				return EmptyPromise.GetFuture();
-			}
-
-			FrameRate = PayloadDescPtr->Key;
-		}
-
-		return Async(EAsyncExecution::TaskGraph, [this, AnimationElement = MoveTemp(AnimationElement), FrameRate]
-			{
-				T TransformPayloadData;
-				TOptional<T> Result;
-
-				if (GetAnimationPayloadData(*AnimationElement, FrameRate, TransformPayloadData))
-				{
-					Result.Emplace(MoveTemp(TransformPayloadData));
-				}
-
-				return Result;
-			}
-		);
-	}
-
 	mutable uint64 StartTime = 0;
 	mutable FString FileName;
 
diff --git a/Engine/Plugins/Interchange/Runtime/Source/FactoryNodes/Private/InterchangeAnimSequenceFactoryNode.cpp b/Engine/Plugins/Interchange/Runtime/Source/FactoryNodes/Private/InterchangeAnimSequenceFactoryNode.cpp
--- a/Engine/Plugins/Interchange/Runtime/Source/FactoryNodes/Private/InterchangeAnimSequenceFactoryNode.cpp
+++ b/Engine/Plugins/Interchange/Runtime/Source/FactoryNodes/Private/InterchangeAnimSequenceFactoryNode.cpp
@@ -2,6 +2,7 @@
 
 #include "InterchangeAnimSequenceFactoryNode.h"
 #include "Animation/AnimSequence.h"
+#include "InterchangeAnimationTrackSetNode.h"
 
 #include UE_INLINE_GENERATED_CPP_BY_NAME(InterchangeAnimSequenceFactoryNode)
 
@@ -46,15 +47,27 @@ namespace UE::Interchange
 		return AttributeKey;
 	}
 
-	const FAttributeKey& FAnimSequenceNodeStaticData::GetSceneNodeAnimationPayloadKeyMapKey()
+	const FAttributeKey& FAnimSequenceNodeStaticData::GetSceneNodeAnimationPayloadKeyUidMapKey()
 	{
-		static FAttributeKey AttributeKey(TEXT("__SceneNodeAnimationPayloadKeyMap__"));
+		static FAttributeKey AttributeKey(TEXT("__SceneNodeAnimationPayloadKeyUidMap__"));
 		return AttributeKey;
 	}
 
-	const FAttributeKey& FAnimSequenceNodeStaticData::GetMorphTargetNodePayloadKeyMapKey()
+	const FAttributeKey& FAnimSequenceNodeStaticData::GetSceneNodeAnimationPayloadKeyTypeMapKey()
 	{
-		static FAttributeKey AttributeKey(TEXT("__MorphTargetNodePayloadKeyMap__"));
+		static FAttributeKey AttributeKey(TEXT("__SceneNodeAnimationPayloadKeyTypeMap__"));
+		return AttributeKey;
+	}
+
+	const FAttributeKey& FAnimSequenceNodeStaticData::GetMorphTargetNodePayloadKeyUidMapKey()
+	{
+		static FAttributeKey AttributeKey(TEXT("__MorphTargetNodePayloadKeyUidMap__"));
+		return AttributeKey;
+	}
+
+	const FAttributeKey& FAnimSequenceNodeStaticData::GetMorphTargetNodePayloadKeyTypeMapKey()
+	{
+		static FAttributeKey AttributeKey(TEXT("__MorphTargetNodePayloadKeyTypeMap__"));
 		return AttributeKey;
 	}
 }//ns UE::Interchange
@@ -67,8 +80,11 @@ UInterchangeAnimSequenceFactoryNode::UInterchangeAnimSequenceFactoryNode()
 	AnimatedAttributeStepCurveNames.Initialize(Attributes, UE::Interchange::FAnimSequenceNodeStaticData::GetAnimatedAttributeStepCurveNamesKey().ToString());
 	AnimatedMaterialCurveSuffixes.Initialize(Attributes, UE::Interchange::FAnimSequenceNodeStaticData::GetAnimatedMaterialCurveSuffixesKey().ToString());
 
-	SceneNodeAnimationPayloadKeyMap.Initialize(Attributes.ToSharedRef(), UE::Interchange::FAnimSequenceNodeStaticData::GetSceneNodeAnimationPayloadKeyMapKey().ToString());
-	MorphTargetNodePayloadKeyMap.Initialize(Attributes.ToSharedRef(), UE::Interchange::FAnimSequenceNodeStaticData::GetMorphTargetNodePayloadKeyMapKey().ToString());
+	SceneNodeAnimationPayLoadKeyUidMap.Initialize(Attributes.ToSharedRef(), UE::Interchange::FAnimSequenceNodeStaticData::GetSceneNodeAnimationPayloadKeyUidMapKey().ToString());
+	SceneNodeAnimationPayLoadKeyTypeMap.Initialize(Attributes.ToSharedRef(), UE::Interchange::FAnimSequenceNodeStaticData::GetSceneNodeAnimationPayloadKeyTypeMapKey().ToString());
+
+	MorphTargetNodePayloadKeyUidMap.Initialize(Attributes.ToSharedRef(), UE::Interchange::FAnimSequenceNodeStaticData::GetMorphTargetNodePayloadKeyUidMapKey().ToString());
+	MorphTargetNodePayloadKeyTypeMap.Initialize(Attributes.ToSharedRef(), UE::Interchange::FAnimSequenceNodeStaticData::GetMorphTargetNodePayloadKeyTypeMapKey().ToString());
 }
 
 void UInterchangeAnimSequenceFactoryNode::InitializeAnimSequenceNode(const FString& UniqueID, const FString& DisplayLabel)
@@ -323,44 +339,65 @@ bool UInterchangeAnimSequenceFactoryNode::SetCustomDeleteExistingNonCurveCustomA
 	IMPLEMENT_NODE_ATTRIBUTE_SETTER_NODELEGATE(DeleteExistingNonCurveCustomAttributes, bool);
 }
 
-void UInterchangeAnimSequenceFactoryNode::GetSceneNodeAnimationPayloadKeys(TMap<FString, FString>& OutSceneNodeAnimationPayloads) const
+void UInterchangeAnimSequenceFactoryNode::GetSceneNodeAnimationPayloadKeys(TMap<FString, FInterchangeAnimationPayLoadKey>& OutSceneNodeAnimationPayloadKeys) const
 {
-	OutSceneNodeAnimationPayloads = SceneNodeAnimationPayloadKeyMap.ToMap();
-}
+	TMap<FString, FString> SceneNodeToPayLoadUidMap = SceneNodeAnimationPayLoadKeyUidMap.ToMap();
+	TMap<FString, uint8> SceneNodeToPayLoadTypeMap = SceneNodeAnimationPayLoadKeyTypeMap.ToMap();
 
-bool UInterchangeAnimSequenceFactoryNode::GetAnimationPayloadKeyFromSceneNodeUid(const FString& SceneNodeUid, FString& OutPayloadKey) const
-{
-	return SceneNodeAnimationPayloadKeyMap.GetValue(SceneNodeUid, OutPayloadKey);
-}
+	OutSceneNodeAnimationPayloadKeys.Reserve(SceneNodeToPayLoadTypeMap.Num());
 
-bool UInterchangeAnimSequenceFactoryNode::SetAnimationPayloadKeyForSceneNodeUid(const FString& SceneNodeUid, const FString& PayloadKey)
-{
-	return SceneNodeAnimationPayloadKeyMap.SetKeyValue(SceneNodeUid, PayloadKey);
+	for (const TPair<FString, FString>& SceneNodePayloadUidPair : SceneNodeToPayLoadUidMap)
+	{
+		uint8* Value = SceneNodeToPayLoadTypeMap.Find(SceneNodePayloadUidPair.Key);
+		if (Value)
+		{
+			OutSceneNodeAnimationPayloadKeys.Add(SceneNodePayloadUidPair.Key, FInterchangeAnimationPayLoadKey(SceneNodePayloadUidPair.Value, (EInterchangeAnimationPayLoadType)*Value));
+		}
+	}
 }
 
-bool UInterchangeAnimSequenceFactoryNode::RemoveAnimationPayloadKeyForSceneNodeUid(const FString& SceneNodeUid)
+void UInterchangeAnimSequenceFactoryNode::SetAnimationPayloadKeysForSceneNodeUids(const TMap<FString, FString>& SceneNodeAnimationPayloadKeyUids, const TMap<FString, uint8>& SceneNodeAnimationPayloadKeyTypes)
 {
-	return SceneNodeAnimationPayloadKeyMap.RemoveKey(SceneNodeUid);
-}
+	for (const TPair<FString, FString>& SceneNodePayLoadUidPair : SceneNodeAnimationPayloadKeyUids)
+	{
+		SceneNodeAnimationPayLoadKeyUidMap.SetKeyValue(SceneNodePayLoadUidPair.Key, SceneNodePayLoadUidPair.Value);
+	}
 
-void UInterchangeAnimSequenceFactoryNode::GetMorphTargetNodeAnimationPayloadKeys(TMap<FString, FString>& OutMorphTargetAnimationPayloads) const
-{
-	OutMorphTargetAnimationPayloads = MorphTargetNodePayloadKeyMap.ToMap();
+	for (const TPair<FString, uint8>& SceneNodePayLoadTypePair : SceneNodeAnimationPayloadKeyTypes)
+	{
+		SceneNodeAnimationPayLoadKeyTypeMap.SetKeyValue(SceneNodePayLoadTypePair.Key, SceneNodePayLoadTypePair.Value);
+	}
 }
 
-bool UInterchangeAnimSequenceFactoryNode::GetAnimationPayloadKeyFromMorphTargetNodeUid(const FString& MorphTargetNodeUid, FString& OutPayloadKey) const
+void UInterchangeAnimSequenceFactoryNode::GetMorphTargetNodeAnimationPayloadKeys(TMap<FString, FInterchangeAnimationPayLoadKey>& OutMorphTargetAnimationPayloads) const
 {
-	return MorphTargetNodePayloadKeyMap.GetValue(MorphTargetNodeUid, OutPayloadKey);
-}
+	TMap<FString, FString> MorphTargetToPayLoadUidMap = MorphTargetNodePayloadKeyUidMap.ToMap();
+	TMap<FString, uint8> MorphTargetToPayLoadTypeMap = MorphTargetNodePayloadKeyTypeMap.ToMap();
+
+	OutMorphTargetAnimationPayloads.Reserve(MorphTargetToPayLoadTypeMap.Num());
+
+	for (const TPair<FString, FString>& SceneNodePayloadUidPair : MorphTargetToPayLoadUidMap)
+	{
+		uint8* Value = MorphTargetToPayLoadTypeMap.Find(SceneNodePayloadUidPair.Key);
+		if (Value)
+		{
+			OutMorphTargetAnimationPayloads.Add(SceneNodePayloadUidPair.Key, FInterchangeAnimationPayLoadKey(SceneNodePayloadUidPair.Value, (EInterchangeAnimationPayLoadType)*Value));
+		}
+	}
 
-bool UInterchangeAnimSequenceFactoryNode::SetAnimationPayloadKeyForMorphTargetNodeUid(const FString& MorphTargetNodeUid, const FString& PayloadKey)
-{
-	return MorphTargetNodePayloadKeyMap.SetKeyValue(MorphTargetNodeUid, PayloadKey);
 }
 
-bool UInterchangeAnimSequenceFactoryNode::RemoveAnimationPayloadKeyForMorphTargetNodeUid(const FString& MorphTargetNodeUid)
+void UInterchangeAnimSequenceFactoryNode::SetAnimationPayloadKeysForMorphTargetNodeUids(const TMap<FString, FString>& MorphTargetAnimationPayloadKeyUids, const TMap<FString, uint8>& MorphTargetAnimationPayloadKeyTypes)
 {
-	return MorphTargetNodePayloadKeyMap.RemoveKey(MorphTargetNodeUid);
+	for (const TPair<FString, FString>& MorphTargetPayLoadUidPair : MorphTargetAnimationPayloadKeyUids)
+	{
+		MorphTargetNodePayloadKeyUidMap.SetKeyValue(MorphTargetPayLoadUidPair.Key, MorphTargetPayLoadUidPair.Value);
+	}
+
+	for (const TPair<FString, uint8>& MorphTargetPayLoadTypePair : MorphTargetAnimationPayloadKeyTypes)
+	{
+		MorphTargetNodePayloadKeyTypeMap.SetKeyValue(MorphTargetPayLoadTypePair.Key, MorphTargetPayLoadTypePair.Value);
+	}
 }
 
 /************************************************************************/
diff --git a/Engine/Plugins/Interchange/Runtime/Source/FactoryNodes/Public/InterchangeAnimSequenceFactoryNode.h b/Engine/Plugins/Interchange/Runtime/Source/FactoryNodes/Public/InterchangeAnimSequenceFactoryNode.h
--- a/Engine/Plugins/Interchange/Runtime/Source/FactoryNodes/Public/InterchangeAnimSequenceFactoryNode.h
+++ b/Engine/Plugins/Interchange/Runtime/Source/FactoryNodes/Public/InterchangeAnimSequenceFactoryNode.h
@@ -21,11 +21,15 @@ namespace UE::Interchange
 		static const FAttributeKey& GetAnimatedAttributeCurveNamesKey();
 		static const FAttributeKey& GetAnimatedAttributeStepCurveNamesKey();
 		static const FAttributeKey& GetAnimatedMaterialCurveSuffixesKey();
-		static const FAttributeKey& GetSceneNodeAnimationPayloadKeyMapKey();
-		static const FAttributeKey& GetMorphTargetNodePayloadKeyMapKey();
+		static const FAttributeKey& GetSceneNodeAnimationPayloadKeyUidMapKey();
+		static const FAttributeKey& GetSceneNodeAnimationPayloadKeyTypeMapKey();
+		static const FAttributeKey& GetMorphTargetNodePayloadKeyUidMapKey();
+		static const FAttributeKey& GetMorphTargetNodePayloadKeyTypeMapKey();
 	};
 }//ns UE::Interchange
 
+struct FInterchangeAnimationPayLoadKey;
+
 UCLASS(BlueprintType)
 class INTERCHANGEFACTORYNODES_API UInterchangeAnimSequenceFactoryNode : public UInterchangeFactoryBaseNode
 {
@@ -51,8 +55,11 @@ public:
 
 		if (Ar.IsLoading() && bIsInitialized)
 		{
-			SceneNodeAnimationPayloadKeyMap.RebuildCache();
-			MorphTargetNodePayloadKeyMap.RebuildCache();
+			SceneNodeAnimationPayLoadKeyUidMap.RebuildCache();
+			SceneNodeAnimationPayLoadKeyTypeMap.RebuildCache();
+
+			MorphTargetNodePayloadKeyUidMap.RebuildCache();
+			MorphTargetNodePayloadKeyTypeMap.RebuildCache();
 		}
 	}
 
@@ -386,28 +393,16 @@ public:
 	bool SetCustomSkeletonSoftObjectPath(const FSoftObjectPath& AttributeValue);
 
 	UFUNCTION(BlueprintCallable, Category = "Interchange | Node | SkeletalAnimationTrack")
-	void GetSceneNodeAnimationPayloadKeys(TMap<FString, FString>& OutSceneNodeAnimationPayloads) const;
-
-	UFUNCTION(BlueprintCallable, Category = "Interchange | Node | SkeletalAnimationTrack")
-	bool GetAnimationPayloadKeyFromSceneNodeUid(const FString& SceneNodeUid, FString& OutPayloadKey) const;
-
-	UFUNCTION(BlueprintCallable, Category = "Interchange | Node | SkeletalAnimationTrack")
-	bool SetAnimationPayloadKeyForSceneNodeUid(const FString& SceneNodeUid, const FString& PayloadKey);
+	void GetSceneNodeAnimationPayloadKeys(TMap<FString, FInterchangeAnimationPayLoadKey>& OutSceneNodeAnimationPayloadKeys) const;
 
 	UFUNCTION(BlueprintCallable, Category = "Interchange | Node | SkeletalAnimationTrack")
-	bool RemoveAnimationPayloadKeyForSceneNodeUid(const FString& SceneNodeUid);
+	void SetAnimationPayloadKeysForSceneNodeUids(const TMap<FString, FString>& SceneNodeAnimationPayloadKeyUids, const TMap<FString, uint8>& SceneNodeAnimationPayloadKeyTypes);
 
 	UFUNCTION(BlueprintCallable, Category = "Interchange | Node | SkeletalAnimationTrack")
-	void GetMorphTargetNodeAnimationPayloadKeys(TMap<FString, FString>& OutMorphTargetNodeAnimationPayloads) const;
+	void GetMorphTargetNodeAnimationPayloadKeys(TMap<FString, FInterchangeAnimationPayLoadKey>& OutMorphTargetNodeAnimationPayloads) const;
 
 	UFUNCTION(BlueprintCallable, Category = "Interchange | Node | SkeletalAnimationTrack")
-	bool GetAnimationPayloadKeyFromMorphTargetNodeUid(const FString& MorphTargetNodeUid, FString& OutPayloadKey) const;
-
-	UFUNCTION(BlueprintCallable, Category = "Interchange | Node | SkeletalAnimationTrack")
-	bool SetAnimationPayloadKeyForMorphTargetNodeUid(const FString& MorphTargetNodeUid, const FString& PayloadKey);
-
-	UFUNCTION(BlueprintCallable, Category = "Interchange | Node | SkeletalAnimationTrack")
-	bool RemoveAnimationPayloadKeyForMorphTargetNodeUid(const FString& MorphTargetNodeUid);
+	void SetAnimationPayloadKeysForMorphTargetNodeUids(const TMap<FString, FString>& MorphTargetAnimationPayloadKeyUids, const TMap<FString, uint8>& MorphTargetAnimationPayloadKeyTypes);
 
 private:
 	const UE::Interchange::FAttributeKey Macro_CustomSkeletonFactoryNodeUidKey = UE::Interchange::FAttributeKey(TEXT("SkeletonFactoryNodeUid"));
@@ -434,6 +429,9 @@ private:
 	UE::Interchange::TArrayAttributeHelper<FString> AnimatedAttributeStepCurveNames;
 	UE::Interchange::TArrayAttributeHelper<FString> AnimatedMaterialCurveSuffixes;
 
-	UE::Interchange::TMapAttributeHelper<FString, FString> SceneNodeAnimationPayloadKeyMap;
-	UE::Interchange::TMapAttributeHelper<FString, FString> MorphTargetNodePayloadKeyMap;
+	UE::Interchange::TMapAttributeHelper<FString, FString>	SceneNodeAnimationPayLoadKeyUidMap;
+	UE::Interchange::TMapAttributeHelper<FString, uint8>	SceneNodeAnimationPayLoadKeyTypeMap;
+
+	UE::Interchange::TMapAttributeHelper<FString, FString>	MorphTargetNodePayloadKeyUidMap;
+	UE::Interchange::TMapAttributeHelper<FString, uint8>	MorphTargetNodePayloadKeyTypeMap;
 };
diff --git a/Engine/Plugins/Interchange/Runtime/Source/Import/Private/Animation/InterchangeAnimSequenceFactory.cpp b/Engine/Plugins/Interchange/Runtime/Source/Import/Private/Animation/InterchangeAnimSequenceFactory.cpp
--- a/Engine/Plugins/Interchange/Runtime/Source/Import/Private/Animation/InterchangeAnimSequenceFactory.cpp
+++ b/Engine/Plugins/Interchange/Runtime/Source/Import/Private/Animation/InterchangeAnimSequenceFactory.cpp
@@ -3,7 +3,6 @@
 
 #include "Animation/AnimSequence.h"
 #include "Animation/BuiltInAttributeTypes.h"
-#include "Animation/InterchangeAnimationPayload.h"
 #include "Animation/InterchangeAnimationPayloadInterface.h"
 #include "InterchangeCommonPipelineDataFactoryNode.h"
 #include "InterchangeAssetImportData.h"
@@ -20,6 +19,7 @@
 #include "Nodes/InterchangeBaseNodeContainer.h"
 #include "Nodes/InterchangeSourceNode.h"
 #include "Nodes/InterchangeUserDefinedAttribute.h"
+#include "InterchangeAnimationTrackSetNode.h"
 
 #include UE_INLINE_GENERATED_CPP_BY_NAME(InterchangeAnimSequenceFactory)
 
@@ -58,14 +58,14 @@ namespace UE::Interchange::Private
 	}
 
 	bool CreateAttributeStepCurve(UAnimSequence* TargetSequence
-		, FAnimationStepCurvePayloadData& StepCurvePayload
+		, TArray<FInterchangeStepCurve>& StepCurves
 		, const FString& CurveName
 		, const FString& BoneName
 		, const int32 CurveFlags
 		, const bool bDoNotImportCurveWithZero)
 	{
 		//For bone attribute we support only single curve type (structured type like vector are not allowed)
-		if (!TargetSequence || StepCurvePayload.StepCurves.Num() != 1 || CurveName.IsEmpty())
+		if (!TargetSequence || StepCurves.Num() != 1 || CurveName.IsEmpty())
 		{
 			return false;
 		}
@@ -73,7 +73,7 @@ namespace UE::Interchange::Private
 		if (bDoNotImportCurveWithZero)
 		{
 			bool bAllCurveValuesZero = true;
-			for (const FInterchangeStepCurve& StepCurve : StepCurvePayload.StepCurves)
+			for (const FInterchangeStepCurve& StepCurve : StepCurves)
 			{
 				if (StepCurve.FloatKeyValues.IsSet())
 				{
@@ -116,7 +116,7 @@ namespace UE::Interchange::Private
 		}
 
 		bool bResult = false;
-		for (const FInterchangeStepCurve& StepCurve : StepCurvePayload.StepCurves)
+		for (const FInterchangeStepCurve& StepCurve : StepCurves)
 		{
 			if (StepCurve.FloatKeyValues.IsSet())
 			{
@@ -147,7 +147,7 @@ namespace UE::Interchange::Private
 	}
 
 	bool InternalCreateCurve(UAnimSequence* TargetSequence
-		, FAnimationCurvePayloadData& CurvePayload
+		, TArray<FRichCurve>& Curves
 		, const FString& CurveName
 		, const int32 CurveFlags
 		, const bool bDoNotImportCurveWithZero
@@ -156,7 +156,7 @@ namespace UE::Interchange::Private
 		, const bool bMaterialCurve
 		, const bool bShouldTransact)
 	{
-		if (!TargetSequence || CurvePayload.Curves.Num() != 1 || CurveName.IsEmpty())
+		if (!TargetSequence || Curves.Num() != 1 || CurveName.IsEmpty())
 		{
 			return false;
 		}
@@ -164,7 +164,7 @@ namespace UE::Interchange::Private
 		if (bDoNotImportCurveWithZero)
 		{
 			bool bAllCurveValueAreZero = true;
-			for (const FRichCurve& Curve : CurvePayload.Curves)
+			for (const FRichCurve& Curve : Curves)
 			{
 				FKeyHandle KeyHandle = Curve.GetFirstKeyHandle();
 				while (KeyHandle != FKeyHandle::Invalid())
@@ -207,13 +207,13 @@ namespace UE::Interchange::Private
 		ensure(TargetCurve);
 
 		//MorphTarget curves are shift to 0 second
-		const float CurveStartTime = CurvePayload.Curves[0].GetFirstKey().Time;
+		const float CurveStartTime = Curves[0].GetFirstKey().Time;
 		if (CurveStartTime > UE_SMALL_NUMBER)
 		{
-			CurvePayload.Curves[0].ShiftCurve(-CurveStartTime);
+			Curves[0].ShiftCurve(-CurveStartTime);
 		}
 		// Set actual keys on curve within the model
-		Controller.SetCurveKeys(FloatCurveId, CurvePayload.Curves[0].GetConstRefOfKeys(), bShouldTransact);
+		Controller.SetCurveKeys(FloatCurveId, Curves[0].GetConstRefOfKeys(), bShouldTransact);
 
 		if (bMaterialCurve || bMorphTargetCurve)
 		{
@@ -226,26 +226,26 @@ namespace UE::Interchange::Private
 		return true;
 	}
 
-	bool CreateMorphTargetCurve(UAnimSequence* TargetSequence, FAnimationCurvePayloadData& CurvePayload, const FString& CurveName, int32 CurveFlags, bool bDoNotImportCurveWithZero, bool bAddCurveMetadataToSkeleton, bool bShouldTransact)
+	bool CreateMorphTargetCurve(UAnimSequence* TargetSequence, TArray<FRichCurve>& Curves, const FString& CurveName, int32 CurveFlags, bool bDoNotImportCurveWithZero, bool bAddCurveMetadataToSkeleton, bool bShouldTransact)
 	{
 		constexpr bool bIsMorphTargetCurve = true;
 		constexpr bool bIsMaterialCurve = false;
-		return InternalCreateCurve(TargetSequence, CurvePayload, CurveName, CurveFlags, bDoNotImportCurveWithZero, bAddCurveMetadataToSkeleton, bIsMorphTargetCurve, bIsMaterialCurve, bShouldTransact);
+		return InternalCreateCurve(TargetSequence, Curves, CurveName, CurveFlags, bDoNotImportCurveWithZero, bAddCurveMetadataToSkeleton, bIsMorphTargetCurve, bIsMaterialCurve, bShouldTransact);
 	}
 
-	bool CreateMaterialCurve(UAnimSequence* TargetSequence, FAnimationCurvePayloadData& CurvePayload, const FString& CurveName, int32 CurveFlags, bool bDoNotImportCurveWithZero, bool bAddCurveMetadataToSkeleton, bool bShouldTransact)
+	bool CreateMaterialCurve(UAnimSequence* TargetSequence, TArray<FRichCurve>& Curves, const FString& CurveName, int32 CurveFlags, bool bDoNotImportCurveWithZero, bool bAddCurveMetadataToSkeleton, bool bShouldTransact)
 	{
 		constexpr bool bIsMorphTargetCurve = false;
 		constexpr bool bIsMaterialCurve = true;
-		return InternalCreateCurve(TargetSequence, CurvePayload, CurveName, CurveFlags, bDoNotImportCurveWithZero, bAddCurveMetadataToSkeleton, bIsMorphTargetCurve, bIsMaterialCurve, bShouldTransact);
+		return InternalCreateCurve(TargetSequence, Curves, CurveName, CurveFlags, bDoNotImportCurveWithZero, bAddCurveMetadataToSkeleton, bIsMorphTargetCurve, bIsMaterialCurve, bShouldTransact);
 	}
 
-	bool CreateAttributeCurve(UAnimSequence* TargetSequence, FAnimationCurvePayloadData& CurvePayload, const FString& CurveName, int32 CurveFlags, bool bDoNotImportCurveWithZero, bool bAddCurveMetadataToSkeleton, bool bShouldTransact)
+	bool CreateAttributeCurve(UAnimSequence* TargetSequence, TArray<FRichCurve>& Curves, const FString& CurveName, int32 CurveFlags, bool bDoNotImportCurveWithZero, bool bAddCurveMetadataToSkeleton, bool bShouldTransact)
 	{
 		//This curve don't animate morph target or material parameter.
 		constexpr bool bIsMorphTargetCurve = false;
 		constexpr bool bIsMaterialCurve = false;
-		return InternalCreateCurve(TargetSequence, CurvePayload, CurveName, CurveFlags, bDoNotImportCurveWithZero, bAddCurveMetadataToSkeleton, bIsMorphTargetCurve, bIsMaterialCurve, bShouldTransact);
+		return InternalCreateCurve(TargetSequence, Curves, CurveName, CurveFlags, bDoNotImportCurveWithZero, bAddCurveMetadataToSkeleton, bIsMorphTargetCurve, bIsMaterialCurve, bShouldTransact);
 	}
 
 	void RetrieveAnimationPayloads(UAnimSequence* AnimSequence
@@ -257,7 +257,7 @@ namespace UE::Interchange::Private
 		, const bool bIsReimporting
 		, TArray<FString> OutCurvesNotFound)
 	{
-		TMap<const UInterchangeSceneNode*, TFuture<TOptional<UE::Interchange::FAnimationBakeTransformPayloadData>>> AnimationPayloads;
+		TMap<const UInterchangeSceneNode*, TFuture<TOptional<UE::Interchange::FAnimationPayloadData>>> AnimationPayloads;
 
 		FString SkeletonRootUid;
 		if (!SkeletonFactoryNode->GetCustomRootJointUid(SkeletonRootUid))
@@ -291,28 +291,21 @@ namespace UE::Interchange::Private
 			AnimSequenceFactoryNode->GetCustomImportBoneTracksRangeStop(RangeEnd);
 
 			const double BakeInterval = 1.0 / SampleRate;
-			const double SequenceLength = FMath::Max<double>(RangeEnd - RangeStart, MINIMUM_ANIMATION_LENGTH);
-			int32 BakeKeyCount = FMath::RoundToInt32(SequenceLength * SampleRate) + 1;
 
-			TMap<FString, FString> PayloadKeys;
+			TMap<FString, FInterchangeAnimationPayLoadKey> PayloadKeys;
 			AnimSequenceFactoryNode->GetSceneNodeAnimationPayloadKeys(PayloadKeys);
 
-			for (const TTuple<FString, FString>& SceneNodeUidAndPayloadKey : PayloadKeys)
+			for (const TTuple<FString, FInterchangeAnimationPayLoadKey>& SceneNodeUidAndPayloadKey : PayloadKeys)
 			{
 				if (const UInterchangeSceneNode* SkeletonSceneNode = Cast<UInterchangeSceneNode>(NodeContainer->GetNode(SceneNodeUidAndPayloadKey.Key)))
 				{
-					AnimationPayloads.Add(SkeletonSceneNode, AnimSequenceTranslatorPayloadInterface->GetAnimationBakeTransformPayloadData(SceneNodeUidAndPayloadKey.Value, SampleRate, RangeStart, RangeEnd));
+					AnimationPayloads.Add(SkeletonSceneNode, AnimSequenceTranslatorPayloadInterface->GetAnimationPayloadData(SceneNodeUidAndPayloadKey.Value, SampleRate, RangeStart, RangeEnd));
 				}
 			}
 
 			//This destroy all previously imported animation raw data
 			Controller.RemoveAllBoneTracks(bShouldTransact);
 
-			const FFrameRate ResampleFrameRate(SampleRate, 1);
-			const FFrameNumber NumberOfFrames = ResampleFrameRate.AsFrameNumber(SequenceLength);
-			Controller.SetFrameRate(ResampleFrameRate, bShouldTransact);
-			Controller.SetNumberOfFrames(FGenericPlatformMath::Max<int32>(NumberOfFrames.Value, 1), bShouldTransact);
-
 			FTransform GlobalOffsetTransform;
 			bool bBakeMeshes = false;
 			{
@@ -324,7 +317,10 @@ namespace UE::Interchange::Private
 				}
 			}
 
-			for (const TTuple< const UInterchangeSceneNode*, TFuture<TOptional<UE::Interchange::FAnimationBakeTransformPayloadData>>>& AnimationPayload : AnimationPayloads)
+			double MergedRangeEnd = RangeEnd;
+			double MergedRangeStart = RangeStart;
+			TMap<const UInterchangeSceneNode*, UE::Interchange::FAnimationPayloadData> PreProcessedAnimationPayloads;
+			for (const TTuple< const UInterchangeSceneNode*, TFuture<TOptional<UE::Interchange::FAnimationPayloadData>>>& AnimationPayload : AnimationPayloads)
 			{
 				const FName BoneName = FName(*(AnimationPayload.Key->GetDisplayLabel()));
 				const int32 BoneIndex = Skeleton->GetReferenceSkeleton().FindBoneIndex(BoneName);
@@ -335,39 +331,82 @@ namespace UE::Interchange::Private
 				}
 				//If we are getting the root 
 				bool bApplyGlobalOffset = AnimationPayload.Key->GetUniqueID().Equals(SkeletonRootUid);
-				
-				TOptional<UE::Interchange::FAnimationBakeTransformPayloadData> AnimationTransformPayload = AnimationPayload.Value.Get();
-				if (!AnimationTransformPayload.IsSet())
+
+				TOptional<UE::Interchange::FAnimationPayloadData> OptionalAnimationTransformPayload = AnimationPayload.Value.Get();
+				if (!OptionalAnimationTransformPayload.IsSet())
 				{
-					FString PayloadKey = PayloadKeys[AnimationPayload.Key->GetUniqueID()];
+					FString PayloadKey = PayloadKeys[AnimationPayload.Key->GetUniqueID()].UniqueId;
 					UE_LOG(LogInterchangeImport, Warning, TEXT("Invalid animation transform payload key [%s] AnimSequence asset %s"), *PayloadKey, *AssetName);
 					continue;
 				}
 
-				if (AnimationTransformPayload->Transforms.Num() == 0)
+				UE::Interchange::FAnimationPayloadData& AnimationTransformPayload = OptionalAnimationTransformPayload.GetValue();
+
+				if (AnimationTransformPayload.Type != EInterchangeAnimationPayLoadType::BAKED)
+				{
+					//Where Curve is null the LocalTransform should be used for the Baked Transform generation.
+					FTransform LocalTransform;
+					AnimationPayload.Key->GetCustomLocalTransform(LocalTransform);
+
+					//Currently only Curve -> Baked conversion (for LevelSequence->AnimSequence conversion by ForceMeshType Skeletal use case)
+					//and Curve -> Step Curve conversion (for custom attributes)
+					AnimationTransformPayload.CalculateDataFor(EInterchangeAnimationPayLoadType::BAKED, LocalTransform);
+					//Range End will be calculated as well:
+					if (MergedRangeEnd < AnimationTransformPayload.RangeEndTime)
+					{
+						MergedRangeEnd = AnimationTransformPayload.RangeEndTime;
+					}
+				}
+
+				PreProcessedAnimationPayloads.Add(AnimationPayload.Key, AnimationTransformPayload);
+			}
+
+			const double SequenceLength = FMath::Max<double>(MergedRangeEnd - MergedRangeStart, MINIMUM_ANIMATION_LENGTH);
+			int32 FrameCount = FMath::RoundToInt32(SequenceLength * SampleRate);
+			int32 BakeKeyCount = FrameCount + 1;
+			const FFrameRate ResampleFrameRate(SampleRate, 1);
+			Controller.SetFrameRate(ResampleFrameRate, bShouldTransact);
+			Controller.SetNumberOfFrames(FrameCount, bShouldTransact);
+
+			for (TTuple< const UInterchangeSceneNode*, UE::Interchange::FAnimationPayloadData>& AnimationPayload : PreProcessedAnimationPayloads)
+			{
+				const FName BoneName = FName(*(AnimationPayload.Key->GetDisplayLabel()));
+				UE::Interchange::FAnimationPayloadData& AnimationTransformPayload = AnimationPayload.Value;
+				
+				//If we are getting the root 
+				bool bApplyGlobalOffset = AnimationPayload.Key->GetUniqueID().Equals(SkeletonRootUid);
+
+				if (AnimationTransformPayload.Transforms.Num() == 0)
 				{
 					//We need at least one transform
-					AnimationTransformPayload->Transforms.Add(FTransform::Identity);
+					AnimationTransformPayload.Transforms.Add(FTransform::Identity);
 				}
 
+				const double SequenceLengthForAnimationPayload = FMath::Max<double>(AnimationTransformPayload.RangeEndTime - AnimationTransformPayload.RangeStartTime, MINIMUM_ANIMATION_LENGTH);
+				int32 BakeKeyCountForAnimationPayload = FMath::RoundToInt32(SequenceLengthForAnimationPayload * AnimationTransformPayload.BakeFrequency) + 1;
+
 				FRawAnimSequenceTrack RawTrack;
-				RawTrack.PosKeys.Reserve(BakeKeyCount);
-				RawTrack.RotKeys.Reserve(BakeKeyCount);
-				RawTrack.ScaleKeys.Reserve(BakeKeyCount);
+				RawTrack.PosKeys.Reserve(BakeKeyCountForAnimationPayload);
+				RawTrack.RotKeys.Reserve(BakeKeyCountForAnimationPayload);
+				RawTrack.ScaleKeys.Reserve(BakeKeyCountForAnimationPayload);
 				TArray<float> TimeKeys;
-				TimeKeys.Reserve(BakeKeyCount);
+				TimeKeys.Reserve(BakeKeyCountForAnimationPayload);
 
-				//Everything should match Key count, sample rate and range
-				check(AnimationTransformPayload->Transforms.Num() == BakeKeyCount);
-				check(FMath::IsNearlyEqual(AnimationTransformPayload->BakeFrequency, SampleRate, UE_DOUBLE_KINDA_SMALL_NUMBER));
-				check(FMath::IsNearlyEqual(AnimationTransformPayload->RangeStartTime, RangeStart, UE_DOUBLE_KINDA_SMALL_NUMBER));
-				check(FMath::IsNearlyEqual(AnimationTransformPayload->RangeEndTime, RangeEnd, UE_DOUBLE_KINDA_SMALL_NUMBER));
+				check(AnimationTransformPayload.Transforms.Num() == BakeKeyCountForAnimationPayload);
+
+				if (AnimationTransformPayload.Type == EInterchangeAnimationPayLoadType::BAKED)
+				{
+					//Everything should match Key count, sample rate and range
+					check(FMath::IsNearlyEqual(AnimationTransformPayload.BakeFrequency, SampleRate, UE_DOUBLE_KINDA_SMALL_NUMBER));
+					check(FMath::IsNearlyEqual(AnimationTransformPayload.RangeStartTime, RangeStart, UE_DOUBLE_KINDA_SMALL_NUMBER));
+					check(FMath::IsNearlyEqual(AnimationTransformPayload.RangeEndTime, RangeEnd, UE_DOUBLE_KINDA_SMALL_NUMBER));
+				}
 
-				int32 TransformIndex = 0;
-				for (double CurrentTime = RangeStart; CurrentTime <= RangeEnd + UE_DOUBLE_SMALL_NUMBER; CurrentTime += BakeInterval, ++TransformIndex)
+				double CurrentTime = 0;
+				for (int32 BakeIndex = 0; BakeIndex < BakeKeyCountForAnimationPayload; BakeIndex++, CurrentTime += BakeInterval)
 				{
-					check(AnimationTransformPayload->Transforms.IsValidIndex(TransformIndex));
-					FTransform3f AnimKeyTransform = FTransform3f(AnimationTransformPayload->Transforms[TransformIndex]);
+					check(AnimationTransformPayload.Transforms.IsValidIndex(BakeIndex));
+					FTransform3f AnimKeyTransform = FTransform3f(AnimationTransformPayload.Transforms[BakeIndex]);
 					if (bApplyGlobalOffset)
 					{
 						if (bBakeMeshes)
@@ -399,14 +438,14 @@ namespace UE::Interchange::Private
 					RawTrack.PosKeys.Add(Position);
 					RawTrack.RotKeys.Add(Quaternion);
 					//Animation are always translated to zero
-					TimeKeys.Add(CurrentTime - RangeStart);
+					TimeKeys.Add(CurrentTime - AnimationTransformPayload.RangeStartTime);
 				}
 
 				//Make sure we create the correct amount of keys
-				check(RawTrack.ScaleKeys.Num() == BakeKeyCount
-					&& RawTrack.PosKeys.Num() == BakeKeyCount
-					&& RawTrack.RotKeys.Num() == BakeKeyCount
-					&& TimeKeys.Num() == BakeKeyCount);
+				check(RawTrack.ScaleKeys.Num() == BakeKeyCountForAnimationPayload
+					&& RawTrack.PosKeys.Num() == BakeKeyCountForAnimationPayload
+					&& RawTrack.RotKeys.Num() == BakeKeyCountForAnimationPayload
+					&& TimeKeys.Num() == BakeKeyCountForAnimationPayload);
 
 				//add new track
 				Controller.AddBoneCurve(BoneName, bShouldTransact);
@@ -494,35 +533,35 @@ namespace UE::Interchange::Private
 
 			//Import morph target curves
 			{
-				TMap<FString, FString> MorphTargetNodeAnimationPayloads;
+				TMap<FString, FInterchangeAnimationPayLoadKey> MorphTargetNodeAnimationPayloads;
 				AnimSequenceFactoryNode->GetMorphTargetNodeAnimationPayloadKeys(MorphTargetNodeAnimationPayloads);
 
-				TMap<FString, TFuture<TOptional<UE::Interchange::FAnimationCurvePayloadData>>> AnimationCurvesPayloads;
+				TMap<FString, TFuture<TOptional<UE::Interchange::FAnimationPayloadData>>> AnimationCurvesPayloads;
 				TMap<FString, FString> AnimationCurveMorphTargetNodeNames;
 
-				for (const TPair<FString, FString>& MorphTargetNodeUidAnimationPayload : MorphTargetNodeAnimationPayloads)
+				for (const TPair<FString, FInterchangeAnimationPayLoadKey>& MorphTargetNodeUidAnimationPayload : MorphTargetNodeAnimationPayloads)
 				{
-					FString PayloadKey = MorphTargetNodeUidAnimationPayload.Value;
+					FString PayloadKey = MorphTargetNodeUidAnimationPayload.Value.UniqueId;
 					if (PayloadKey.Len() == 0)
 					{
 						continue;
 					}
 					if (const UInterchangeMeshNode* MorphTargetNode = Cast<UInterchangeMeshNode>(NodeContainer->GetNode(MorphTargetNodeUidAnimationPayload.Key)))
 					{
-						AnimationCurvesPayloads.Add(PayloadKey, AnimSequenceTranslatorPayloadInterface->GetAnimationCurvePayloadData(PayloadKey));
+						AnimationCurvesPayloads.Add(PayloadKey, AnimSequenceTranslatorPayloadInterface->GetAnimationPayloadData(MorphTargetNodeUidAnimationPayload.Value));
 						AnimationCurveMorphTargetNodeNames.Add(PayloadKey, MorphTargetNode->GetDisplayLabel());
 					}
 				}
 
-				for (TPair<FString, TFuture<TOptional<UE::Interchange::FAnimationCurvePayloadData>>>& CurveNameAndPayload : AnimationCurvesPayloads)
+				for (TPair<FString, TFuture<TOptional<UE::Interchange::FAnimationPayloadData>>>& CurveNameAndPayload : AnimationCurvesPayloads)
 				{
-					TOptional<UE::Interchange::FAnimationCurvePayloadData> AnimationCurvePayload = CurveNameAndPayload.Value.Get();
+					TOptional<UE::Interchange::FAnimationPayloadData> AnimationCurvePayload = CurveNameAndPayload.Value.Get();
 					if (!AnimationCurvePayload.IsSet())
 					{
 						UE_LOG(LogInterchangeImport, Warning, TEXT("Invalid animation morph target curve payload key [%s] AnimSequence asset %s"), *CurveNameAndPayload.Key, *AssetName);
 						continue;
 					}
-					FAnimationCurvePayloadData& CurvePayload = AnimationCurvePayload.GetValue();
+					FAnimationPayloadData& CurvePayload = AnimationCurvePayload.GetValue();
 					if (bRemoveCurveRedundantKeys)
 					{
 						for (FRichCurve& RichCurve : CurvePayload.Curves)
@@ -531,7 +570,7 @@ namespace UE::Interchange::Private
 						}
 					}
 					constexpr int32 CurveFlags = 0;
-					CreateMorphTargetCurve(AnimSequence, CurvePayload, AnimationCurveMorphTargetNodeNames.FindChecked(CurveNameAndPayload.Key), CurveFlags, bDoNotImportCurveWithZero, bAddCurveMetadataToSkeleton, bShouldTransact);
+					CreateMorphTargetCurve(AnimSequence, CurvePayload.Curves, AnimationCurveMorphTargetNodeNames.FindChecked(CurveNameAndPayload.Key), CurveFlags, bDoNotImportCurveWithZero, bAddCurveMetadataToSkeleton, bShouldTransact);
 				}
 			}
 
@@ -559,7 +598,7 @@ namespace UE::Interchange::Private
 				{
 					TArray<FString> AttributeCurveNames;
 					AnimSequenceFactoryNode->GetAnimatedAttributeCurveNames(AttributeCurveNames);
-					TMap<FString, TFuture<TOptional<UE::Interchange::FAnimationCurvePayloadData>>> AnimationCurvesPayloads;
+					TMap<FString, TFuture<TOptional<UE::Interchange::FAnimationPayloadData>>> AnimationCurvesPayloads;
 					for (const FString& NodeUid : SkeletonNodes)
 					{
 						if (const UInterchangeSceneNode* SkeletonSceneNode = Cast<UInterchangeSceneNode>(NodeContainer->GetNode(NodeUid)))
@@ -577,22 +616,25 @@ namespace UE::Interchange::Private
 							}
 							for (const TPair<FString, FString>& CurveNamePayload : CurveNamePayloads)
 							{
-								AnimationCurvesPayloads.Add(CurveNamePayload.Key, AnimSequenceTranslatorPayloadInterface->GetAnimationCurvePayloadData(CurveNamePayload.Value));
+								//This goes slightly against the intent of the Type / FInterchangeAnimationPayLoadKey usage as we set the Type here, outside of the Translator
+								// this is due to the nature of the Attribute curves.
+								// Could be potentially reworked so that with the AttributeInfos we store the Type as well.
+								AnimationCurvesPayloads.Add(CurveNamePayload.Key, AnimSequenceTranslatorPayloadInterface->GetAnimationPayloadData(FInterchangeAnimationPayLoadKey(CurveNamePayload.Value, EInterchangeAnimationPayLoadType::CURVE)));
 							}
 						}
 					}
 
-					for (TPair<FString, TFuture<TOptional<UE::Interchange::FAnimationCurvePayloadData>>>& CurveNameAndPayload : AnimationCurvesPayloads)
+					for (TPair<FString, TFuture<TOptional<UE::Interchange::FAnimationPayloadData>>>& CurveNameAndPayload : AnimationCurvesPayloads)
 					{
 						const FString& CurveName = CurveNameAndPayload.Key;
-						TOptional<UE::Interchange::FAnimationCurvePayloadData> AnimationCurvePayload = CurveNameAndPayload.Value.Get();
+						TOptional<UE::Interchange::FAnimationPayloadData> AnimationCurvePayload = CurveNameAndPayload.Value.Get();
 						if (!AnimationCurvePayload.IsSet())
 						{
 							UE_LOG(LogInterchangeImport, Warning, TEXT("Invalid animation material curve payload key [%s] AnimSequence asset %s"), *CurveName, *AssetName);
 							continue;
 						}
 
-						FAnimationCurvePayloadData& CurvePayload = AnimationCurvePayload.GetValue();
+						FAnimationPayloadData& CurvePayload = AnimationCurvePayload.GetValue();
 						if (bRemoveCurveRedundantKeys)
 						{
 							for (FRichCurve& RichCurve : CurvePayload.Curves)
@@ -604,11 +646,11 @@ namespace UE::Interchange::Private
 						constexpr int32 CurveFlags = 0;
 						if (bMaterialDriveParameterOnCustomAttribute || IsCurveHookToMaterial(CurveName))
 						{
-							CreateMaterialCurve(AnimSequence, CurvePayload, CurveName, CurveFlags, bDoNotImportCurveWithZero, bAddCurveMetadataToSkeleton, bShouldTransact);
+							CreateMaterialCurve(AnimSequence, CurvePayload.Curves, CurveName, CurveFlags, bDoNotImportCurveWithZero, bAddCurveMetadataToSkeleton, bShouldTransact);
 						}
 						else
 						{
-							CreateAttributeCurve(AnimSequence, CurvePayload, CurveName, CurveFlags, bDoNotImportCurveWithZero, bAddCurveMetadataToSkeleton, bShouldTransact);
+							CreateAttributeCurve(AnimSequence, CurvePayload.Curves, CurveName, CurveFlags, bDoNotImportCurveWithZero, bAddCurveMetadataToSkeleton, bShouldTransact);
 						}
 					}
 				}
@@ -617,8 +659,8 @@ namespace UE::Interchange::Private
 				{
 					TArray<FString> AttributeStepCurveNames;
 					AnimSequenceFactoryNode->GetAnimatedAttributeStepCurveNames(AttributeStepCurveNames);
-					TMap<FString, TFuture<TOptional<UE::Interchange::FAnimationCurvePayloadData>>> AnimationCurvesPayloads;
-					TMap<FString, TFuture<TOptional<UE::Interchange::FAnimationStepCurvePayloadData>>> AnimationStepCurvesPayloads;
+					TMap<FString, TFuture<TOptional<UE::Interchange::FAnimationPayloadData>>> AnimationCurvesPayloads;
+					TMap<FString, TFuture<TOptional<UE::Interchange::FAnimationPayloadData>>> AnimationStepCurvesPayloads;
 					TMap<FString, FString> AnimationBoneNames;
 					for (const FString& NodeUid : SkeletonNodes)
 					{
@@ -646,18 +688,24 @@ namespace UE::Interchange::Private
 							}
 							for (const TPair<FString, FString>& CurveNamePayload : CurveNameFloatPayloads)
 							{
-								AnimationCurvesPayloads.Add(CurveNamePayload.Key, AnimSequenceTranslatorPayloadInterface->GetAnimationCurvePayloadData(CurveNamePayload.Value));
+								//This goes slightly against the intent of the Type / FInterchangeAnimationPayLoadKey usage as we set the Type here, outside of the Translator
+								// this is due to the nature of the Attribute curves.
+								// Could be potentially reworked so that with the AttributeInfos we store the Type as well.
+								AnimationCurvesPayloads.Add(CurveNamePayload.Key, AnimSequenceTranslatorPayloadInterface->GetAnimationPayloadData(FInterchangeAnimationPayLoadKey(CurveNamePayload.Value, EInterchangeAnimationPayLoadType::CURVE)));
 								AnimationBoneNames.Add(CurveNamePayload.Key, BoneName);
 							}
 							for (const TPair<FString, FString>& StepCurveNamePayload : CurveNameStepCurvePayloads)
 							{
-								AnimationStepCurvesPayloads.Add(StepCurveNamePayload.Key, AnimSequenceTranslatorPayloadInterface->GetAnimationStepCurvePayloadData(StepCurveNamePayload.Value));
+								//This goes slightly against the intent of the Type / FInterchangeAnimationPayLoadKey usage as we set the Type here, outside of the Translator
+								// this is due to the nature of the Attribute curves.
+								// Could be potentially reworked so that with the AttributeInfos we store the Type as well.
+								AnimationStepCurvesPayloads.Add(StepCurveNamePayload.Key, AnimSequenceTranslatorPayloadInterface->GetAnimationPayloadData(FInterchangeAnimationPayLoadKey(StepCurveNamePayload.Value, EInterchangeAnimationPayLoadType::STEPCURVE)));
 								AnimationBoneNames.Add(StepCurveNamePayload.Key, BoneName);
 							}
 						}
 					}
 
-					auto AddAttributeCurveToAnimSequence = [bRemoveCurveRedundantKeys, bDoNotImportCurveWithZero, &AnimSequence](FAnimationStepCurvePayloadData& StepCurvePayload, const FString& CurveName, const FString& BoneName)
+					auto AddAttributeCurveToAnimSequence = [bRemoveCurveRedundantKeys, bDoNotImportCurveWithZero, &AnimSequence](FAnimationPayloadData& StepCurvePayload, const FString& CurveName, const FString& BoneName)
 					{
 						if (bRemoveCurveRedundantKeys)
 						{
@@ -668,25 +716,25 @@ namespace UE::Interchange::Private
 						}
 
 						constexpr int32 CurveFlags = 0;
-						CreateAttributeStepCurve(AnimSequence, StepCurvePayload, CurveName, BoneName, CurveFlags, bDoNotImportCurveWithZero);
+						CreateAttributeStepCurve(AnimSequence, StepCurvePayload.StepCurves, CurveName, BoneName, CurveFlags, bDoNotImportCurveWithZero);
 					};
 
-					for (TPair<FString, TFuture<TOptional<UE::Interchange::FAnimationCurvePayloadData>>>& CurveNameAndPayload : AnimationCurvesPayloads)
+					for (TPair<FString, TFuture<TOptional<UE::Interchange::FAnimationPayloadData>>>& CurveNameAndPayload : AnimationCurvesPayloads)
 					{
-						TOptional<UE::Interchange::FAnimationCurvePayloadData> AnimationCurvePayload = CurveNameAndPayload.Value.Get();
+						TOptional<UE::Interchange::FAnimationPayloadData> AnimationCurvePayload = CurveNameAndPayload.Value.Get();
 						if (!AnimationCurvePayload.IsSet())
 						{
 							UE_LOG(LogInterchangeImport, Warning, TEXT("Invalid animation material curve payload key [%s] AnimSequence asset %s"), *CurveNameAndPayload.Key, *AssetName);
 							continue;
 						}
-						FAnimationCurvePayloadData& CurvePayload = AnimationCurvePayload.GetValue();
-						FAnimationStepCurvePayloadData StepCurvePayload = FAnimationStepCurvePayloadData::FromAnimationCurvePayloadData(CurvePayload);
-						AddAttributeCurveToAnimSequence(StepCurvePayload, CurveNameAndPayload.Key, AnimationBoneNames.FindChecked(CurveNameAndPayload.Key));
+						FAnimationPayloadData& AnimationPayloadData = AnimationCurvePayload.GetValue();
+						AnimationPayloadData.CalculateDataFor(EInterchangeAnimationPayLoadType::STEPCURVE);
+						AddAttributeCurveToAnimSequence(AnimationPayloadData, CurveNameAndPayload.Key, AnimationBoneNames.FindChecked(CurveNameAndPayload.Key));
 					}
 
-					for (TPair<FString, TFuture<TOptional<UE::Interchange::FAnimationStepCurvePayloadData>>>& StepCurveNameAndPayload : AnimationStepCurvesPayloads)
+					for (TPair<FString, TFuture<TOptional<UE::Interchange::FAnimationPayloadData>>>& StepCurveNameAndPayload : AnimationStepCurvesPayloads)
 					{
-						TOptional<UE::Interchange::FAnimationStepCurvePayloadData> AnimationStepCurvePayload = StepCurveNameAndPayload.Value.Get();
+						TOptional<UE::Interchange::FAnimationPayloadData> AnimationStepCurvePayload = StepCurveNameAndPayload.Value.Get();
 						if (!AnimationStepCurvePayload.IsSet())
 						{
 							UE_LOG(LogInterchangeImport, Warning, TEXT("Invalid animation material curve payload key [%s] AnimSequence asset %s"), *StepCurveNameAndPayload.Key, *AssetName);
diff --git a/Engine/Plugins/Interchange/Runtime/Source/Import/Private/Animation/InterchangeAnimationTrackSetFactory.cpp b/Engine/Plugins/Interchange/Runtime/Source/Import/Private/Animation/InterchangeAnimationTrackSetFactory.cpp
--- a/Engine/Plugins/Interchange/Runtime/Source/Import/Private/Animation/InterchangeAnimationTrackSetFactory.cpp
+++ b/Engine/Plugins/Interchange/Runtime/Source/Import/Private/Animation/InterchangeAnimationTrackSetFactory.cpp
@@ -2,7 +2,6 @@
 
 #include "Animation/InterchangeAnimationTrackSetFactory.h"
 
-#include "Animation/InterchangeAnimationPayload.h"
 #include "Animation/InterchangeAnimationPayloadInterface.h"
 #include "InterchangeAnimationTrackSetFactoryNode.h"
 #include "InterchangeAnimationTrackSetNode.h"
@@ -258,15 +257,15 @@ namespace UE::Interchange::Private
 		}
 
 		// Get payload
-		FString PayloadKey;
+		FInterchangeAnimationPayLoadKey PayloadKey;
 		if (!TransformTrackNode.GetCustomAnimationPayloadKey(PayloadKey))
 		{
 			UE_LOG(LogInterchangeImport, Warning, TEXT("No payload key for animation track %s on actor %s"), *TransformTrackNode.GetDisplayLabel(), *Actor->GetActorLabel());
 			return;
 		}
 
-		TFuture<TOptional<UE::Interchange::FAnimationCurvePayloadData>> Result = PayloadInterface.GetAnimationCurvePayloadData(PayloadKey);
-		const TOptional<UE::Interchange::FAnimationCurvePayloadData>& PayloadData = Result.Get();
+		TFuture<TOptional<UE::Interchange::FAnimationPayloadData>> Result = PayloadInterface.GetAnimationPayloadData(PayloadKey);
+		const TOptional<UE::Interchange::FAnimationPayloadData>& PayloadData = Result.Get();
 		if (!PayloadData.IsSet() || PayloadData->Curves.Num() != 9)
 		{
 			UE_LOG(LogInterchangeImport, Warning, TEXT("No payload for animation track %s on actor %s"), *TransformTrackNode.GetDisplayLabel(), *Actor->GetActorLabel());
@@ -467,15 +466,15 @@ namespace UE::Interchange::Private
 		}
 
 		// Get payload
-		FString PayloadKey;
+		FInterchangeAnimationPayLoadKey PayloadKey;
 		if (!AnimationTrackNode.GetCustomAnimationPayloadKey(PayloadKey))
 		{
 			UE_LOG(LogInterchangeImport, Warning, TEXT("No payload key for animation track %s on actor %s"), *AnimationTrackNode.GetDisplayLabel(), *Actor->GetActorLabel());
 			return;
 		}
 
-		TFuture<TOptional<UE::Interchange::FAnimationStepCurvePayloadData>> Result = PayloadInterface.GetAnimationStepCurvePayloadData(PayloadKey);
-		const TOptional<UE::Interchange::FAnimationStepCurvePayloadData>& PayloadData = Result.Get();
+		TFuture<TOptional<UE::Interchange::FAnimationPayloadData>> Result = PayloadInterface.GetAnimationPayloadData(PayloadKey);
+		const TOptional<UE::Interchange::FAnimationPayloadData>& PayloadData = Result.Get();
 		if (!PayloadData.IsSet() || PayloadData->StepCurves.Num() != 1)
 		{
 			UE_LOG(LogInterchangeImport, Warning, TEXT("No payload for animation track %s on actor %s"), *AnimationTrackNode.GetDisplayLabel(), *Actor->GetActorLabel());
diff --git a/Engine/Plugins/Interchange/Runtime/Source/Import/Private/Fbx/InterchangeFbxTranslator.cpp b/Engine/Plugins/Interchange/Runtime/Source/Import/Private/Fbx/InterchangeFbxTranslator.cpp
--- a/Engine/Plugins/Interchange/Runtime/Source/Import/Private/Fbx/InterchangeFbxTranslator.cpp
+++ b/Engine/Plugins/Interchange/Runtime/Source/Import/Private/Fbx/InterchangeFbxTranslator.cpp
@@ -284,101 +284,25 @@ TFuture<TOptional<UE::Interchange::FMeshPayloadData>> UInterchangeFbxTranslator:
 	return Promise->GetFuture();
 }
 
-TFuture<TOptional<UE::Interchange::FAnimationCurvePayloadData>> UInterchangeFbxTranslator::GetAnimationCurvePayloadData(const FString& PayLoadKey) const
+TFuture<TOptional<UE::Interchange::FAnimationPayloadData>> UInterchangeFbxTranslator::GetAnimationPayloadData(const FInterchangeAnimationPayLoadKey& PayLoadKey, const double BakeFrequency, const double RangeStartSecond, const double RangeStopSecond) const
 {
-	TSharedPtr<TPromise<TOptional<UE::Interchange::FAnimationCurvePayloadData>>> Promise = MakeShared<TPromise<TOptional<UE::Interchange::FAnimationCurvePayloadData>>>();
+	TSharedPtr<TPromise<TOptional<UE::Interchange::FAnimationPayloadData>>> Promise = MakeShared<TPromise<TOptional<UE::Interchange::FAnimationPayloadData>>>();
 
 	if (!Dispatcher.IsValid())
 	{
-		Promise->SetValue(TOptional<UE::Interchange::FAnimationCurvePayloadData>());
+		Promise->SetValue(TOptional<UE::Interchange::FAnimationPayloadData>());
 		return Promise->GetFuture();
 	}
 
 	// Create a json command to read the fbx file
-	FString JsonCommand = CreateFetchPayloadFbxCommand(PayLoadKey);
-	const int32 CreatedTaskIndex = Dispatcher->AddTask(JsonCommand, FInterchangeDispatcherTaskCompleted::CreateLambda([this, Promise](const int32 TaskIndex)
-		{
-			TRACE_CPUPROFILER_EVENT_SCOPE("UInterchangeFbxTranslator::GetAnimationCurvePayloadData::Dispatcher->AddTaskDone")
-			UE::Interchange::ETaskState TaskState;
-			FString JsonResult;
-			TArray<FString> JsonMessages;
-			Dispatcher->GetTaskState(TaskIndex, TaskState, JsonResult, JsonMessages);
-
-			// Parse the Json messages into UInterchangeResults
-			for (const FString& JsonMessage : JsonMessages)
-			{
-				AddMessage(UInterchangeResult::FromJson(JsonMessage));
-			}
-
-			if (TaskState != UE::Interchange::ETaskState::ProcessOk)
-			{
-				Promise->SetValue(TOptional<UE::Interchange::FAnimationCurvePayloadData>());
-				return;
-			}
-
-			// Grab the result file and fill the BaseNodeContainer
-			UE::Interchange::FJsonFetchPayloadCmd::JsonResultParser ResultParser;
-			ResultParser.FromJson(JsonResult);
-			FString AnimationTransformPayloadFilename = ResultParser.GetResultFilename();
-
-			if (!ensure(FPaths::FileExists(AnimationTransformPayloadFilename)))
-			{
-				// TODO log an error saying the payload file does not exist even if the get payload command succeeded
-				Promise->SetValue(TOptional<UE::Interchange::FAnimationCurvePayloadData>());
-				return;
-			}
-
-			UE::Interchange::FAnimationCurvePayloadData AnimationTransformPayload;
-
-			// All sub object should be gone with the reset
-			TArray64<uint8> Buffer;
-			FFileHelper::LoadFileToArray(Buffer, *AnimationTransformPayloadFilename);
-			uint8* FileData = Buffer.GetData();
-			int64 FileDataSize = Buffer.Num();
-			if (FileDataSize < 1)
-			{
-				// Nothing to load from this file
-				Promise->SetValue(TOptional<UE::Interchange::FAnimationCurvePayloadData>());
-				return;
-			}
-
-			// Buffer keeps the ownership of the data, the large memory reader is use to serialize the TMap
-			FLargeMemoryReader Ar(FileData, FileDataSize);
-			TArray<FInterchangeCurve> InterchangeCurves;
-			Ar << InterchangeCurves;
-			AnimationTransformPayload.Curves.AddDefaulted(InterchangeCurves.Num());
-			for (int32 CurveIndex = 0; CurveIndex < InterchangeCurves.Num(); ++CurveIndex)
-			{
-				const FInterchangeCurve& InterchangeCurve = InterchangeCurves[CurveIndex];
-				InterchangeCurve.ToRichCurve(AnimationTransformPayload.Curves[CurveIndex]);
-			}
-			Promise->SetValue(MoveTemp(AnimationTransformPayload));
-		}));
-
-	// The task was not added to the dispatcher
-	if (CreatedTaskIndex == INDEX_NONE)
-	{
-		Promise->SetValue(TOptional<UE::Interchange::FAnimationCurvePayloadData>{});
-	}
-
-	return Promise->GetFuture();
-}
-
-TFuture<TOptional<UE::Interchange::FAnimationStepCurvePayloadData>> UInterchangeFbxTranslator::GetAnimationStepCurvePayloadData(const FString& PayLoadKey) const
-{
-	TSharedPtr<TPromise<TOptional<UE::Interchange::FAnimationStepCurvePayloadData>>> Promise = MakeShared<TPromise<TOptional<UE::Interchange::FAnimationStepCurvePayloadData>>>();
-
-	if (!Dispatcher.IsValid())
-	{
-		Promise->SetValue(TOptional<UE::Interchange::FAnimationStepCurvePayloadData>());
-		return Promise->GetFuture();
-	}
+	FString JsonCommand = 
+		(PayLoadKey.Type == EInterchangeAnimationPayLoadType::BAKED)
+		? CreateFetchAnimationBakeTransformPayloadFbxCommand(PayLoadKey.UniqueId, BakeFrequency, RangeStartSecond, RangeStopSecond)
+		: CreateFetchPayloadFbxCommand(PayLoadKey.UniqueId);
 
-	// Create a json command to read the fbx file
-	FString JsonCommand = CreateFetchPayloadFbxCommand(PayLoadKey);
-	const int32 CreatedTaskIndex = Dispatcher->AddTask(JsonCommand, FInterchangeDispatcherTaskCompleted::CreateLambda([this, Promise](const int32 TaskIndex)
+	const int32 CreatedTaskIndex = Dispatcher->AddTask(JsonCommand, FInterchangeDispatcherTaskCompleted::CreateLambda([this, Promise, PayLoadKey](const int32 TaskIndex)
 		{
-			TRACE_CPUPROFILER_EVENT_SCOPE("UInterchangeFbxTranslator::GetAnimationStepCurvePayloadData::Dispatcher->AddTaskDone")
+			TRACE_CPUPROFILER_EVENT_SCOPE("UInterchangeFbxTranslator::GetAnimationCurvePayloadData::Dispatcher->AddTaskDone")
 			UE::Interchange::ETaskState TaskState;
 			FString JsonResult;
 			TArray<FString> JsonMessages;
@@ -392,7 +316,7 @@ TFuture<TOptional<UE::Interchange::FAnimationStepCurvePayloadData>> UInterchange
 
 			if (TaskState != UE::Interchange::ETaskState::ProcessOk)
 			{
-				Promise->SetValue(TOptional<UE::Interchange::FAnimationStepCurvePayloadData>());
+				Promise->SetValue(TOptional<UE::Interchange::FAnimationPayloadData>());
 				return;
 			}
 
@@ -404,11 +328,11 @@ TFuture<TOptional<UE::Interchange::FAnimationStepCurvePayloadData>> UInterchange
 			if (!ensure(FPaths::FileExists(AnimationTransformPayloadFilename)))
 			{
 				// TODO log an error saying the payload file does not exist even if the get payload command succeeded
-				Promise->SetValue(TOptional<UE::Interchange::FAnimationStepCurvePayloadData>());
+				Promise->SetValue(TOptional<UE::Interchange::FAnimationPayloadData>());
 				return;
 			}
 
-			UE::Interchange::FAnimationStepCurvePayloadData AnimationTransformPayload;
+			UE::Interchange::FAnimationPayloadData AnimationTransformPayload(PayLoadKey.Type);
 
 			// All sub object should be gone with the reset
 			TArray64<uint8> Buffer;
@@ -418,94 +342,48 @@ TFuture<TOptional<UE::Interchange::FAnimationStepCurvePayloadData>> UInterchange
 			if (FileDataSize < 1)
 			{
 				// Nothing to load from this file
-				Promise->SetValue(TOptional<UE::Interchange::FAnimationStepCurvePayloadData>());
+				Promise->SetValue(TOptional<UE::Interchange::FAnimationPayloadData>());
 				return;
 			}
 
 			// Buffer keeps the ownership of the data, the large memory reader is use to serialize the TMap
 			FLargeMemoryReader Ar(FileData, FileDataSize);
-			Ar << AnimationTransformPayload.StepCurves;
-			Promise->SetValue(MoveTemp(AnimationTransformPayload));
-		}));
-
-	// The task was not added to the dispatcher
-	if (CreatedTaskIndex == INDEX_NONE)
-	{
-		Promise->SetValue(TOptional<UE::Interchange::FAnimationStepCurvePayloadData>{});
-	}
-
-	return Promise->GetFuture();
-}
-
-TFuture<TOptional<UE::Interchange::FAnimationBakeTransformPayloadData>> UInterchangeFbxTranslator::GetAnimationBakeTransformPayloadData(const FString& PayLoadKey, const double BakeFrequency, const double RangeStartSecond, const double RangeStopSecond) const
-{
-	TSharedPtr<TPromise<TOptional<UE::Interchange::FAnimationBakeTransformPayloadData>>> Promise = MakeShared<TPromise<TOptional<UE::Interchange::FAnimationBakeTransformPayloadData>>>();
-
-	if (!Dispatcher.IsValid())
-	{
-		Promise->SetValue(TOptional<UE::Interchange::FAnimationBakeTransformPayloadData>());
-		return Promise->GetFuture();
-	}
-
-
-	// Create a json command to read the fbx file
-	FString JsonCommand = CreateFetchAnimationBakeTransformPayloadFbxCommand(PayLoadKey, BakeFrequency, RangeStartSecond, RangeStopSecond);
-	const int32 CreatedTaskIndex = Dispatcher->AddTask(JsonCommand, FInterchangeDispatcherTaskCompleted::CreateLambda([this, Promise](const int32 TaskIndex)
-		{
-			TRACE_CPUPROFILER_EVENT_SCOPE("UInterchangeFbxTranslator::GetAnimationBakeTransformPayloadData::Dispatcher->AddTaskDone")
-			UE::Interchange::ETaskState TaskState;
-			FString JsonResult;
-			TArray<FString> JsonMessages;
-			Dispatcher->GetTaskState(TaskIndex, TaskState, JsonResult, JsonMessages);
 
-			// Parse the Json messages into UInterchangeResults
-			for (const FString& JsonMessage : JsonMessages)
+			switch (PayLoadKey.Type)
 			{
-				AddMessage(UInterchangeResult::FromJson(JsonMessage));
-			}
-
-			if (TaskState != UE::Interchange::ETaskState::ProcessOk)
-			{
-				Promise->SetValue(TOptional<UE::Interchange::FAnimationBakeTransformPayloadData>());
-				return;
-			}
-
-			// Grab the result file and fill the BaseNodeContainer
-			UE::Interchange::FJsonFetchPayloadCmd::JsonResultParser ResultParser;
-			ResultParser.FromJson(JsonResult);
-			FString AnimationTransformPayloadFilename = ResultParser.GetResultFilename();
-
-			if (!ensure(FPaths::FileExists(AnimationTransformPayloadFilename)))
-			{
-				// TODO log an error saying the payload file does not exist even if the get payload command succeeded
-				Promise->SetValue(TOptional<UE::Interchange::FAnimationBakeTransformPayloadData>());
-				return;
-			}
-
-			UE::Interchange::FAnimationBakeTransformPayloadData AnimationTransformPayload;
-
-			// All sub object should be gone with the reset
-			TArray64<uint8> Buffer;
-			FFileHelper::LoadFileToArray(Buffer, *AnimationTransformPayloadFilename);
-			uint8* FileData = Buffer.GetData();
-			int64 FileDataSize = Buffer.Num();
-			if (FileDataSize < 1)
-			{
-				// Nothing to load from this file
-				Promise->SetValue(TOptional<UE::Interchange::FAnimationBakeTransformPayloadData>());
-				return;
+			case EInterchangeAnimationPayLoadType::CURVE:
+			case EInterchangeAnimationPayLoadType::MORPHTARGETCURVE:
+				{
+					TArray<FInterchangeCurve> InterchangeCurves;
+					Ar << InterchangeCurves;
+					AnimationTransformPayload.Curves.AddDefaulted(InterchangeCurves.Num());
+					for (int32 CurveIndex = 0; CurveIndex < InterchangeCurves.Num(); ++CurveIndex)
+					{
+						const FInterchangeCurve& InterchangeCurve = InterchangeCurves[CurveIndex];
+						InterchangeCurve.ToRichCurve(AnimationTransformPayload.Curves[CurveIndex]);
+					}
+				}
+				break;
+			case EInterchangeAnimationPayLoadType::STEPCURVE:
+				{
+					Ar << AnimationTransformPayload.StepCurves;
+				}
+				break;
+			case EInterchangeAnimationPayLoadType::BAKED:
+				AnimationTransformPayload.SerializeBaked(Ar);
+				break;
+			case EInterchangeAnimationPayLoadType::NONE:
+			default:
+				break;
 			}
 
-			// Buffer keeps the ownership of the data, the large memory reader is use to serialize the TMap
-			FLargeMemoryReader Ar(FileData, FileDataSize);
-			AnimationTransformPayload.Serialize(Ar);
 			Promise->SetValue(MoveTemp(AnimationTransformPayload));
 		}));
 
 	// The task was not added to the dispatcher
 	if (CreatedTaskIndex == INDEX_NONE)
 	{
-		Promise->SetValue(TOptional<UE::Interchange::FAnimationBakeTransformPayloadData>{});
+		Promise->SetValue(TOptional<UE::Interchange::FAnimationPayloadData>{});
 	}
 
 	return Promise->GetFuture();
diff --git a/Engine/Plugins/Interchange/Runtime/Source/Import/Private/Gltf/InterchangeGltfAnimation.cpp b/Engine/Plugins/Interchange/Runtime/Source/Import/Private/Gltf/InterchangeGltfAnimation.cpp
--- a/Engine/Plugins/Interchange/Runtime/Source/Import/Private/Gltf/InterchangeGltfAnimation.cpp
+++ b/Engine/Plugins/Interchange/Runtime/Source/Import/Private/Gltf/InterchangeGltfAnimation.cpp
@@ -2,7 +2,7 @@
 
 #include "Gltf/InterchangeGltfPrivate.h"
 
-#include "Animation/InterchangeAnimationPayload.h"
+#include "InterchangeCommonAnimationPayload.h"
 #include "GLTFAccessor.h"
 #include "GLTFAnimation.h"
 #include "GLTFAsset.h"
@@ -201,7 +201,7 @@ namespace UE::Interchange::Gltf::Private
 		return ChannelIndices.Num() != 0;
 	}
 
-	bool GetTransformAnimationPayloadData(const FString& PayLoadKey, const GLTF::FAsset& GltfAsset, FAnimationCurvePayloadData& OutPayloadData)
+	bool GetTransformAnimationPayloadData(const FString& PayLoadKey, const GLTF::FAsset& GltfAsset, FAnimationPayloadData& OutPayloadData)
 	{
 		//translation, rotation, scale time 3 channels:
 		OutPayloadData.Curves.SetNum(9);
@@ -368,7 +368,7 @@ namespace UE::Interchange::Gltf::Private
 		return true;
 	}
 
-	bool GetMorphTargetAnimationPayloadData(const FString& PayLoadKey, const GLTF::FAsset& GltfAsset, FAnimationCurvePayloadData& OutPayloadData)
+	bool GetMorphTargetAnimationPayloadData(const FString& PayLoadKey, const GLTF::FAsset& GltfAsset, FAnimationPayloadData& OutPayloadData)
 	{
 		OutPayloadData.Curves.SetNum(1);
 
@@ -449,22 +449,7 @@ namespace UE::Interchange::Gltf::Private
 		return true;
 	}
 
-	bool GetAnimationPayloadData(const FString& PayLoadKey, const GLTF::FAsset& GltfAsset, FAnimationCurvePayloadData& OutPayloadData)
-	{
-		TArray<FString> MorphTargetCheckPayLoadKeys;
-		PayLoadKey.ParseIntoArray(MorphTargetCheckPayLoadKeys, TEXT("~"));
-
-		if (MorphTargetCheckPayLoadKeys.Num() == 2)
-		{
-			return GetMorphTargetAnimationPayloadData(MorphTargetCheckPayLoadKeys[1], GltfAsset, OutPayloadData);
-		}
-		else
-		{
-			return GetTransformAnimationPayloadData(PayLoadKey, GltfAsset, OutPayloadData);
-		}
-	}
-
-	bool GetBakedAnimationTransformPayloadData(const FString& PayLoadKey, const GLTF::FAsset& GltfAsset, FAnimationBakeTransformPayloadData& PayloadData)
+	bool GetBakedAnimationTransformPayloadData(const FString& PayLoadKey, const GLTF::FAsset& GltfAsset, FAnimationPayloadData& PayloadData)
 	{
 		TArray<int32> ChannelIndices;
 		int32 AnimationIndex;
diff --git a/Engine/Plugins/Interchange/Runtime/Source/Import/Private/Gltf/InterchangeGltfPrivate.h b/Engine/Plugins/Interchange/Runtime/Source/Import/Private/Gltf/InterchangeGltfPrivate.h
--- a/Engine/Plugins/Interchange/Runtime/Source/Import/Private/Gltf/InterchangeGltfPrivate.h
+++ b/Engine/Plugins/Interchange/Runtime/Source/Import/Private/Gltf/InterchangeGltfPrivate.h
@@ -12,8 +12,7 @@ namespace GLTF
 
 namespace UE::Interchange
 {
-	struct FAnimationCurvePayloadData;
-	struct FAnimationBakeTransformPayloadData;
+	struct FAnimationPayloadData;
 	struct FMeshPayloadData;
 
 	namespace Gltf::Private
@@ -21,8 +20,9 @@ namespace UE::Interchange
 		static float GltfUnitConversionMultiplier = 100.f;
 
 		// Animation related functions
-		bool GetAnimationPayloadData(const FString& PayLoadKey, const GLTF::FAsset& GltfAsset, FAnimationCurvePayloadData& OutPayloadData);
-		bool GetBakedAnimationTransformPayloadData(const FString& PayLoadKey, const GLTF::FAsset& GltfAsset, FAnimationBakeTransformPayloadData& PayloadData);
+		bool GetTransformAnimationPayloadData(const FString& PayLoadKey, const GLTF::FAsset& GltfAsset, FAnimationPayloadData& OutPayloadData);
+		bool GetMorphTargetAnimationPayloadData(const FString& PayLoadKey, const GLTF::FAsset& GltfAsset, FAnimationPayloadData& OutPayloadData);
+		bool GetBakedAnimationTransformPayloadData(const FString& PayLoadKey, const GLTF::FAsset& GltfAsset, FAnimationPayloadData& PayloadData);
 		// 
 
 		// Mesh related functions
diff --git a/Engine/Plugins/Interchange/Runtime/Source/Import/Private/Gltf/InterchangeGltfTranslator.cpp b/Engine/Plugins/Interchange/Runtime/Source/Import/Private/Gltf/InterchangeGltfTranslator.cpp
--- a/Engine/Plugins/Interchange/Runtime/Source/Import/Private/Gltf/InterchangeGltfTranslator.cpp
+++ b/Engine/Plugins/Interchange/Runtime/Source/Import/Private/Gltf/InterchangeGltfTranslator.cpp
@@ -1321,17 +1321,41 @@ TOptional< UE::Interchange::FImportImage > UInterchangeGltfTranslator::GetTextur
 	}
 }
 
-TFuture<TOptional<UE::Interchange::FAnimationCurvePayloadData>> UInterchangeGltfTranslator::GetAnimationCurvePayloadData(const FString& PayLoadKey) const
+TFuture<TOptional<UE::Interchange::FAnimationPayloadData>> UInterchangeGltfTranslator::GetAnimationPayloadData(const FInterchangeAnimationPayLoadKey& PayLoadKey, const double BakeFrequency, const double RangeStartSecond, const double RangeStopSecond) const
 {
-	return Async(EAsyncExecution::TaskGraph, [this, PayLoadKey]
+	return Async(EAsyncExecution::TaskGraph, [this, PayLoadKey, BakeFrequency, RangeStartSecond, RangeStopSecond]
 		{
 
-			TOptional<UE::Interchange::FAnimationCurvePayloadData> Result;
-			UE::Interchange::FAnimationCurvePayloadData AnimationCurvePayloadData;
+			TOptional<UE::Interchange::FAnimationPayloadData> Result;
+			UE::Interchange::FAnimationPayloadData AnimationPayLoadData(PayLoadKey.Type);
 
-			if (UE::Interchange::Gltf::Private::GetAnimationPayloadData(PayLoadKey, GltfAsset, AnimationCurvePayloadData))
+			switch (PayLoadKey.Type)
 			{
-				Result.Emplace(AnimationCurvePayloadData);
+			case EInterchangeAnimationPayLoadType::CURVE:
+				if (UE::Interchange::Gltf::Private::GetTransformAnimationPayloadData(PayLoadKey.UniqueId, GltfAsset, AnimationPayLoadData))
+				{
+					Result.Emplace(AnimationPayLoadData);
+				}
+				break;
+			case EInterchangeAnimationPayLoadType::MORPHTARGETCURVE:
+				if (UE::Interchange::Gltf::Private::GetMorphTargetAnimationPayloadData(PayLoadKey.UniqueId, GltfAsset, AnimationPayLoadData))
+				{
+					Result.Emplace(AnimationPayLoadData);
+				}
+				break;
+			case EInterchangeAnimationPayLoadType::BAKED:
+				AnimationPayLoadData.BakeFrequency = BakeFrequency;
+				AnimationPayLoadData.RangeStartTime = RangeStartSecond;
+				AnimationPayLoadData.RangeEndTime = RangeStopSecond;
+				if (UE::Interchange::Gltf::Private::GetBakedAnimationTransformPayloadData(PayLoadKey.UniqueId, GltfAsset, AnimationPayLoadData))
+				{
+					Result.Emplace(AnimationPayLoadData);
+				}
+				break;
+			case EInterchangeAnimationPayLoadType::STEPCURVE:
+			case EInterchangeAnimationPayLoadType::NONE:
+			default:
+				break;
 			}
 
 			return Result;
@@ -1339,36 +1363,6 @@ TFuture<TOptional<UE::Interchange::FAnimationCurvePayloadData>> UInterchangeGltf
 	);
 }
 
-TFuture<TOptional<UE::Interchange::FAnimationStepCurvePayloadData>> UInterchangeGltfTranslator::GetAnimationStepCurvePayloadData(const FString& PayLoadKey) const
-{
-	using namespace UE::Interchange;
-
-	TSharedPtr<TPromise<TOptional<FAnimationStepCurvePayloadData>>> Promise = MakeShared<TPromise<TOptional<FAnimationStepCurvePayloadData>>>();
-	Promise->SetValue(TOptional<FAnimationStepCurvePayloadData>());
-	return Promise->GetFuture();
-}
-
-TFuture<TOptional<UE::Interchange::FAnimationBakeTransformPayloadData>> UInterchangeGltfTranslator::GetAnimationBakeTransformPayloadData(const FString& PayLoadKey, const double BakeFrequency, const double RangeStartSecond, const double RangeStopSecond) const
-{
-	return Async(EAsyncExecution::TaskGraph, [this, PayLoadKey, BakeFrequency, RangeStartSecond, RangeStopSecond]
-		{
-			TOptional<UE::Interchange::FAnimationBakeTransformPayloadData> Result;
-
-			UE::Interchange::FAnimationBakeTransformPayloadData PayloadData;
-
-			PayloadData.BakeFrequency = BakeFrequency;
-			PayloadData.RangeStartTime = RangeStartSecond;
-			PayloadData.RangeEndTime = RangeStopSecond;
-
-			if (UE::Interchange::Gltf::Private::GetBakedAnimationTransformPayloadData(PayLoadKey, GltfAsset, PayloadData))
-			{
-				Result.Emplace(PayloadData);
-			}
-
-			return Result;
-		});
-}
-
 void UInterchangeGltfTranslator::HandleGltfAnimation(UInterchangeBaseNodeContainer& NodeContainer, int32 AnimationIndex) const
 {
 	const GLTF::FAnimation& GltfAnimation = GltfAsset.Animations[AnimationIndex];
@@ -1414,7 +1408,7 @@ void UInterchangeGltfTranslator::HandleGltfAnimation(UInterchangeBaseNodeContain
 
 				for (const TPair<FString, FString>& AnimationPayloadKeyForMorphTargetNodeUid : AnimationPayloadKeyForMorphTargetNodeUids)
 				{
-					TrackNode->SetAnimationPayloadKeyForMorphTargetNodeUid(AnimationPayloadKeyForMorphTargetNodeUid.Key, AnimationPayloadKeyForMorphTargetNodeUid.Value);
+					TrackNode->SetAnimationPayloadKeyForMorphTargetNodeUid(AnimationPayloadKeyForMorphTargetNodeUid.Key, AnimationPayloadKeyForMorphTargetNodeUid.Value, EInterchangeAnimationPayLoadType::MORPHTARGETCURVE);
 				}
 
 				TMap<FString, TArray<int32>>& JointUidWithChannelsUsedMap = TrackNodeToJointUidWithChannelsUsedMap.FindOrAdd(TrackNode);
@@ -1458,7 +1452,7 @@ void UInterchangeGltfTranslator::HandleGltfAnimation(UInterchangeBaseNodeContain
 								if (MorphTargetNodeConst->GetPayLoadKey().IsSet())
 								{
 									FInterchangeMeshPayLoadKey PayLoadKey = MorphTargetNodeConst->GetPayLoadKey().GetValue();
-									FString PayLoadKeyUniqueId = TEXT("MorphTargetAnimation~") + LexToString(AnimationIndex) + TEXT(":") + LexToString(ChannelIndex) + TEXT(":") + PayLoadKey.UniqueId;
+									FString PayLoadKeyUniqueId = LexToString(AnimationIndex) + TEXT(":") + LexToString(ChannelIndex) + TEXT(":") + PayLoadKey.UniqueId;
 
 									AnimationPayloadKeyForMorphTargetNodeUids.Add(MorphTargetDependencyUid, PayLoadKeyUniqueId);
 								}
@@ -1570,7 +1564,7 @@ void UInterchangeGltfTranslator::HandleGltfAnimation(UInterchangeBaseNodeContain
 				{
 					if (bHasNonWeightAnimationChannel)
 					{
-						TrackNodeAndJointNodeChannels.Key->SetAnimationPayloadKeyForSceneNodeUid(JointNodeUidAndChannelsUsedPair.Key, Payload);
+						TrackNodeAndJointNodeChannels.Key->SetAnimationPayloadKeyForSceneNodeUid(JointNodeUidAndChannelsUsedPair.Key, Payload, EInterchangeAnimationPayLoadType::BAKED);
 					}
 					
 					bHasAnimationPayloadSet = true;
@@ -1659,7 +1653,7 @@ void UInterchangeGltfTranslator::HandleGltfAnimation(UInterchangeBaseNodeContain
 			}
 		}
 
-		TransformAnimTrackNode->SetCustomAnimationPayloadKey(PayloadKey);
+		TransformAnimTrackNode->SetCustomAnimationPayloadKey(PayloadKey, EInterchangeAnimationPayLoadType::CURVE);
 		TransformAnimTrackNode->SetCustomUsedChannels(UsedChannels);
 
 		NodeContainer.AddNode(TransformAnimTrackNode);
diff --git a/Engine/Plugins/Interchange/Runtime/Source/Import/Public/Animation/InterchangeAnimationPayload.h b/Engine/Plugins/Interchange/Runtime/Source/Import/Public/Animation/InterchangeAnimationPayload.h
deleted file mode 100644
diff --git a/Engine/Plugins/Interchange/Runtime/Source/Import/Public/Animation/InterchangeAnimationPayloadInterface.h b/Engine/Plugins/Interchange/Runtime/Source/Import/Public/Animation/InterchangeAnimationPayloadInterface.h
--- a/Engine/Plugins/Interchange/Runtime/Source/Import/Public/Animation/InterchangeAnimationPayloadInterface.h
+++ b/Engine/Plugins/Interchange/Runtime/Source/Import/Public/Animation/InterchangeAnimationPayloadInterface.h
@@ -3,11 +3,12 @@
 
 #include "Async/Future.h"
 #include "CoreMinimal.h"
-#include "Animation/InterchangeAnimationPayload.h"
 #include "Misc/FrameRate.h"
 #include "UObject/Interface.h"
 #include "UObject/Object.h"
 #include "UObject/ObjectMacros.h"
+#include "InterchangeAnimationTrackSetNode.h"
+#include "InterchangeCommonAnimationPayload.h"
 
 #include "InterchangeAnimationPayloadInterface.generated.h"
 
@@ -26,24 +27,15 @@ class INTERCHANGEIMPORT_API IInterchangeAnimationPayloadInterface
 public:
 
 	/**
-	 * Get animation payload data for the specified payload key. It return an array of FRichCurve (Rich curve are float curve we can interpolate).
+	 * Get animation payload data for the specified payload key.
+	 * It return an array of FRichCurve (Rich curve are float curve we can interpolate) or array of "Step" curve or an array of Baked Transformations, depending on Type
 	 *
-	 * @param PayLoadKey - The key to retrieve the a particular payload contain into the specified source data.
+	 * @param PayLoadKey.UniqueId - The key to retrieve the a particular payload contain into the specified source data.
+	 * @param PayLoadKey.Type - The type of animation data to be retrieved (Type is provided by the Translator upon initialization of PayLoadKeys)
 	 * @return - The resulting PayloadData as a TFuture point by the PayloadKey. The TOptional will not be set if there is an error retrieving the payload.
 	 * 
 	 */
-	virtual TFuture<TOptional<UE::Interchange::FAnimationCurvePayloadData>> GetAnimationCurvePayloadData(const FString& PayLoadKey) const = 0;
-
-	/**
-	 * Get animation payload data for the specified payload key. It return an array of "Step" curve data (step curve are a combination of Key time and Key Value, no interpolation data)
-	 *
-	 * @param PayLoadKey - The key to retrieve the a particular payload contain into the specified source data.
-	 * @return - The resulting PayloadData as a TFuture point by the PayloadKey. The TOptional will not be set if there is an error retrieving the payload.
-	 *
-	 */
-	virtual TFuture<TOptional<UE::Interchange::FAnimationStepCurvePayloadData>> GetAnimationStepCurvePayloadData(const FString& PayLoadKey) const = 0;
-
-	virtual TFuture<TOptional<UE::Interchange::FAnimationBakeTransformPayloadData>> GetAnimationBakeTransformPayloadData(const FString& PayLoadKey, const double BakeFrequency, const double RangeStartSecond, const double RangeStopSecond) const = 0;
+	virtual TFuture<TOptional<UE::Interchange::FAnimationPayloadData>> GetAnimationPayloadData(const FInterchangeAnimationPayLoadKey& PayLoadKey, const double BakeFrequency = 0, const double RangeStartSecond = 0, const double RangeStopSecond = 0) const = 0;
 };
 
 
diff --git a/Engine/Plugins/Interchange/Runtime/Source/Import/Public/Fbx/InterchangeFbxTranslator.h b/Engine/Plugins/Interchange/Runtime/Source/Import/Public/Fbx/InterchangeFbxTranslator.h
--- a/Engine/Plugins/Interchange/Runtime/Source/Import/Public/Fbx/InterchangeFbxTranslator.h
+++ b/Engine/Plugins/Interchange/Runtime/Source/Import/Public/Fbx/InterchangeFbxTranslator.h
@@ -69,13 +69,7 @@ public:
 
 	//////////////////////////////////////////////////////////////////////////
 	/* IInterchangeAnimationPayloadInterface Begin */
-
-	virtual TFuture<TOptional<UE::Interchange::FAnimationCurvePayloadData>> GetAnimationCurvePayloadData(const FString& PayLoadKey) const override;
-
-	virtual TFuture<TOptional<UE::Interchange::FAnimationStepCurvePayloadData>> GetAnimationStepCurvePayloadData(const FString& PayLoadKey) const override;
-
-	virtual TFuture<TOptional<UE::Interchange::FAnimationBakeTransformPayloadData>> GetAnimationBakeTransformPayloadData(const FString& PayLoadKey, const double BakeFrequency, const double RangeStartSecond, const double RangeStopSecond) const override;
-
+	virtual TFuture<TOptional<UE::Interchange::FAnimationPayloadData>> GetAnimationPayloadData(const FInterchangeAnimationPayLoadKey& PayLoadKey, const double BakeFrequency = 0, const double RangeStartSecond = 0, const double RangeStopSecond = 0) const override;
 	/* IInterchangeAnimationPayloadInterface End */
 private:
 	FString CreateLoadFbxFileCommand(const FString& FbxFilePath) const;
diff --git a/Engine/Plugins/Interchange/Runtime/Source/Import/Public/Gltf/InterchangeGltfTranslator.h b/Engine/Plugins/Interchange/Runtime/Source/Import/Public/Gltf/InterchangeGltfTranslator.h
--- a/Engine/Plugins/Interchange/Runtime/Source/Import/Public/Gltf/InterchangeGltfTranslator.h
+++ b/Engine/Plugins/Interchange/Runtime/Source/Import/Public/Gltf/InterchangeGltfTranslator.h
@@ -52,9 +52,7 @@ public:
 	/* IInterchangeTexturePayloadInterface End */
 
 	/* IInterchangeAnimationPayloadInterface Begin */
-	virtual TFuture<TOptional<UE::Interchange::FAnimationCurvePayloadData>> GetAnimationCurvePayloadData(const FString& PayLoadKey) const override;
-	virtual TFuture<TOptional<UE::Interchange::FAnimationStepCurvePayloadData>> GetAnimationStepCurvePayloadData(const FString& PayLoadKey) const override;
-	virtual TFuture<TOptional<UE::Interchange::FAnimationBakeTransformPayloadData>> GetAnimationBakeTransformPayloadData(const FString& PayLoadKey, const double BakeFrequency, const double RangeStartSecond, const double RangeStopSecond) const override;
+	virtual TFuture<TOptional<UE::Interchange::FAnimationPayloadData>> GetAnimationPayloadData(const FInterchangeAnimationPayLoadKey& PayLoadKey, const double BakeFrequency = 0, const double RangeStartSecond = 0, const double RangeStopSecond = 0) const override;
 	/* IInterchangeAnimationPayloadInterface End */
 
 	/* IInterchangeVariantSetPayloadInterface Begin */
diff --git a/Engine/Plugins/Interchange/Runtime/Source/Nodes/Private/InterchangeAnimationTrackSetNode.cpp b/Engine/Plugins/Interchange/Runtime/Source/Nodes/Private/InterchangeAnimationTrackSetNode.cpp
--- a/Engine/Plugins/Interchange/Runtime/Source/Nodes/Private/InterchangeAnimationTrackSetNode.cpp
+++ b/Engine/Plugins/Interchange/Runtime/Source/Nodes/Private/InterchangeAnimationTrackSetNode.cpp
@@ -4,6 +4,37 @@
 
 // UInterchangeAnimationTrackSetNode
 
+//Interchange namespace
+namespace UE
+{
+	namespace Interchange
+	{
+		const FAttributeKey& FAnimationStaticData::AnimationPayLoadUidKey()
+		{
+			static FAttributeKey AttributeKey(TEXT("__AnimationPayloadUidKey__"));
+			return AttributeKey;
+		}
+
+		const FAttributeKey& FAnimationStaticData::AnimationPayLoadTypeKey()
+		{
+			static FAttributeKey AttributeKey(TEXT("__AnimationPayloadTypeKey__"));
+			return AttributeKey;
+		}
+
+		const FAttributeKey& FAnimationStaticData::MorphTargetAnimationPayLoadUidKey()
+		{
+			static FAttributeKey AttributeKey(TEXT("__MorphTargetAnimationPayloadUidKey__"));
+			return AttributeKey;
+		}
+
+		const FAttributeKey& FAnimationStaticData::MorphTargetAnimationPayLoadTypeKey()
+		{
+			static FAttributeKey AttributeKey(TEXT("__MorphTargetAnimationPayloadTypeKey__"));
+			return AttributeKey;
+		}
+	}
+}
+
 UInterchangeAnimationTrackSetNode::UInterchangeAnimationTrackSetNode()
 {
 	CustomAnimationTrackUids.Initialize(Attributes, TEXT("AnimationTrackUids"));
@@ -69,14 +100,75 @@ bool UInterchangeAnimationTrackNode::SetCustomActorDependencyUid(const FString&
 	IMPLEMENT_NODE_ATTRIBUTE_SETTER_NODELEGATE(ActorDependency, FString);
 }
 
-bool UInterchangeAnimationTrackNode::GetCustomAnimationPayloadKey(FString& AttributeValue) const
-{
-	IMPLEMENT_NODE_ATTRIBUTE_GETTER(AnimationPayload, FString);
-}
-
-bool UInterchangeAnimationTrackNode::SetCustomAnimationPayloadKey(const FString& AttributeValue)
-{
-	IMPLEMENT_NODE_ATTRIBUTE_SETTER_NODELEGATE(AnimationPayload, FString);
+bool UInterchangeAnimationTrackNode::GetCustomAnimationPayloadKey(FInterchangeAnimationPayLoadKey& AnimationPayLoadKey) const
+{
+	FString UniqueId;
+	EInterchangeAnimationPayLoadType PayLoadType;
+
+	//PayLoadKey
+	{
+		if (!Attributes->ContainAttribute(UE::Interchange::FAnimationStaticData::AnimationPayLoadUidKey()))
+		{
+			return false;
+		}
+		UE::Interchange::FAttributeStorage::TAttributeHandle<FString> AttributeHandle = Attributes->GetAttributeHandle<FString>(UE::Interchange::FAnimationStaticData::AnimationPayLoadUidKey());
+		if (!AttributeHandle.IsValid())
+		{
+			return false;
+		}
+		UE::Interchange::EAttributeStorageResult Result = AttributeHandle.Get(UniqueId);
+		if (!IsAttributeStorageResultSuccess(Result))
+		{
+			LogAttributeStorageErrors(Result, TEXT("UInterchangeMeshNode.GetPayLoadKey"), UE::Interchange::FAnimationStaticData::AnimationPayLoadUidKey());
+			return false;
+		}
+	}
+
+	//PayLoadType
+	{
+		if (!Attributes->ContainAttribute(UE::Interchange::FAnimationStaticData::AnimationPayLoadTypeKey()))
+		{
+			return false;
+		}
+		UE::Interchange::FAttributeStorage::TAttributeHandle<EInterchangeAnimationPayLoadType> AttributeHandle = Attributes->GetAttributeHandle<EInterchangeAnimationPayLoadType>(UE::Interchange::FAnimationStaticData::AnimationPayLoadTypeKey());
+		if (!AttributeHandle.IsValid())
+		{
+			return false;
+		}
+
+		UE::Interchange::EAttributeStorageResult Result = AttributeHandle.Get(PayLoadType);
+		if (!IsAttributeStorageResultSuccess(Result))
+		{
+			LogAttributeStorageErrors(Result, TEXT("UInterchangeMeshNode.GetPayLoadTypeKey"), UE::Interchange::FAnimationStaticData::AnimationPayLoadTypeKey());
+			return false;
+		}
+	}
+	
+	AnimationPayLoadKey.UniqueId = UniqueId;
+	AnimationPayLoadKey.Type = PayLoadType;
+
+	return true;
+}
+
+bool UInterchangeAnimationTrackNode::SetCustomAnimationPayloadKey(const FString& InUniqueId, const EInterchangeAnimationPayLoadType& InType)
+{
+	UE::Interchange::EAttributeStorageResult Result = Attributes->RegisterAttribute(UE::Interchange::FAnimationStaticData::AnimationPayLoadUidKey(), InUniqueId);
+	if (!IsAttributeStorageResultSuccess(Result))
+	{
+		LogAttributeStorageErrors(Result, TEXT("UInterchangeAnimationTrackNode.SetPayLoadKey"), UE::Interchange::FAnimationStaticData::AnimationPayLoadUidKey());
+		return false;
+	}
+	else
+	{
+		Result = Attributes->RegisterAttribute(UE::Interchange::FAnimationStaticData::AnimationPayLoadTypeKey(), InType);
+		if (!IsAttributeStorageResultSuccess(Result))
+		{
+			LogAttributeStorageErrors(Result, TEXT("UInterchangeAnimationTrackNode.SetPayLoadTypeKey"), UE::Interchange::FAnimationStaticData::AnimationPayLoadTypeKey());
+			return false;
+		}
+	}
+
+	return true;
 }
 
 bool UInterchangeAnimationTrackNode::GetCustomFrameCount(int32& AttributeValue) const
@@ -157,8 +249,11 @@ bool UInterchangeTransformAnimationTrackNode::GetCustomUsedChannels(int32& Attri
 
 UInterchangeSkeletalAnimationTrackNode::UInterchangeSkeletalAnimationTrackNode()
 {
-	SceneNodeAnimationPayloadKeyMap.Initialize(Attributes.ToSharedRef(), TEXT("__SceneNodeAnimationPayloadKeyMap__"));
-	MorphTargetPayloadKeyMap.Initialize(Attributes.ToSharedRef(), TEXT("__MorphTargetPayloadKeyMap__"));
+	SceneNodeAnimationPayloadKeyUidMap.Initialize(Attributes.ToSharedRef(), TEXT("__SceneNodeAnimationPayloadKeyUidMap__"));
+	SceneNodeAnimationPayloadKeyTypeMap.Initialize(Attributes.ToSharedRef(), TEXT("__SceneNodeAnimationPayloadKeyTypeMap__"));
+
+	MorphTargetPayloadKeyUidMap.Initialize(Attributes.ToSharedRef(), TEXT("__MorphTargetPayloadKeyUidMap__"));
+	MorphTargetPayloadKeyTypeMap.Initialize(Attributes.ToSharedRef(), TEXT("__MorphTargetPayloadKeyTypeMap__"));
 }
 
 bool UInterchangeSkeletalAnimationTrackNode::GetCustomSkeletonNodeUid(FString& AttributeValue) const
@@ -201,42 +296,28 @@ bool UInterchangeSkeletalAnimationTrackNode::SetCustomAnimationStopTime(const do
 	IMPLEMENT_NODE_ATTRIBUTE_SETTER_NODELEGATE(AnimationStopTime, double);
 }
 
-void UInterchangeSkeletalAnimationTrackNode::GetSceneNodeAnimationPayloadKeys(TMap<FString, FString>& OutSceneNodeAnimationPayloads) const
-{
-	OutSceneNodeAnimationPayloads = SceneNodeAnimationPayloadKeyMap.ToMap();
-}
-
-bool UInterchangeSkeletalAnimationTrackNode::GetAnimationPayloadKeyFromSceneNodeUid(const FString& SceneNodeUid, FString& OutPayloadKey) const
-{
-	return SceneNodeAnimationPayloadKeyMap.GetValue(SceneNodeUid, OutPayloadKey);
-}
-
-bool UInterchangeSkeletalAnimationTrackNode::SetAnimationPayloadKeyForSceneNodeUid(const FString& SceneNodeUid, const FString& PayloadKey)
-{
-	return SceneNodeAnimationPayloadKeyMap.SetKeyValue(SceneNodeUid, PayloadKey);
-}
-
-bool UInterchangeSkeletalAnimationTrackNode::RemoveAnimationPayloadKeyForSceneNodeUid(const FString& SceneNodeUid)
-{
-	return SceneNodeAnimationPayloadKeyMap.RemoveKey(SceneNodeUid);
-}
-
-void UInterchangeSkeletalAnimationTrackNode::GetMorphTargetNodeAnimationPayloadKeys(TMap<FString, FString>& OutMorphTargetNodeAnimationPayloads) const
+void UInterchangeSkeletalAnimationTrackNode::GetSceneNodeAnimationPayloadKeys(TMap<FString, FString>& OutSceneNodeAnimationPayloadKeyUids, TMap<FString, uint8>& OutSceneNodeAnimationPayloadKeyTypes) const
 {
-	OutMorphTargetNodeAnimationPayloads = MorphTargetPayloadKeyMap.ToMap();
+	OutSceneNodeAnimationPayloadKeyUids = SceneNodeAnimationPayloadKeyUidMap.ToMap();
+	OutSceneNodeAnimationPayloadKeyTypes = SceneNodeAnimationPayloadKeyTypeMap.ToMap();
 }
 
-bool UInterchangeSkeletalAnimationTrackNode::GetAnimationPayloadKeyFromMorphTargetNodeUid(const FString& MorphTargetNodeUid, FString& OutPayloadKey) const
+bool UInterchangeSkeletalAnimationTrackNode::SetAnimationPayloadKeyForSceneNodeUid(const FString& SceneNodeUid, const FString& InUniqueId, const EInterchangeAnimationPayLoadType& InType)
 {
-	return MorphTargetPayloadKeyMap.GetValue(MorphTargetNodeUid, OutPayloadKey);
+	bool bSuccess = SceneNodeAnimationPayloadKeyUidMap.SetKeyValue(SceneNodeUid, InUniqueId);
+	bSuccess &= SceneNodeAnimationPayloadKeyTypeMap.SetKeyValue(SceneNodeUid, (uint8)InType);
+	return bSuccess;
 }
 
-bool UInterchangeSkeletalAnimationTrackNode::SetAnimationPayloadKeyForMorphTargetNodeUid(const FString& MorphTargetNodeUid, const FString& PayloadKey)
+void UInterchangeSkeletalAnimationTrackNode::GetMorphTargetNodeAnimationPayloadKeys(TMap<FString, FString>& OutMorphTargetNodeAnimationPayloadKeyUids, TMap<FString, uint8>& OutMorphTargetNodeAnimationPayloadKeyTypes) const
 {
-	return MorphTargetPayloadKeyMap.SetKeyValue(MorphTargetNodeUid, PayloadKey);
+	OutMorphTargetNodeAnimationPayloadKeyUids = MorphTargetPayloadKeyUidMap.ToMap();
+	OutMorphTargetNodeAnimationPayloadKeyTypes = MorphTargetPayloadKeyTypeMap.ToMap();
 }
 
-bool UInterchangeSkeletalAnimationTrackNode::RemoveAnimationPayloadKeyForMorphTargetNodeUid(const FString& MorphTargetNodeUid)
+bool UInterchangeSkeletalAnimationTrackNode::SetAnimationPayloadKeyForMorphTargetNodeUid(const FString& MorphTargetNodeUid, const FString& InUniqueId, const EInterchangeAnimationPayLoadType& InType)
 {
-	return MorphTargetPayloadKeyMap.RemoveKey(MorphTargetNodeUid);
+	bool bSuccess = MorphTargetPayloadKeyUidMap.SetKeyValue(MorphTargetNodeUid, InUniqueId);
+	bSuccess &= MorphTargetPayloadKeyTypeMap.SetKeyValue(MorphTargetNodeUid, (uint8)InType);
+	return bSuccess;
 }
diff --git a/Engine/Plugins/Interchange/Runtime/Source/Nodes/Public/InterchangeAnimationTrackSetNode.h b/Engine/Plugins/Interchange/Runtime/Source/Nodes/Public/InterchangeAnimationTrackSetNode.h
--- a/Engine/Plugins/Interchange/Runtime/Source/Nodes/Public/InterchangeAnimationTrackSetNode.h
+++ b/Engine/Plugins/Interchange/Runtime/Source/Nodes/Public/InterchangeAnimationTrackSetNode.h
@@ -7,6 +7,55 @@
 
 #include "InterchangeAnimationTrackSetNode.generated.h"
 
+UENUM(BlueprintType)
+enum class EInterchangeAnimationPayLoadType : uint8
+{
+	NONE = 0,
+	CURVE,
+	MORPHTARGETCURVE, //handles/generates the same properties as the CURVE variation, but they way it is acquired might be different (depending on the Translator)
+	STEPCURVE,
+	BAKED
+};
+
+
+USTRUCT(BlueprintType)
+struct INTERCHANGENODES_API FInterchangeAnimationPayLoadKey
+{
+	GENERATED_BODY()
+
+public:
+
+	UPROPERTY(EditAnywhere, BlueprintReadWrite, Category = "Interchange | Node")
+	FString UniqueId = "";
+
+	UPROPERTY(EditAnywhere, BlueprintReadWrite, Category = "Interchange | Node")
+	EInterchangeAnimationPayLoadType Type = EInterchangeAnimationPayLoadType::NONE;
+
+	FInterchangeAnimationPayLoadKey() {}
+
+	FInterchangeAnimationPayLoadKey(const FString& InUniqueId, const EInterchangeAnimationPayLoadType& InType)
+		: UniqueId(InUniqueId)
+		, Type(InType)
+	{
+	}
+};
+
+//Interchange namespace
+namespace UE
+{
+	namespace Interchange
+	{
+
+		struct INTERCHANGENODES_API FAnimationStaticData : public FBaseNodeStaticData
+		{
+			static const FAttributeKey& AnimationPayLoadUidKey();
+			static const FAttributeKey& AnimationPayLoadTypeKey();
+			static const FAttributeKey& MorphTargetAnimationPayLoadUidKey();
+			static const FAttributeKey& MorphTargetAnimationPayLoadTypeKey();
+		};
+	}
+}
+
 /**
  * Enumeration specifying which properties of a camera, light or scene node can be animated besides transform.
  */
@@ -270,13 +319,13 @@ public:
 	 * Set the payload key needed to retrieve the animation for this track.
 	 */
 	UFUNCTION(BlueprintCallable, Category = "Interchange | Node | AnimationTrack")
-	bool SetCustomAnimationPayloadKey(const FString& PayloadKey);
+	bool SetCustomAnimationPayloadKey(const FString& InUniqueId, const EInterchangeAnimationPayLoadType& InType);
 
 	/**
 	 * Get the payload key needed to retrieve the animation for this track.
 	 */
 	UFUNCTION(BlueprintCallable, Category = "Interchange | Node | AnimationTrack")
-	bool GetCustomAnimationPayloadKey(FString& PayloadKey) const;
+	bool GetCustomAnimationPayloadKey(FInterchangeAnimationPayLoadKey& AnimationPayLoadKey) const;
 
 	/**
 	 * Set the number of frames for the animation of this track.
@@ -292,7 +341,8 @@ public:
 
 private:
 	const UE::Interchange::FAttributeKey Macro_CustomActorDependencyKey = UE::Interchange::FAttributeKey(TEXT("ActorDependency"));
-	const UE::Interchange::FAttributeKey Macro_CustomAnimationPayloadKey = UE::Interchange::FAttributeKey(TEXT("AnimationPayload"));
+	const UE::Interchange::FAttributeKey Macro_CustomAnimationPayloadUidKey = UE::Interchange::FAttributeKey(TEXT("AnimationPayloadUid"));
+	const UE::Interchange::FAttributeKey Macro_CustomAnimationPayloadTypeKey = UE::Interchange::FAttributeKey(TEXT("AnimationPayloadType"));
 	const UE::Interchange::FAttributeKey Macro_CustomFrameCountKey = UE::Interchange::FAttributeKey(TEXT("FrameCount"));
 	const UE::Interchange::FAttributeKey Macro_CustomTargetedPropertyKey = UE::Interchange::FAttributeKey(TEXT("TargetedProperty"));
 };
@@ -374,8 +424,11 @@ public:
 
 		if (Ar.IsLoading() && bIsInitialized)
 		{
-			SceneNodeAnimationPayloadKeyMap.RebuildCache();
-			MorphTargetPayloadKeyMap.RebuildCache();
+			SceneNodeAnimationPayloadKeyUidMap.RebuildCache();
+			SceneNodeAnimationPayloadKeyTypeMap.RebuildCache();
+
+			MorphTargetPayloadKeyUidMap.RebuildCache();
+			MorphTargetPayloadKeyTypeMap.RebuildCache();
 		}
 	}
 
@@ -440,28 +493,16 @@ public:
 	bool GetCustomAnimationStopTime(double& StopTime) const;
 
 	UFUNCTION(BlueprintCallable, Category = "Interchange | Node | SkeletalAnimationTrack")
-	void GetSceneNodeAnimationPayloadKeys(TMap<FString, FString>& OutSceneNodeAnimationPayloadKeys) const;
-
-	UFUNCTION(BlueprintCallable, Category = "Interchange | Node | SkeletalAnimationTrack")
-	bool GetAnimationPayloadKeyFromSceneNodeUid(const FString& SceneNodeUid, FString& OutPayloadKey) const;
-
-	UFUNCTION(BlueprintCallable, Category = "Interchange | Node | SkeletalAnimationTrack")
-	bool SetAnimationPayloadKeyForSceneNodeUid(const FString& SceneNodeUid, const FString& PayloadKey);
-
-	UFUNCTION(BlueprintCallable, Category = "Interchange | Node | SkeletalAnimationTrack")
-	bool RemoveAnimationPayloadKeyForSceneNodeUid(const FString& SceneNodeUid);
-
-	UFUNCTION(BlueprintCallable, Category = "Interchange | Node | SkeletalAnimationTrack")
-	void GetMorphTargetNodeAnimationPayloadKeys(TMap<FString, FString>& OutMorphTargetNodeAnimationPayloads) const;
+	void GetSceneNodeAnimationPayloadKeys(TMap<FString, FString>& OutSceneNodeAnimationPayloadKeyUids, TMap<FString, uint8>& OutSceneNodeAnimationPayloadKeyTypes) const;
 
 	UFUNCTION(BlueprintCallable, Category = "Interchange | Node | SkeletalAnimationTrack")
-	bool GetAnimationPayloadKeyFromMorphTargetNodeUid(const FString& MorphTargetNodeUid, FString& OutPayloadKey) const;
+	bool SetAnimationPayloadKeyForSceneNodeUid(const FString& SceneNodeUid, const FString& InUniqueId, const EInterchangeAnimationPayLoadType& InType);
 
 	UFUNCTION(BlueprintCallable, Category = "Interchange | Node | SkeletalAnimationTrack")
-	bool SetAnimationPayloadKeyForMorphTargetNodeUid(const FString& MorphTargetNodeUid, const FString& PayloadKey);
+	void GetMorphTargetNodeAnimationPayloadKeys(TMap<FString, FString>& OutMorphTargetNodeAnimationPayloadKeyUids, TMap<FString, uint8>& OutMorphTargetNodeAnimationPayloadKeyTypes) const;
 
 	UFUNCTION(BlueprintCallable, Category = "Interchange | Node | SkeletalAnimationTrack")
-	bool RemoveAnimationPayloadKeyForMorphTargetNodeUid(const FString& MorphTargetNodeUid);
+	bool SetAnimationPayloadKeyForMorphTargetNodeUid(const FString& MorphTargetNodeUid, const FString& InUniqueId, const EInterchangeAnimationPayLoadType& InType);
 
 private:
 	const UE::Interchange::FAttributeKey Macro_CustomSkeletonNodeUidKey = UE::Interchange::FAttributeKey(TEXT("SkeletonNodeUid"));
@@ -469,6 +510,9 @@ private:
 	const UE::Interchange::FAttributeKey Macro_CustomAnimationStartTimeKey = UE::Interchange::FAttributeKey(TEXT("AnimationStartTime"));
 	const UE::Interchange::FAttributeKey Macro_CustomAnimationStopTimeKey = UE::Interchange::FAttributeKey(TEXT("AnimationStopTime"));
 	
-	UE::Interchange::TMapAttributeHelper<FString, FString> SceneNodeAnimationPayloadKeyMap;
-	UE::Interchange::TMapAttributeHelper<FString, FString> MorphTargetPayloadKeyMap;
+	UE::Interchange::TMapAttributeHelper<FString, FString>	SceneNodeAnimationPayloadKeyUidMap;
+	UE::Interchange::TMapAttributeHelper<FString, uint8>	SceneNodeAnimationPayloadKeyTypeMap;
+
+	UE::Interchange::TMapAttributeHelper<FString, FString>	MorphTargetPayloadKeyUidMap;
+	UE::Interchange::TMapAttributeHelper<FString, uint8>	MorphTargetPayloadKeyTypeMap;
 };
diff --git a/Engine/Plugins/Interchange/Runtime/Source/Parsers/CommonParser/InterchangeCommonParser.Build.cs b/Engine/Plugins/Interchange/Runtime/Source/Parsers/CommonParser/InterchangeCommonParser.Build.cs
--- a/Engine/Plugins/Interchange/Runtime/Source/Parsers/CommonParser/InterchangeCommonParser.Build.cs
+++ b/Engine/Plugins/Interchange/Runtime/Source/Parsers/CommonParser/InterchangeCommonParser.Build.cs
@@ -14,6 +14,7 @@ namespace UnrealBuildTool.Rules
 					"Core",
 					"CoreUObject",
 					"InterchangeCore",
+					"InterchangeNodes",
 				}
 			);
 
diff --git a/Engine/Plugins/Interchange/Runtime/Source/Parsers/CommonParser/Private/InterchangeCommonAnimationPayload.cpp b/Engine/Plugins/Interchange/Runtime/Source/Parsers/CommonParser/Private/InterchangeCommonAnimationPayload.cpp
--- a/Engine/Plugins/Interchange/Runtime/Source/Parsers/CommonParser/Private/InterchangeCommonAnimationPayload.cpp
+++ b/Engine/Plugins/Interchange/Runtime/Source/Parsers/CommonParser/Private/InterchangeCommonAnimationPayload.cpp
@@ -6,19 +6,115 @@
 
 #include UE_INLINE_GENERATED_CPP_BY_NAME(InterchangeCommonAnimationPayload)
 
-#if WITH_ENGINE
-#include "Curves/RichCurve.h"
-#endif //WITH_ENGINE
-
 namespace UE::Interchange
 {
-	void FAnimationBakeTransformPayloadData::Serialize(FArchive& Ar)
+	void FAnimationPayloadData::SerializeBaked(FArchive& Ar)
 	{
 		Ar << BakeFrequency;
 		Ar << RangeStartTime;
 		Ar << RangeEndTime;
 		Ar << Transforms;
 	}
+
+	void FAnimationPayloadData::CalculateDataFor(const EInterchangeAnimationPayLoadType& ToType, const FTransform& DefaultTransform)
+	{
+#if WITH_ENGINE
+		if (Type == EInterchangeAnimationPayLoadType::CURVE
+			&& ToType == EInterchangeAnimationPayLoadType::STEPCURVE)
+		{
+
+			StepCurves.Reserve(Curves.Num());
+			for (const FRichCurve& RichCurve : Curves)
+			{
+				FInterchangeStepCurve& StepCurve = StepCurves.AddDefaulted_GetRef();
+				const int32 KeyCount = RichCurve.GetNumKeys();
+				StepCurve.KeyTimes.AddZeroed(KeyCount);
+				TArray<float> KeyValues;
+				KeyValues.AddZeroed(KeyCount);
+				int32 KeyIndex = 0;
+				for (FKeyHandle KeyHandle = RichCurve.GetFirstKeyHandle(); KeyHandle != FKeyHandle::Invalid(); KeyHandle = RichCurve.GetNextKey(KeyHandle), ++KeyIndex)
+				{
+					StepCurve.KeyTimes[KeyIndex] = RichCurve.GetKeyTime(KeyHandle);
+					KeyValues[KeyIndex] = RichCurve.GetKeyValue(KeyHandle);
+				}
+				StepCurve.FloatKeyValues = MoveTemp(KeyValues);
+			}
+			AdditionalSupportedType = ToType;
+		}
+		else if (Type == EInterchangeAnimationPayLoadType::CURVE
+				 && ToType == EInterchangeAnimationPayLoadType::BAKED)
+		{
+			if (Curves.Num() != 9)
+			{
+				return;
+			}
+
+			//calculate RangeEnd -> 'BakeKeyCount'
+			RangeEndTime = -FLT_MAX;
+			for (const FRichCurve& Curve : Curves)
+			{
+				const TArray<FRichCurveKey>& CurveKeys = Curve.GetConstRefOfKeys();
+				for (const FRichCurveKey& CurveKey : CurveKeys)
+				{
+					if (RangeEndTime < CurveKey.Time)
+					{
+						RangeEndTime = CurveKey.Time;
+					}
+				}
+			}
+			if (RangeEndTime < 0)
+			{
+				return;
+			}
+
+			const double BakeInterval = 1.0 / BakeFrequency;
+			const double SequenceLength = FMath::Max<double>(RangeEndTime - RangeStartTime, BakeInterval);
+			int32 BakeKeyCount = FMath::RoundToInt32(SequenceLength * BakeFrequency) + 1;
+			const FFrameRate ResampleFrameRate(BakeFrequency, 1);
+
+			auto EvaluateCurve = [=](const int32& CurveIndex, double CurrentTime, float DefaultValue)
+			{
+				if (Curves[CurveIndex].IsEmpty())
+				{
+					return DefaultValue;
+				}
+				else
+				{
+					return Curves[CurveIndex].Eval(CurrentTime);
+				}
+			};
+
+			FVector DefaultTranslation = DefaultTransform.GetTranslation();
+			FVector DefaultEuler = DefaultTransform.GetRotation().Euler();
+			FVector DefaultScale = DefaultTransform.GetScale3D();
+
+			double CurrentTime = 0;
+			for (int32 BakeIndex = 0; BakeIndex < BakeKeyCount; BakeIndex++, CurrentTime += BakeInterval)
+			{
+				FVector Translation;
+				Translation.X = EvaluateCurve(0, CurrentTime, DefaultTranslation.X);
+				Translation.Y = EvaluateCurve(1, CurrentTime, DefaultTranslation.Y);
+				Translation.Z = EvaluateCurve(2, CurrentTime, DefaultTranslation.Z);
+
+				FVector RotationEuler;
+				RotationEuler.X = EvaluateCurve(3, CurrentTime, DefaultEuler.X);
+				RotationEuler.Y = EvaluateCurve(4, CurrentTime, DefaultEuler.Y);
+				RotationEuler.Z = EvaluateCurve(5, CurrentTime, DefaultEuler.Z);
+
+				FVector Scale3d;
+				Scale3d.X = EvaluateCurve(6, CurrentTime, DefaultScale.X);
+				Scale3d.Y = EvaluateCurve(7, CurrentTime, DefaultScale.Y);
+				Scale3d.Z = EvaluateCurve(8, CurrentTime, DefaultScale.Z);
+
+				FTransform AnimKeyTransform(FRotator::MakeFromEuler(RotationEuler), Translation, Scale3d);
+
+				Transforms.Add(AnimKeyTransform);
+			}
+
+			AdditionalSupportedType = ToType;
+		}
+#endif
+	}
 }
 
 #if WITH_ENGINE
diff --git a/Engine/Plugins/Interchange/Runtime/Source/Parsers/CommonParser/Public/InterchangeCommonAnimationPayload.h b/Engine/Plugins/Interchange/Runtime/Source/Parsers/CommonParser/Public/InterchangeCommonAnimationPayload.h
--- a/Engine/Plugins/Interchange/Runtime/Source/Parsers/CommonParser/Public/InterchangeCommonAnimationPayload.h
+++ b/Engine/Plugins/Interchange/Runtime/Source/Parsers/CommonParser/Public/InterchangeCommonAnimationPayload.h
@@ -7,6 +7,7 @@
 #include "UObject/Class.h"
 #include "UObject/Object.h"
 #include "UObject/ObjectMacros.h"
+#include "InterchangeAnimationTrackSetNode.h"
 
 #if WITH_ENGINE
 #include "Curves/RichCurve.h"
@@ -16,18 +17,39 @@
 
 namespace UE::Interchange
 {
-	/**
-	 * This payload class is use to get a scene node bake transform payload
-	 * The translator should bake the scene node transform using the bake settings provide by the factory.
-	 */
-	struct INTERCHANGECOMMONPARSER_API FAnimationBakeTransformPayloadData
+	struct INTERCHANGECOMMONPARSER_API FAnimationPayloadData
 	{
+#if WITH_ENGINE
+		//CURVE
+		TArray<FRichCurve> Curves;
+#endif
+		//STEP CURVE
+		TArray<FInterchangeStepCurve> StepCurves;
+
+		//BAKED TRANSFORMS
+		/**
+		 * This payload class part is used to get a scene node bake transform payload
+		 * The translator should bake the scene node transform using the bake settings provided by the factory.
+		 */
 		double BakeFrequency = 30.0;
 		double RangeStartTime = 0.0;
 		double RangeEndTime = 1.0 / BakeFrequency;
 		TArray<FTransform> Transforms;
 
-		void Serialize(FArchive& Ar);
+		//TYPE
+		EInterchangeAnimationPayLoadType Type = EInterchangeAnimationPayLoadType::NONE; //Original
+		EInterchangeAnimationPayLoadType AdditionalSupportedType = EInterchangeAnimationPayLoadType::NONE;
+
+		//
+		FAnimationPayloadData(const EInterchangeAnimationPayLoadType& InType)
+			: Type(InType)
+		{
+		}
+
+		void SerializeBaked(FArchive& Ar);
+
+		//Conversions:
+		void CalculateDataFor(const EInterchangeAnimationPayLoadType& ToType, const FTransform& DefaultTransform = FTransform());
 	};
 }
 
diff --git a/Engine/Plugins/Interchange/Runtime/Source/Parsers/Fbx/Private/FbxAnimation.cpp b/Engine/Plugins/Interchange/Runtime/Source/Parsers/Fbx/Private/FbxAnimation.cpp
--- a/Engine/Plugins/Interchange/Runtime/Source/Parsers/Fbx/Private/FbxAnimation.cpp
+++ b/Engine/Plugins/Interchange/Runtime/Source/Parsers/Fbx/Private/FbxAnimation.cpp
@@ -384,7 +384,7 @@ namespace UE::Interchange::Private
 		DestinationCurve.StringKeyValues = StepCurveValues;
 	}
 
-	bool ImportBakeTransforms(FNodeTransformFetchPayloadData& FetchPayloadData, FAnimationBakeTransformPayloadData& AnimationBakeTransformPayloadData, FbxScene* SDKScene)
+	bool ImportBakeTransforms(FNodeTransformFetchPayloadData& FetchPayloadData, FAnimationPayloadData& AnimationBakeTransformPayloadData, FbxScene* SDKScene)
 	{
 		if (!ensure(!FMath::IsNearlyZero(AnimationBakeTransformPayloadData.BakeFrequency)))
 		{
@@ -540,7 +540,13 @@ namespace UE::Interchange::Private
 		}
 
 		bool bBakeTransform = false;
-		FAnimationBakeTransformPayloadData AnimationBakeTransformPayloadData;
+
+		//At the Moment the Interchange Worker does not transfer the PayLoadType
+		//FBX parser does not need it however because the existing systems identify the appropriate fetch mechanism required in a different manner
+		//For that reason we pass PayLoadType here instead of receiving it and passing it through.
+		//FAnimationPayloadData is used here only for the Baked properties which then gets Serialized via SerializeBaked() function
+		//Resulting in the FAnimatinoPayloadData.Type to be inconsequential here.
+		FAnimationPayloadData AnimationBakeTransformPayloadData(EInterchangeAnimationPayLoadType::BAKED);
 		AnimationBakeTransformPayloadData.BakeFrequency = BakeFrequency;
 		AnimationBakeTransformPayloadData.RangeStartTime = RangeStartTime;
 		AnimationBakeTransformPayloadData.RangeEndTime = RangeEndTime;
@@ -548,7 +554,7 @@ namespace UE::Interchange::Private
 		ImportBakeTransforms(FetchPayloadData, AnimationBakeTransformPayloadData, Parser.GetSDKScene());
 		{
 			FLargeMemoryWriter Ar;
-			AnimationBakeTransformPayloadData.Serialize(Ar);
+			AnimationBakeTransformPayloadData.SerializeBaked(Ar);
 			uint8* ArchiveData = Ar.GetData();
 			int64 ArchiveSize = Ar.TotalSize();
 			TArray64<uint8> Buffer(ArchiveData, ArchiveSize);
@@ -771,7 +777,7 @@ namespace UE::Interchange::Private
 				AnimPayload->NodeTransformFetchPayloadData = FetchPayloadData;
 				PayloadContexts.Add(PayLoadKey, AnimPayload);
 			}
-			SkeletalAnimationTrackNode->SetAnimationPayloadKeyForSceneNodeUid(SceneNode->GetUniqueID(), PayLoadKey);
+			SkeletalAnimationTrackNode->SetAnimationPayloadKeyForSceneNodeUid(SceneNode->GetUniqueID(), PayLoadKey, EInterchangeAnimationPayLoadType::BAKED);
 		}
 
 		return Parameters.IsNodeAnimated;
@@ -853,7 +859,7 @@ namespace UE::Interchange::Private
 			PayloadContexts.Add(PayLoadKey, AnimPayload);
 		}
 
-		SkeletalAnimationTrackNode->SetAnimationPayloadKeyForMorphTargetNodeUid(MorphTargetAnimationBuildingData.MorphTargetNodeUid, PayLoadKey);
+		SkeletalAnimationTrackNode->SetAnimationPayloadKeyForMorphTargetNodeUid(MorphTargetAnimationBuildingData.MorphTargetNodeUid, PayLoadKey, EInterchangeAnimationPayLoadType::MORPHTARGETCURVE);
 	}
 
 	bool FFbxAnimation::IsFbxPropertyTypeSupported(EFbxType PropertyType)
diff --git a/Engine/Plugins/Interchange/Runtime/Source/Parsers/Fbx/Private/FbxScene.cpp b/Engine/Plugins/Interchange/Runtime/Source/Parsers/Fbx/Private/FbxScene.cpp
--- a/Engine/Plugins/Interchange/Runtime/Source/Parsers/Fbx/Private/FbxScene.cpp
+++ b/Engine/Plugins/Interchange/Runtime/Source/Parsers/Fbx/Private/FbxScene.cpp
@@ -523,7 +523,7 @@ namespace UE
 
 						TransformAnimTrackNode->SetCustomActorDependencyUid(*UnrealNode->GetUniqueID());
 
-						TransformAnimTrackNode->SetCustomAnimationPayloadKey(PayloadKey.GetValue());
+						TransformAnimTrackNode->SetCustomAnimationPayloadKey(PayloadKey.GetValue(), EInterchangeAnimationPayLoadType::CURVE);
 
 						TransformAnimTrackNode->SetCustomUsedChannels(UsedChannels);
 
diff --git a/Engine/Plugins/Interchange/Runtime/Source/Pipelines/Private/InterchangeGenericAnimationPipeline.cpp b/Engine/Plugins/Interchange/Runtime/Source/Pipelines/Private/InterchangeGenericAnimationPipeline.cpp
--- a/Engine/Plugins/Interchange/Runtime/Source/Pipelines/Private/InterchangeGenericAnimationPipeline.cpp
+++ b/Engine/Plugins/Interchange/Runtime/Source/Pipelines/Private/InterchangeGenericAnimationPipeline.cpp
@@ -82,11 +82,81 @@ void UInterchangeGenericAnimationPipeline::ExecutePipeline(UInterchangeBaseNodeC
 			TrackSetNodes.Add(Node);
 		});
 
-	for (UInterchangeAnimationTrackSetNode* TrackSetNode : TrackSetNodes)
+	if (CommonMeshesProperties->ForceAllMeshAsType == EInterchangeForceMeshType::IFMT_SkeletalMesh)
 	{
-		if (TrackSetNode)
+		for (UInterchangeAnimationTrackSetNode* TrackSetNode : TrackSetNodes)
 		{
-			CreateAnimationTrackSetFactoryNode(*TrackSetNode);
+			if (!TrackSetNode) continue;
+
+			UInterchangeSkeletalAnimationTrackNode* SkeletalAnimationNode = NewObject< UInterchangeSkeletalAnimationTrackNode >(BaseNodeContainer);
+			FString SkeletalAnimationNodeUid = "\\SkeletalAnimation\\ConvertedFromRigidAnimation\\" + TrackSetNode->GetUniqueID();
+			SkeletalAnimationNode->InitializeNode(SkeletalAnimationNodeUid, TrackSetNode->GetDisplayLabel(), EInterchangeNodeContainerType::TranslatedAsset);
+
+			bool bCustomSkeletonNodeUidSet = false;
+
+			TArray<FString> AnimationTrackUids;
+			TrackSetNode->GetCustomAnimationTrackUids(AnimationTrackUids);
+			float CustomFrameRate;
+			if (!TrackSetNode->GetCustomFrameRate(CustomFrameRate))
+			{
+				CustomFrameRate = 30.0f;
+			}
+			SkeletalAnimationNode->SetCustomAnimationSampleRate(CustomFrameRate);
+			SkeletalAnimationNode->SetCustomAnimationStartTime(0);
+			//stop time will be calculated once we get the curves from the Translators.
+			SkeletalAnimationNode->SetCustomAnimationStopTime(0);
+
+			for (const FString& AnimationTrackUid : AnimationTrackUids)
+			{
+				if (const UInterchangeTransformAnimationTrackNode* TransformTrackNode = Cast<UInterchangeTransformAnimationTrackNode>(BaseNodeContainer->GetNode(AnimationTrackUid)))
+				{
+					FString ActorNodeUid;
+					if (TransformTrackNode->GetCustomActorDependencyUid(ActorNodeUid))
+					{
+						if (const UInterchangeSceneNode* ActorNode = Cast<UInterchangeSceneNode>(BaseNodeContainer->GetNode(ActorNodeUid)))
+						{
+							FInterchangeAnimationPayLoadKey AnimationPayLoadKey;
+							if (TransformTrackNode->GetCustomAnimationPayloadKey(AnimationPayLoadKey))
+							{
+								if (!bCustomSkeletonNodeUidSet)
+								{
+									FString SkeletonRootUid;
+									if (const UInterchangeSceneNode* SceneNode = ActorNode)
+									{
+										FString ParentUid = SceneNode->GetParentUid();
+										while (!ParentUid.Equals(UInterchangeBaseNode::InvalidNodeUid()))
+										{
+											if (const UInterchangeSceneNode* ParentNode = Cast<UInterchangeSceneNode>(BaseNodeContainer->GetNode(ParentUid)))
+											{
+												SkeletonRootUid = ParentUid;
+												ParentUid = ParentNode->GetParentUid();
+											}
+										}
+									}
+
+									SkeletalAnimationNode->SetCustomSkeletonNodeUid(SkeletonRootUid);
+									bCustomSkeletonNodeUidSet = true;
+								}
+
+								//add the payload key:
+								SkeletalAnimationNode->SetAnimationPayloadKeyForSceneNodeUid(ActorNode->GetUniqueID(), AnimationPayLoadKey.UniqueId, AnimationPayLoadKey.Type);
+							}
+						}
+					}
+				}
+			}
+
+			BaseNodeContainer->AddNode(SkeletalAnimationNode);
+		}
+	}
+	else
+	{
+		for (UInterchangeAnimationTrackSetNode* TrackSetNode : TrackSetNodes)
+		{
+			if (TrackSetNode)
+			{
+				CreateAnimationTrackSetFactoryNode(*TrackSetNode);
+			}
 		}
 	}
 
@@ -447,21 +517,17 @@ void UInterchangeGenericAnimationPipeline::CreateAnimSequenceFactoryNode(UInterc
 	}
 
 	{
-		TMap<FString, FString> SceneNodeAnimationPayloadKeys;
-		TrackNode.GetSceneNodeAnimationPayloadKeys(SceneNodeAnimationPayloadKeys);
-		for (const TPair<FString, FString>& Entry : SceneNodeAnimationPayloadKeys)
-		{
-			AnimSequenceFactoryNode->SetAnimationPayloadKeyForSceneNodeUid(Entry.Key, Entry.Value);
-		}
+		TMap<FString, FString> SceneNodeAnimationPayloadKeyUids;
+		TMap<FString, uint8> SceneNodeAnimationPayloadKeyTypes;
+		TrackNode.GetSceneNodeAnimationPayloadKeys(SceneNodeAnimationPayloadKeyUids, SceneNodeAnimationPayloadKeyTypes);
+		AnimSequenceFactoryNode->SetAnimationPayloadKeysForSceneNodeUids(SceneNodeAnimationPayloadKeyUids, SceneNodeAnimationPayloadKeyTypes);
 	}
 
 	{
-		TMap<FString, FString> MorphTargetNodeAnimationPayloads;
-		TrackNode.GetMorphTargetNodeAnimationPayloadKeys(MorphTargetNodeAnimationPayloads);
-		for (const TPair<FString, FString>& Entry : MorphTargetNodeAnimationPayloads)
-		{
-			AnimSequenceFactoryNode->SetAnimationPayloadKeyForMorphTargetNodeUid(Entry.Key, Entry.Value);
-		}
+		TMap<FString, FString> MorphTargetNodeAnimationPayloadKeyUids;
+		TMap<FString, uint8> MorphTargetNodeAnimationPayloadKeyTypes;
+		TrackNode.GetMorphTargetNodeAnimationPayloadKeys(MorphTargetNodeAnimationPayloadKeyUids, MorphTargetNodeAnimationPayloadKeyTypes);
+		AnimSequenceFactoryNode->SetAnimationPayloadKeysForMorphTargetNodeUids(MorphTargetNodeAnimationPayloadKeyUids, MorphTargetNodeAnimationPayloadKeyTypes);
 	}
 
 	UInterchangeUserDefinedAttributesAPI::DuplicateAllUserDefinedAttribute(&TrackNode, AnimSequenceFactoryNode, false);
diff --git a/Engine/Plugins/Interchange/Runtime/Source/Pipelines/Private/InterchangeGenericAssetsPipeline.cpp b/Engine/Plugins/Interchange/Runtime/Source/Pipelines/Private/InterchangeGenericAssetsPipeline.cpp
--- a/Engine/Plugins/Interchange/Runtime/Source/Pipelines/Private/InterchangeGenericAssetsPipeline.cpp
+++ b/Engine/Plugins/Interchange/Runtime/Source/Pipelines/Private/InterchangeGenericAssetsPipeline.cpp
@@ -42,6 +42,7 @@ UInterchangeGenericAssetsPipeline::UInterchangeGenericAssetsPipeline()
 	MeshPipeline->CommonSkeletalMeshesAndAnimationsProperties = CommonSkeletalMeshesAndAnimationsProperties;
 	AnimationPipeline = CreateDefaultSubobject<UInterchangeGenericAnimationPipeline>("AnimationPipeline");
 	AnimationPipeline->CommonSkeletalMeshesAndAnimationsProperties = CommonSkeletalMeshesAndAnimationsProperties;
+	AnimationPipeline->CommonMeshesProperties = CommonMeshesProperties;
 }
 
 void UInterchangeGenericAssetsPipeline::PreDialogCleanup(const FName PipelineStackName)
diff --git a/Engine/Plugins/Interchange/Runtime/Source/Pipelines/Public/InterchangeGenericAnimationPipeline.h b/Engine/Plugins/Interchange/Runtime/Source/Pipelines/Public/InterchangeGenericAnimationPipeline.h
--- a/Engine/Plugins/Interchange/Runtime/Source/Pipelines/Public/InterchangeGenericAnimationPipeline.h
+++ b/Engine/Plugins/Interchange/Runtime/Source/Pipelines/Public/InterchangeGenericAnimationPipeline.h
@@ -41,6 +41,10 @@ public:
 	UPROPERTY(Transient)
 	TWeakObjectPtr<UInterchangeGenericCommonSkeletalMeshesAndAnimationsProperties> CommonSkeletalMeshesAndAnimationsProperties;
 
+	//Common Meshes Properties Settings Pointer
+	UPROPERTY(Transient)
+	TWeakObjectPtr<UInterchangeGenericCommonMeshesProperties> CommonMeshesProperties;
+
 	//////	ANIMATION_CATEGORY Properties //////
 	/** If enable, import all animation assets find in the sources. */
 	UPROPERTY(EditAnywhere, BlueprintReadWrite, Category = "Animations")
