Refactoring the MaterialX Translator, using a hierarchy for LightShaders, SurfaceShaders.
The MaterialXManager is here to allocate only once the different structures that do the mapping between MaterialX categories and Unreal Engine Material Expression, Inputs, etc.
Removing lightcompound test in MTLX_light_rig_test_1, now it's up to the user to override his own lightshaders, no more default point light

diff --git a/Engine/Plugins/Interchange/Runtime/Source/Import/Private/MaterialX/InterchangeMaterialXTranslator.cpp b/Engine/Plugins/Interchange/Runtime/Source/Import/Private/MaterialX/InterchangeMaterialXTranslator.cpp
--- a/Engine/Plugins/Interchange/Runtime/Source/Import/Private/MaterialX/InterchangeMaterialXTranslator.cpp
+++ b/Engine/Plugins/Interchange/Runtime/Source/Import/Private/MaterialX/InterchangeMaterialXTranslator.cpp
@@ -1,27 +1,22 @@
 // Copyright Epic Games, Inc. All Rights Reserved. 
 
 #include "MaterialX/InterchangeMaterialXTranslator.h"
-#if WITH_EDITOR
-#include "MaterialXFormat/Util.h"
-#endif
 #include "InterchangeImportLog.h"
-#include "InterchangeLightNode.h"
 #include "InterchangeManager.h"
-#include "InterchangeMaterialDefinitions.h"
-#include "InterchangeSceneNode.h"
-#include "InterchangeMaterialXDefinitions.h"
-#include "InterchangeTexture2DNode.h"
-#include "InterchangeTextureBlurNode.h"
 #include "Nodes/InterchangeSourceNode.h"
 #include "UObject/GCObjectScopeGuard.h"
 
-#include "Materials/MaterialExpressionNoise.h"
-#include "Materials/MaterialExpressionTransform.h"
-#include "Materials/MaterialExpressionTransformPosition.h"
-#include "Materials/MaterialExpressionVectorNoise.h"
-
 #include "Misc/PackageName.h"
 
+#if WITH_EDITOR
+#include "MaterialX/MaterialXUtils/MaterialXManager.h"
+#include "MaterialXFormat/Util.h"
+#include "MaterialXUtils/MaterialXSurfaceMaterial.h"
+#include "MaterialXUtils/MaterialXSpotLightShader.h"
+#include "MaterialXUtils/MaterialXPointLightShader.h"
+#include "MaterialXUtils/MaterialXDirectionalLightShader.h"
+#endif
+
 #include UE_INLINE_GENERATED_CPP_BY_NAME(InterchangeMaterialXTranslator)
 
 #define LOCTEXT_NAMESPACE "InterchangeMaterialXTranslator"
@@ -33,9 +28,6 @@ static FAutoConsoleVariableRef CCvarInterchangeEnableMaterialXImport(
 	TEXT("Whether MaterialX support is enabled."),
 	ECVF_Default);
 
-
-namespace mx = MaterialX;
-
 namespace UE::Interchange::MaterialX
 {
 	bool AreMaterialFunctionPackagesLoaded()
@@ -71,174 +63,6 @@ namespace UE::Interchange::MaterialX
 	}
 }
 
-UInterchangeMaterialXTranslator::UInterchangeMaterialXTranslator()
-#if WITH_EDITOR
-	:
-InputNamesMaterialX2UE
-{
-	{{TEXT(""),                            TEXT("amplitude")}, TEXT("Amplitude")},
-	{{TEXT(""),                            TEXT("bg")},        TEXT("B")},
-	{{TEXT(""),                            TEXT("center")},    TEXT("Center")},
-	{{TEXT(""),                            TEXT("diminish")},  TEXT("Diminish")},
-	{{TEXT(""),                            TEXT("fg")},        TEXT("A")},
-	{{TEXT(""),                            TEXT("high")},      TEXT("Max")},
-	{{TEXT(""),                            TEXT("in")},        TEXT("Input")},
-	{{TEXT(""),                            TEXT("in1")},       TEXT("A")},
-	{{TEXT(""),                            TEXT("in2")},       TEXT("B")},
-	{{TEXT(""),                            TEXT("in3")},       TEXT("C")},
-	{{TEXT(""),                            TEXT("in4")},       TEXT("D")},
-	{{TEXT(""),                            TEXT("inlow")},     TEXT("InputLow")},
-	{{TEXT(""),                            TEXT("inhigh")},    TEXT("InputHigh")},
-	{{TEXT(""),                            TEXT("lacunarity")},TEXT("Lacunarity")},
-	{{TEXT(""),                            TEXT("low")},       TEXT("Min")},
-	{{TEXT(""),                            TEXT("lumacoeffs")},TEXT("LuminanceFactors")}, // for the moment not yet handled by Interchange, because of the attribute being an advanced pin
-	{{TEXT(""),                            TEXT("mix")},       TEXT("Alpha")},
-	{{TEXT(""),                            TEXT("position")},  TEXT("Position")},
-	{{TEXT(""),                            TEXT("texcoord")},  TEXT("Coordinates")},
-	{{TEXT(""),                            TEXT("octaves")},   TEXT("Octaves")},
-	{{TEXT(""),                            TEXT("outlow")},    TEXT("TargetLow")},
-	{{TEXT(""),                            TEXT("outhigh")},   TEXT("TargetHigh")},
-	{{TEXT(""),                            TEXT("valuel")},    TEXT("A")},
-	{{TEXT(""),                            TEXT("valuer")},    TEXT("B")},
-	{{MaterialX::Category::Atan2,          TEXT("in1")},       TEXT("Y")},
-	{{MaterialX::Category::Atan2,          TEXT("in2")},       TEXT("X")},
-	{{MaterialX::Category::HeightToNormal, TEXT("scale")},     UE::Interchange::Materials::Standard::Nodes::NormalFromHeightMap::Inputs::Intensity.ToString()},
-	{{MaterialX::Category::IfGreater,      TEXT("value1")},    TEXT("A")},
-	{{MaterialX::Category::IfGreater,      TEXT("value2")},    TEXT("B")},
-	{{MaterialX::Category::IfGreater,      TEXT("in1")},       TEXT("AGreaterThanB")},
-	{{MaterialX::Category::IfGreater,      TEXT("in2")},       TEXT("ALessThanB")}, //another input is added for the case equal see ConnectIfGreater
-	{{MaterialX::Category::IfGreaterEq,    TEXT("value1")},    TEXT("A")},
-	{{MaterialX::Category::IfGreaterEq,    TEXT("value2")},    TEXT("B")},
-	{{MaterialX::Category::IfGreaterEq,    TEXT("in1")},       TEXT("AGreaterThanB")}, //another is added for the case equal see ConnectIfGreaterEq
-	{{MaterialX::Category::IfGreaterEq,    TEXT("in2")},       TEXT("ALessThanB")},
-	{{MaterialX::Category::IfEqual,        TEXT("value1")},    TEXT("A")},
-	{{MaterialX::Category::IfEqual,        TEXT("value2")},    TEXT("B")},
-	{{MaterialX::Category::IfEqual,        TEXT("in1")},       TEXT("AEqualsB")},
-	{{MaterialX::Category::IfEqual,        TEXT("in2")},       TEXT("ALessThanB")},  // another input is added for the case greater see ConnectIfEqual
-	{{MaterialX::Category::Inside,         TEXT("in")},        TEXT("A")},			  // Inside is treated as a Multiply node
-	{{MaterialX::Category::Inside,         TEXT("mask")},      TEXT("B")},			  // Inside is treated as a Multiply node
-	{{MaterialX::Category::Invert,         TEXT("amount")},    TEXT("A")},
-	{{MaterialX::Category::Invert,         TEXT("in")},        TEXT("B")},
-	{{MaterialX::Category::Magnitude,      TEXT("in")},        TEXT("A")},
-	{{MaterialX::Category::Mix,            TEXT("fg")},        TEXT("B")},
-	{{MaterialX::Category::Mix,            TEXT("bg")},        TEXT("A")},
-	{{MaterialX::Category::Mix,            TEXT("mix")},		TEXT("Factor")},
-	{{MaterialX::Category::Noise3D,        TEXT("amplitude")}, TEXT("B")},            // The amplitude of the noise is connected to a multiply node
-	{{MaterialX::Category::Noise3D,        TEXT("pivot")},     TEXT("B")},            // The pivot of the noise is connected to a add node
-	{{MaterialX::Category::Normalize,      TEXT("in")},        TEXT("VectorInput")},
-	{{MaterialX::Category::Outside,        TEXT("in")},        TEXT("A")},				// Outside is treated as Multiply node
-	{{MaterialX::Category::Outside,        TEXT("mask")},      TEXT("B")},				// Outside is treated as Multiply node
-	{{MaterialX::Category::Power,          TEXT("in1")},       TEXT("Base")},
-	{{MaterialX::Category::Power,	       TEXT("in2")},       TEXT("Exponent")},
-	{{MaterialX::Category::Rotate2D,       TEXT("amount")},    TEXT("RotationAngle")},
-	{{MaterialX::Category::Rotate2D,       TEXT("in")},        TEXT("Position")},
-	{{MaterialX::Category::Rotate3D,       TEXT("amount")},    TEXT("RotationAngle")},
-	{{MaterialX::Category::Rotate3D,       TEXT("axis")},		TEXT("NormalizedRotationAxis")},
-	{{MaterialX::Category::Rotate3D,       TEXT("in")},        TEXT("Position")},
-	{{MaterialX::Category::Saturate,       TEXT("amount")},    TEXT("Fraction")},
-	{{MaterialX::Category::Smoothstep,     TEXT("in")},        TEXT("Value")},
-},
-NodeNamesMaterialX2UE{
-	// Math nodes
-	{MaterialX::Category::Absval,       TEXT("Abs")},
-	{MaterialX::Category::Add,          TEXT("Add")},
-	{MaterialX::Category::Acos,         TEXT("Arccosine")},
-	{MaterialX::Category::Asin,         TEXT("Arcsine")},
-	{MaterialX::Category::Atan2,        TEXT("Arctangent2")},
-	{MaterialX::Category::Ceil,         TEXT("Ceil")},
-	{MaterialX::Category::Clamp,        TEXT("Clamp")},
-	{MaterialX::Category::Cos,          TEXT("Cosine")},
-	{MaterialX::Category::CrossProduct, TEXT("Crossproduct")},
-	{MaterialX::Category::Divide,       TEXT("Divide")},
-	{MaterialX::Category::DotProduct,   TEXT("Dotproduct")},
-	{MaterialX::Category::Exp,          TEXT("Exponential")},
-	{MaterialX::Category::Floor,        TEXT("Floor")},
-	{MaterialX::Category::Invert,       TEXT("Subtract")},
-	{MaterialX::Category::Ln,           TEXT("Logarithm")},
-	{MaterialX::Category::Magnitude,    TEXT("Length")},
-	{MaterialX::Category::Max,          TEXT("Max")},
-	{MaterialX::Category::Min,          TEXT("Min")},
-	{MaterialX::Category::Modulo,       TEXT("Fmod")},
-	{MaterialX::Category::Multiply,     TEXT("Multiply")},
-	{MaterialX::Category::Normalize,    TEXT("Normalize")},
-	{MaterialX::Category::Power,        TEXT("Power")},
-	{MaterialX::Category::RampLR,       TEXT("RampLeftRight")},
-	{MaterialX::Category::RampTB,       TEXT("RampTopBottom")},
-	{MaterialX::Category::Sign,         TEXT("Sign")},
-	{MaterialX::Category::Sin,          TEXT("Sine")},
-	{MaterialX::Category::SplitLR,      TEXT("SplitLeftRight")},
-	{MaterialX::Category::SplitTB,      TEXT("SplitTopBottom")},
-	{MaterialX::Category::Sqrt,         TEXT("SquareRoot")},
-	{MaterialX::Category::Sub,          TEXT("Subtract")},
-	{MaterialX::Category::Tan,          TEXT("Tangent")},
-	// Compositing nodes
-	{MaterialX::Category::Burn,         TEXT("Burn")},
-	{MaterialX::Category::Difference,   TEXT("Difference")},
-	{MaterialX::Category::Disjointover, TEXT("Disjointover")},
-	{MaterialX::Category::Dodge,        TEXT("Dodge")},
-	{MaterialX::Category::In,           TEXT("In")},
-	{MaterialX::Category::Inside,       TEXT("Multiply")},
-	{MaterialX::Category::Mask,         TEXT("Mask")},
-	{MaterialX::Category::Matte,        TEXT("Matte")},
-	{MaterialX::Category::Minus,        TEXT("Minus")},
-	{MaterialX::Category::Mix,          TEXT("Lerp")},
-	{MaterialX::Category::Out,          TEXT("Out")},
-	{MaterialX::Category::Over,         TEXT("Over")},
-	{MaterialX::Category::Overlay,      TEXT("Overlay")},
-	{MaterialX::Category::Plus,         TEXT("Plus")},
-	{MaterialX::Category::Premult,      TEXT("Premult")},
-	{MaterialX::Category::Screen,       TEXT("Screen")},
-	{MaterialX::Category::Unpremult,    TEXT("Unpremult")},
-	// Channel nodes
-	{MaterialX::Category::Combine2,     TEXT("AppendVector")},
-	{MaterialX::Category::Combine3,     TEXT("Append3Vector")},
-	{MaterialX::Category::Combine4,     TEXT("Append4Vector")},
-	// Procedural3D nodes
-	{MaterialX::Category::Fractal3D,    TEXT("Fractal3D")}, 
-	// Geometric nodes 
-	{MaterialX::Category::GeomColor,	TEXT("VertexColor")},
-	{MaterialX::Category::TexCoord,		TEXT("TextureCoordinate")},
-	// Adjustment nodes,
-	{MaterialX::Category::HsvToRgb,		TEXT("HsvToRgb")},
-	{MaterialX::Category::Luminance,    TEXT("Luminance")},
-	{MaterialX::Category::Remap,        TEXT("Remap")},
-	{MaterialX::Category::RgbToHsv,		TEXT("RgbToHsv")},
-	{MaterialX::Category::Saturate,		TEXT("Desaturation")},
-	{MaterialX::Category::Smoothstep,   TEXT("SmoothStep")}
-},
-UEInputs
-{
-    TEXT("A"),
-	TEXT("Alpha"),
-	TEXT("Amplitude"),
-	TEXT("B"),
-    TEXT("Base"),
-    TEXT("C"),
-	TEXT("Center"),
-    TEXT("D"),
-	TEXT("Exponent"),
-    TEXT("Factor"),
-	TEXT("Fraction"),
-    TEXT("Input"),
-	TEXT("InputLow"),
-	TEXT("InputHigh"),
-	TEXT("Lacunarity"),
-	TEXT("LuminanceFactors"),
-    TEXT("Max"),
-    TEXT("Min"),
-	TEXT("Octaves"),
-	TEXT("Position"),
-	TEXT("TargetLow"),
-	TEXT("TargetHigh"),
-	TEXT("Value"),
-    TEXT("VectorInput"),
-    TEXT("X"),
-    TEXT("Y")
-}
-#endif // WITH_EDITOR
-{
-}
-
 EInterchangeTranslatorType UInterchangeMaterialXTranslator::GetTranslatorType() const
 {
 	return EInterchangeTranslatorType::Scenes;
@@ -267,6 +91,8 @@ bool UInterchangeMaterialXTranslator::Translate(UInterchangeBaseNodeContainer& B
 	bool bIsReferencesValid = true;
 
 #if WITH_EDITOR
+	namespace mx = MaterialX;
+
 	FString Filename = GetSourceData()->GetFilename();
 
 	if(!FPaths::FileExists(Filename))
@@ -307,7 +133,7 @@ bool UInterchangeMaterialXTranslator::Translate(UInterchangeBaseNodeContainer& B
 		}
 
 		//Update the document by initializing and reorganizing the different nodes and subgraphs
-		UpdateDocumentRecursively(Document);
+		FMaterialXBase::UpdateDocumentRecursively(Document);
 
 		for(mx::ElementPtr Elem : Document->traverseTree())
 		{
@@ -324,6 +150,7 @@ bool UInterchangeMaterialXTranslator::Translate(UInterchangeBaseNodeContainer& B
 				bool bIsMaterialShader = Node->getType() == mx::Type::Material;
 				bool bIsLightShader = Node->getType() == mx::Type::LightShader;
 
+				//The entry point is only surfacematerial or lightshader
 				if(bIsMaterialShader || bIsLightShader)
 				{
 					if(!Node->getTypeDef())
@@ -335,40 +162,10 @@ bool UInterchangeMaterialXTranslator::Translate(UInterchangeBaseNodeContainer& B
 						break;
 					}
 
-					//The entry point for materials is only on a surfacematerial node
-					if(bIsMaterialShader && Node->getCategory() == mx::SURFACE_MATERIAL_NODE_STRING)
-					{
-						//standard_surface and surface_unlit are the only surface shaders supported for now
-						bool bHasSupportedSurfaceShader = false;
-
-						for(mx::InputPtr Input : Node->getInputs())
-						{
-							//we only support standard_surface for now
-							mx::NodePtr ConnectedNode = Input->getConnectedNode();
-
-							if(ConnectedNode && ConnectedNode->getCategory() == mx::Category::StandardSurface)
-							{
-								ProcessStandardSurface(BaseNodeContainer, ConnectedNode, Document);
-								bHasSupportedSurfaceShader = true;
-							}
-							else if(ConnectedNode && ConnectedNode->getCategory() == mx::Category::SurfaceUnlit)
-							{
-								ProcessSurfaceUnlit(BaseNodeContainer, ConnectedNode, Document);
-								bHasSupportedSurfaceShader = true;
-							}
-						}
-
-						if(!bHasSupportedSurfaceShader)
-						{
-							UInterchangeResultWarning_Generic* Message = AddMessage<UInterchangeResultWarning_Generic>();
-							Message->Text = FText::Format(LOCTEXT("SurfaceShaderNotFound", "the surfaceshader of <{0}> is not supported"),
-														  FText::FromString(Node->getName().c_str()));
-						}
-					}
-					else if(bIsLightShader)
+					TSharedPtr<FMaterialXBase> ShaderTranslator = FMaterialXManager::GetInstance().GetShaderTranslator(Node->getCategory().c_str(), BaseNodeContainer);
+					if(ShaderTranslator)
 					{
-						bIsReferencesValid = true;
-						ProcessLightShader(BaseNodeContainer, Node, Document);
+						ShaderTranslator->Translate(Node);
 					}
 				}
 			}
@@ -416,1873 +213,4 @@ TOptional<UE::Interchange::FImportImage> UInterchangeMaterialXTranslator::GetTex
 	return TextureTranslator->GetTexturePayloadData(PayloadSourceData, PayloadKey);
 }
 
-#if WITH_EDITOR
-void UInterchangeMaterialXTranslator::ProcessStandardSurface(UInterchangeBaseNodeContainer& NodeContainer, MaterialX::NodePtr StandardSurfaceNode, MaterialX::DocumentPtr Document) const
-{
-	using namespace UE::Interchange::Materials;
-
-	TMap<FString, UInterchangeShaderNode*> NamesToShaderNodes;
-
-	UInterchangeShaderGraphNode* ShaderGraphNode = CreateShaderNode<UInterchangeShaderGraphNode>(StandardSurfaceNode->getName().c_str(), StandardSurface::Name.ToString(), NamesToShaderNodes, NodeContainer);
-
-	auto ConnectNodeOutputToInput = [&](const char* InputName, UInterchangeShaderNode * ShaderGraphNode, const FString & InputShaderName, auto DefaultValue)
-	{
-		mx::InputPtr Input = GetStandardSurfaceInput(StandardSurfaceNode, InputName, Document);
-
-		if(!ConnectNodeGraphOutputToInput(Input, ShaderGraphNode, InputShaderName, NamesToShaderNodes, NodeContainer))
-		{
-			if(!ConnectNodeNameOutputToInput(Input, ShaderGraphNode, InputShaderName, NamesToShaderNodes, NodeContainer))
-			{
-				//the different inputs from a standard surface are either float, linear color, or no attribute at all
-				if constexpr(std::is_same_v<decltype(DefaultValue), float>)
-				{
-					AddFloatAttribute(Input, InputShaderName, ShaderGraphNode, DefaultValue);
-				}
-				else if constexpr(std::is_same_v<decltype(DefaultValue), FLinearColor>)
-				{
-					AddLinearColorAttribute(Input, InputShaderName, ShaderGraphNode, DefaultValue);
-				}
-			}
-		}
-	};
-
-	//Base
-	ConnectNodeOutputToInput(mx::StandardSurface::Input::Base, ShaderGraphNode, StandardSurface::Parameters::Base.ToString(), mx::StandardSurface::DefaultValue::Float::Base);
-	
-	//Base Color
-	ConnectNodeOutputToInput(mx::StandardSurface::Input::BaseColor, ShaderGraphNode, StandardSurface::Parameters::BaseColor.ToString(), mx::StandardSurface::DefaultValue::Color3::BaseColor);
-
-	//Diffuse Roughness
-	ConnectNodeOutputToInput(mx::StandardSurface::Input::DiffuseRoughness, ShaderGraphNode, StandardSurface::Parameters::DiffuseRoughness.ToString(), mx::StandardSurface::DefaultValue::Float::DiffuseRoughness);
-
-	//Specular
-	ConnectNodeOutputToInput(mx::StandardSurface::Input::Specular, ShaderGraphNode, StandardSurface::Parameters::Specular.ToString(), mx::StandardSurface::DefaultValue::Float::Specular);
-	
-	//Specular Roughness
-	ConnectNodeOutputToInput(mx::StandardSurface::Input::SpecularRoughness, ShaderGraphNode, StandardSurface::Parameters::SpecularRoughness.ToString(), mx::StandardSurface::DefaultValue::Float::SpecularRoughness);
-	
-	//Specular IOR
-	ConnectNodeOutputToInput(mx::StandardSurface::Input::SpecularIOR, ShaderGraphNode, StandardSurface::Parameters::SpecularIOR.ToString(), mx::StandardSurface::DefaultValue::Float::SpecularIOR);
-	
-	//Specular Anisotropy
-	ConnectNodeOutputToInput(mx::StandardSurface::Input::SpecularAnisotropy, ShaderGraphNode, StandardSurface::Parameters::SpecularAnisotropy.ToString(), mx::StandardSurface::DefaultValue::Float::SpecularAnisotropy);
-	
-	//Specular Rotation
-	ConnectNodeOutputToInput(mx::StandardSurface::Input::SpecularRotation, ShaderGraphNode, StandardSurface::Parameters::SpecularRotation.ToString(), mx::StandardSurface::DefaultValue::Float::SpecularRotation);
-
-	//Metallic
-	ConnectNodeOutputToInput(mx::StandardSurface::Input::Metalness, ShaderGraphNode, StandardSurface::Parameters::Metalness.ToString(), mx::StandardSurface::DefaultValue::Float::Metalness);
-
-	//Subsurface
-	ConnectNodeOutputToInput(mx::StandardSurface::Input::Subsurface, ShaderGraphNode, StandardSurface::Parameters::Subsurface.ToString(), mx::StandardSurface::DefaultValue::Float::Subsurface);
-	
-	//Subsurface Color
-	ConnectNodeOutputToInput(mx::StandardSurface::Input::SubsurfaceColor, ShaderGraphNode, StandardSurface::Parameters::SubsurfaceColor.ToString(), mx::StandardSurface::DefaultValue::Color3::SubsurfaceColor);
-	
-	//Subsurface Radius
-	ConnectNodeOutputToInput(mx::StandardSurface::Input::SubsurfaceRadius, ShaderGraphNode, StandardSurface::Parameters::SubsurfaceRadius.ToString(), mx::StandardSurface::DefaultValue::Color3::SubsurfaceRadius);
-	
-	//Subsurface Scale
-	ConnectNodeOutputToInput(mx::StandardSurface::Input::SubsurfaceScale, ShaderGraphNode, StandardSurface::Parameters::SubsurfaceScale.ToString(), mx::StandardSurface::DefaultValue::Float::SubsurfaceScale);
-
-	//Sheen
-	ConnectNodeOutputToInput(mx::StandardSurface::Input::Sheen, ShaderGraphNode, StandardSurface::Parameters::Sheen.ToString(), mx::StandardSurface::DefaultValue::Float::Sheen);
-	
-	//Sheen Color
-	ConnectNodeOutputToInput(mx::StandardSurface::Input::SheenColor, ShaderGraphNode, StandardSurface::Parameters::SheenColor.ToString(), mx::StandardSurface::DefaultValue::Color3::SheenColor);
-	
-	//Sheen Roughness
-	ConnectNodeOutputToInput(mx::StandardSurface::Input::SheenRoughness, ShaderGraphNode, StandardSurface::Parameters::SheenRoughness.ToString(), mx::StandardSurface::DefaultValue::Float::SheenRoughness);
-
-	//Coat
-	ConnectNodeOutputToInput(mx::StandardSurface::Input::Coat, ShaderGraphNode, StandardSurface::Parameters::Coat.ToString(), mx::StandardSurface::DefaultValue::Float::Coat);
-	
-	//Coat Color
-	ConnectNodeOutputToInput(mx::StandardSurface::Input::CoatColor, ShaderGraphNode, StandardSurface::Parameters::CoatColor.ToString(), mx::StandardSurface::DefaultValue::Color3::CoatColor);
-	
-	//Coat Roughness
-	ConnectNodeOutputToInput(mx::StandardSurface::Input::CoatRoughness, ShaderGraphNode, StandardSurface::Parameters::CoatRoughness.ToString(), mx::StandardSurface::DefaultValue::Float::CoatRoughness);
-	
-	//Coat Normal: No need to take the default input if there is no CoatNormal input
-	ConnectNodeOutputToInput(mx::StandardSurface::Input::CoatNormal, ShaderGraphNode, StandardSurface::Parameters::CoatNormal.ToString(), nullptr);
-
-	//Thin Film Thickness
-	ConnectNodeOutputToInput(mx::StandardSurface::Input::ThinFilmThickness, ShaderGraphNode, StandardSurface::Parameters::ThinFilmThickness.ToString(), mx::StandardSurface::DefaultValue::Float::ThinFilmThickness);
-
-	//Emission
-	ConnectNodeOutputToInput(mx::StandardSurface::Input::Emission, ShaderGraphNode, StandardSurface::Parameters::Emission.ToString(), mx::StandardSurface::DefaultValue::Float::Emission);
-	
-	//Emission Color
-	ConnectNodeOutputToInput(mx::StandardSurface::Input::EmissionColor, ShaderGraphNode, StandardSurface::Parameters::EmissionColor.ToString(), mx::StandardSurface::DefaultValue::Color3::EmissionColor);
-
-	//Normal: No need to take the default input if there is no Normal input
-	ConnectNodeOutputToInput(mx::StandardSurface::Input::Normal, ShaderGraphNode, StandardSurface::Parameters::Normal.ToString(), nullptr);
-
-	//Tangent: No need to take the default input if there is no Tangent input
-	ConnectNodeOutputToInput(mx::StandardSurface::Input::Tangent, ShaderGraphNode, StandardSurface::Parameters::Tangent.ToString(), nullptr);
-
-	//Transmission
-	ConnectNodeOutputToInput(mx::StandardSurface::Input::Transmission, ShaderGraphNode, StandardSurface::Parameters::Transmission.ToString(), mx::StandardSurface::DefaultValue::Float::Transmission);
-	
-	//Transmission Color
-	ConnectNodeOutputToInput(mx::StandardSurface::Input::TransmissionColor, ShaderGraphNode, StandardSurface::Parameters::TransmissionColor.ToString(), mx::StandardSurface::DefaultValue::Color3::TransmissionColor);
-	
-	//Transmission Depth
-	ConnectNodeOutputToInput(mx::StandardSurface::Input::TransmissionDepth, ShaderGraphNode, StandardSurface::Parameters::TransmissionDepth.ToString(), mx::StandardSurface::DefaultValue::Float::TransmissionDepth);
-	
-	//Transmission Scatter
-	ConnectNodeOutputToInput(mx::StandardSurface::Input::TransmissionScatter, ShaderGraphNode, StandardSurface::Parameters::TransmissionScatter.ToString(), mx::StandardSurface::DefaultValue::Color3::TransmissionScatter);
-	
-	//Transmission Scatter Anisotropy
-	ConnectNodeOutputToInput(mx::StandardSurface::Input::TransmissionScatterAnisotropy, ShaderGraphNode, StandardSurface::Parameters::TransmissionScatterAnisotropy.ToString(), mx::StandardSurface::DefaultValue::Float::TransmissionScatterAnisotropy);
-
-	//Transmission Dispersion
-	ConnectNodeOutputToInput(mx::StandardSurface::Input::TransmissionDispersion, ShaderGraphNode, StandardSurface::Parameters::TransmissionDispersion.ToString(), mx::StandardSurface::DefaultValue::Float::TransmissionDispersion);
-	
-	//Transmission Extra Roughness
-	ConnectNodeOutputToInput(mx::StandardSurface::Input::TransmissionExtraRoughness, ShaderGraphNode, StandardSurface::Parameters::TransmissionExtraRoughness.ToString(), mx::StandardSurface::DefaultValue::Float::TransmissionExtraRoughness);
-}
-
-void UInterchangeMaterialXTranslator::ProcessSurfaceUnlit(UInterchangeBaseNodeContainer& NodeContainer, MaterialX::NodePtr SurfaceUnlitNode, MaterialX::DocumentPtr Document) const
-{
-	using namespace UE::Interchange::Materials;
-
-	TMap<FString, UInterchangeShaderNode*> NamesToShaderNodes;
-
-	UInterchangeShaderGraphNode* ShaderGraphNode = CreateShaderNode<UInterchangeShaderGraphNode>(SurfaceUnlitNode->getName().c_str(), SurfaceUnlit::Name.ToString(), NamesToShaderNodes, NodeContainer);
-
-	auto ConnectNodeOutputToInput = [&](const char* InputName, UInterchangeShaderNode* ShaderGraphNode, const FString& InputShaderName, auto DefaultValue)
-	{
-		mx::InputPtr Input = GetSurfaceUnlitInput(SurfaceUnlitNode, InputName, Document);
-
-		if(!ConnectNodeGraphOutputToInput(Input, ShaderGraphNode, InputShaderName, NamesToShaderNodes, NodeContainer))
-		{
-			if(!ConnectNodeNameOutputToInput(Input, ShaderGraphNode, InputShaderName, NamesToShaderNodes, NodeContainer))
-			{
-				//the different inputs from a surface unlit are either float or linear color
-				if constexpr(std::is_same_v<decltype(DefaultValue), float>)
-				{
-					AddFloatAttribute(Input, InputShaderName, ShaderGraphNode, DefaultValue);
-				}
-				else if constexpr(std::is_same_v<decltype(DefaultValue), FLinearColor>)
-				{
-					AddLinearColorAttribute(Input, InputShaderName, ShaderGraphNode, DefaultValue);
-				}
-			}
-		}
-	};
-
-	//Emission
-	ConnectNodeOutputToInput(mx::SurfaceUnlit::Input::Emission, ShaderGraphNode, SurfaceUnlit::Parameters::Emission.ToString(), mx::SurfaceUnlit::DefaultValue::Float::Emission);
-
-	//Emission Color
-	ConnectNodeOutputToInput(mx::SurfaceUnlit::Input::EmissionColor, ShaderGraphNode, SurfaceUnlit::Parameters::EmissionColor.ToString(), mx::SurfaceUnlit::DefaultValue::Color3::EmissionColor);
-
-	//Opacity
-	ConnectNodeOutputToInput(mx::SurfaceUnlit::Input::Opacity, ShaderGraphNode, SurfaceUnlit::Parameters::Opacity.ToString(), mx::SurfaceUnlit::DefaultValue::Float::Opacity);
-
-	//Transmission
-	ConnectNodeOutputToInput(mx::SurfaceUnlit::Input::Transmission, ShaderGraphNode, SurfaceUnlit::Parameters::Transmission.ToString(), mx::SurfaceUnlit::DefaultValue::Float::Transmission);
-
-	//Transmission Color
-	ConnectNodeOutputToInput(mx::SurfaceUnlit::Input::TransmissionColor, ShaderGraphNode, SurfaceUnlit::Parameters::TransmissionColor.ToString(), mx::SurfaceUnlit::DefaultValue::Color3::TransmissionColor);
-}
-
-void UInterchangeMaterialXTranslator::ProcessLightShader(UInterchangeBaseNodeContainer& NodeContainer, MaterialX::NodePtr LightShaderNode, MaterialX::DocumentPtr Document) const
-{
-	const FString FileName{ FPaths::GetBaseFilename(LightShaderNode->getActiveSourceUri().c_str()) };
-	const FString LightNodeLabel = FString(LightShaderNode->getName().c_str());
-	UInterchangeSceneNode* SceneNode = NewObject<UInterchangeSceneNode>(&NodeContainer);
-	const FString SceneNodeUID = TEXT("\\Light\\") + FileName + TEXT("\\") + FString(LightShaderNode->getName().c_str());
-	SceneNode->InitializeNode(SceneNodeUID, LightNodeLabel, EInterchangeNodeContainerType::TranslatedScene);
-	NodeContainer.AddNode(SceneNode);
-
-	UInterchangeBaseLightNode* LightNode = nullptr;
-	if(LightShaderNode->getCategory() == mx::Category::PointLight)
-	{
-		LightNode = CreatePointLightNode(LightShaderNode, SceneNode, NodeContainer, Document);
-	}
-	else if(LightShaderNode->getCategory() == mx::Category::DirectionalLight)
-	{
-		LightNode = CreateDirectionalLightNode(LightShaderNode, SceneNode, NodeContainer, Document);
-	}
-	else if(LightShaderNode->getCategory() == mx::Category::SpotLight)
-	{
-		LightNode = CreateSpotLightNode(LightShaderNode, SceneNode, NodeContainer, Document);
-	}
-	else // MaterialX has no standardized lights, these 3 are the most common ones though and serves as an example in the format
-	{
-		LightNode = CreatePointLightNode(LightShaderNode, SceneNode, NodeContainer, Document);
-	}
-
-	const FString LightNodeUid = TEXT("\\Light\\") + LightNodeLabel;
-	LightNode->InitializeNode(LightNodeUid, LightNodeLabel, EInterchangeNodeContainerType::TranslatedAsset);
-	NodeContainer.AddNode(LightNode);
-	const FString Uid = LightNode->GetUniqueID();
-	SceneNode->SetCustomAssetInstanceUid(Uid);
-
-	// Color
-	{
-		mx::InputPtr LightColor = LightShaderNode->getInput(mx::Lights::Input::Color);
-		const FLinearColor Color = GetLinearColor(LightColor);
-		LightNode->SetCustomLightColor(Color);
-	}
-
-	// Intensity
-	{
-		mx::InputPtr LightIntensity = LightShaderNode->getInput(mx::Lights::Input::Intensity);
-		LightNode->SetCustomIntensity(mx::fromValueString<float>(LightIntensity->getValueString()));
-	}
-}
-
-static void GetLightDirection(mx::InputPtr DirectionInput, FTransform& Transform)
-{
-	mx::Vector3 Direction = mx::fromValueString<mx::Vector3>(DirectionInput->getValueString());
-
-	//Get rotation to go from UE's default direction of directional light and the direction of the MX directional light node
-	const FVector LightDirection{ Direction[2], Direction[0], Direction[1] };
-	const FVector TransformDirection = Transform.GetUnitAxis(EAxis::X); // it's the default direction of a UE directional light
-	Transform.SetRotation(FQuat::FindBetween(LightDirection, TransformDirection));
-}
-
-static void GetLightPosition(mx::InputPtr PositionInput, FTransform& Transform)
-{
-	const mx::Vector3 Position = mx::fromValueString<mx::Vector3>(PositionInput->getValueString());
-	Transform.SetLocation(FVector{ Position[0] * 100, Position[1] * 100, Position[2] * 100 });
-}
-
-UInterchangeBaseLightNode* UInterchangeMaterialXTranslator::CreateDirectionalLightNode(MaterialX::NodePtr DirectionalLightShaderNode, UInterchangeSceneNode* SceneNode, UInterchangeBaseNodeContainer& NodeContainer, MaterialX::DocumentPtr Document) const
-{
-	UInterchangeDirectionalLightNode* LightNode = NewObject<UInterchangeDirectionalLightNode>(&NodeContainer);
-
-	// Direction
-	{
-		mx::InputPtr DirectionInput = GetDirectionalLightInput(DirectionalLightShaderNode, mx::Lights::DirectionalLight::Input::Direction, Document);
-
-		FTransform Transform;
-		GetLightDirection(DirectionInput, Transform);
-
-		SceneNode->SetCustomLocalTransform(&NodeContainer, Transform);
-	}
-
-	return LightNode;
-}
-
-UInterchangeBaseLightNode* UInterchangeMaterialXTranslator::CreatePointLightNode(MaterialX::NodePtr PointLightShaderNode, UInterchangeSceneNode* SceneNode, UInterchangeBaseNodeContainer& NodeContainer, MaterialX::DocumentPtr Document) const
-{
-	UInterchangePointLightNode* LightNode = NewObject<UInterchangePointLightNode>(&NodeContainer);
-	LightNode->SetCustomIntensityUnits(EInterchangeLightUnits::Candelas);
-
-	// Decay rate
-	{
-		mx::InputPtr DecayRateInput = GetPointLightInput(PointLightShaderNode, mx::Lights::PointLight::Input::DecayRate, Document);
-		const float DecayRate = mx::fromValueString<float>(DecayRateInput->getValueString());
-		LightNode->SetCustomUseInverseSquaredFalloff(false);
-		LightNode->SetCustomLightFalloffExponent(DecayRate);
-	}
-
-	// Position
-	{
-		mx::InputPtr PositionInput = GetPointLightInput(PointLightShaderNode, mx::Lights::PointLight::Input::Position, Document);
-		FTransform Transform;
-		GetLightPosition(PositionInput, Transform);
-		SceneNode->SetCustomLocalTransform(&NodeContainer, Transform);
-	}
-
-	return LightNode;
-}
-
-UInterchangeBaseLightNode* UInterchangeMaterialXTranslator::CreateSpotLightNode(MaterialX::NodePtr SpotLightShaderNode, UInterchangeSceneNode* SceneNode, UInterchangeBaseNodeContainer& NodeContainer, MaterialX::DocumentPtr Document) const
-{
-	UInterchangeSpotLightNode* LightNode = NewObject<UInterchangeSpotLightNode>(&NodeContainer);
-
-	LightNode->SetCustomIntensityUnits(EInterchangeLightUnits::Candelas);
-
-	// Decay rate
-	{
-		mx::InputPtr DecayRateInput = GetSpotLightInput(SpotLightShaderNode, mx::Lights::SpotLight::Input::DecayRate, Document);
-		const float DecayRate = mx::fromValueString<float>(DecayRateInput->getValueString());
-		LightNode->SetCustomUseInverseSquaredFalloff(false);
-		LightNode->SetCustomLightFalloffExponent(DecayRate);
-	}
-
-	{
-		FTransform Transform;
-		// Position
-		{
-			mx::InputPtr PositionInput = GetSpotLightInput(SpotLightShaderNode, mx::Lights::SpotLight::Input::Position, Document);
-			GetLightPosition(PositionInput, Transform);
-		}
-
-		// Direction
-		{
-			mx::InputPtr DirectionInput = GetSpotLightInput(SpotLightShaderNode, mx::Lights::DirectionalLight::Input::Direction, Document);
-			GetLightDirection(DirectionInput, Transform);
-		}
-		SceneNode->SetCustomLocalTransform(&NodeContainer, Transform);
-	}
-
-	// Inner angle
-	{
-		mx::InputPtr InnerAngleInput = GetSpotLightInput(SpotLightShaderNode, mx::Lights::SpotLight::Input::InnerAngle, Document);
-		const float InnerAngle = FMath::RadiansToDegrees(mx::fromValueString<float>(InnerAngleInput->getValueString()));
-		LightNode->SetCustomInnerConeAngle(InnerAngle);
-	}
-
-	// Outer angle
-	{
-		mx::InputPtr OuterAngleInput = GetSpotLightInput(SpotLightShaderNode, mx::Lights::SpotLight::Input::OuterAngle, Document);
-		const float OuterAngle = FMath::RadiansToDegrees(mx::fromValueString<float>(OuterAngleInput->getValueString()));
-		LightNode->SetCustomOuterConeAngle(OuterAngle);
-	}
-
-	return LightNode;
-}
-
-bool UInterchangeMaterialXTranslator::ConnectNodeGraphOutputToInput(MaterialX::InputPtr InputToNodeGraph, UInterchangeShaderNode* ShaderNode, const FString& ParentInputName, TMap<FString, UInterchangeShaderNode*>& NamesToShaderNodes, UInterchangeBaseNodeContainer& NodeContainer) const
-{
-	using namespace UE::Interchange::Materials::Standard::Nodes;
-
-	bool bHasNodeGraph = false;
-
-	if(InputToNodeGraph->hasNodeGraphString())
-	{
-		bHasNodeGraph = true;
-
-		mx::OutputPtr Output = InputToNodeGraph->getConnectedOutput();
-
-		if(!Output)
-		{
-			UInterchangeResultWarning_Generic* Message = AddMessage<UInterchangeResultWarning_Generic>();
-			Message->Text = FText::Format(LOCTEXT("OutputNotFound", "Couldn't find a connected output to ({0})"),
-										  FText::FromString(GetInputName(InputToNodeGraph)));
-			return false;
-		}
-
-		for(mx::Edge Edge : Output->traverseGraph())
-		{
-			ConnectNodeCategoryOutputToInput(Edge, ShaderNode, ParentInputName, NamesToShaderNodes, NodeContainer);
-		}
-	}
-
-	return bHasNodeGraph;
-}
-
-bool UInterchangeMaterialXTranslator::ConnectNodeOutputToInput(MaterialX::NodePtr Node, UInterchangeShaderNode* ParentShaderNode, const FString& InputChannelName, TMap<FString, UInterchangeShaderNode*>& NamesToShaderNodes, UInterchangeBaseNodeContainer& NodeContainer) const
-{
-	bool bIsConnected = false;
-
-	if(const FString* ShaderType = NodeNamesMaterialX2UE.Find(Node->getCategory().c_str()))
-	{
-		UInterchangeShaderNode* OperatorNode = nullptr;
-
-		OperatorNode = CreateShaderNode<UInterchangeShaderNode>(Node->getName().c_str(), *ShaderType, NamesToShaderNodes, NodeContainer);
-
-		for(mx::InputPtr Input : Node->getInputs())
-		{
-			if(const FString* InputNameFound = UEInputs.Find(GetInputName(Input)))
-			{
-				AddAttributeFromValueOrInterface(Input, *InputNameFound, OperatorNode);
-			}
-		}
-
-		bIsConnected = UInterchangeShaderPortsAPI::ConnectDefaultOuputToInput(ParentShaderNode, InputChannelName, OperatorNode->GetUniqueID());
-	}
-
-	return bIsConnected;
-}
-
-void UInterchangeMaterialXTranslator::ConnectNodeCategoryOutputToInput(const MaterialX::Edge& Edge, UInterchangeShaderNode* ShaderNode, const FString& ParentInputName, TMap<FString, UInterchangeShaderNode*>& NamesToShaderNodes, UInterchangeBaseNodeContainer& NodeContainer) const
-{
-	if(mx::NodePtr UpstreamNode = Edge.getUpstreamElement()->asA<mx::Node>())
-	{
-		UInterchangeShaderNode* ParentShaderNode = ShaderNode;
-		FString InputChannelName = ParentInputName;
-
-		//Replace the input's name by the one used in UE
-		RenameNodeInputs(UpstreamNode);
-
-		if(mx::ElementPtr DownstreamElement = Edge.getDownstreamElement())
-		{
-			if(mx::NodePtr DownstreamNode = DownstreamElement->asA<mx::Node>())
-			{
-				if(UInterchangeShaderNode** FoundNode = NamesToShaderNodes.Find(GetAttributeParentName(DownstreamNode)))// DownstreamNode->getName().c_str()))
-				{
-					ParentShaderNode = *FoundNode;
-				}
-
-				if(mx::InputPtr ConnectedInput = Edge.getConnectingElement()->asA<mx::Input>())
-				{
-					InputChannelName = GetInputName(ConnectedInput);
-				}
-			}
-		}
-
-		if(ConnectNodeOutputToInput(UpstreamNode, ParentShaderNode, InputChannelName, NamesToShaderNodes, NodeContainer))
-		{
-			return;
-		}
-		else if(UpstreamNode->getCategory() == mx::Category::Constant)
-		{
-			ConnectConstantInputToOutput(UpstreamNode, ParentShaderNode, InputChannelName, NamesToShaderNodes, NodeContainer);
-		}
-		else if(UpstreamNode->getCategory() == mx::Category::Extract)
-		{
-			ConnectExtractInputToOutput(UpstreamNode, ParentShaderNode, InputChannelName, NamesToShaderNodes, NodeContainer);
-		}
-		//dot means identity, input == output
-		else if(UpstreamNode->getCategory() == mx::Category::Dot || UpstreamNode->getCategory() == mx::Category::NormalMap)
-		{
-			ConnectDotInputToOutput(UpstreamNode, ParentShaderNode, InputChannelName, NamesToShaderNodes, NodeContainer);
-		}
-		else if(UpstreamNode->getCategory() == mx::Category::TransformPoint)
-		{
-			ConnectTransformPositionInputToOutput(UpstreamNode, ParentShaderNode, InputChannelName, NamesToShaderNodes, NodeContainer);
-		}
-		else if(UpstreamNode->getCategory() == mx::Category::TransformVector || UpstreamNode->getCategory() == mx::Category::TransformNormal)
-		{
-			ConnectTransformVectorInputToOutput(UpstreamNode, ParentShaderNode, InputChannelName, NamesToShaderNodes, NodeContainer);
-		}
-		else if(UpstreamNode->getCategory() == mx::Category::Rotate2D)
-		{
-			ConnectRotate2DInputToOutput(UpstreamNode, ParentShaderNode, InputChannelName, NamesToShaderNodes, NodeContainer);
-		}
-		else if(UpstreamNode->getCategory() == mx::Category::Rotate3D)
-		{
-			ConnectRotate3DInputToOutput(UpstreamNode, ParentShaderNode, InputChannelName, NamesToShaderNodes, NodeContainer);
-		}
-		else if(UpstreamNode->getCategory() == mx::Category::Image)
-		{
-			ConnectImageInputToOutput(UpstreamNode, ParentShaderNode, InputChannelName, NamesToShaderNodes, NodeContainer);
-		}
-		else if(UpstreamNode->getCategory() == mx::Category::Convert)
-		{
-			ConnectConvertInputToOutput(UpstreamNode, ParentShaderNode, InputChannelName, NamesToShaderNodes, NodeContainer);
-		}
-		else if(UpstreamNode->getCategory() == mx::Category::IfGreater)
-		{
-			ConnectIfGreaterInputToOutput(UpstreamNode, ParentShaderNode, InputChannelName, NamesToShaderNodes, NodeContainer);
-		}
-		else if(UpstreamNode->getCategory() == mx::Category::IfGreaterEq)
-		{
-			ConnectIfGreaterEqInputToOutput(UpstreamNode, ParentShaderNode, InputChannelName, NamesToShaderNodes, NodeContainer);
-		}
-		else if(UpstreamNode->getCategory() == mx::Category::IfEqual)
-		{
-			ConnectIfEqualInputToOutput(UpstreamNode, ParentShaderNode, InputChannelName, NamesToShaderNodes, NodeContainer);
-		}
-		else if(UpstreamNode->getCategory() == mx::Category::Outside)
-		{
-			ConnectOutsideInputToOutput(UpstreamNode, ParentShaderNode, InputChannelName, NamesToShaderNodes, NodeContainer);
-		}
-		else if(UpstreamNode->getCategory() == mx::Category::Position)
-		{
-			ConnectPositionInputToOutput(UpstreamNode, ParentShaderNode, InputChannelName, NamesToShaderNodes, NodeContainer);
-		}
-		else if(UpstreamNode->getCategory() == mx::Category::Normal)
-		{
-			ConnectNormalInputToOutput(UpstreamNode, ParentShaderNode, InputChannelName, NamesToShaderNodes, NodeContainer);
-		}
-		else if(UpstreamNode->getCategory() == mx::Category::Tangent)
-		{
-			ConnectTangentInputToOutput(UpstreamNode, ParentShaderNode, InputChannelName, NamesToShaderNodes, NodeContainer);
-		}
-		else if(UpstreamNode->getCategory() == mx::Category::Bitangent)
-		{
-			ConnectBitangentInputToOutput(UpstreamNode, ParentShaderNode, InputChannelName, NamesToShaderNodes, NodeContainer);
-		}
-		else if(UpstreamNode->getCategory() == mx::Category::Time)
-		{
-			ConnectTimeInputToOutput(UpstreamNode, ParentShaderNode, InputChannelName, NamesToShaderNodes, NodeContainer);
-		}
-		else if(UpstreamNode->getCategory() == mx::Category::Noise3D)
-		{
-			ConnectNoise3DInputToOutput(UpstreamNode, ParentShaderNode, InputChannelName, NamesToShaderNodes, NodeContainer);
-		}
-		else if(UpstreamNode->getCategory() == mx::Category::CellNoise3D)
-		{
-			ConnectCellNoise3DInputToOutput(UpstreamNode, ParentShaderNode, InputChannelName, NamesToShaderNodes, NodeContainer);
-		}
-		else if(UpstreamNode->getCategory() == mx::Category::WorleyNoise3D)
-		{
-			ConnectWorleyNoise3DInputToOutput(UpstreamNode, ParentShaderNode, InputChannelName, NamesToShaderNodes, NodeContainer);
-		}
-		else if(UpstreamNode->getCategory() == mx::Category::Blur)
-		{
-			ConnectBlurInputToOutput(UpstreamNode, ParentShaderNode, InputChannelName, NamesToShaderNodes, NodeContainer);
-		}
-		else if(UpstreamNode->getCategory() == mx::Category::HeightToNormal)
-		{
-			ConnectHeightToNormalInputToOutput(UpstreamNode, ParentShaderNode, InputChannelName, NamesToShaderNodes, NodeContainer);
-		}
-		else
-		{
-			UInterchangeResultWarning_Generic* Message = AddMessage<UInterchangeResultWarning_Generic>();
-			Message->Text = FText::Format(LOCTEXT("NodeCategoryNotSupported", "<{0}> is not supported yet"),
-										  FText::FromString(UpstreamNode->getCategory().c_str()));
-		}
-	}
-}
-
-bool UInterchangeMaterialXTranslator::ConnectNodeNameOutputToInput(MaterialX::InputPtr InputToConnectedNode, UInterchangeShaderNode* ShaderNode, const FString& ParentInputName, TMap<FString, UInterchangeShaderNode*>& NamesToShaderNodes, UInterchangeBaseNodeContainer& NodeContainer) const
-{
-	mx::NodePtr ConnectedNode = InputToConnectedNode->getConnectedNode();
-
-	if(!ConnectedNode)
-	{
-		return false;
-	}
-
-	UInterchangeShaderNode* ParentShaderNode = ShaderNode;
-	FString InputChannelName = ParentInputName;
-
-	mx::Edge Edge(nullptr, InputToConnectedNode, ConnectedNode);
-
-	TArray<mx::Edge> Stack{ Edge };
-
-	while(!Stack.IsEmpty())
-	{
-		Edge = Stack.Pop();
-		
-		if(Edge.getUpstreamElement())
-		{
-			ConnectNodeCategoryOutputToInput(Edge, ShaderNode, ParentInputName, NamesToShaderNodes, NodeContainer);
-			ConnectedNode = Edge.getUpstreamElement()->asA<mx::Node>();
-			for(mx::InputPtr Input : ConnectedNode->getInputs())
-			{
-				Stack.Emplace(ConnectedNode, Input, Input->getConnectedNode());
-			}
-		}
-	}
-
-	return true;
-}
-
-const FString& UInterchangeMaterialXTranslator::GetMatchedInputName(MaterialX::NodePtr Node, MaterialX::InputPtr Input) const
-{
-	static FString EmptyString;
-
-	if(Input)
-	{
-		const FString NodeCategory{ Node->getCategory().c_str() };
-		const FString InputName{ GetInputName(Input) };
-		TPair<FString, FString> Pair{ NodeCategory, InputName };
-		if(const FString* Result = InputNamesMaterialX2UE.Find(Pair))
-		{
-			return *Result;
-		}
-		else if((Result = InputNamesMaterialX2UE.Find({ EmptyString, InputName })))
-		{
-			return *Result;
-		}
-	}
-
-	return EmptyString;
-}
-
-void UInterchangeMaterialXTranslator::RenameNodeInputs(MaterialX::NodePtr Node) const
-{
-	if(Node)
-	{
-		if(const std::string& IsVisited = Node->getAttribute(mx::Attributes::IsVisited); IsVisited.empty())
-		{
-			Node->setAttribute(mx::Attributes::IsVisited, "true");
-
-			for(mx::InputPtr Input : Node->getInputs())
-			{
-				if(const FString& Name = GetMatchedInputName(Node, Input); !Name.IsEmpty())
-				{
-					RenameInput(Input, TCHAR_TO_UTF8(*Name));
-				}
-			}
-		}
-	}
-}
-
-void UInterchangeMaterialXTranslator::RenameInput(MaterialX::InputPtr Input, const char* NewName) const
-{
-	Input->setAttribute(mx::Attributes::NewName, NewName);
-}
-
-FString UInterchangeMaterialXTranslator::GetInputName(MaterialX::InputPtr Input) const
-{
-	if(Input->hasAttribute(mx::Attributes::NewName))
-	{
-		return Input->getAttribute(mx::Attributes::NewName).c_str();
-	}
-	else
-	{
-		return Input->getName().c_str();
-	}
-}
-
-MaterialX::InputPtr UInterchangeMaterialXTranslator::GetStandardSurfaceInput(MaterialX::NodePtr StandardSurface, const char* InputName, MaterialX::DocumentPtr Document) const
-{
-	mx::InputPtr Input = StandardSurface->getInput(InputName);
-
-	if(!Input)
-	{
-		Input = Document->getNodeDef(mx::NodeDefinition::StandardSurface)->getActiveInput(InputName);
-	}
-
-	return Input;
-}
-
-MaterialX::InputPtr UInterchangeMaterialXTranslator::GetSurfaceUnlitInput(MaterialX::NodePtr SurfaceUnlit, const char* InputName, MaterialX::DocumentPtr Document) const
-{
-	mx::InputPtr Input = SurfaceUnlit->getInput(InputName);
-
-	if(!Input)
-	{
-		Input = Document->getNodeDef(mx::NodeDefinition::SurfaceUnlit)->getActiveInput(InputName);
-	}
-
-	return Input;
-}
-
-MaterialX::InputPtr UInterchangeMaterialXTranslator::GetPointLightInput(MaterialX::NodePtr PointLight, const char* InputName, MaterialX::DocumentPtr Document) const
-{
-	mx::InputPtr Input = PointLight->getInput(InputName);
-
-	if(!Input)
-	{
-		Input = Document->getNodeDef(mx::NodeDefinition::PointLight)->getActiveInput(InputName);
-	}
-
-	return Input;
-}
-
-MaterialX::InputPtr UInterchangeMaterialXTranslator::GetDirectionalLightInput(MaterialX::NodePtr DirectionalLight, const char* InputName, MaterialX::DocumentPtr Document) const
-{
-	mx::InputPtr Input = DirectionalLight->getInput(InputName);
-
-	if(!Input)
-	{
-		Input = Document->getNodeDef(mx::NodeDefinition::DirectionalLight)->getActiveInput(InputName);
-	}
-
-	return Input;
-}
-
-MaterialX::InputPtr UInterchangeMaterialXTranslator::GetSpotLightInput(MaterialX::NodePtr SpotLight, const char* InputName, MaterialX::DocumentPtr Document) const
-{
-	mx::InputPtr Input = SpotLight->getInput(InputName);
-
-	if(!Input)
-	{
-		Input = Document->getNodeDef(mx::NodeDefinition::SpotLight)->getActiveInput(InputName);
-	}
-
-	return Input;
-}
-
-bool UInterchangeMaterialXTranslator::AddAttributeFromValueOrInterface(MaterialX::InputPtr Input, const FString& InputChannelName, UInterchangeShaderNode* ShaderNode) const
-{
-	if(Input)
-	{
-		if(Input->hasValue())
-		{
-			return AddAttribute(Input, InputChannelName, ShaderNode);
-		}
-		else if(Input->hasInterfaceName())
-		{
-			if(mx::InputPtr InputInterface = Input->getInterfaceInput(); InputInterface->hasValue())
-			{
-				return AddAttribute(InputInterface, InputChannelName, ShaderNode);
-			}
-		}
-	}
-
-	return false;
-}
-
-bool UInterchangeMaterialXTranslator::AddAttribute(MaterialX::InputPtr Input, const FString& InputChannelName, UInterchangeShaderNode* ShaderNode) const
-{
-	if(Input)
-	{
-		if(Input->getType() == mx::Type::Float)
-		{
-			return ShaderNode->AddFloatAttribute(UInterchangeShaderPortsAPI::MakeInputValueKey(InputChannelName), mx::fromValueString<float>(Input->getValueString()));
-		}
-		else if(Input->getType() == mx::Type::Integer) //Let's add it a Float attribute, because Interchange doesn't create a scalar if it's an int
-		{
-			return ShaderNode->AddFloatAttribute(UInterchangeShaderPortsAPI::MakeInputValueKey(InputChannelName), mx::fromValueString<int32>(Input->getValueString()));
-		}
-		else if(Input->getType() == mx::Type::Color3 || Input->getType() == mx::Type::Color4)
-		{
-			FLinearColor LinearColor = GetLinearColor(Input);
-			return ShaderNode->AddLinearColorAttribute(UInterchangeShaderPortsAPI::MakeInputValueKey(InputChannelName), LinearColor);
-		}
-		else if(Input->getType() == mx::Type::Vector2)
-		{
-			FLinearColor Vector = GetVector(Input);
-			return ShaderNode->AddVector2Attribute(UInterchangeShaderPortsAPI::MakeInputValueKey(InputChannelName), FVector2f{ Vector.R, Vector.G });
-		}
-		else if(Input->getType() == mx::Type::Vector3 || Input->getType() == mx::Type::Vector4)
-		{
-			FLinearColor Vector = GetVector(Input);
-			return ShaderNode->AddLinearColorAttribute(UInterchangeShaderPortsAPI::MakeInputValueKey(InputChannelName), Vector);
-		}
-	}
-
-	return false;
-}
-
-bool UInterchangeMaterialXTranslator::AddFloatAttribute(MaterialX::InputPtr Input, const FString& InputChannelName, UInterchangeShaderNode* ShaderNode, float DefaultValue) const
-{
-	if(Input)
-	{
-		if(Input->hasValue())
-		{
-			float Value = mx::fromValueString<float>(Input->getValueString());
-
-			if(!FMath::IsNearlyEqual(Value, DefaultValue))
-			{
-				return ShaderNode->AddFloatAttribute(UInterchangeShaderPortsAPI::MakeInputValueKey(InputChannelName), Value);
-			}
-		}
-	}
-
-	return false;
-}
-
-bool UInterchangeMaterialXTranslator::AddLinearColorAttribute(MaterialX::InputPtr Input, const FString& InputChannelName, UInterchangeShaderNode* ShaderNode, const FLinearColor& DefaultValue) const
-{
-	if(Input)
-	{
-		if(Input->hasValue())
-		{
-			const FLinearColor Value = GetLinearColor(Input);
-
-			if(!Value.Equals(DefaultValue))
-			{
-				return ShaderNode->AddLinearColorAttribute(UInterchangeShaderPortsAPI::MakeInputValueKey(InputChannelName), Value);
-			}
-		}
-	}
-
-	return false;
-}
-
-FString UInterchangeMaterialXTranslator::GetFilePrefix(MaterialX::ElementPtr Element) const
-{
-	FString FilePrefix;
-
-	if(Element)
-	{
-		if(Element->hasFilePrefix())
-		{
-			return FString(Element->getFilePrefix().c_str());
-		}
-		else
-		{
-			return GetFilePrefix(Element->getParent());
-		}
-	}
-
-	return FilePrefix;
-}
-
-FString UInterchangeMaterialXTranslator::GetColorSpace(MaterialX::ElementPtr Element) const
-{
-	FString ColorSpace;
-
-	if(Element)
-	{
-		if(Element->hasColorSpace())
-		{
-			return FString(Element->getColorSpace().c_str());
-		}
-		else
-		{
-			return GetColorSpace(Element->getParent());
-		}
-	}
-
-	return ColorSpace;
-}
-
-FLinearColor UInterchangeMaterialXTranslator::GetLinearColor(MaterialX::InputPtr Input) const
-{
-	//we assume that the default color space is linear
-	FLinearColor LinearColor;
-
-	if(Input->getType() == mx::Type::Color3)
-	{
-		mx::Color3 Color = mx::fromValueString<mx::Color3>(Input->getValueString());
-		LinearColor = FLinearColor{ Color[0], Color[1], Color[2] };
-	}
-	else if(Input->getType() == mx::Type::Color4)
-	{
-		mx::Color4 Color = mx::fromValueString<mx::Color4>(Input->getValueString());
-		LinearColor = FLinearColor{ Color[0], Color[1], Color[2], Color[3] };
-	}
-	else
-	{
-		ensureMsgf(false, TEXT("input type can only be color3 or color4"));
-	}
-
-	const FString ColorSpace = GetColorSpace(Input);
-
-	if(ColorSpace.IsEmpty() || ColorSpace == TEXT("lin_rec709") || ColorSpace == TEXT("none"))
-	{
-		;//noop
-	}
-	else if(ColorSpace == TEXT("gamma22"))
-	{
-		LinearColor = FLinearColor::FromPow22Color(FColor(LinearColor.R / 255.f, LinearColor.G / 255.f, LinearColor.B / 255.f, LinearColor.A / 255.f));
-	}
-	else
-	{
-		UInterchangeResultWarning_Generic* Message = AddMessage<UInterchangeResultWarning_Generic>();
-		Message->Text = FText::Format(LOCTEXT("ColorSpaceNotSupported", "<{0}>-<{1}>: Colorspace {2} is not supported yet, falling back to linear"),
-									  FText::FromString(Input->getParent()->getName().c_str()),
-									  FText::FromString(Input->getName().c_str()),
-									  FText::FromString(ColorSpace));
-	}
-
-	return LinearColor;
-}
-
-void UInterchangeMaterialXTranslator::ConnectConstantInputToOutput(MaterialX::NodePtr UpstreamNode, UInterchangeShaderNode* ParentShaderNode, const FString& InputChannelName, TMap<FString, UInterchangeShaderNode*>& NamesToShaderNodes, UInterchangeBaseNodeContainer& NodeContainer) const
-{
-	AddAttributeFromValueOrInterface(UpstreamNode->getInput("value"), InputChannelName, ParentShaderNode);
-}
-
-void UInterchangeMaterialXTranslator::ConnectExtractInputToOutput(MaterialX::NodePtr UpstreamNode, UInterchangeShaderNode* ParentShaderNode, const FString& InputChannelName, TMap<FString, UInterchangeShaderNode*>& NamesToShaderNodes, UInterchangeBaseNodeContainer& NodeContainer) const
-{
-	using namespace UE::Interchange::Materials::Standard::Nodes;
-
-	UInterchangeShaderNode* MaskShaderNode = CreateShaderNode<UInterchangeShaderNode>(UpstreamNode->getName().c_str(), Mask::Name.ToString(), NamesToShaderNodes, NodeContainer);
-
-	if(mx::InputPtr InputIndex = UpstreamNode->getInput("index"))
-	{
-		const int32 Index = mx::fromValueString<int>(InputIndex->getValueString());
-		switch(Index)
-		{
-		case 0: MaskShaderNode->AddBooleanAttribute(Mask::Attributes::R, true); break;
-		case 1: MaskShaderNode->AddBooleanAttribute(Mask::Attributes::G, true); break;
-		case 2: MaskShaderNode->AddBooleanAttribute(Mask::Attributes::B, true); break;
-		case 3: MaskShaderNode->AddBooleanAttribute(Mask::Attributes::A, true); break;
-		default:
-			UInterchangeResultError_Generic* Message = AddMessage<UInterchangeResultError_Generic>();
-			Message->Text = FText::FromString(TEXT("Wrong index number for extract node, values are from [0-3]"));
-			break;
-		}
-	}
-
-	UInterchangeShaderPortsAPI::ConnectDefaultOuputToInput(ParentShaderNode, InputChannelName, MaskShaderNode->GetUniqueID());
-}
-
-void UInterchangeMaterialXTranslator::ConnectDotInputToOutput(MaterialX::NodePtr UpstreamNode, UInterchangeShaderNode* ParentShaderNode, const FString& InputChannelName, TMap<FString, UInterchangeShaderNode*>& NamesToShaderNodes, UInterchangeBaseNodeContainer& NodeContainer) const
-{
-	if(mx::InputPtr Input = UpstreamNode->getInput("in"))
-	{
-		RenameInput(Input, TCHAR_TO_UTF8(*InputChannelName)); //let's take the parent node's input name
-		NamesToShaderNodes.Add(UpstreamNode->getName().c_str(), ParentShaderNode);
-	}
-}
-
-void UInterchangeMaterialXTranslator::ConnectTransformPositionInputToOutput(MaterialX::NodePtr UpstreamNode, UInterchangeShaderNode* ParentShaderNode, const FString& InputChannelName, TMap<FString, UInterchangeShaderNode*>& NamesToShaderNodes, UInterchangeBaseNodeContainer& NodeContainer) const
-{
-	UInterchangeShaderNode* TransformNode = CreateShaderNode<UInterchangeShaderNode>(UpstreamNode->getName().c_str(), TEXT("TransformPosition"), NamesToShaderNodes, NodeContainer);
-	AddAttributeFromValueOrInterface(UpstreamNode->getInput("in"), TEXT("Input"), TransformNode);
-	UInterchangeShaderPortsAPI::ConnectDefaultOuputToInput(ParentShaderNode, InputChannelName, TransformNode->GetUniqueID());
-}
-
-void UInterchangeMaterialXTranslator::ConnectTransformVectorInputToOutput(MaterialX::NodePtr UpstreamNode, UInterchangeShaderNode* ParentShaderNode, const FString& InputChannelName, TMap<FString, UInterchangeShaderNode*>& NamesToShaderNodes, UInterchangeBaseNodeContainer& NodeContainer) const
-{
-	UInterchangeShaderNode* TransformNode = CreateShaderNode<UInterchangeShaderNode>(UpstreamNode->getName().c_str(), TEXT("Transform"), NamesToShaderNodes, NodeContainer);
-	AddAttributeFromValueOrInterface(UpstreamNode->getInput("in"), TEXT("Input"), TransformNode);
-	UInterchangeShaderPortsAPI::ConnectDefaultOuputToInput(ParentShaderNode, InputChannelName, TransformNode->GetUniqueID());
-}
-
-void UInterchangeMaterialXTranslator::ConnectRotate2DInputToOutput(mx::NodePtr UpstreamNode, UInterchangeShaderNode* ParentShaderNode, const FString& InputChannelName, TMap<FString, UInterchangeShaderNode*>& NamesToShaderNodes, UInterchangeBaseNodeContainer& NodeContainer) const
-{
-	UInterchangeShaderNode* Rotate2DNode = CreateShaderNode<UInterchangeShaderNode>(UpstreamNode->getName().c_str(), TEXT("RotateAboutAxis"), NamesToShaderNodes, NodeContainer);
-	Rotate2DNode->AddLinearColorAttribute(UInterchangeShaderPortsAPI::MakeInputValueKey(TEXT("NormalizedRotationAxis")), FLinearColor(0, 0, 1));
-	Rotate2DNode->AddLinearColorAttribute(UInterchangeShaderPortsAPI::MakeInputValueKey(TEXT("PivotPoint")), FLinearColor(0.5, 0.5, 0));
-
-	{
-		mx::InputPtr Input = UpstreamNode->getInput("in");
-		AddAttributeFromValueOrInterface(UpstreamNode->getInput("in"), TEXT("Position"), Rotate2DNode);
-
-		// Rotate2D "in" input takes a vector2 input, but RotateAboutAxis takes a vector3 so we convert
-		if(Input)
-		{
-			mx::NodePtr ConvertNode = CreateNode(UpstreamNode->getParent()->asA<mx::NodeGraph>(),
-												 UpstreamNode->getName().c_str(),
-												 mx::Category::Convert,
-												 { {"in", Input} },
-												 {});
-			ConvertNode->setType(mx::Type::Vector3);
-			Input->setNodeName(ConvertNode->getName());
-		}
-	}
-
-	{
-
-		mx::InputPtr Input = UpstreamNode->getInput("amount");
-		AddAttributeFromValueOrInterface(Input, TEXT("RotationAngle"), Rotate2DNode);
-
-		// we create a Divide node to convert MaterialX angle in degrees to UE's angle which is a value between [0,1]
-		if(Input)
-		{
-			//copy "in" input into "in1" and create "amount" input as "in2" with the value 360 because UE's angle is a value between [0,1]
-			mx::NodePtr DivideNode = CreateNode(UpstreamNode->getParent()->asA<mx::NodeGraph>(),
-												UpstreamNode->getName().c_str(),
-												mx::Category::Divide,
-												{ {"in1", Input} }, // rename input to match <divide> input "in1"
-												{ {"in2", FAttributeValueArray{{"type", "float"}, {"value", "360"}}} });
-
-			// Input now points to the new node
-			Input->setNodeName(DivideNode->getName());
-		}
-	}
-
-	//We insert a mask node because the output has to be a vector2
-	UInterchangeShaderNode * MaskNode = CreateMaskShaderNode(0b1100, UpstreamNode->getName().c_str() + FString(TEXT("_Mask")), NamesToShaderNodes, NodeContainer);
-	UInterchangeShaderPortsAPI::ConnectDefaultOuputToInput(MaskNode, TEXT("Input"), Rotate2DNode->GetUniqueID());
-
-	UInterchangeShaderPortsAPI::ConnectDefaultOuputToInput(ParentShaderNode, InputChannelName, MaskNode->GetUniqueID());
-}
-
-void UInterchangeMaterialXTranslator::ConnectRotate3DInputToOutput(mx::NodePtr UpstreamNode, UInterchangeShaderNode* ParentShaderNode, const FString& InputChannelName, TMap<FString, UInterchangeShaderNode*>& NamesToShaderNodes, UInterchangeBaseNodeContainer& NodeContainer) const
-{
-	UInterchangeShaderNode* Rotate3DNode = CreateShaderNode<UInterchangeShaderNode>(UpstreamNode->getName().c_str(), TEXT("RotateAboutAxis"), NamesToShaderNodes, NodeContainer);
-	Rotate3DNode->AddLinearColorAttribute(UInterchangeShaderPortsAPI::MakeInputValueKey(TEXT("PivotPoint")), FLinearColor(0.5, 0.5, 0));
-
-	AddAttributeFromValueOrInterface(UpstreamNode->getInput("in"), TEXT("Position"), Rotate3DNode);
-	AddAttributeFromValueOrInterface(UpstreamNode->getInput("axis"), TEXT("NormalizedRotationAxis"), Rotate3DNode);
-
-	// we create a Divide node to convert MaterialX angle in degrees to UE's angle which is a value between [0,1]
-	if(mx::InputPtr Input = UpstreamNode->getInput("amount"))
-	{
-		if(!AddAttributeFromValueOrInterface(Input, TEXT("RotationAngle"), Rotate3DNode))
-		{
-			//copy "in" input into "in1" and create "amount" input as "in2" with the value 360 because UE's angle is a value between [0,1]
-			mx::NodePtr NewDivideNode = CreateNode(UpstreamNode->getParent()->asA<mx::NodeGraph>(),
-												   UpstreamNode->getName().c_str(),
-												   mx::Category::Divide,
-												   { {"in1", Input} }, // rename input to match <divide> input "in1"
-												   { {"in2", FAttributeValueArray{{"type", "float"}, {"value", "360"}}} });
-
-				// Input now points to the new node
-				Input->setNodeName(NewDivideNode->getName());
-		}
-	}
-
-	UInterchangeShaderPortsAPI::ConnectDefaultOuputToInput(ParentShaderNode, InputChannelName, Rotate3DNode->GetUniqueID());
-}
-
-void UInterchangeMaterialXTranslator::ConnectImageInputToOutput(MaterialX::NodePtr UpstreamNode, UInterchangeShaderNode* ParentShaderNode, const FString& InputChannelName, TMap<FString, UInterchangeShaderNode*>& NamesToShaderNodes, UInterchangeBaseNodeContainer& NodeContainer) const
-{
-	using namespace UE::Interchange::Materials::Standard::Nodes;
-
-	if(UInterchangeTextureNode* TextureNode = CreateTextureNode<UInterchangeTexture2DNode>(UpstreamNode, NodeContainer))
-	{
-		//By default set the output of a texture to RGB
-		FString OutputChannel{ TEXT("RGB") };
-
-		if(UpstreamNode->getType() == mx::Type::Vector4 || UpstreamNode->getType() == mx::Type::Color4)
-		{
-			OutputChannel = TEXT("RGBA");
-		}
-		else if(UpstreamNode->getType() == mx::Type::Float)
-		{
-			OutputChannel = TEXT("R");
-		}
-
-		UInterchangeShaderNode* TextureShaderNode = CreateShaderNode<UInterchangeShaderNode>(UpstreamNode->getName().c_str(), TextureSample::Name.ToString(), NamesToShaderNodes, NodeContainer);
-		TextureShaderNode->AddStringAttribute(UInterchangeShaderPortsAPI::MakeInputValueKey(TextureSample::Inputs::Texture.ToString()), TextureNode->GetUniqueID());
-		UInterchangeShaderPortsAPI::ConnectOuputToInput(ParentShaderNode, InputChannelName, TextureShaderNode->GetUniqueID(), OutputChannel);
-	}
-	else
-	{
-		AddAttributeFromValueOrInterface(UpstreamNode->getInput(mx::NodeGroup::Texture2D::Inputs::Default), InputChannelName, ParentShaderNode);
-	}
-}
-
-void UInterchangeMaterialXTranslator::ConnectConvertInputToOutput(MaterialX::NodePtr UpstreamNode, UInterchangeShaderNode* ParentShaderNode, const FString& InputChannelName, TMap<FString, UInterchangeShaderNode*>& NamesToShaderNodes, UInterchangeBaseNodeContainer& NodeContainer) const
-{
-	using namespace UE::Interchange::Materials::Standard::Nodes;
-
-	//In case of an upwards conversion, let's do an append, downwards a mask, otherwise leave as is
-	mx::InputPtr Input = UpstreamNode->getInput("in");
-	const std::string& NodeType = UpstreamNode->getType();
-	const std::string& InputType = Input->getType();
-
-	//Ensure that the types are supported
-	const bool bIsNodeTypeSupported =
-		NodeType == mx::Type::Color4 ||
-		NodeType == mx::Type::Color3 ||
-		NodeType == mx::Type::Vector4 ||
-		NodeType == mx::Type::Vector3 ||
-		NodeType == mx::Type::Vector2 ||
-		NodeType == mx::Type::Float ||
-		NodeType == mx::Type::Integer ||
-		NodeType == mx::Type::Boolean;
-
-	const bool bIsInputTypeSupported =
-		InputType == mx::Type::Color4 ||
-		InputType == mx::Type::Color3 ||
-		InputType == mx::Type::Vector4 ||
-		InputType == mx::Type::Vector3 ||
-		InputType == mx::Type::Vector2 ||
-		InputType == mx::Type::Float ||
-		InputType == mx::Type::Integer ||
-		InputType == mx::Type::Boolean;
-
-	if(!bIsNodeTypeSupported || !bIsInputTypeSupported)
-	{
-		UInterchangeResultError_Generic* Message = AddMessage<UInterchangeResultError_Generic>();
-		Message->Text = FText::FromString("<convert> node has non supported types");
-		return;
-	}
-
-	char NodeN = NodeType.back();
-	char InputN = InputType.back();
-
-	constexpr auto Remap = [](auto Value, auto Low1, auto High1, auto Low2, auto High2)
-	{
-		return Low2 + (Value - Low1) * (High2 - Low2) / (High1 - Low1);
-	};
-
-	//Remap NodeN and InputN in a range lower than '2' in case of integer/boolean/float
-	if(NodeN > '4')
-	{
-		NodeN = Remap(NodeN, 'a', 'z', 0, '1');
-	}
-	if(InputN > '4')
-	{
-		InputN = Remap(InputN, 'a', 'z', 0, '1');
-	}
-
-	if(InputN > NodeN) // Mask
-	{
-		UInterchangeShaderNode* MaskShaderNode = nullptr;
-		if(NodeN == '3')
-		{
-			MaskShaderNode = CreateMaskShaderNode(0b1110, UpstreamNode->getName().c_str(), NamesToShaderNodes, NodeContainer);
-		}
-		else if(NodeN == '2')
-		{
-			MaskShaderNode = CreateMaskShaderNode(0b1100, UpstreamNode->getName().c_str(), NamesToShaderNodes, NodeContainer);
-		}
-		else
-		{
-			MaskShaderNode = CreateMaskShaderNode(0b1000, UpstreamNode->getName().c_str(), NamesToShaderNodes, NodeContainer);
-		}
-
-		AddAttributeFromValueOrInterface(Input, TEXT("Input"), MaskShaderNode);
-
-		UInterchangeShaderPortsAPI::ConnectDefaultOuputToInput(ParentShaderNode, InputChannelName, MaskShaderNode->GetUniqueID());
-	}
-	else // Append 
-	{
-		//same as dot, just connect the next output to this parent input
-		RenameInput(Input, TCHAR_TO_UTF8(*InputChannelName)); //let's take the parent node's input name
-		NamesToShaderNodes.Add(UpstreamNode->getName().c_str(), ParentShaderNode);
-
-		//No need to create a node, since the input and the node have the same channels, we just check if there's a value
-		if(NodeN == InputN || (NodeN < '2' && InputN < '2'))
-		{
-			AddAttributeFromValueOrInterface(Input, InputChannelName, ParentShaderNode);
-			return;
-		}
-		std::string Category;
-		TArray<FInputToCopy> InputsToCopy;
-		TArray<FInputToCreate> InputsToCreate;
-
-		// float to N
-		if(InputN < '2')
-		{
-			if(NodeN == '2')
-			{
-				Category = mx::Category::Combine2;
-				InputsToCopy.Add({ "in1", Input });
-				InputsToCopy.Add({ "in2", Input });
-			}
-			else if(NodeN == '3')
-			{
-				Category = mx::Category::Combine3;
-				InputsToCopy.Add({ "in1", Input });
-				InputsToCopy.Add({ "in2", Input });
-				InputsToCopy.Add({ "in3", Input });
-			}
-			else if(NodeN == '4')
-			{
-				Category = mx::Category::Combine4;
-				InputsToCopy.Add({ "in1", Input });
-				InputsToCopy.Add({ "in2", Input });
-				InputsToCopy.Add({ "in3", Input });
-				InputsToCopy.Add({ "in4", Input });
-			}
-		}
-		else if((InputN == '2' && NodeN == '3') || (InputN == '3' && NodeN == '4'))
-		{
-			Category = mx::Category::Combine2;
-			InputsToCopy.Add({ "in1", Input });
-			InputsToCreate.Add({ "in2", FAttributeValueArray{{"type","float"}, {"value","1"}}});
-		}
-
-		//copy "in" input into "in1" and create "amount" input as "in2" with the value 360 because UE's angle is a value between [0,1]
-		mx::NodePtr CombineNode = CreateNode(UpstreamNode->getParent()->asA<mx::NodeGraph>(),
-											 UpstreamNode->getName().c_str(),
-											 Category.c_str(),
-											 InputsToCopy, // rename input to match <divide> input "in1"
-											 InputsToCreate);
-
-		// Input now points to the new node
-		Input->setNodeName(CombineNode->getName());
-	}
-}
-
-void UInterchangeMaterialXTranslator::ConnectIfGreaterInputToOutput(MaterialX::NodePtr UpstreamNode, UInterchangeShaderNode* ParentShaderNode, const FString& InputChannelName, TMap<FString, UInterchangeShaderNode*>& NamesToShaderNodes, UInterchangeBaseNodeContainer& NodeContainer) const
-{
-	UInterchangeShaderNode* NodeIf = CreateShaderNode<UInterchangeShaderNode>(UpstreamNode->getName().c_str(), TEXT("If"), NamesToShaderNodes, NodeContainer);
-	UInterchangeShaderPortsAPI::ConnectDefaultOuputToInput(ParentShaderNode, InputChannelName, NodeIf->GetUniqueID());
-
-	if(mx::InputPtr Input = UpstreamNode->getInput("value1"))
-	{
-		AddAttributeFromValueOrInterface(Input, GetInputName(Input), NodeIf);
-	}
-	
-	if(mx::InputPtr Input = UpstreamNode->getInput("value2"))
-	{
-		AddAttributeFromValueOrInterface(Input, GetInputName(Input), NodeIf);
-	}
-	
-	if(mx::InputPtr Input = UpstreamNode->getInput("in1"))
-	{
-		AddAttributeFromValueOrInterface(Input, GetInputName(Input), NodeIf);
-	}
-
-	//In that case we also need to add an attribute to AEqualsB
-	mx::InputPtr Input = UpstreamNode->getInput("in2");
-	if(Input)
-	{
-		bool bHasValue = AddAttributeFromValueOrInterface(Input, GetInputName(Input), NodeIf);
-		AddAttributeFromValueOrInterface(Input, TEXT("AEqualsB"), NodeIf);
-
-		if(bHasValue)
-		{
-			//Let's add a new input that is a copy of in2 to connect it to the equal input
-			mx::InputPtr Input3 = UpstreamNode->addInput("in3");
-			Input3->copyContentFrom(Input);
-			RenameInput(Input3, "AEqualsB");
-		}
-	}
-}
-
-void UInterchangeMaterialXTranslator::ConnectIfGreaterEqInputToOutput(MaterialX::NodePtr UpstreamNode, UInterchangeShaderNode * ParentShaderNode, const FString & InputChannelName, TMap<FString, UInterchangeShaderNode*>&NamesToShaderNodes, UInterchangeBaseNodeContainer & NodeContainer) const
-{
-	UInterchangeShaderNode* NodeIf = CreateShaderNode<UInterchangeShaderNode>(UpstreamNode->getName().c_str(), TEXT("If"), NamesToShaderNodes, NodeContainer);
-	UInterchangeShaderPortsAPI::ConnectDefaultOuputToInput(ParentShaderNode, InputChannelName, NodeIf->GetUniqueID());
-
-	if(mx::InputPtr Input = UpstreamNode->getInput("value1"))
-	{
-		AddAttributeFromValueOrInterface(Input, GetInputName(Input), NodeIf);
-	}
-
-	if(mx::InputPtr Input = UpstreamNode->getInput("value2"))
-	{
-		AddAttributeFromValueOrInterface(Input, GetInputName(Input), NodeIf);
-	}
-
-	if(mx::InputPtr Input = UpstreamNode->getInput("in2"))
-	{
-		AddAttributeFromValueOrInterface(Input, GetInputName(Input), NodeIf);
-	}
-
-	//In that case we also need to add an attribute to AEqualsB
-	mx::InputPtr Input = UpstreamNode->getInput("in1");
-	if(Input)
-	{
-		bool bHasValue = AddAttributeFromValueOrInterface(Input, GetInputName(Input), NodeIf);
-		AddAttributeFromValueOrInterface(Input, TEXT("AEqualsB"), NodeIf);
-
-		if(bHasValue)
-		{
-			//Let's add a new input that is a copy of in2 to connect it to the equal input
-			mx::InputPtr Input3 = UpstreamNode->addInput("in3");
-			Input3->copyContentFrom(Input);
-			RenameInput(Input3, "AEqualsB");
-		}
-	}
-}
-
-void UInterchangeMaterialXTranslator::ConnectIfEqualInputToOutput(MaterialX::NodePtr UpstreamNode, UInterchangeShaderNode * ParentShaderNode, const FString & InputChannelName, TMap<FString, UInterchangeShaderNode*>&NamesToShaderNodes, UInterchangeBaseNodeContainer & NodeContainer) const
-{
-	UInterchangeShaderNode* NodeIf = CreateShaderNode<UInterchangeShaderNode>(UpstreamNode->getName().c_str(), TEXT("If"), NamesToShaderNodes, NodeContainer);
-	UInterchangeShaderPortsAPI::ConnectDefaultOuputToInput(ParentShaderNode, InputChannelName, NodeIf->GetUniqueID());
-
-	if(mx::InputPtr Input = UpstreamNode->getInput("value1"); Input && Input->hasValue())
-	{
-		AddAttributeFromValueOrInterface(Input, GetInputName(Input), NodeIf);
-	}
-
-	if(mx::InputPtr Input = UpstreamNode->getInput("value2"); Input && Input->hasValue())
-	{
-		AddAttributeFromValueOrInterface(Input, GetInputName(Input), NodeIf);
-	}
-
-	if(mx::InputPtr Input = UpstreamNode->getInput("in1"); Input && Input->hasValue())
-	{
-		AddAttributeFromValueOrInterface(Input, GetInputName(Input), NodeIf);
-	}
-
-	//In that case we also need to add an attribute to AGreaterThanB
-	mx::InputPtr Input = UpstreamNode->getInput("in2");
-	if(Input)
-	{
-		bool bHasValue = AddAttributeFromValueOrInterface(Input, GetInputName(Input), NodeIf);
-		AddAttributeFromValueOrInterface(Input, TEXT("AGreaterThanB"), NodeIf);
-
-		if(bHasValue)
-		{
-			//Let's add a new input that is a copy of in2 to connect it to the equal input
-			mx::InputPtr Input3 = UpstreamNode->addInput("in3");
-			Input3->copyContentFrom(Input);
-			RenameInput(Input3, "AGreaterThanB");
-		}
-	}
-}
-
-void UInterchangeMaterialXTranslator::ConnectOutsideInputToOutput(MaterialX::NodePtr UpstreamNode, UInterchangeShaderNode* ParentShaderNode, const FString& InputChannelName, TMap<FString, UInterchangeShaderNode*>& NamesToShaderNodes, UInterchangeBaseNodeContainer& NodeContainer) const
-{
-	//in * (1 - mask)
-	UInterchangeShaderNode* NodeMultiply = CreateShaderNode<UInterchangeShaderNode>(UpstreamNode->getName().c_str(), TEXT("Multiply"), NamesToShaderNodes, NodeContainer);
-	AddAttributeFromValueOrInterface(UpstreamNode->getInput("in"), TEXT("A"), NodeMultiply);
-	UInterchangeShaderPortsAPI::ConnectDefaultOuputToInput(ParentShaderNode, InputChannelName, NodeMultiply->GetUniqueID());
-
-	UInterchangeShaderNode* NodeOneMinus = CreateShaderNode<UInterchangeShaderNode>(UpstreamNode->getName().c_str() + FString(TEXT("_OneMinus")), TEXT("OneMinus"), NamesToShaderNodes, NodeContainer);
-	AddAttributeFromValueOrInterface(UpstreamNode->getInput("mask"), TEXT("Input"), NodeOneMinus);
-	UInterchangeShaderPortsAPI::ConnectDefaultOuputToInput(NodeMultiply, TEXT("B"), NodeOneMinus->GetUniqueID());
-}
-
-void UInterchangeMaterialXTranslator::ConnectPositionInputToOutput(MaterialX::NodePtr UpstreamNode, UInterchangeShaderNode* ParentShaderNode, const FString& InputChannelName, TMap<FString, UInterchangeShaderNode*>& NamesToShaderNodes, UInterchangeBaseNodeContainer& NodeContainer) const
-{
-	// MaterialX defines the space as: object, model, world
-	// model: The local coordinate space of the geometry, before any local deformations or global transforms have been applied.
-	// object: The local coordinate space of the geometry, after local deformations have been applied, but before any global transforms.
-	// world : The global coordinate space of the geometry, after local deformationsand global transforms have been applied.
-
-	// For the moment we don't have the distinction between model/object, so let's just create an UMaterialExpressionWorldPosition
-	// In case of model/object we need to add a TransformPoint from world to local space
-	UInterchangeShaderNode* PositionNode = CreateShaderNode<UInterchangeShaderNode>(UpstreamNode->getName().c_str() + FString(TEXT("_Position")), TEXT("WorldPosition"), NamesToShaderNodes, NodeContainer);
-
-	// In case of the position node, it seems that the unit is different, we assume for now a conversion from mm -> m, even if UE by default is cm
-	// See standard_surface_marble_solid file, especially on the fractal3d node
-	UInterchangeShaderNode* UnitNode = CreateShaderNode<UInterchangeShaderNode>(UpstreamNode->getName().c_str(), TEXT("Multiply"), NamesToShaderNodes, NodeContainer);
-	UnitNode->AddFloatAttribute(UInterchangeShaderPortsAPI::MakeInputValueKey(TEXT("B")), 0.001f);
-	UInterchangeShaderPortsAPI::ConnectDefaultOuputToInput(UnitNode, TEXT("A"), PositionNode->GetUniqueID());
-
-	UInterchangeShaderNode* NodeToConnectTo = ParentShaderNode;
-	FString InputToConnectTo = InputChannelName;
-
-	std::string Space = "world";
-	mx::InputPtr InputSpace = UpstreamNode->getInput("space");
-
-	if(InputSpace)
-	{
-		Space = InputSpace->getValueString();
-	}
-
-	//the default space defined by the nodedef is "object"
-	if(Space != "world" || !InputSpace)
-	{
-		using namespace UE::Interchange::Materials::Standard::Nodes;
-
-		UInterchangeShaderNode* TransformNode = CreateShaderNode<UInterchangeShaderNode>(UpstreamNode->getName().c_str() + FString(TEXT("_Transform")), TransformPosition::Name.ToString(), NamesToShaderNodes, NodeContainer);
-		NodeToConnectTo = TransformNode;
-		InputToConnectTo = TransformPosition::Inputs::Input.ToString();
-		TransformNode->AddInt32Attribute(TransformPosition::Attributes::TransformSourceType, EMaterialPositionTransformSource::TRANSFORMPOSSOURCE_World);
-		TransformNode->AddInt32Attribute(TransformPosition::Attributes::TransformType, EMaterialPositionTransformSource::TRANSFORMPOSSOURCE_Local);
-		UInterchangeShaderPortsAPI::ConnectDefaultOuputToInput(ParentShaderNode, InputChannelName, TransformNode->GetUniqueID());
-	}
-
-	UInterchangeShaderPortsAPI::ConnectDefaultOuputToInput(NodeToConnectTo, InputToConnectTo, UnitNode->GetUniqueID());
-}
-
-void UInterchangeMaterialXTranslator::ConnectNormalInputToOutput(MaterialX::NodePtr UpstreamNode, UInterchangeShaderNode* ParentShaderNode, const FString& InputChannelName, TMap<FString, UInterchangeShaderNode*>& NamesToShaderNodes, UInterchangeBaseNodeContainer& NodeContainer) const
-{
-	// MaterialX defines the space as: object, model, world
-	// model: The local coordinate space of the geometry, before any local deformations or global transforms have been applied.
-	// object: The local coordinate space of the geometry, after local deformations have been applied, but before any global transforms.
-	// world : The global coordinate space of the geometry, after local deformationsand global transforms have been applied.
-
-	// For the moment we don't have the distinction between model/object, so let's just create an UMaterialExpressionVertexNormalWS
-	// In case of model/object we need to add a TransformVector from world to local space
-	UInterchangeShaderNode* NormalNode = CreateShaderNode<UInterchangeShaderNode>(UpstreamNode->getName().c_str(), TEXT("VertexNormalWS"), NamesToShaderNodes, NodeContainer);
-	UInterchangeShaderNode* NodeToConnectTo = ParentShaderNode;
-	FString InputToConnectTo = InputChannelName;
-
-	std::string Space = "world";
-	mx::InputPtr InputSpace = UpstreamNode->getInput("space");
-
-	if(InputSpace)
-	{
-		Space = InputSpace->getValueString();
-	}
-
-	//the default space defined by the nodedef is "object"
-	if(Space != "world" || !InputSpace)
-	{
-		using namespace UE::Interchange::Materials::Standard::Nodes;
-
-		UInterchangeShaderNode* TransformNode = CreateShaderNode<UInterchangeShaderNode>(UpstreamNode->getName().c_str() + FString(TEXT("_Transform")), TransformVector::Name.ToString(), NamesToShaderNodes, NodeContainer);
-		NodeToConnectTo = TransformNode;
-		InputToConnectTo = TransformVector::Inputs::Input.ToString();
-		TransformNode->AddInt32Attribute(TransformVector::Attributes::TransformSourceType, EMaterialVectorCoordTransformSource::TRANSFORMSOURCE_World);
-		TransformNode->AddInt32Attribute(TransformVector::Attributes::TransformType, EMaterialVectorCoordTransform::TRANSFORM_Local);
-		UInterchangeShaderPortsAPI::ConnectDefaultOuputToInput(ParentShaderNode, InputChannelName, TransformNode->GetUniqueID());
-	}
-
-	UInterchangeShaderPortsAPI::ConnectDefaultOuputToInput(NodeToConnectTo, InputToConnectTo, NormalNode->GetUniqueID());
-}
-
-void UInterchangeMaterialXTranslator::ConnectTangentInputToOutput(MaterialX::NodePtr UpstreamNode, UInterchangeShaderNode* ParentShaderNode, const FString& InputChannelName, TMap<FString, UInterchangeShaderNode*>& NamesToShaderNodes, UInterchangeBaseNodeContainer& NodeContainer) const
-{
-	// MaterialX defines the space as: object, model, world
-	// model: The local coordinate space of the geometry, before any local deformations or global transforms have been applied.
-	// object: The local coordinate space of the geometry, after local deformations have been applied, but before any global transforms.
-	// world : The global coordinate space of the geometry, after local deformationsand global transforms have been applied.
-
-	// For the moment we don't have the distinction between model/object, so let's just create an UMaterialExpressionVertexTangentWS
-	// In case of model/object we need to add a TransformVector from world to local space
-	UInterchangeShaderNode* TangentNode = CreateShaderNode<UInterchangeShaderNode>(UpstreamNode->getName().c_str(), TEXT("VertexTangentWS"), NamesToShaderNodes, NodeContainer);
-	UInterchangeShaderNode* NodeToConnectTo = ParentShaderNode;
-	FString InputToConnectTo = InputChannelName;
-
-	std::string Space = "world";
-	mx::InputPtr InputSpace = UpstreamNode->getInput("space");
-
-	if(InputSpace)
-	{
-		Space = InputSpace->getValueString();
-	}
-
-	//the default space defined by the nodedef is "object"
-	if(Space != "world" || !InputSpace)
-	{
-		using namespace UE::Interchange::Materials::Standard::Nodes;
-
-		UInterchangeShaderNode* TransformNode = CreateShaderNode<UInterchangeShaderNode>(UpstreamNode->getName().c_str() + FString(TEXT("_Transform")), TransformVector::Name.ToString(), NamesToShaderNodes, NodeContainer);
-		NodeToConnectTo = TransformNode;
-		InputToConnectTo = TransformVector::Inputs::Input.ToString();
-		TransformNode->AddInt32Attribute(TransformVector::Attributes::TransformSourceType, EMaterialVectorCoordTransformSource::TRANSFORMSOURCE_World);
-		TransformNode->AddInt32Attribute(TransformVector::Attributes::TransformType, EMaterialVectorCoordTransform::TRANSFORM_Local);
-		UInterchangeShaderPortsAPI::ConnectDefaultOuputToInput(ParentShaderNode, InputChannelName, TransformNode->GetUniqueID());
-	}
-
-	UInterchangeShaderPortsAPI::ConnectDefaultOuputToInput(NodeToConnectTo, InputToConnectTo, TangentNode->GetUniqueID());
-}
-
-void UInterchangeMaterialXTranslator::ConnectBitangentInputToOutput(MaterialX::NodePtr UpstreamNode, UInterchangeShaderNode* ParentShaderNode, const FString& InputChannelName, TMap<FString, UInterchangeShaderNode*>& NamesToShaderNodes, UInterchangeBaseNodeContainer& NodeContainer) const
-{
-	// MaterialX defines the space as: object, model, world
-	// model: The local coordinate space of the geometry, before any local deformations or global transforms have been applied.
-	// object: The local coordinate space of the geometry, after local deformations have been applied, but before any global transforms.
-	// world : The global coordinate space of the geometry, after local deformationsand global transforms have been applied.
-
-	// For the moment we don't have the distinction between model/object, so let's just do the cross product between the normal and the tangent
-	// In case of model/object we need to add a TransformVector from world to local space
-	UInterchangeShaderNode* NormalNode = CreateShaderNode<UInterchangeShaderNode>(UpstreamNode->getName().c_str() + FString(TEXT("_Normal")), TEXT("VertexNormalWS"), NamesToShaderNodes, NodeContainer);
-	UInterchangeShaderNode* TangentNode = CreateShaderNode<UInterchangeShaderNode>(UpstreamNode->getName().c_str() + FString(TEXT("_Tangent")), TEXT("VertexTangentWS"), NamesToShaderNodes, NodeContainer);
-	UInterchangeShaderNode* BitangentNode = CreateShaderNode<UInterchangeShaderNode>(UpstreamNode->getName().c_str(), TEXT("CrossProduct"), NamesToShaderNodes, NodeContainer);
-	
-	UInterchangeShaderPortsAPI::ConnectDefaultOuputToInput(BitangentNode, TEXT("A"), NormalNode->GetUniqueID());
-	UInterchangeShaderPortsAPI::ConnectDefaultOuputToInput(BitangentNode, TEXT("B"), TangentNode->GetUniqueID());
-
-
-	UInterchangeShaderNode* NodeToConnectTo = ParentShaderNode;
-	FString InputToConnectTo = InputChannelName;
-
-	std::string Space = "world";
-	mx::InputPtr InputSpace = UpstreamNode->getInput("space");
-
-	if(InputSpace)
-	{
-		Space = InputSpace->getValueString();
-	}
-
-	//the default space defined by the nodedef is "object"
-	if(Space != "world" || !InputSpace)
-	{
-		using namespace UE::Interchange::Materials::Standard::Nodes;
-
-		UInterchangeShaderNode* TransformNode = CreateShaderNode<UInterchangeShaderNode>(UpstreamNode->getName().c_str() + FString(TEXT("_Transform")), TransformVector::Name.ToString(), NamesToShaderNodes, NodeContainer);
-		NodeToConnectTo = TransformNode;
-		InputToConnectTo = TransformVector::Inputs::Input.ToString();
-		TransformNode->AddInt32Attribute(TransformVector::Attributes::TransformSourceType, EMaterialVectorCoordTransformSource::TRANSFORMSOURCE_World);
-		TransformNode->AddInt32Attribute(TransformVector::Attributes::TransformType, EMaterialVectorCoordTransform::TRANSFORM_Local);
-		UInterchangeShaderPortsAPI::ConnectDefaultOuputToInput(ParentShaderNode, InputChannelName, TransformNode->GetUniqueID());
-	}
-
-	UInterchangeShaderPortsAPI::ConnectDefaultOuputToInput(NodeToConnectTo, InputToConnectTo, BitangentNode->GetUniqueID());
-}
-
-void UInterchangeMaterialXTranslator::ConnectTimeInputToOutput(MaterialX::NodePtr UpstreamNode, UInterchangeShaderNode* ParentShaderNode, const FString& InputChannelName, TMap<FString, UInterchangeShaderNode*>& NamesToShaderNodes, UInterchangeBaseNodeContainer& NodeContainer) const
-{
-	using namespace UE::Interchange::Materials::Standard::Nodes;
-	UInterchangeShaderNode* TimeNode = CreateShaderNode<UInterchangeShaderNode>(UpstreamNode->getName().c_str(), TEXT("Time"), NamesToShaderNodes, NodeContainer);
-	TimeNode->AddBooleanAttribute(Time::Attributes::OverridePeriod, true);
-
-	float FPS; 
-	mx::InputPtr Input = UpstreamNode->getInput("fps");
-
-	//Take the default value from the node definition
-	if(!Input)
-	{
-		Input = UpstreamNode->getNodeDef()->getInput("fps");
-	}
-
-	FPS = mx::fromValueString<float>(Input->getValueString());
-
-	//UE is a period
-	TimeNode->AddFloatAttribute(Time::Attributes::Period, 1.f / FPS);
-
-	UInterchangeShaderPortsAPI::ConnectDefaultOuputToInput(ParentShaderNode, InputChannelName, TimeNode->GetUniqueID());
-}
-
-void UInterchangeMaterialXTranslator::ConnectNoise3DInputToOutput(MaterialX::NodePtr UpstreamNode, UInterchangeShaderNode* ParentShaderNode, const FString& InputChannelName, TMap<FString, UInterchangeShaderNode*>& NamesToShaderNodes, UInterchangeBaseNodeContainer& NodeContainer) const
-{
-	using namespace UE::Interchange::Materials::Standard::Nodes;
-
-	// MaterialX defines the Noise3d as Perlin Noise which is multiplied by the Amplitude then Added to Pivot
-	UInterchangeShaderNode* NoiseNode = CreateShaderNode<UInterchangeShaderNode>(UpstreamNode->getName().c_str(), Noise::Name.ToString(), NamesToShaderNodes, NodeContainer);
-	NoiseNode->AddInt32Attribute(Noise::Attributes::Function, ENoiseFunction::NOISEFUNCTION_GradientTex);
-	NoiseNode->AddBooleanAttribute(Noise::Attributes::Turbulence, false);
-	NoiseNode->AddFloatAttribute(Noise::Attributes::OutputMin, 0);
-
-	UInterchangeShaderNode* NodeToConnect = NoiseNode;
-
-	// Multiply Node
-	auto ConnectNodeToInput = [&](mx::InputPtr Input, UInterchangeShaderNode * NodeToConnectTo, const FString & ShaderType, int32 IndexAttrib) -> UInterchangeShaderNode*
-	{
-		if(!Input)
-		{
-			return nullptr;
-		}
-
-		const FString ShaderNodeName = UpstreamNode->getName().c_str() + FString{ TEXT("_") } + ShaderType;
-		UInterchangeShaderNode* ShaderNode = CreateShaderNode<UInterchangeShaderNode>(ShaderNodeName, ShaderType, NamesToShaderNodes, NodeContainer);
-
-		UInterchangeShaderPortsAPI::ConnectDefaultOuputToInput(ShaderNode, TEXT("A"), NodeToConnectTo->GetUniqueID());
-
-		// Connect the amplitude node to the shader node not the noise
-		// it will be handle during the upstream-downstream connection phase
-		// The index is here for the unicity of the attribute
-		UpstreamNode->setAttribute(mx::Attributes::ParentName + std::to_string(IndexAttrib), TCHAR_TO_ANSI(*ShaderNodeName));
-		AddAttributeFromValueOrInterface(Input, TEXT("B"), ShaderNode);
-
-		return ShaderNode;
-	};
-
-	if(UInterchangeShaderNode* MultiplyNode = ConnectNodeToInput(UpstreamNode->getInput("amplitude"), NoiseNode, TEXT("Multiply"), 0))
-	{
-		NodeToConnect = MultiplyNode;
-	}
-
-
-	if(UInterchangeShaderNode* AddNode = ConnectNodeToInput(UpstreamNode->getInput("pivot"), NodeToConnect, TEXT("Add"), 1))
-	{
-		NodeToConnect = AddNode;
-	}
-
-	UInterchangeShaderPortsAPI::ConnectDefaultOuputToInput(ParentShaderNode, InputChannelName, NodeToConnect->GetUniqueID());
-}
-
-void UInterchangeMaterialXTranslator::ConnectCellNoise3DInputToOutput(MaterialX::NodePtr UpstreamNode, UInterchangeShaderNode* ParentShaderNode, const FString& InputChannelName, TMap<FString, UInterchangeShaderNode*>& NamesToShaderNodes, UInterchangeBaseNodeContainer& NodeContainer) const
-{
-	using namespace UE::Interchange::Materials::Standard::Nodes;
-
-	// Let's use a vector noise for this one, the only one that is close to MaterialX implementation
-	UInterchangeShaderNode* NoiseNode = CreateShaderNode<UInterchangeShaderNode>(UpstreamNode->getName().c_str(), VectorNoise::Name.ToString(), NamesToShaderNodes, NodeContainer);
-	NoiseNode->AddInt32Attribute(VectorNoise::Attributes::Function, EVectorNoiseFunction::VNF_CellnoiseALU);
-
-	UInterchangeShaderPortsAPI::ConnectDefaultOuputToInput(ParentShaderNode, InputChannelName, NoiseNode->GetUniqueID());
-}
-
-void UInterchangeMaterialXTranslator::ConnectWorleyNoise3DInputToOutput(MaterialX::NodePtr UpstreamNode, UInterchangeShaderNode* ParentShaderNode, const FString& InputChannelName, TMap<FString, UInterchangeShaderNode*>& NamesToShaderNodes, UInterchangeBaseNodeContainer& NodeContainer) const
-{
-	using namespace UE::Interchange::Materials::Standard::Nodes;
-
-	//Also called Voronoi, the implementation is a bit different in UE, especially we don't have access to the jitter
-	UInterchangeShaderNode* NoiseNode = CreateShaderNode<UInterchangeShaderNode>(UpstreamNode->getName().c_str(), Noise::Name.ToString(), NamesToShaderNodes, NodeContainer);
-	NoiseNode->AddInt32Attribute(Noise::Attributes::Function, ENoiseFunction::NOISEFUNCTION_VoronoiALU);
-
-	UInterchangeShaderPortsAPI::ConnectDefaultOuputToInput(ParentShaderNode, InputChannelName, NoiseNode->GetUniqueID());
-}
-
-void UInterchangeMaterialXTranslator::ConnectHeightToNormalInputToOutput(MaterialX::NodePtr UpstreamNode, UInterchangeShaderNode* ParentShaderNode, const FString& InputChannelName, TMap<FString, UInterchangeShaderNode*>& NamesToShaderNodes, UInterchangeBaseNodeContainer& NodeContainer) const
-{
-	if(mx::InputPtr Input = UpstreamNode->getInput("in"))
-	{
-		// Image node will become this node
-		if(mx::NodePtr ConnectedNode = Input->getConnectedNode();
-		   ConnectedNode && ConnectedNode->getCategory() == mx::Category::Image)
-		{
-			//we need to copy the content of the image node to this node
-			UpstreamNode->copyContentFrom(ConnectedNode);
-
-			RenameNodeInputs(UpstreamNode);
-
-			//the copy overwrite every attribute of the node, so we need to get them back, essentially the type and the renaming
-			// the output is always a vec3
-			UpstreamNode->setType(mx::Type::Vector3);
-			
-			mx::NodeGraphPtr Graph = UpstreamNode->getParent()->asA<mx::NodeGraph>();
-			Graph->removeNode(ConnectedNode->getName());
-
-			using namespace UE::Interchange::Materials::Standard::Nodes;
-
-			if(UInterchangeTextureNode* TextureNode = CreateTextureNode<UInterchangeTexture2DNode>(UpstreamNode, NodeContainer))
-			{
-				UInterchangeShaderNode* HeightMapNode = CreateShaderNode<UInterchangeShaderNode>(UpstreamNode->getName().c_str(), NormalFromHeightMap::Name.ToString(), NamesToShaderNodes, NodeContainer);
-				UInterchangeShaderPortsAPI::ConnectDefaultOuputToInput(ParentShaderNode, InputChannelName, HeightMapNode->GetUniqueID());
-
-				const FString TextureNodeName = UpstreamNode->getName().c_str() + FString{ "_texture" };
-				UInterchangeShaderNode* TextureShaderNode = CreateShaderNode<UInterchangeShaderNode>(TextureNodeName, TextureObject::Name.ToString(), NamesToShaderNodes, NodeContainer);
-				TextureShaderNode->AddStringAttribute(UInterchangeShaderPortsAPI::MakeInputValueKey(TextureObject::Inputs::Texture.ToString()), TextureNode->GetUniqueID());
-				UInterchangeShaderPortsAPI::ConnectDefaultOuputToInput(HeightMapNode, NormalFromHeightMap::Inputs::HeightMap.ToString(), TextureShaderNode->GetUniqueID());
-
-				AddAttributeFromValueOrInterface(UpstreamNode->getInput("scale"), NormalFromHeightMap::Inputs::Intensity.ToString(), HeightMapNode);
-			}
-			else
-			{
-				AddAttributeFromValueOrInterface(UpstreamNode->getInput(mx::NodeGroup::Texture2D::Inputs::Default), InputChannelName, ParentShaderNode);
-			}
-		}
-		else
-		{
-			// For the moment it doesn't make sense to plug a value to it, so let's plug directly the child to the parent, in the future we could implement a Sobel and handle a multi output
-			RenameInput(Input, TCHAR_TO_UTF8(*InputChannelName)); //let's take the parent node's input name
-			NamesToShaderNodes.Add(UpstreamNode->getName().c_str(), ParentShaderNode);
-		}
-	}
-}
-
-void UInterchangeMaterialXTranslator::ConnectBlurInputToOutput(MaterialX::NodePtr UpstreamNode, UInterchangeShaderNode* ParentShaderNode, const FString& InputChannelName, TMap<FString, UInterchangeShaderNode*>& NamesToShaderNodes, UInterchangeBaseNodeContainer& NodeContainer) const
-{
-	if(mx::InputPtr Input = UpstreamNode->getInput("in"))
-	{
-		// Image node will become this node
-		if(mx::NodePtr ConnectedNode = Input->getConnectedNode();
-		   ConnectedNode && ConnectedNode->getCategory() == mx::Category::Image)
-		{
-			std::string NodeType = UpstreamNode->getType();
-
-			//we need to copy the content of the image node to this node
-			UpstreamNode->copyContentFrom(ConnectedNode);
-
-			RenameNodeInputs(UpstreamNode);
-
-			//the copy overwrites every attribute of the node, so we need to get them back, essentially the type and the renaming
-			UpstreamNode->setType(NodeType);
-
-			mx::NodeGraphPtr Graph = UpstreamNode->getParent()->asA<mx::NodeGraph>();
-			Graph->removeNode(ConnectedNode->getName());
-
-			using namespace UE::Interchange::Materials::Standard::Nodes;
-
-			if(UInterchangeTextureNode* TextureNode = CreateTextureNode<UInterchangeTextureBlurNode>(UpstreamNode, NodeContainer))
-			{
-				FString OutputChannel{ TEXT("RGB") };
-
-				if(NodeType == mx::Type::Vector4 || NodeType == mx::Type::Color4)
-				{
-					OutputChannel = TEXT("RGBA");
-				}
-				else if(NodeType == mx::Type::Float)
-				{
-					OutputChannel = TEXT("R");
-				}
-
-				UInterchangeShaderNode* TextureShaderNode = CreateShaderNode<UInterchangeShaderNode>(UpstreamNode->getName().c_str(), TextureSampleBlur::Name.ToString(), NamesToShaderNodes, NodeContainer);
-				TextureShaderNode->AddStringAttribute(UInterchangeShaderPortsAPI::MakeInputValueKey(TextureSampleBlur::Inputs::Texture.ToString()), TextureNode->GetUniqueID());
-				UInterchangeShaderPortsAPI::ConnectOuputToInput(ParentShaderNode, InputChannelName, TextureShaderNode->GetUniqueID(), OutputChannel);
-
-				if(mx::InputPtr InputKernel = UpstreamNode->getInput("filtertype"))
-				{
-					// By default TextureSampleBox uses gaussian filter
-					if(InputKernel->getValueString() == "box")
-					{
-						TextureShaderNode->AddInt32Attribute(TextureSampleBlur::Attributes::Filter, 0);
-					}
-				}
-
-				if(mx::InputPtr InputKernel = UpstreamNode->getInput("size"))
-				{
-					float KernelSize = mx::fromValueString<float>(InputKernel->getValueString());
-					TextureShaderNode->AddFloatAttribute(TextureSampleBlur::Attributes::KernelSize, KernelSize);
-				}
-			}
-			else
-			{
-				AddAttributeFromValueOrInterface(UpstreamNode->getInput(mx::NodeGroup::Texture2D::Inputs::Default), InputChannelName, ParentShaderNode);
-			}
-		}
-		else
-		{
-			// For a blur it doesn't make sense if there's no image input
-			RenameInput(Input, TCHAR_TO_UTF8(*InputChannelName)); //let's take the parent node's input name
-			NamesToShaderNodes.Add(UpstreamNode->getName().c_str(), ParentShaderNode);
-		}
-	}
-}
-
-UInterchangeShaderNode* UInterchangeMaterialXTranslator::CreateMaskShaderNode(uint8 RGBA, const FString& NodeName, TMap<FString, UInterchangeShaderNode*>& NamesToShaderNodes, UInterchangeBaseNodeContainer& NodeContainer) const
-{
-	bool bR = (0b1000 & RGBA) >> 3;
-	bool bG = (0b0100 & RGBA) >> 2;
-	bool bB = (0b0010 & RGBA) >> 1;
-	bool bA = (0b0001 & RGBA) >> 0;
-	using namespace UE::Interchange::Materials::Standard::Nodes;
-	UInterchangeShaderNode* MaskShaderNode = CreateShaderNode<UInterchangeShaderNode>(NodeName, Mask::Name.ToString(), NamesToShaderNodes, NodeContainer);
-	MaskShaderNode->AddBooleanAttribute(Mask::Attributes::R, bR);
-	MaskShaderNode->AddBooleanAttribute(Mask::Attributes::G, bG);
-	MaskShaderNode->AddBooleanAttribute(Mask::Attributes::B, bB);
-	MaskShaderNode->AddBooleanAttribute(Mask::Attributes::A, bA);
-
-	return MaskShaderNode;
-}
-
-FLinearColor UInterchangeMaterialXTranslator::GetVector(MaterialX::InputPtr Input) const
-{
-	FLinearColor LinearColor;
-
-	if(Input->getType() == mx::Type::Vector2)
-	{
-		mx::Vector2 Color = mx::fromValueString<mx::Vector2>(Input->getValueString());
-		LinearColor = FLinearColor{ Color[0], Color[1], 0 };
-	}
-	else if(Input->getType() == mx::Type::Vector3)
-	{
-		mx::Vector3 Color = mx::fromValueString<mx::Vector3>(Input->getValueString());
-		LinearColor = FLinearColor{ Color[0], Color[1], Color[2] };
-	}
-	else if(Input->getType() == mx::Type::Vector4)
-	{
-		mx::Vector4 Color = mx::fromValueString<mx::Vector4>(Input->getValueString());
-		LinearColor = FLinearColor{ Color[0], Color[1], Color[2], Color[3] };
-	}
-	else
-	{
-		ensureMsgf(false, TEXT("input type can only be a vectorN"));
-	}
-
-	return LinearColor;
-}
-
-MaterialX::NodePtr UInterchangeMaterialXTranslator::CreateNode(MaterialX::ElementPtr NodeGraph, const char* NodeName, const char* Category, TArray<FInputToCopy> InputsToCopy, TArray<FInputToCreate> InputsToCreate) const
-{
-	std::string UniqueNodeName{ NodeName + std::string{"_"} + Category };
-
-	mx::NodePtr Node = NodeGraph->getChildOfType<mx::Node>(UniqueNodeName);
-
-	if(!Node)
-	{
-		Node = NodeGraph->addChildOfCategory(Category, UniqueNodeName)->asA<mx::Node>();
-
-		for(const FInputToCopy& Pair : InputsToCopy)
-		{
-			const char* NewInputName = Pair.Get<0>();
-			mx::InputPtr Input = Pair.Get<1>();
-
-			mx::InputPtr InputCopy = Node->addInput();
-			InputCopy->copyContentFrom(Input);
-			//remove the attribute NewName since when the input will be copied, the name will be the one from the spec (before a renaming)
-			InputCopy->removeAttribute(mx::Attributes::NewName);
-
-			if(NewInputName)
-			{
-				InputCopy->setName(NewInputName);
-			}
-		}
-
-		for(const FInputToCreate& Input : InputsToCreate)
-		{
-			const char* InputName = Input.Get<0>();
-			const FAttributeValueArray& Attributes = Input.Get<1>();
-
-			mx::InputPtr NewInput = Node->addInput();
-			NewInput->setName(InputName);
-			for(const FAttributeValue& AttributeValue : Attributes)
-			{
-				const char* Attribute = AttributeValue.Get<0>();
-				const char* Value = AttributeValue.Get<1>();
-				NewInput->setAttribute(Attribute, Value);
-			}
-		}
-	}
-
-	return Node;
-}
-
-void UInterchangeMaterialXTranslator::ConvertNeutralNodesToDot(MaterialX::ElementPtr NodeGraph) const
-{
-	const std::vector<mx::NodePtr> Nodes = NodeGraph->getChildrenOfType<mx::Node>();
-	for(mx::NodePtr Node : Nodes)
-	{
-		const std::string& Category = Node->getCategory();
-		if(Category == mx::Category::Multiply || Category == mx::Category::Divide || Category == mx::Category::Add || Category == mx::Category::Sub)
-		{
-			mx::InputPtr Input2 = Node->getInput("in2");
-			bool bIsNeutral = false;
-
-			if(!Input2)
-			{
-				bIsNeutral = true;
-				if(mx::InputPtr Input1 = Node->getInput("in1"))
-				{
-					Input1->setName("in");
-				}
-			}
-			else
-			{
-				if(!Input2->hasValue() && !Input2->hasNodeName() && !Input2->hasInterfaceName())
-				{
-					bIsNeutral = true;
-					Node->removeInput("in2");
-					if(mx::InputPtr Input1 = Node->getInput("in1"))
-					{
-						Input1->setName("in");
-					}
-				}
-				else if(Input2->hasValue())
-				{
-					using namespace mx::NodeGroup::Math;
-
-					auto IsValueStringNeutral = [&bIsNeutral, &Category, Input2](auto T)
-					{
-						bIsNeutral = (Category == mx::Category::Multiply || Category == mx::Category::Divide) ?
-							mx::fromValueString<decltype(T)>(Input2->getValueString()) == NeutralOne<decltype(T)> :
-							mx::fromValueString<decltype(T)>(Input2->getValueString()) == NeutralZero<decltype(T)>;
-					};
-
-					if(Input2->getType() == mx::Type::Float)
-					{
-						IsValueStringNeutral(float());
-					}
-					else if(Input2->getType() == mx::Type::Vector2)
-					{
-						IsValueStringNeutral(mx::Vector2());
-					}
-					else if(Input2->getType() == mx::Type::Vector3)
-					{
-						IsValueStringNeutral(mx::Vector3());
-					}
-					else if(Input2->getType() == mx::Type::Vector4)
-					{
-						IsValueStringNeutral(mx::Vector4());
-					}
-					else if(Input2->getType() == mx::Type::Color3)
-					{
-						IsValueStringNeutral(mx::Color3());
-					}
-					else if(Input2->getType() == mx::Type::Color4)
-					{
-						IsValueStringNeutral(mx::Color4());
-					}
-
-					if(bIsNeutral)
-					{
-						Node->removeInput("in2");
-						if(mx::InputPtr Input1 = Node->getInput("in1"))
-						{
-							Input1->setName("in");
-						}
-					}
-				}
-			}
-
-			if(bIsNeutral)
-			{
-				Node->setCategory(mx::Category::Dot);
-			}
-		}
-	}
-}
-
-void UInterchangeMaterialXTranslator::AddTexCoordToTiledImageNodes(MaterialX::ElementPtr Graph) const
-{
-	for(mx::NodePtr Node : Graph->getChildrenOfType<mx::Node>(mx::Category::TiledImage))
-	{
-		if(Node->getCategory() == mx::Category::TiledImage)
-		{
-			mx::InputPtr InputTexCoord = Node->getInput("texcoord");
-
-			if(!InputTexCoord)
-			{
-				mx::NodePtr NodeTexCoord = CreateNode(Node->getParent(),
-													  "node_ue_texcoord",
-													  mx::Category::TexCoord,
-													  {},
-													  { {"index", FAttributeValueArray{{"type", "integer"}, {"value", "0"}}} });
-
-				InputTexCoord = Node->addInput(mx::NodeGroup::Texture2D::Inputs::TexCoord, mx::Type::Vector2);
-				InputTexCoord->setNodeName(NodeTexCoord->getName());
-			}
-		}
-	}
-}
-
-void UInterchangeMaterialXTranslator::UpdateDocumentRecursively(MaterialX::GraphElementPtr Graph) const
-{
-	AddTexCoordToTiledImageNodes(Graph);
-
-	auto FilterNode = [](mx::NodePtr Node)
-	{
-		const std::string& Category = Node->getCategory();
-		return
-			Category != mx::Category::StandardSurface &&
-			Category != mx::Category::Extract &&
-			Category != mx::Category::Saturate &&
-			Category != mx::Category::Separate2 &&
-			Category != mx::Category::Separate3 &&
-			Category != mx::Category::Separate4;
-	};
-
-	//This function will replace all the nodes that are defined with a nodegraph with their equivalent node network
-	//This allows us to handle nodes defined with a nodegraph, with nodes that we already support while traversing the current node graph
-	//For example the node tiledimage is just a succession of the nodes image, multiply, add, etc.
-	Graph->flattenSubgraphs(mx::EMPTY_STRING, FilterNode);
-
-	ConvertNeutralNodesToDot(Graph);
-
-	mx::DocumentPtr Document = Graph->getDocument();
-	//we also need need to flatten the subgraphs
-	for(mx::NodeGraphPtr SubGraph : Graph->getChildrenOfType<mx::NodeGraph>())
-	{
-		if(SubGraph->getActiveSourceUri() == Document->getSourceUri())
-		{
-			UpdateDocumentRecursively(SubGraph);
-		}
-	}
-}
-
-FString UInterchangeMaterialXTranslator::GetAttributeParentName(MaterialX::NodePtr Node) const
-{
-	FString ParentName;
-	const mx::StringVec & Attributes =  Node->getAttributeNames();
-
-	// For consistency the parent attribute has an index attach to it to ensure unicity
-	// Attributes are set in order, we only need to take the first one (it will be remove after that)
-	for(const std::string& Attrib : Attributes)
-	{
-		if(Attrib.find(mx::Attributes::ParentName) != std::string::npos)
-		{
-			ParentName = Node->getAttribute(Attrib).c_str();
-			Node->removeAttribute(Attrib);
-			break;
-		}
-	}
-
-	return ParentName.IsEmpty() ? Node->getName().c_str() : ParentName;
-}
-
-#endif //WITH_EDITOR
-
 #undef LOCTEXT_NAMESPACE
diff --git a/Engine/Plugins/Interchange/Runtime/Source/Import/Private/MaterialX/MaterialXUtils/MaterialXBase.cpp b/Engine/Plugins/Interchange/Runtime/Source/Import/Private/MaterialX/MaterialXUtils/MaterialXBase.cpp
new file mode 100644
--- /dev/null
+++ b/Engine/Plugins/Interchange/Runtime/Source/Import/Private/MaterialX/MaterialXUtils/MaterialXBase.cpp
@@ -0,0 +1,277 @@
+// Copyright Epic Games, Inc. All Rights Reserved.
+
+#if WITH_EDITOR
+
+#include "MaterialX/MaterialXUtils/MaterialXBase.h"
+#include "MaterialX/MaterialXUtils/MaterialXManager.h"
+#include "InterchangeImportLog.h"
+
+namespace mx = MaterialX;
+
+FMaterialXBase::FMaterialXBase(UInterchangeBaseNodeContainer& BaseNodeContainer)
+	: NodeContainer(BaseNodeContainer)
+{}
+
+void FMaterialXBase::UpdateDocumentRecursively(MaterialX::GraphElementPtr Graph)
+{
+	AddTexCoordToTiledImageNodes(Graph);
+
+	static std::set<std::string> SkippedCategories{
+		mx::Category::StandardSurface,
+		mx::Category::Extract,
+		mx::Category::Saturate,
+		mx::Category::Separate2,
+		mx::Category::Separate3,
+		mx::Category::Separate4
+	};
+
+	auto FilterNode = [](mx::NodePtr Node)
+	{
+		const std::string& Category = Node->getCategory();
+		return SkippedCategories.find(Category) == SkippedCategories.end();
+	};
+
+	//This function will replace all the nodes that are defined with a nodegraph with their equivalent node network
+	//This allows us to handle nodes defined with a nodegraph, with nodes that we already support while traversing the current node graph
+	//For example the node tiledimage is just a succession of the nodes image, multiply, add, etc.
+	Graph->flattenSubgraphs(mx::EMPTY_STRING, FilterNode);
+
+	ConvertNeutralNodesToDot(Graph);
+
+	mx::DocumentPtr Document = Graph->getDocument();
+	//we also need need to flatten the subgraphs
+	for(mx::NodeGraphPtr SubGraph : Graph->getChildrenOfType<mx::NodeGraph>())
+	{
+		if(SubGraph->getActiveSourceUri() == Document->getSourceUri())
+		{
+			UpdateDocumentRecursively(SubGraph);
+		}
+	}
+}
+
+
+void FMaterialXBase::AddTexCoordToTiledImageNodes(MaterialX::ElementPtr Graph)
+{
+	for(mx::NodePtr Node : Graph->getChildrenOfType<mx::Node>(mx::Category::TiledImage))
+	{
+		if(Node->getCategory() == mx::Category::TiledImage)
+		{
+			mx::InputPtr InputTexCoord = Node->getInput("texcoord");
+
+			if(!InputTexCoord)
+			{
+				mx::NodePtr NodeTexCoord = CreateNode(Node->getParent(),
+													  "node_ue_texcoord",
+													  mx::Category::TexCoord,
+													  {},
+													  { {"index", FAttributeValueArray{{"type", "integer"}, {"value", "0"}}} });
+
+				InputTexCoord = Node->addInput(mx::NodeGroup::Texture2D::Inputs::TexCoord, mx::Type::Vector2);
+				InputTexCoord->setNodeName(NodeTexCoord->getName());
+			}
+		}
+	}
+}
+
+void FMaterialXBase::ConvertNeutralNodesToDot(MaterialX::ElementPtr NodeGraph)
+{
+	const std::vector<mx::NodePtr> Nodes = NodeGraph->getChildrenOfType<mx::Node>();
+	for(mx::NodePtr Node : Nodes)
+	{
+		const std::string& Category = Node->getCategory();
+		if(Category == mx::Category::Multiply || Category == mx::Category::Divide || Category == mx::Category::Add || Category == mx::Category::Sub)
+		{
+			mx::InputPtr Input2 = Node->getInput("in2");
+			bool bIsNeutral = false;
+
+			if(!Input2)
+			{
+				bIsNeutral = true;
+				if(mx::InputPtr Input1 = Node->getInput("in1"))
+				{
+					Input1->setName("in");
+				}
+			}
+			else
+			{
+				if(!Input2->hasValue() && !Input2->hasNodeName() && !Input2->hasInterfaceName())
+				{
+					bIsNeutral = true;
+					Node->removeInput("in2");
+					if(mx::InputPtr Input1 = Node->getInput("in1"))
+					{
+						Input1->setName("in");
+					}
+				}
+				else if(Input2->hasValue())
+				{
+					using namespace mx::NodeGroup::Math;
+
+					auto IsValueStringNeutral = [&bIsNeutral, &Category, Input2](auto T)
+					{
+						bIsNeutral = (Category == mx::Category::Multiply || Category == mx::Category::Divide) ?
+							mx::fromValueString<decltype(T)>(Input2->getValueString()) == NeutralOne<decltype(T)> :
+							mx::fromValueString<decltype(T)>(Input2->getValueString()) == NeutralZero<decltype(T)>;
+					};
+
+					if(Input2->getType() == mx::Type::Float)
+					{
+						IsValueStringNeutral(float());
+					}
+					else if(Input2->getType() == mx::Type::Vector2)
+					{
+						IsValueStringNeutral(mx::Vector2());
+					}
+					else if(Input2->getType() == mx::Type::Vector3)
+					{
+						IsValueStringNeutral(mx::Vector3());
+					}
+					else if(Input2->getType() == mx::Type::Vector4)
+					{
+						IsValueStringNeutral(mx::Vector4());
+					}
+					else if(Input2->getType() == mx::Type::Color3)
+					{
+						IsValueStringNeutral(mx::Color3());
+					}
+					else if(Input2->getType() == mx::Type::Color4)
+					{
+						IsValueStringNeutral(mx::Color4());
+					}
+
+					if(bIsNeutral)
+					{
+						Node->removeInput("in2");
+						if(mx::InputPtr Input1 = Node->getInput("in1"))
+						{
+							Input1->setName("in");
+						}
+					}
+				}
+			}
+
+			if(bIsNeutral)
+			{
+				Node->setCategory(mx::Category::Dot);
+			}
+		}
+	}
+}
+
+MaterialX::NodePtr FMaterialXBase::CreateNode(MaterialX::ElementPtr NodeGraph, const char* NodeName, const char* Category, const TArray<FInputToCopy>& InputsToCopy, const TArray<FInputToCreate>& InputsToCreate)
+{
+	std::string UniqueNodeName{ NodeName + std::string{"_"} + Category };
+
+	mx::NodePtr Node = NodeGraph->getChildOfType<mx::Node>(UniqueNodeName);
+
+	if(!Node)
+	{
+		Node = NodeGraph->addChildOfCategory(Category, UniqueNodeName)->asA<mx::Node>();
+
+		for(const FInputToCopy& Pair : InputsToCopy)
+		{
+			const char* NewInputName = Pair.Get<0>();
+			mx::InputPtr Input = Pair.Get<1>();
+
+			mx::InputPtr InputCopy = Node->addInput();
+			InputCopy->copyContentFrom(Input);
+			//remove the attribute NewName since when the input will be copied, the name will be the one from the spec (before a renaming)
+			InputCopy->removeAttribute(mx::Attributes::NewName);
+
+			if(NewInputName)
+			{
+				InputCopy->setName(NewInputName);
+			}
+		}
+
+		for(const FInputToCreate& Input : InputsToCreate)
+		{
+			const char* InputName = Input.Get<0>();
+			const FAttributeValueArray& Attributes = Input.Get<1>();
+
+			mx::InputPtr NewInput = Node->addInput();
+			NewInput->setName(InputName);
+			for(const FAttributeValue& AttributeValue : Attributes)
+			{
+				const char* Attribute = AttributeValue.Get<0>();
+				const char* Value = AttributeValue.Get<1>();
+				NewInput->setAttribute(Attribute, Value);
+			}
+		}
+	}
+
+	return Node;
+}
+
+FString FMaterialXBase::GetColorSpace(MaterialX::ElementPtr Element)
+{
+	FString ColorSpace;
+
+	if(Element)
+	{
+		if(Element->hasColorSpace())
+		{
+			return FString(Element->getColorSpace().c_str());
+		}
+		else
+		{
+			return GetColorSpace(Element->getParent());
+		}
+	}
+
+	return ColorSpace;
+}
+
+MaterialX::InputPtr FMaterialXBase::GetInput(MaterialX::NodePtr Node, const char* InputName)
+{
+	mx::InputPtr Input = Node->getInput(InputName);
+
+	if(!Input)
+	{
+		Input = Node->getDocument()->getNodeDef(NodeDefinition)->getActiveInput(InputName);
+	}
+
+	return Input;
+}
+
+FLinearColor FMaterialXBase::GetLinearColor(MaterialX::InputPtr Input)
+{
+	//we assume that the default color space is linear
+	FLinearColor LinearColor;
+
+	if(Input->getType() == mx::Type::Color3)
+	{
+		mx::Color3 Color = mx::fromValueString<mx::Color3>(Input->getValueString());
+		LinearColor = FLinearColor{ Color[0], Color[1], Color[2] };
+	}
+	else if(Input->getType() == mx::Type::Color4)
+	{
+		mx::Color4 Color = mx::fromValueString<mx::Color4>(Input->getValueString());
+		LinearColor = FLinearColor{ Color[0], Color[1], Color[2], Color[3] };
+	}
+	else
+	{
+		ensureMsgf(false, TEXT("input type can only be color3 or color4"));
+	}
+
+	const FString ColorSpace = GetColorSpace(Input);
+
+	if(ColorSpace.IsEmpty() || ColorSpace == TEXT("lin_rec709") || ColorSpace == TEXT("none"))
+	{
+		;//noop
+	}
+	else if(ColorSpace == TEXT("gamma22"))
+	{
+		LinearColor = FLinearColor::FromPow22Color(FColor(LinearColor.R / 255.f, LinearColor.G / 255.f, LinearColor.B / 255.f, LinearColor.A / 255.f));
+	}
+	else
+	{
+		UE_LOG(LogInterchangeImport, Warning, TEXT("<%s>-<%s>: Colorspace %s is not supported yet, falling back to linear"),
+			   Input->getParent()->getName().c_str(),
+			   Input->getName().c_str(),
+			   *ColorSpace);
+	}
+
+	return LinearColor;
+}
+#endif
\ No newline at end of file
diff --git a/Engine/Plugins/Interchange/Runtime/Source/Import/Private/MaterialX/MaterialXUtils/MaterialXDirectionalLightShader.cpp b/Engine/Plugins/Interchange/Runtime/Source/Import/Private/MaterialX/MaterialXUtils/MaterialXDirectionalLightShader.cpp
new file mode 100644
--- /dev/null
+++ b/Engine/Plugins/Interchange/Runtime/Source/Import/Private/MaterialX/MaterialXUtils/MaterialXDirectionalLightShader.cpp
@@ -0,0 +1,47 @@
+// Copyright Epic Games, Inc. All Rights Reserved.
+
+#if WITH_EDITOR
+#include "MaterialXDirectionalLightShader.h"
+
+namespace mx = MaterialX;
+
+FMaterialXDirectionalLightShader::FMaterialXDirectionalLightShader(UInterchangeBaseNodeContainer& BaseNodeContainer)
+	: FMaterialXLightShaderAbstract(BaseNodeContainer)
+{
+	NodeDefinition = mx::NodeDefinition::DirectionalLight;
+}
+
+TSharedRef<FMaterialXBase> FMaterialXDirectionalLightShader::MakeInstance(UInterchangeBaseNodeContainer& BaseNodeContainer)
+{
+	return MakeShared<FMaterialXDirectionalLightShader>(BaseNodeContainer);
+}
+
+void FMaterialXDirectionalLightShader::Translate(MaterialX::NodePtr DirectionalLightShaderNode)
+{
+	this->LightShaderNode = DirectionalLightShaderNode;
+
+	PreTranslate();
+
+	UInterchangeDirectionalLightNode* DirectionalLightNode = NewObject<UInterchangeDirectionalLightNode>(&NodeContainer);
+
+	// Direction
+	{
+		mx::InputPtr DirectionInput = GetInput(DirectionalLightShaderNode, mx::Lights::DirectionalLight::Input::Direction);
+		GetDirection(DirectionInput, Transform);
+	}
+
+	this->LightNode = DirectionalLightNode;
+
+	PostTranslate();
+}
+
+void FMaterialXDirectionalLightShader::GetDirection(MaterialX::InputPtr DirectionInput, FTransform& Transform)
+{
+	mx::Vector3 Direction = mx::fromValueString<mx::Vector3>(DirectionInput->getValueString());
+
+	//Get rotation to go from UE's default direction of directional light and the direction of the MX directional light node
+	const FVector LightDirection{ Direction[2], Direction[0], Direction[1] };
+	const FVector TransformDirection = Transform.GetUnitAxis(EAxis::X); // it's the default direction of a UE directional light
+	Transform.SetRotation(FQuat::FindBetween(LightDirection, TransformDirection));
+}
+#endif
\ No newline at end of file
diff --git a/Engine/Plugins/Interchange/Runtime/Source/Import/Private/MaterialX/MaterialXUtils/MaterialXDirectionalLightShader.h b/Engine/Plugins/Interchange/Runtime/Source/Import/Private/MaterialX/MaterialXUtils/MaterialXDirectionalLightShader.h
new file mode 100644
--- /dev/null
+++ b/Engine/Plugins/Interchange/Runtime/Source/Import/Private/MaterialX/MaterialXUtils/MaterialXDirectionalLightShader.h
@@ -0,0 +1,26 @@
+// Copyright Epic Games, Inc. All Rights Reserved.
+
+#pragma once
+
+#if WITH_EDITOR
+
+#include "MaterialX/MaterialXUtils/MaterialXLightShaderAbstract.h"
+
+class FMaterialXDirectionalLightShader : public FMaterialXLightShaderAbstract
+{
+protected:
+
+	template<typename T, ESPMode>
+	friend class SharedPointerInternals::TIntrusiveReferenceController;
+
+	FMaterialXDirectionalLightShader(UInterchangeBaseNodeContainer& BaseNodeContainer);
+
+public:
+
+	static TSharedRef<FMaterialXBase> MakeInstance(UInterchangeBaseNodeContainer& BaseNodeContainer);
+
+	virtual void Translate(MaterialX::NodePtr PointLightShaderNode) override;
+
+	static void GetDirection(MaterialX::InputPtr DirectionInput, FTransform& Transform);
+};
+#endif
diff --git a/Engine/Plugins/Interchange/Runtime/Source/Import/Private/MaterialX/MaterialXUtils/MaterialXLightShaderAbstract.cpp b/Engine/Plugins/Interchange/Runtime/Source/Import/Private/MaterialX/MaterialXUtils/MaterialXLightShaderAbstract.cpp
new file mode 100644
--- /dev/null
+++ b/Engine/Plugins/Interchange/Runtime/Source/Import/Private/MaterialX/MaterialXUtils/MaterialXLightShaderAbstract.cpp
@@ -0,0 +1,48 @@
+// Copyright Epic Games, Inc. All Rights Reserved. 
+
+#if WITH_EDITOR
+
+#include "MaterialX/MaterialXUtils/MaterialXLightShaderAbstract.h"
+#include "Misc/Paths.h"
+
+namespace mx = MaterialX;
+
+FMaterialXLightShaderAbstract::FMaterialXLightShaderAbstract(UInterchangeBaseNodeContainer& BaseNodeContainer)
+	: FMaterialXBase(BaseNodeContainer)
+{}
+
+void FMaterialXLightShaderAbstract::PreTranslate()
+{
+	const FString FileName{ FPaths::GetBaseFilename(LightShaderNode->getActiveSourceUri().c_str()) };
+	const FString LightNodeLabel = FString{ LightShaderNode->getName().c_str() };
+	SceneNode = NewObject<UInterchangeSceneNode>(&NodeContainer);
+	const FString SceneNodeUID = TEXT("\\Light\\") + FileName + TEXT("\\") + LightNodeLabel;
+	SceneNode->InitializeNode(SceneNodeUID, LightNodeLabel, EInterchangeNodeContainerType::TranslatedScene);
+	NodeContainer.AddNode(SceneNode);
+}
+
+void FMaterialXLightShaderAbstract::PostTranslate()
+{
+	const FString LightNodeLabel = FString{ LightShaderNode->getName().c_str() };
+	const FString LightNodeUid = TEXT("\\Light\\") + LightNodeLabel;
+	LightNode->InitializeNode(LightNodeUid, LightNodeLabel, EInterchangeNodeContainerType::TranslatedAsset);
+	NodeContainer.AddNode(LightNode);
+	const FString Uid = LightNode->GetUniqueID();
+	SceneNode->SetCustomAssetInstanceUid(Uid);
+
+	// Color
+	{
+		mx::InputPtr LightColor = LightShaderNode->getInput(mx::Lights::Input::Color);
+		const FLinearColor Color = GetLinearColor(LightColor);
+		LightNode->SetCustomLightColor(Color);
+	}
+
+	// Intensity
+	{
+		mx::InputPtr LightIntensity = LightShaderNode->getInput(mx::Lights::Input::Intensity);
+		LightNode->SetCustomIntensity(mx::fromValueString<float>(LightIntensity->getValueString()));
+	}
+
+	SceneNode->SetCustomLocalTransform(&NodeContainer, Transform);
+}
+#endif
\ No newline at end of file
diff --git a/Engine/Plugins/Interchange/Runtime/Source/Import/Private/MaterialX/MaterialXUtils/MaterialXManager.cpp b/Engine/Plugins/Interchange/Runtime/Source/Import/Private/MaterialX/MaterialXUtils/MaterialXManager.cpp
new file mode 100644
--- /dev/null
+++ b/Engine/Plugins/Interchange/Runtime/Source/Import/Private/MaterialX/MaterialXUtils/MaterialXManager.cpp
@@ -0,0 +1,229 @@
+// Copyright Epic Games, Inc. All Rights Reserved.
+
+#if WITH_EDITOR
+#include "MaterialX/MaterialXUtils/MaterialXManager.h"
+#include "MaterialX/InterchangeMaterialXDefinitions.h"
+#include "InterchangeMaterialDefinitions.h"
+#include "InterchangeShaderGraphNode.h"
+#include "MaterialX/MaterialXUtils/MaterialXSurfaceShaderAbstract.h"
+#include "MaterialX/MaterialXUtils/MaterialXStandardSurfaceShader.h"
+#include "MaterialX/MaterialXUtils/MaterialXSurfaceUnlitShader.h"
+#include "MaterialX/MaterialXUtils/MaterialXPointLightShader.h"
+#include "MaterialX/MaterialXUtils/MaterialXDirectionalLightShader.h"
+#include "MaterialX/MaterialXUtils/MaterialXSpotLightShader.h"
+#include "MaterialX/MaterialXUtils/MaterialXSurfaceMaterial.h"
+
+namespace mx = MaterialX;
+
+FMaterialXManager::FMaterialXManager()
+	: MatchingInputNames {
+		{{TEXT(""),                            TEXT("amplitude")}, TEXT("Amplitude")},
+		{{TEXT(""),                            TEXT("bg")},        TEXT("B")},
+		{{TEXT(""),                            TEXT("center")},    TEXT("Center")},
+		{{TEXT(""),                            TEXT("diminish")},  TEXT("Diminish")},
+		{{TEXT(""),                            TEXT("fg")},        TEXT("A")},
+		{{TEXT(""),                            TEXT("high")},      TEXT("Max")},
+		{{TEXT(""),                            TEXT("in")},        TEXT("Input")},
+		{{TEXT(""),                            TEXT("in1")},       TEXT("A")},
+		{{TEXT(""),                            TEXT("in2")},       TEXT("B")},
+		{{TEXT(""),                            TEXT("in3")},       TEXT("C")},
+		{{TEXT(""),                            TEXT("in4")},       TEXT("D")},
+		{{TEXT(""),                            TEXT("inlow")},     TEXT("InputLow")},
+		{{TEXT(""),                            TEXT("inhigh")},    TEXT("InputHigh")},
+		{{TEXT(""),                            TEXT("lacunarity")},TEXT("Lacunarity")},
+		{{TEXT(""),                            TEXT("low")},       TEXT("Min")},
+		{{TEXT(""),                            TEXT("lumacoeffs")},TEXT("LuminanceFactors")}, // for the moment not yet handled by Interchange, because of the attribute being an advanced pin
+		{{TEXT(""),                            TEXT("mix")},       TEXT("Alpha")},
+		{{TEXT(""),                            TEXT("position")},  TEXT("Position")},
+		{{TEXT(""),                            TEXT("texcoord")},  TEXT("Coordinates")},
+		{{TEXT(""),                            TEXT("octaves")},   TEXT("Octaves")},
+		{{TEXT(""),                            TEXT("outlow")},    TEXT("TargetLow")},
+		{{TEXT(""),                            TEXT("outhigh")},   TEXT("TargetHigh")},
+		{{TEXT(""),                            TEXT("valuel")},    TEXT("A")},
+		{{TEXT(""),                            TEXT("valuer")},    TEXT("B")},
+		{{MaterialX::Category::Atan2,          TEXT("in1")},       TEXT("Y")},
+		{{MaterialX::Category::Atan2,          TEXT("in2")},       TEXT("X")},
+		{{MaterialX::Category::HeightToNormal, TEXT("scale")},     UE::Interchange::Materials::Standard::Nodes::NormalFromHeightMap::Inputs::Intensity.ToString()},
+		{{MaterialX::Category::IfGreater,      TEXT("value1")},    TEXT("A")},
+		{{MaterialX::Category::IfGreater,      TEXT("value2")},    TEXT("B")},
+		{{MaterialX::Category::IfGreater,      TEXT("in1")},       TEXT("AGreaterThanB")},
+		{{MaterialX::Category::IfGreater,      TEXT("in2")},       TEXT("ALessThanB")}, //another input is added for the case equal see ConnectIfGreater
+		{{MaterialX::Category::IfGreaterEq,    TEXT("value1")},    TEXT("A")},
+		{{MaterialX::Category::IfGreaterEq,    TEXT("value2")},    TEXT("B")},
+		{{MaterialX::Category::IfGreaterEq,    TEXT("in1")},       TEXT("AGreaterThanB")}, //another is added for the case equal see ConnectIfGreaterEq
+		{{MaterialX::Category::IfGreaterEq,    TEXT("in2")},       TEXT("ALessThanB")},
+		{{MaterialX::Category::IfEqual,        TEXT("value1")},    TEXT("A")},
+		{{MaterialX::Category::IfEqual,        TEXT("value2")},    TEXT("B")},
+		{{MaterialX::Category::IfEqual,        TEXT("in1")},       TEXT("AEqualsB")},
+		{{MaterialX::Category::IfEqual,        TEXT("in2")},       TEXT("ALessThanB")},  // another input is added for the case greater see ConnectIfEqual
+		{{MaterialX::Category::Inside,         TEXT("in")},        TEXT("A")},			  // Inside is treated as a Multiply node
+		{{MaterialX::Category::Inside,         TEXT("mask")},      TEXT("B")},			  // Inside is treated as a Multiply node
+		{{MaterialX::Category::Invert,         TEXT("amount")},    TEXT("A")},
+		{{MaterialX::Category::Invert,         TEXT("in")},        TEXT("B")},
+		{{MaterialX::Category::Magnitude,      TEXT("in")},        TEXT("A")},
+		{{MaterialX::Category::Mix,            TEXT("fg")},        TEXT("B")},
+		{{MaterialX::Category::Mix,            TEXT("bg")},        TEXT("A")},
+		{{MaterialX::Category::Mix,            TEXT("mix")},		TEXT("Factor")},
+		{{MaterialX::Category::Noise3D,        TEXT("amplitude")}, TEXT("B")},            // The amplitude of the noise is connected to a multiply node
+		{{MaterialX::Category::Noise3D,        TEXT("pivot")},     TEXT("B")},            // The pivot of the noise is connected to a add node
+		{{MaterialX::Category::Normalize,      TEXT("in")},        TEXT("VectorInput")},
+		{{MaterialX::Category::Outside,        TEXT("in")},        TEXT("A")},				// Outside is treated as Multiply node
+		{{MaterialX::Category::Outside,        TEXT("mask")},      TEXT("B")},				// Outside is treated as Multiply node
+		{{MaterialX::Category::Power,          TEXT("in1")},       TEXT("Base")},
+		{{MaterialX::Category::Power,	       TEXT("in2")},       TEXT("Exponent")},
+		{{MaterialX::Category::Rotate2D,       TEXT("amount")},    TEXT("RotationAngle")},
+		{{MaterialX::Category::Rotate2D,       TEXT("in")},        TEXT("Position")},
+		{{MaterialX::Category::Rotate3D,       TEXT("amount")},    TEXT("RotationAngle")},
+		{{MaterialX::Category::Rotate3D,       TEXT("axis")},		TEXT("NormalizedRotationAxis")},
+		{{MaterialX::Category::Rotate3D,       TEXT("in")},        TEXT("Position")},
+		{{MaterialX::Category::Saturate,       TEXT("amount")},    TEXT("Fraction")},
+		{{MaterialX::Category::Smoothstep,     TEXT("in")},        TEXT("Value")},
+}
+	, MatchingMaterialExpressions {
+		// Math nodes
+		{MaterialX::Category::Absval,       TEXT("Abs")},
+		{MaterialX::Category::Add,          TEXT("Add")},
+		{MaterialX::Category::Acos,         TEXT("Arccosine")},
+		{MaterialX::Category::Asin,         TEXT("Arcsine")},
+		{MaterialX::Category::Atan2,        TEXT("Arctangent2")},
+		{MaterialX::Category::Ceil,         TEXT("Ceil")},
+		{MaterialX::Category::Clamp,        TEXT("Clamp")},
+		{MaterialX::Category::Cos,          TEXT("Cosine")},
+		{MaterialX::Category::CrossProduct, TEXT("Crossproduct")},
+		{MaterialX::Category::Divide,       TEXT("Divide")},
+		{MaterialX::Category::DotProduct,   TEXT("Dotproduct")},
+		{MaterialX::Category::Exp,          TEXT("Exponential")},
+		{MaterialX::Category::Floor,        TEXT("Floor")},
+		{MaterialX::Category::Invert,       TEXT("Subtract")},
+		{MaterialX::Category::Ln,           TEXT("Logarithm")},
+		{MaterialX::Category::Magnitude,    TEXT("Length")},
+		{MaterialX::Category::Max,          TEXT("Max")},
+		{MaterialX::Category::Min,          TEXT("Min")},
+		{MaterialX::Category::Modulo,       TEXT("Fmod")},
+		{MaterialX::Category::Multiply,     TEXT("Multiply")},
+		{MaterialX::Category::Normalize,    TEXT("Normalize")},
+		{MaterialX::Category::Power,        TEXT("Power")},
+		{MaterialX::Category::RampLR,       TEXT("RampLeftRight")},
+		{MaterialX::Category::RampTB,       TEXT("RampTopBottom")},
+		{MaterialX::Category::Sign,         TEXT("Sign")},
+		{MaterialX::Category::Sin,          TEXT("Sine")},
+		{MaterialX::Category::SplitLR,      TEXT("SplitLeftRight")},
+		{MaterialX::Category::SplitTB,      TEXT("SplitTopBottom")},
+		{MaterialX::Category::Sqrt,         TEXT("SquareRoot")},
+		{MaterialX::Category::Sub,          TEXT("Subtract")},
+		{MaterialX::Category::Tan,          TEXT("Tangent")},
+		// Compositing nodes
+		{MaterialX::Category::Burn,         TEXT("Burn")},
+		{MaterialX::Category::Difference,   TEXT("Difference")},
+		{MaterialX::Category::Disjointover, TEXT("Disjointover")},
+		{MaterialX::Category::Dodge,        TEXT("Dodge")},
+		{MaterialX::Category::In,           TEXT("In")},
+		{MaterialX::Category::Inside,       TEXT("Multiply")},
+		{MaterialX::Category::Mask,         TEXT("Mask")},
+		{MaterialX::Category::Matte,        TEXT("Matte")},
+		{MaterialX::Category::Minus,        TEXT("Minus")},
+		{MaterialX::Category::Mix,          TEXT("Lerp")},
+		{MaterialX::Category::Out,          TEXT("Out")},
+		{MaterialX::Category::Over,         TEXT("Over")},
+		{MaterialX::Category::Overlay,      TEXT("Overlay")},
+		{MaterialX::Category::Plus,         TEXT("Plus")},
+		{MaterialX::Category::Premult,      TEXT("Premult")},
+		{MaterialX::Category::Screen,       TEXT("Screen")},
+		{MaterialX::Category::Unpremult,    TEXT("Unpremult")},
+		// Channel nodes
+		{MaterialX::Category::Combine2,     TEXT("AppendVector")},
+		{MaterialX::Category::Combine3,     TEXT("Append3Vector")},
+		{MaterialX::Category::Combine4,     TEXT("Append4Vector")},
+		// Procedural3D nodes
+		{MaterialX::Category::Fractal3D,    TEXT("Fractal3D")},
+		// Geometric nodes 
+		{MaterialX::Category::GeomColor,	TEXT("VertexColor")},
+		{MaterialX::Category::TexCoord,		TEXT("TextureCoordinate")},
+		// Adjustment nodes,
+		{MaterialX::Category::HsvToRgb,		TEXT("HsvToRgb")},
+		{MaterialX::Category::Luminance,    TEXT("Luminance")},
+		{MaterialX::Category::Remap,        TEXT("Remap")},
+		{MaterialX::Category::RgbToHsv,		TEXT("RgbToHsv")},
+		{MaterialX::Category::Saturate,		TEXT("Desaturation")},
+		{MaterialX::Category::Smoothstep,   TEXT("SmoothStep")}
+	}
+	, MaterialExpressionInputs {
+		TEXT("A"),
+		TEXT("Alpha"),
+		TEXT("Amplitude"),
+		TEXT("B"),
+		TEXT("Base"),
+		TEXT("C"),
+		TEXT("Center"),
+		TEXT("D"),
+		TEXT("Exponent"),
+		TEXT("Factor"),
+		TEXT("Fraction"),
+		TEXT("Input"),
+		TEXT("InputLow"),
+		TEXT("InputHigh"),
+		TEXT("Lacunarity"),
+		TEXT("LuminanceFactors"),
+		TEXT("Max"),
+		TEXT("Min"),
+		TEXT("Octaves"),
+		TEXT("Position"),
+		TEXT("TargetLow"),
+		TEXT("TargetHigh"),
+		TEXT("Value"),
+		TEXT("VectorInput"),
+		TEXT("X"),
+		TEXT("Y")
+}
+	, MaterialXContainerDelegates{
+		{mx::Category::SurfaceUnlit, FMaterialXManager::FOnGetMaterialXInstance::CreateStatic(&FMaterialXSurfaceUnlitShader::MakeInstance)},
+		{mx::Category::StandardSurface, FMaterialXManager::FOnGetMaterialXInstance::CreateStatic(&FMaterialXStandardSurfaceShader::MakeInstance)},
+		{mx::Category::PointLight, FMaterialXManager::FOnGetMaterialXInstance::CreateStatic(&FMaterialXPointLightShader::MakeInstance)},
+		{mx::Category::DirectionalLight, FMaterialXManager::FOnGetMaterialXInstance::CreateStatic(&FMaterialXDirectionalLightShader::MakeInstance)},
+		{mx::Category::SpotLight, FMaterialXManager::FOnGetMaterialXInstance::CreateStatic(&FMaterialXSpotLightShader::MakeInstance)},
+		{mx::SURFACE_MATERIAL_NODE_STRING.c_str(), FMaterialXManager::FOnGetMaterialXInstance::CreateStatic(&FMaterialXSurfaceMaterial::MakeInstance)}
+	}
+{}
+
+FMaterialXManager& FMaterialXManager::GetInstance()
+{
+	static FMaterialXManager Instance;
+	return Instance;
+}
+
+const FString* FMaterialXManager::FindMatchingInput(const TPair<FString, FString>& CategoryInputKey) const
+{
+	return MatchingInputNames.Find(CategoryInputKey);
+}
+
+const FString* FMaterialXManager::FindMaterialExpressionInput(const FString& InputKey) const
+{
+	return MaterialExpressionInputs.Find(InputKey);
+}
+
+const FString* FMaterialXManager::FindMatchingMaterialExpression(const FString& InputKey) const
+{
+	return MatchingMaterialExpressions.Find(InputKey);
+}
+
+TSharedPtr<FMaterialXBase> FMaterialXManager::GetShaderTranslator(const FString& CategoryShader, UInterchangeBaseNodeContainer & NodeContainer)
+{
+	if(FOnGetMaterialXInstance* InstanceDelegate = MaterialXContainerDelegates.Find(CategoryShader))
+	{
+		if(InstanceDelegate->IsBound())
+		{
+			return InstanceDelegate->Execute(NodeContainer);
+		}
+	}
+
+	return nullptr;
+}
+
+void FMaterialXManager::RegisterMaterialXInstance(const FString& Category, FOnGetMaterialXInstance MaterialXInstanceDelegate)
+{
+	if(!Category.IsEmpty())
+	{
+		MaterialXContainerDelegates.Add(Category, MaterialXInstanceDelegate);
+	}
+}
+#endif
\ No newline at end of file
diff --git a/Engine/Plugins/Interchange/Runtime/Source/Import/Private/MaterialX/MaterialXUtils/MaterialXPointLightShader.cpp b/Engine/Plugins/Interchange/Runtime/Source/Import/Private/MaterialX/MaterialXUtils/MaterialXPointLightShader.cpp
new file mode 100644
--- /dev/null
+++ b/Engine/Plugins/Interchange/Runtime/Source/Import/Private/MaterialX/MaterialXUtils/MaterialXPointLightShader.cpp
@@ -0,0 +1,63 @@
+// Copyright Epic Games, Inc. All Rights Reserved.
+
+#if WITH_EDITOR
+
+#include "MaterialXPointLightShader.h"
+
+namespace mx = MaterialX;
+
+FMaterialXPointLightShader::FMaterialXPointLightShader(UInterchangeBaseNodeContainer& BaseNodeContainer)
+	: FMaterialXLightShaderAbstract(BaseNodeContainer)
+{
+	NodeDefinition = mx::NodeDefinition::PointLight;
+}
+
+TSharedRef<FMaterialXBase> FMaterialXPointLightShader::MakeInstance(UInterchangeBaseNodeContainer& BaseNodeContainer)
+{
+	return MakeShared<FMaterialXPointLightShader>(BaseNodeContainer);
+}
+
+void FMaterialXPointLightShader::Translate(MaterialX::NodePtr PointLightShaderNode)
+{
+	this->LightShaderNode = PointLightShaderNode;
+
+	PreTranslate();
+
+	LightNode = NewObject<UInterchangePointLightNode>(&NodeContainer);
+	Cast<UInterchangePointLightNode>(LightNode)->SetCustomIntensityUnits(EInterchangeLightUnits::Candelas);
+
+	// Decay rate
+	SetDecayRate();
+
+	// Position
+	SetPosition();
+
+	PostTranslate();
+}
+
+float FMaterialXPointLightShader::GetDecayRate(mx::InputPtr DecayRateInput)
+{
+	return 0.0f;
+}
+
+inline void FMaterialXPointLightShader::GetPosition(mx::InputPtr PositionInput, FTransform& Transform)
+{
+	const mx::Vector3 Position = mx::fromValueString<mx::Vector3>(PositionInput->getValueString());
+	Transform.SetLocation(FVector{ Position[0] * 100, Position[1] * 100, Position[2] * 100 });
+}
+
+void FMaterialXPointLightShader::SetDecayRate()
+{
+	UInterchangePointLightNode* PointLightNode = Cast<UInterchangePointLightNode>(LightNode);
+	mx::InputPtr DecayRateInput = GetInput(LightShaderNode, mx::Lights::PointLight::Input::DecayRate);
+	const float DecayRate = mx::fromValueString<float>(DecayRateInput->getValueString());
+	PointLightNode->SetCustomUseInverseSquaredFalloff(false);
+	PointLightNode->SetCustomLightFalloffExponent(DecayRate);
+}
+
+void FMaterialXPointLightShader::SetPosition()
+{
+	mx::InputPtr PositionInput = GetInput(LightShaderNode, mx::Lights::PointLight::Input::Position);
+	GetPosition(PositionInput, Transform);
+}
+#endif
\ No newline at end of file
diff --git a/Engine/Plugins/Interchange/Runtime/Source/Import/Private/MaterialX/MaterialXUtils/MaterialXPointLightShader.h b/Engine/Plugins/Interchange/Runtime/Source/Import/Private/MaterialX/MaterialXUtils/MaterialXPointLightShader.h
new file mode 100644
--- /dev/null
+++ b/Engine/Plugins/Interchange/Runtime/Source/Import/Private/MaterialX/MaterialXUtils/MaterialXPointLightShader.h
@@ -0,0 +1,31 @@
+// Copyright Epic Games, Inc. All Rights Reserved.
+
+#pragma once
+
+#if WITH_EDITOR
+#include "MaterialX/MaterialXUtils/MaterialXLightShaderAbstract.h"
+
+class FMaterialXPointLightShader : public FMaterialXLightShaderAbstract
+{
+protected:
+
+	template<typename T, ESPMode>
+	friend class SharedPointerInternals::TIntrusiveReferenceController;
+
+	FMaterialXPointLightShader(UInterchangeBaseNodeContainer & BaseNodeContainer);
+
+public:
+
+	static TSharedRef<FMaterialXBase> MakeInstance(UInterchangeBaseNodeContainer& BaseNodeContainer);
+
+	virtual void Translate(MaterialX::NodePtr PointLightShaderNode) override;
+
+	static float GetDecayRate(MaterialX::InputPtr DecayRateInput);
+
+	static void GetPosition(MaterialX::InputPtr PositionInput, FTransform& Transform);
+
+	void SetDecayRate();
+
+	void SetPosition();
+};
+#endif
\ No newline at end of file
diff --git a/Engine/Plugins/Interchange/Runtime/Source/Import/Private/MaterialX/MaterialXUtils/MaterialXSpotLightShader.cpp b/Engine/Plugins/Interchange/Runtime/Source/Import/Private/MaterialX/MaterialXUtils/MaterialXSpotLightShader.cpp
new file mode 100644
--- /dev/null
+++ b/Engine/Plugins/Interchange/Runtime/Source/Import/Private/MaterialX/MaterialXUtils/MaterialXSpotLightShader.cpp
@@ -0,0 +1,67 @@
+// Copyright Epic Games, Inc. All Rights Reserved.
+
+#if WITH_EDITOR
+
+#include "MaterialXSpotLightShader.h"
+#include "MaterialXDirectionalLightShader.h"
+
+namespace mx = MaterialX;
+
+FMaterialXSpotLightShader::FMaterialXSpotLightShader(UInterchangeBaseNodeContainer& BaseNodeContainer)
+	: FMaterialXPointLightShader(BaseNodeContainer)
+{
+	NodeDefinition = mx::NodeDefinition::SpotLight;
+}
+
+TSharedRef<FMaterialXBase> FMaterialXSpotLightShader::MakeInstance(UInterchangeBaseNodeContainer& BaseNodeContainer)
+{
+	return MakeShared<FMaterialXSpotLightShader>(BaseNodeContainer);
+}
+
+void FMaterialXSpotLightShader::Translate(MaterialX::NodePtr SpotLightShaderNode)
+{
+	this->LightShaderNode = SpotLightShaderNode;
+
+	PreTranslate();
+
+	LightNode = NewObject<UInterchangeSpotLightNode>(&NodeContainer);
+	Cast<UInterchangeSpotLightNode>(LightNode)->SetCustomIntensityUnits(EInterchangeLightUnits::Candelas);
+
+	// Decay rate
+	SetDecayRate();
+
+	// Decay rate
+	SetDecayRate();
+
+	// Position
+	SetPosition();
+
+	// Direction
+	{
+		mx::InputPtr DirectionInput = GetInput(SpotLightShaderNode, mx::Lights::DirectionalLight::Input::Direction);
+		FMaterialXDirectionalLightShader::GetDirection(DirectionInput, Transform);
+	}
+
+	// Inner angle
+	SetInnerAngle();
+
+	// Outer angle
+	SetOuterAngle();
+
+	PostTranslate();
+}
+
+void FMaterialXSpotLightShader::SetInnerAngle()
+{
+	mx::InputPtr InnerAngleInput = GetInput(LightShaderNode, mx::Lights::SpotLight::Input::InnerAngle);
+	const float InnerAngle = FMath::RadiansToDegrees(mx::fromValueString<float>(InnerAngleInput->getValueString()));
+	Cast<UInterchangeSpotLightNode>(LightNode)->SetCustomInnerConeAngle(InnerAngle);
+}
+
+void FMaterialXSpotLightShader::SetOuterAngle()
+{
+	mx::InputPtr OuterAngleInput = GetInput(LightShaderNode, mx::Lights::SpotLight::Input::OuterAngle);
+	const float OuterAngle = FMath::RadiansToDegrees(mx::fromValueString<float>(OuterAngleInput->getValueString()));
+	Cast<UInterchangeSpotLightNode>(LightNode)->SetCustomOuterConeAngle(OuterAngle);
+}
+#endif
\ No newline at end of file
diff --git a/Engine/Plugins/Interchange/Runtime/Source/Import/Private/MaterialX/MaterialXUtils/MaterialXSpotLightShader.h b/Engine/Plugins/Interchange/Runtime/Source/Import/Private/MaterialX/MaterialXUtils/MaterialXSpotLightShader.h
new file mode 100644
--- /dev/null
+++ b/Engine/Plugins/Interchange/Runtime/Source/Import/Private/MaterialX/MaterialXUtils/MaterialXSpotLightShader.h
@@ -0,0 +1,28 @@
+// Copyright Epic Games, Inc. All Rights Reserved.
+
+#pragma once
+
+#if WITH_EDITOR
+#include "MaterialXPointLightShader.h"
+
+class FMaterialXSpotLightShader : public FMaterialXPointLightShader
+{
+protected:
+
+	template<typename T, ESPMode>
+	friend class SharedPointerInternals::TIntrusiveReferenceController;
+
+	FMaterialXSpotLightShader(UInterchangeBaseNodeContainer & BaseNodeContainer);
+
+public:
+
+	static TSharedRef<FMaterialXBase> MakeInstance(UInterchangeBaseNodeContainer& BaseNodeContainer);
+
+	virtual void Translate(MaterialX::NodePtr SpotLightShaderNode) override;
+
+	void SetInnerAngle();
+
+	void SetOuterAngle();
+};
+
+#endif
\ No newline at end of file
diff --git a/Engine/Plugins/Interchange/Runtime/Source/Import/Private/MaterialX/MaterialXUtils/MaterialXStandardSurfaceShader.cpp b/Engine/Plugins/Interchange/Runtime/Source/Import/Private/MaterialX/MaterialXUtils/MaterialXStandardSurfaceShader.cpp
new file mode 100644
--- /dev/null
+++ b/Engine/Plugins/Interchange/Runtime/Source/Import/Private/MaterialX/MaterialXUtils/MaterialXStandardSurfaceShader.cpp
@@ -0,0 +1,125 @@
+// Copyright Epic Games, Inc. All Rights Reserved.
+
+#if WITH_EDITOR
+#include "MaterialXStandardSurfaceShader.h"
+
+namespace mx = MaterialX;
+
+FMaterialXStandardSurfaceShader::FMaterialXStandardSurfaceShader(UInterchangeBaseNodeContainer& BaseNodeContainer)
+	: FMaterialXSurfaceShaderAbstract(BaseNodeContainer)
+{
+	NodeDefinition = mx::NodeDefinition::StandardSurface;
+}
+
+TSharedRef<FMaterialXBase> FMaterialXStandardSurfaceShader::MakeInstance(UInterchangeBaseNodeContainer& BaseNodeContainer)
+{
+	TSharedRef<FMaterialXStandardSurfaceShader> Result= MakeShared<FMaterialXStandardSurfaceShader>(BaseNodeContainer);
+	Result->RegisterConnectNodeOutputToInputDelegates();
+	return Result;
+}
+
+void FMaterialXStandardSurfaceShader::Translate(mx::NodePtr StandardSurfaceNode)
+{
+	this->SurfaceShaderNode = StandardSurfaceNode;
+
+	using namespace UE::Interchange::Materials;
+
+	UInterchangeShaderGraphNode* ShaderGraphNode = CreateShaderNode<UInterchangeShaderGraphNode>(SurfaceShaderNode->getName().c_str(), StandardSurface::Name.ToString());
+
+	//Base
+	ConnectNodeOutputToInput(mx::StandardSurface::Input::Base, ShaderGraphNode, StandardSurface::Parameters::Base.ToString(), mx::StandardSurface::DefaultValue::Float::Base);
+
+	//Base Color
+	ConnectNodeOutputToInput(mx::StandardSurface::Input::BaseColor, ShaderGraphNode, StandardSurface::Parameters::BaseColor.ToString(), mx::StandardSurface::DefaultValue::Color3::BaseColor);
+
+	//Diffuse Roughness
+	ConnectNodeOutputToInput(mx::StandardSurface::Input::DiffuseRoughness, ShaderGraphNode, StandardSurface::Parameters::DiffuseRoughness.ToString(), mx::StandardSurface::DefaultValue::Float::DiffuseRoughness);
+
+	//Specular
+	ConnectNodeOutputToInput(mx::StandardSurface::Input::Specular, ShaderGraphNode, StandardSurface::Parameters::Specular.ToString(), mx::StandardSurface::DefaultValue::Float::Specular);
+
+	//Specular Roughness
+	ConnectNodeOutputToInput(mx::StandardSurface::Input::SpecularRoughness, ShaderGraphNode, StandardSurface::Parameters::SpecularRoughness.ToString(), mx::StandardSurface::DefaultValue::Float::SpecularRoughness);
+
+	//Specular IOR
+	ConnectNodeOutputToInput(mx::StandardSurface::Input::SpecularIOR, ShaderGraphNode, StandardSurface::Parameters::SpecularIOR.ToString(), mx::StandardSurface::DefaultValue::Float::SpecularIOR);
+
+	//Specular Anisotropy
+	ConnectNodeOutputToInput(mx::StandardSurface::Input::SpecularAnisotropy, ShaderGraphNode, StandardSurface::Parameters::SpecularAnisotropy.ToString(), mx::StandardSurface::DefaultValue::Float::SpecularAnisotropy);
+
+	//Specular Rotation
+	ConnectNodeOutputToInput(mx::StandardSurface::Input::SpecularRotation, ShaderGraphNode, StandardSurface::Parameters::SpecularRotation.ToString(), mx::StandardSurface::DefaultValue::Float::SpecularRotation);
+
+	//Metallic
+	ConnectNodeOutputToInput(mx::StandardSurface::Input::Metalness, ShaderGraphNode, StandardSurface::Parameters::Metalness.ToString(), mx::StandardSurface::DefaultValue::Float::Metalness);
+
+	//Subsurface
+	ConnectNodeOutputToInput(mx::StandardSurface::Input::Subsurface, ShaderGraphNode, StandardSurface::Parameters::Subsurface.ToString(), mx::StandardSurface::DefaultValue::Float::Subsurface);
+
+	//Subsurface Color
+	ConnectNodeOutputToInput(mx::StandardSurface::Input::SubsurfaceColor, ShaderGraphNode, StandardSurface::Parameters::SubsurfaceColor.ToString(), mx::StandardSurface::DefaultValue::Color3::SubsurfaceColor);
+
+	//Subsurface Radius
+	ConnectNodeOutputToInput(mx::StandardSurface::Input::SubsurfaceRadius, ShaderGraphNode, StandardSurface::Parameters::SubsurfaceRadius.ToString(), mx::StandardSurface::DefaultValue::Color3::SubsurfaceRadius);
+
+	//Subsurface Scale
+	ConnectNodeOutputToInput(mx::StandardSurface::Input::SubsurfaceScale, ShaderGraphNode, StandardSurface::Parameters::SubsurfaceScale.ToString(), mx::StandardSurface::DefaultValue::Float::SubsurfaceScale);
+
+	//Sheen
+	ConnectNodeOutputToInput(mx::StandardSurface::Input::Sheen, ShaderGraphNode, StandardSurface::Parameters::Sheen.ToString(), mx::StandardSurface::DefaultValue::Float::Sheen);
+
+	//Sheen Color
+	ConnectNodeOutputToInput(mx::StandardSurface::Input::SheenColor, ShaderGraphNode, StandardSurface::Parameters::SheenColor.ToString(), mx::StandardSurface::DefaultValue::Color3::SheenColor);
+
+	//Sheen Roughness
+	ConnectNodeOutputToInput(mx::StandardSurface::Input::SheenRoughness, ShaderGraphNode, StandardSurface::Parameters::SheenRoughness.ToString(), mx::StandardSurface::DefaultValue::Float::SheenRoughness);
+
+	//Coat
+	ConnectNodeOutputToInput(mx::StandardSurface::Input::Coat, ShaderGraphNode, StandardSurface::Parameters::Coat.ToString(), mx::StandardSurface::DefaultValue::Float::Coat);
+
+	//Coat Color
+	ConnectNodeOutputToInput(mx::StandardSurface::Input::CoatColor, ShaderGraphNode, StandardSurface::Parameters::CoatColor.ToString(), mx::StandardSurface::DefaultValue::Color3::CoatColor);
+
+	//Coat Roughness
+	ConnectNodeOutputToInput(mx::StandardSurface::Input::CoatRoughness, ShaderGraphNode, StandardSurface::Parameters::CoatRoughness.ToString(), mx::StandardSurface::DefaultValue::Float::CoatRoughness);
+
+	//Coat Normal: No need to take the default input if there is no CoatNormal input
+	ConnectNodeOutputToInput(mx::StandardSurface::Input::CoatNormal, ShaderGraphNode, StandardSurface::Parameters::CoatNormal.ToString(), nullptr);
+
+	//Thin Film Thickness
+	ConnectNodeOutputToInput(mx::StandardSurface::Input::ThinFilmThickness, ShaderGraphNode, StandardSurface::Parameters::ThinFilmThickness.ToString(), mx::StandardSurface::DefaultValue::Float::ThinFilmThickness);
+
+	//Emission
+	ConnectNodeOutputToInput(mx::StandardSurface::Input::Emission, ShaderGraphNode, StandardSurface::Parameters::Emission.ToString(), mx::StandardSurface::DefaultValue::Float::Emission);
+
+	//Emission Color
+	ConnectNodeOutputToInput(mx::StandardSurface::Input::EmissionColor, ShaderGraphNode, StandardSurface::Parameters::EmissionColor.ToString(), mx::StandardSurface::DefaultValue::Color3::EmissionColor);
+
+	//Normal: No need to take the default input if there is no Normal input
+	ConnectNodeOutputToInput(mx::StandardSurface::Input::Normal, ShaderGraphNode, StandardSurface::Parameters::Normal.ToString(), nullptr);
+
+	//Tangent: No need to take the default input if there is no Tangent input
+	ConnectNodeOutputToInput(mx::StandardSurface::Input::Tangent, ShaderGraphNode, StandardSurface::Parameters::Tangent.ToString(), nullptr);
+
+	//Transmission
+	ConnectNodeOutputToInput(mx::StandardSurface::Input::Transmission, ShaderGraphNode, StandardSurface::Parameters::Transmission.ToString(), mx::StandardSurface::DefaultValue::Float::Transmission);
+
+	//Transmission Color
+	ConnectNodeOutputToInput(mx::StandardSurface::Input::TransmissionColor, ShaderGraphNode, StandardSurface::Parameters::TransmissionColor.ToString(), mx::StandardSurface::DefaultValue::Color3::TransmissionColor);
+
+	//Transmission Depth
+	ConnectNodeOutputToInput(mx::StandardSurface::Input::TransmissionDepth, ShaderGraphNode, StandardSurface::Parameters::TransmissionDepth.ToString(), mx::StandardSurface::DefaultValue::Float::TransmissionDepth);
+
+	//Transmission Scatter
+	ConnectNodeOutputToInput(mx::StandardSurface::Input::TransmissionScatter, ShaderGraphNode, StandardSurface::Parameters::TransmissionScatter.ToString(), mx::StandardSurface::DefaultValue::Color3::TransmissionScatter);
+
+	//Transmission Scatter Anisotropy
+	ConnectNodeOutputToInput(mx::StandardSurface::Input::TransmissionScatterAnisotropy, ShaderGraphNode, StandardSurface::Parameters::TransmissionScatterAnisotropy.ToString(), mx::StandardSurface::DefaultValue::Float::TransmissionScatterAnisotropy);
+
+	//Transmission Dispersion
+	ConnectNodeOutputToInput(mx::StandardSurface::Input::TransmissionDispersion, ShaderGraphNode, StandardSurface::Parameters::TransmissionDispersion.ToString(), mx::StandardSurface::DefaultValue::Float::TransmissionDispersion);
+
+	//Transmission Extra Roughness
+	ConnectNodeOutputToInput(mx::StandardSurface::Input::TransmissionExtraRoughness, ShaderGraphNode, StandardSurface::Parameters::TransmissionExtraRoughness.ToString(), mx::StandardSurface::DefaultValue::Float::TransmissionExtraRoughness);
+}
+#endif
\ No newline at end of file
diff --git a/Engine/Plugins/Interchange/Runtime/Source/Import/Private/MaterialX/MaterialXUtils/MaterialXStandardSurfaceShader.h b/Engine/Plugins/Interchange/Runtime/Source/Import/Private/MaterialX/MaterialXUtils/MaterialXStandardSurfaceShader.h
new file mode 100644
--- /dev/null
+++ b/Engine/Plugins/Interchange/Runtime/Source/Import/Private/MaterialX/MaterialXUtils/MaterialXStandardSurfaceShader.h
@@ -0,0 +1,23 @@
+// Copyright Epic Games, Inc. All Rights Reserved.
+
+#pragma once
+
+#if WITH_EDITOR
+#include "MaterialX/MaterialXUtils/MaterialXSurfaceShaderAbstract.h"
+
+class FMaterialXStandardSurfaceShader :  public FMaterialXSurfaceShaderAbstract
+{
+protected:
+
+	template<typename T, ESPMode>
+	friend class SharedPointerInternals::TIntrusiveReferenceController;
+
+	FMaterialXStandardSurfaceShader(UInterchangeBaseNodeContainer& BaseNodeContainer);
+
+public:
+
+	static TSharedRef<FMaterialXBase> MakeInstance(UInterchangeBaseNodeContainer& BaseNodeContainer);
+
+	virtual void Translate(MaterialX::NodePtr StandardSurfaceNode) override;
+};
+#endif
\ No newline at end of file
diff --git a/Engine/Plugins/Interchange/Runtime/Source/Import/Private/MaterialX/MaterialXUtils/MaterialXSurfaceMaterial.cpp b/Engine/Plugins/Interchange/Runtime/Source/Import/Private/MaterialX/MaterialXUtils/MaterialXSurfaceMaterial.cpp
new file mode 100644
--- /dev/null
+++ b/Engine/Plugins/Interchange/Runtime/Source/Import/Private/MaterialX/MaterialXUtils/MaterialXSurfaceMaterial.cpp
@@ -0,0 +1,47 @@
+// Copyright Epic Games, Inc. All Rights Reserved.
+
+#if WITH_EDITOR
+
+#include "MaterialXSurfaceMaterial.h"
+#include "InterchangeImportLog.h"
+#include "MaterialX/MaterialXUtils/MaterialXManager.h"
+#include "MaterialX/MaterialXUtils/MaterialXStandardSurfaceShader.h"
+#include "MaterialX/MaterialXUtils/MaterialXSurfaceUnlitShader.h"
+
+namespace mx = MaterialX;
+
+FMaterialXSurfaceMaterial::FMaterialXSurfaceMaterial(UInterchangeBaseNodeContainer& BaseNodeContainer)
+	: FMaterialXBase(BaseNodeContainer)
+{}
+
+TSharedRef<FMaterialXBase> FMaterialXSurfaceMaterial::MakeInstance(UInterchangeBaseNodeContainer& BaseNodeContainer)
+{
+	return MakeShared<FMaterialXSurfaceMaterial>(BaseNodeContainer);
+}
+
+void FMaterialXSurfaceMaterial::Translate(MaterialX::NodePtr SurfaceMaterialNode)
+{
+	bool bHasSupportedSurfaceShader = false;
+
+	for(mx::InputPtr Input : SurfaceMaterialNode->getInputs())
+	{
+		mx::NodePtr ConnectedNode = Input->getConnectedNode();
+
+		if(ConnectedNode)
+		{
+			TSharedPtr<FMaterialXBase> ShaderTranslator = FMaterialXManager::GetInstance().GetShaderTranslator(ConnectedNode->getCategory().c_str(), NodeContainer);
+			bHasSupportedSurfaceShader = ShaderTranslator != nullptr;
+			if(bHasSupportedSurfaceShader)
+			{
+				ShaderTranslator->Translate(ConnectedNode);
+			}
+		}
+	}
+
+	if(!bHasSupportedSurfaceShader)
+	{
+		UE_LOG(LogInterchangeImport, Warning, TEXT("the surfaceshader of <%s> is not supported"), ANSI_TO_TCHAR(SurfaceMaterialNode->getName().c_str()));
+	}
+}
+
+#endif
\ No newline at end of file
diff --git a/Engine/Plugins/Interchange/Runtime/Source/Import/Private/MaterialX/MaterialXUtils/MaterialXSurfaceMaterial.h b/Engine/Plugins/Interchange/Runtime/Source/Import/Private/MaterialX/MaterialXUtils/MaterialXSurfaceMaterial.h
new file mode 100644
--- /dev/null
+++ b/Engine/Plugins/Interchange/Runtime/Source/Import/Private/MaterialX/MaterialXUtils/MaterialXSurfaceMaterial.h
@@ -0,0 +1,23 @@
+// Copyright Epic Games, Inc. All Rights Reserved.
+
+#pragma once
+
+#if WITH_EDITOR
+#include "MaterialX/MaterialXUtils/MaterialXBase.h"
+
+class FMaterialXSurfaceMaterial : public FMaterialXBase
+{
+protected:
+
+	template<typename T, ESPMode>
+	friend class SharedPointerInternals::TIntrusiveReferenceController;
+
+	FMaterialXSurfaceMaterial(UInterchangeBaseNodeContainer& BaseNodeContainer);
+
+public:
+
+	static TSharedRef<FMaterialXBase> MakeInstance(UInterchangeBaseNodeContainer& BaseNodeContainer);
+
+	virtual void Translate(MaterialX::NodePtr SurfaceMaterialNode) override;
+};
+#endif
\ No newline at end of file
diff --git a/Engine/Plugins/Interchange/Runtime/Source/Import/Private/MaterialX/MaterialXUtils/MaterialXSurfaceShaderAbstract.cpp b/Engine/Plugins/Interchange/Runtime/Source/Import/Private/MaterialX/MaterialXUtils/MaterialXSurfaceShaderAbstract.cpp
new file mode 100644
--- /dev/null
+++ b/Engine/Plugins/Interchange/Runtime/Source/Import/Private/MaterialX/MaterialXUtils/MaterialXSurfaceShaderAbstract.cpp
@@ -0,0 +1,1197 @@
+// Copyright Epic Games, Inc. All Rights Reserved. 
+
+#if WITH_EDITOR
+#include "MaterialX/MaterialXUtils/MaterialXSurfaceShaderAbstract.h"
+
+#include "InterchangeImportLog.h"
+#include "InterchangeManager.h"
+#include "InterchangeTexture2DNode.h"
+#include "InterchangeTextureBlurNode.h"
+#include "InterchangeTranslatorBase.h"
+#include "Materials/MaterialExpressionNoise.h"
+#include "Materials/MaterialExpressionTransform.h"
+#include "Materials/MaterialExpressionTransformPosition.h"
+#include "Materials/MaterialExpressionVectorNoise.h"
+#include "MaterialX/MaterialXUtils/MaterialXManager.h"
+
+#define LOCTEXT_NAMESPACE "MaterialXSurfaceShaderAbstract"
+
+namespace mx = MaterialX;
+
+FMaterialXSurfaceShaderAbstract::FMaterialXSurfaceShaderAbstract(UInterchangeBaseNodeContainer& BaseNodeContainer)
+	: FMaterialXBase(BaseNodeContainer)
+{}
+
+bool FMaterialXSurfaceShaderAbstract::AddAttribute(MaterialX::InputPtr Input, const FString& InputChannelName, UInterchangeShaderNode* ShaderNode) const
+{
+	if(Input)
+	{
+		if(Input->getType() == mx::Type::Float)
+		{
+			return ShaderNode->AddFloatAttribute(UInterchangeShaderPortsAPI::MakeInputValueKey(InputChannelName), mx::fromValueString<float>(Input->getValueString()));
+		}
+		else if(Input->getType() == mx::Type::Integer) //Let's add it a Float attribute, because Interchange doesn't create a scalar if it's an int
+		{
+			return ShaderNode->AddFloatAttribute(UInterchangeShaderPortsAPI::MakeInputValueKey(InputChannelName), mx::fromValueString<int32>(Input->getValueString()));
+		}
+		else if(Input->getType() == mx::Type::Color3 || Input->getType() == mx::Type::Color4)
+		{
+			FLinearColor LinearColor = GetLinearColor(Input);
+			return ShaderNode->AddLinearColorAttribute(UInterchangeShaderPortsAPI::MakeInputValueKey(InputChannelName), LinearColor);
+		}
+		else if(Input->getType() == mx::Type::Vector2)
+		{
+			FLinearColor Vector = GetVector(Input);
+			return ShaderNode->AddVector2Attribute(UInterchangeShaderPortsAPI::MakeInputValueKey(InputChannelName), FVector2f{ Vector.R, Vector.G });
+		}
+		else if(Input->getType() == mx::Type::Vector3 || Input->getType() == mx::Type::Vector4)
+		{
+			FLinearColor Vector = GetVector(Input);
+			return ShaderNode->AddLinearColorAttribute(UInterchangeShaderPortsAPI::MakeInputValueKey(InputChannelName), Vector);
+		}
+	}
+
+	return false;
+}
+
+bool FMaterialXSurfaceShaderAbstract::AddAttributeFromValueOrInterface(MaterialX::InputPtr Input, const FString& InputChannelName, UInterchangeShaderNode* ShaderNode) const
+{
+	if(Input)
+	{
+		if(Input->hasValue())
+		{
+			return AddAttribute(Input, InputChannelName, ShaderNode);
+		}
+		else if(Input->hasInterfaceName())
+		{
+			if(mx::InputPtr InputInterface = Input->getInterfaceInput(); InputInterface->hasValue())
+			{
+				return AddAttribute(InputInterface, InputChannelName, ShaderNode);
+			}
+		}
+	}
+
+	return false;
+}
+
+bool FMaterialXSurfaceShaderAbstract::AddFloatAttribute(MaterialX::InputPtr Input, const FString& InputChannelName, UInterchangeShaderNode* ShaderNode, float DefaultValue) const
+{
+	if(Input)
+	{
+		if(Input->hasValue())
+		{
+			float Value = mx::fromValueString<float>(Input->getValueString());
+
+			if(!FMath::IsNearlyEqual(Value, DefaultValue))
+			{
+				return ShaderNode->AddFloatAttribute(UInterchangeShaderPortsAPI::MakeInputValueKey(InputChannelName), Value);
+			}
+		}
+	}
+
+	return false;
+}
+
+bool FMaterialXSurfaceShaderAbstract::AddLinearColorAttribute(MaterialX::InputPtr Input, const FString& InputChannelName, UInterchangeShaderNode* ShaderNode, const FLinearColor& DefaultValue) const
+{
+	if(Input)
+	{
+		if(Input->hasValue())
+		{
+			const FLinearColor Value = GetLinearColor(Input);
+
+			if(!Value.Equals(DefaultValue))
+			{
+				return ShaderNode->AddLinearColorAttribute(UInterchangeShaderPortsAPI::MakeInputValueKey(InputChannelName), Value);
+			}
+		}
+	}
+
+	return false;
+}
+
+bool FMaterialXSurfaceShaderAbstract::ConnectNodeGraphOutputToInput(MaterialX::InputPtr InputToNodeGraph, UInterchangeShaderNode* ShaderNode, const FString& ParentInputName)
+{
+	using namespace UE::Interchange::Materials::Standard::Nodes;
+
+	bool bHasNodeGraph = false;
+
+	if(InputToNodeGraph->hasNodeGraphString())
+	{
+		bHasNodeGraph = true;
+
+		mx::OutputPtr Output = InputToNodeGraph->getConnectedOutput();
+
+		if(!Output)
+		{
+			UE_LOG(LogInterchangeImport, Warning, TEXT("Couldn't find a connected output to (%s)"), *GetInputName(InputToNodeGraph));
+			return false;
+		}
+
+		for(mx::Edge Edge : Output->traverseGraph())
+		{
+			ConnectNodeCategoryOutputToInput(Edge, ShaderNode, ParentInputName);
+		}
+	}
+
+	return bHasNodeGraph;
+}
+
+bool FMaterialXSurfaceShaderAbstract::ConnectMatchingNodeOutputToInput(MaterialX::NodePtr Node, UInterchangeShaderNode* ParentShaderNode, const FString& InputChannelName)
+{
+	FMaterialXManager& Manager = FMaterialXManager::GetInstance();
+
+	bool bIsConnected = false;
+
+	// First search a matching Material Expression
+	if(const FString* ShaderType = Manager.FindMatchingMaterialExpression(Node->getCategory().c_str()))
+	{
+		UInterchangeShaderNode* OperatorNode = nullptr;
+
+		OperatorNode = CreateShaderNode<UInterchangeShaderNode>(Node->getName().c_str(), *ShaderType);
+
+		for(mx::InputPtr Input : Node->getInputs())
+		{
+			if(const FString* InputNameFound = Manager.FindMaterialExpressionInput(GetInputName(Input)))
+			{
+				AddAttributeFromValueOrInterface(Input, *InputNameFound, OperatorNode);
+			}
+		}
+
+		bIsConnected = UInterchangeShaderPortsAPI::ConnectDefaultOuputToInput(ParentShaderNode, InputChannelName, OperatorNode->GetUniqueID());
+	}
+	else if(FOnConnectNodeOutputToInput* Delegate = MatchingConnectNodeDelegates.Find(Node->getCategory().c_str()))
+	{
+		bIsConnected = Delegate->ExecuteIfBound(Node, ParentShaderNode, InputChannelName);
+	}
+
+	return bIsConnected;
+}
+
+void FMaterialXSurfaceShaderAbstract::ConnectNodeCategoryOutputToInput(const MaterialX::Edge& Edge, UInterchangeShaderNode* ShaderNode, const FString& ParentInputName)
+{
+	if(mx::NodePtr UpstreamNode = Edge.getUpstreamElement()->asA<mx::Node>())
+	{
+		UInterchangeShaderNode* ParentShaderNode = ShaderNode;
+		FString InputChannelName = ParentInputName;
+
+		//Replace the input's name by the one used in UE
+		SetMatchingInputsNames(UpstreamNode);
+
+		if(mx::ElementPtr DownstreamElement = Edge.getDownstreamElement())
+		{
+			if(mx::NodePtr DownstreamNode = DownstreamElement->asA<mx::Node>())
+			{
+				if(UInterchangeShaderNode** FoundNode = ShaderNodes.Find(GetAttributeParentName(DownstreamNode)))// DownstreamNode->getName().c_str()))
+				{
+					ParentShaderNode = *FoundNode;
+				}
+
+				if(mx::InputPtr ConnectedInput = Edge.getConnectingElement()->asA<mx::Input>())
+				{
+					InputChannelName = GetInputName(ConnectedInput);
+				}
+			}
+		}
+
+		if(!ConnectMatchingNodeOutputToInput(UpstreamNode, ParentShaderNode, InputChannelName))
+		{
+			UE_LOG(LogInterchangeImport, Warning, TEXT("<%s> is not supported yet"), ANSI_TO_TCHAR(UpstreamNode->getCategory().c_str()));
+		}
+	}
+}
+
+bool FMaterialXSurfaceShaderAbstract::ConnectNodeNameOutputToInput(MaterialX::InputPtr InputToConnectedNode, UInterchangeShaderNode* ShaderNode, const FString& ParentInputName)
+{
+	mx::NodePtr ConnectedNode = InputToConnectedNode->getConnectedNode();
+
+	if(!ConnectedNode)
+	{
+		return false;
+	}
+
+	UInterchangeShaderNode* ParentShaderNode = ShaderNode;
+	FString InputChannelName = ParentInputName;
+
+	mx::Edge Edge(nullptr, InputToConnectedNode, ConnectedNode);
+
+	TArray<mx::Edge> Stack{ Edge };
+
+	while(!Stack.IsEmpty())
+	{
+		Edge = Stack.Pop();
+
+		if(Edge.getUpstreamElement())
+		{
+			ConnectNodeCategoryOutputToInput(Edge, ShaderNode, ParentInputName);
+			ConnectedNode = Edge.getUpstreamElement()->asA<mx::Node>();
+			for(mx::InputPtr Input : ConnectedNode->getInputs())
+			{
+				Stack.Emplace(ConnectedNode, Input, Input->getConnectedNode());
+			}
+		}
+	}
+
+	return true;
+}
+
+void FMaterialXSurfaceShaderAbstract::ConnectConstantInputToOutput(MaterialX::NodePtr UpstreamNode, UInterchangeShaderNode* ParentShaderNode, const FString& InputChannelName)
+{
+	AddAttributeFromValueOrInterface(UpstreamNode->getInput("value"), InputChannelName, ParentShaderNode);
+}
+
+void FMaterialXSurfaceShaderAbstract::ConnectExtractInputToOutput(MaterialX::NodePtr UpstreamNode, UInterchangeShaderNode* ParentShaderNode, const FString& InputChannelName)
+{
+	using namespace UE::Interchange::Materials::Standard::Nodes;
+
+	UInterchangeShaderNode* MaskShaderNode = CreateShaderNode<UInterchangeShaderNode>(UpstreamNode->getName().c_str(), Mask::Name.ToString());
+
+	if(mx::InputPtr InputIndex = UpstreamNode->getInput("index"))
+	{
+		const int32 Index = mx::fromValueString<int>(InputIndex->getValueString());
+		switch(Index)
+		{
+		case 0: MaskShaderNode->AddBooleanAttribute(Mask::Attributes::R, true); break;
+		case 1: MaskShaderNode->AddBooleanAttribute(Mask::Attributes::G, true); break;
+		case 2: MaskShaderNode->AddBooleanAttribute(Mask::Attributes::B, true); break;
+		case 3: MaskShaderNode->AddBooleanAttribute(Mask::Attributes::A, true); break;
+		default:
+			UE_LOG(LogInterchangeImport, Warning, TEXT("Wrong index number for extract node, values are from [0-3]"));
+			break;
+		}
+	}
+
+	UInterchangeShaderPortsAPI::ConnectDefaultOuputToInput(ParentShaderNode, InputChannelName, MaskShaderNode->GetUniqueID());
+}
+
+void FMaterialXSurfaceShaderAbstract::ConnectDotInputToOutput(MaterialX::NodePtr UpstreamNode, UInterchangeShaderNode* ParentShaderNode, const FString& InputChannelName)
+{
+	if(mx::InputPtr Input = UpstreamNode->getInput("in"))
+	{
+		SetAttributeNewName(Input, TCHAR_TO_UTF8(*InputChannelName)); //let's take the parent node's input name
+		ShaderNodes.Add(UpstreamNode->getName().c_str(), ParentShaderNode);
+	}
+}
+
+void FMaterialXSurfaceShaderAbstract::ConnectTransformPositionInputToOutput(MaterialX::NodePtr UpstreamNode, UInterchangeShaderNode* ParentShaderNode, const FString& InputChannelName)
+{
+	UInterchangeShaderNode* TransformNode = CreateShaderNode<UInterchangeShaderNode>(UpstreamNode->getName().c_str(), TEXT("TransformPosition"));
+	AddAttributeFromValueOrInterface(UpstreamNode->getInput("in"), TEXT("Input"), TransformNode);
+	UInterchangeShaderPortsAPI::ConnectDefaultOuputToInput(ParentShaderNode, InputChannelName, TransformNode->GetUniqueID());
+}
+
+void FMaterialXSurfaceShaderAbstract::ConnectTransformVectorInputToOutput(MaterialX::NodePtr UpstreamNode, UInterchangeShaderNode* ParentShaderNode, const FString& InputChannelName)
+{
+	UInterchangeShaderNode* TransformNode = CreateShaderNode<UInterchangeShaderNode>(UpstreamNode->getName().c_str(), TEXT("Transform"));
+	AddAttributeFromValueOrInterface(UpstreamNode->getInput("in"), TEXT("Input"), TransformNode);
+	UInterchangeShaderPortsAPI::ConnectDefaultOuputToInput(ParentShaderNode, InputChannelName, TransformNode->GetUniqueID());
+}
+
+void FMaterialXSurfaceShaderAbstract::ConnectRotate2DInputToOutput(mx::NodePtr UpstreamNode, UInterchangeShaderNode* ParentShaderNode, const FString& InputChannelName)
+{
+	UInterchangeShaderNode* Rotate2DNode = CreateShaderNode<UInterchangeShaderNode>(UpstreamNode->getName().c_str(), TEXT("RotateAboutAxis"));
+	Rotate2DNode->AddLinearColorAttribute(UInterchangeShaderPortsAPI::MakeInputValueKey(TEXT("NormalizedRotationAxis")), FLinearColor(0, 0, 1));
+	Rotate2DNode->AddLinearColorAttribute(UInterchangeShaderPortsAPI::MakeInputValueKey(TEXT("PivotPoint")), FLinearColor(0.5, 0.5, 0));
+
+	{
+		mx::InputPtr Input = UpstreamNode->getInput("in");
+		AddAttributeFromValueOrInterface(UpstreamNode->getInput("in"), TEXT("Position"), Rotate2DNode);
+
+		// Rotate2D "in" input takes a vector2 input, but RotateAboutAxis takes a vector3 so we convert
+		if(Input)
+		{
+			mx::NodePtr ConvertNode = CreateNode(UpstreamNode->getParent()->asA<mx::NodeGraph>(),
+												 UpstreamNode->getName().c_str(),
+												 mx::Category::Convert,
+												 { {"in", Input} },
+												 {});
+			ConvertNode->setType(mx::Type::Vector3);
+			Input->setNodeName(ConvertNode->getName());
+		}
+	}
+
+	{
+
+		mx::InputPtr Input = UpstreamNode->getInput("amount");
+		AddAttributeFromValueOrInterface(Input, TEXT("RotationAngle"), Rotate2DNode);
+
+		// we create a Divide node to convert MaterialX angle in degrees to UE's angle which is a value between [0,1]
+		if(Input)
+		{
+			//copy "in" input into "in1" and create "amount" input as "in2" with the value 360 because UE's angle is a value between [0,1]
+			mx::NodePtr DivideNode = CreateNode(UpstreamNode->getParent()->asA<mx::NodeGraph>(),
+												UpstreamNode->getName().c_str(),
+												mx::Category::Divide,
+												{ {"in1", Input} }, // rename input to match <divide> input "in1"
+												{ {"in2", FAttributeValueArray{{"type", "float"}, {"value", "360"}}} });
+
+			// Input now points to the new node
+			Input->setNodeName(DivideNode->getName());
+		}
+	}
+
+	//We insert a mask node because the output has to be a vector2
+	UInterchangeShaderNode* MaskNode = CreateMaskShaderNode(0b1100, UpstreamNode->getName().c_str() + FString(TEXT("_Mask")));
+	UInterchangeShaderPortsAPI::ConnectDefaultOuputToInput(MaskNode, TEXT("Input"), Rotate2DNode->GetUniqueID());
+
+	UInterchangeShaderPortsAPI::ConnectDefaultOuputToInput(ParentShaderNode, InputChannelName, MaskNode->GetUniqueID());
+}
+
+void FMaterialXSurfaceShaderAbstract::ConnectRotate3DInputToOutput(mx::NodePtr UpstreamNode, UInterchangeShaderNode* ParentShaderNode, const FString& InputChannelName)
+{
+	UInterchangeShaderNode* Rotate3DNode = CreateShaderNode<UInterchangeShaderNode>(UpstreamNode->getName().c_str(), TEXT("RotateAboutAxis"));
+	Rotate3DNode->AddLinearColorAttribute(UInterchangeShaderPortsAPI::MakeInputValueKey(TEXT("PivotPoint")), FLinearColor(0.5, 0.5, 0));
+
+	AddAttributeFromValueOrInterface(UpstreamNode->getInput("in"), TEXT("Position"), Rotate3DNode);
+	AddAttributeFromValueOrInterface(UpstreamNode->getInput("axis"), TEXT("NormalizedRotationAxis"), Rotate3DNode);
+
+	// we create a Divide node to convert MaterialX angle in degrees to UE's angle which is a value between [0,1]
+	if(mx::InputPtr Input = UpstreamNode->getInput("amount"))
+	{
+		if(!AddAttributeFromValueOrInterface(Input, TEXT("RotationAngle"), Rotate3DNode))
+		{
+			//copy "in" input into "in1" and create "amount" input as "in2" with the value 360 because UE's angle is a value between [0,1]
+			mx::NodePtr NewDivideNode = CreateNode(UpstreamNode->getParent()->asA<mx::NodeGraph>(),
+												   UpstreamNode->getName().c_str(),
+												   mx::Category::Divide,
+												   { {"in1", Input} }, // rename input to match <divide> input "in1"
+												   { {"in2", FAttributeValueArray{{"type", "float"}, {"value", "360"}}} });
+
+			// Input now points to the new node
+			Input->setNodeName(NewDivideNode->getName());
+		}
+	}
+
+	UInterchangeShaderPortsAPI::ConnectDefaultOuputToInput(ParentShaderNode, InputChannelName, Rotate3DNode->GetUniqueID());
+}
+
+void FMaterialXSurfaceShaderAbstract::ConnectImageInputToOutput(MaterialX::NodePtr UpstreamNode, UInterchangeShaderNode* ParentShaderNode, const FString& InputChannelName)
+{
+	using namespace UE::Interchange::Materials::Standard::Nodes;
+
+	if(UInterchangeTextureNode* TextureNode = CreateTextureNode<UInterchangeTexture2DNode>(UpstreamNode))
+	{
+		//By default set the output of a texture to RGB
+		FString OutputChannel{ TEXT("RGB") };
+
+		if(UpstreamNode->getType() == mx::Type::Vector4 || UpstreamNode->getType() == mx::Type::Color4)
+		{
+			OutputChannel = TEXT("RGBA");
+		}
+		else if(UpstreamNode->getType() == mx::Type::Float)
+		{
+			OutputChannel = TEXT("R");
+		}
+
+		UInterchangeShaderNode* TextureShaderNode = CreateShaderNode<UInterchangeShaderNode>(UpstreamNode->getName().c_str(), TextureSample::Name.ToString());
+		TextureShaderNode->AddStringAttribute(UInterchangeShaderPortsAPI::MakeInputValueKey(TextureSample::Inputs::Texture.ToString()), TextureNode->GetUniqueID());
+		UInterchangeShaderPortsAPI::ConnectOuputToInput(ParentShaderNode, InputChannelName, TextureShaderNode->GetUniqueID(), OutputChannel);
+	}
+	else
+	{
+		AddAttributeFromValueOrInterface(UpstreamNode->getInput(mx::NodeGroup::Texture2D::Inputs::Default), InputChannelName, ParentShaderNode);
+	}
+}
+
+void FMaterialXSurfaceShaderAbstract::ConnectConvertInputToOutput(MaterialX::NodePtr UpstreamNode, UInterchangeShaderNode* ParentShaderNode, const FString& InputChannelName)
+{
+	using namespace UE::Interchange::Materials::Standard::Nodes;
+
+	//In case of an upwards conversion, let's do an append, downwards a mask, otherwise leave as is
+	mx::InputPtr Input = UpstreamNode->getInput("in");
+	const std::string& NodeType = UpstreamNode->getType();
+	const std::string& InputType = Input->getType();
+
+	//Ensure that the types are supported
+	const bool bIsNodeTypeSupported =
+		NodeType == mx::Type::Color4 ||
+		NodeType == mx::Type::Color3 ||
+		NodeType == mx::Type::Vector4 ||
+		NodeType == mx::Type::Vector3 ||
+		NodeType == mx::Type::Vector2 ||
+		NodeType == mx::Type::Float ||
+		NodeType == mx::Type::Integer ||
+		NodeType == mx::Type::Boolean;
+
+	const bool bIsInputTypeSupported =
+		InputType == mx::Type::Color4 ||
+		InputType == mx::Type::Color3 ||
+		InputType == mx::Type::Vector4 ||
+		InputType == mx::Type::Vector3 ||
+		InputType == mx::Type::Vector2 ||
+		InputType == mx::Type::Float ||
+		InputType == mx::Type::Integer ||
+		InputType == mx::Type::Boolean;
+
+	if(!bIsNodeTypeSupported || !bIsInputTypeSupported)
+	{
+		UE_LOG(LogInterchangeImport, Warning, TEXT("<convert> node has non supported types"));
+		return;
+	}
+
+	char NodeN = NodeType.back();
+	char InputN = InputType.back();
+
+	constexpr auto Remap = [](auto Value, auto Low1, auto High1, auto Low2, auto High2)
+	{
+		return Low2 + (Value - Low1) * (High2 - Low2) / (High1 - Low1);
+	};
+
+	//Remap NodeN and InputN in a range lower than '2' in case of integer/boolean/float
+	if(NodeN > '4')
+	{
+		NodeN = Remap(NodeN, 'a', 'z', 0, '1');
+	}
+	if(InputN > '4')
+	{
+		InputN = Remap(InputN, 'a', 'z', 0, '1');
+	}
+
+	if(InputN > NodeN) // Mask
+	{
+		UInterchangeShaderNode* MaskShaderNode = nullptr;
+		if(NodeN == '3')
+		{
+			MaskShaderNode = CreateMaskShaderNode(0b1110, UpstreamNode->getName().c_str());
+		}
+		else if(NodeN == '2')
+		{
+			MaskShaderNode = CreateMaskShaderNode(0b1100, UpstreamNode->getName().c_str());
+		}
+		else
+		{
+			MaskShaderNode = CreateMaskShaderNode(0b1000, UpstreamNode->getName().c_str());
+		}
+
+		AddAttributeFromValueOrInterface(Input, TEXT("Input"), MaskShaderNode);
+
+		UInterchangeShaderPortsAPI::ConnectDefaultOuputToInput(ParentShaderNode, InputChannelName, MaskShaderNode->GetUniqueID());
+	}
+	else // Append 
+	{
+		//same as dot, just connect the next output to this parent input
+		SetAttributeNewName(Input, TCHAR_TO_UTF8(*InputChannelName)); //let's take the parent node's input name
+		ShaderNodes.Add(UpstreamNode->getName().c_str(), ParentShaderNode);
+
+		//No need to create a node, since the input and the node have the same channels, we just check if there's a value
+		if(NodeN == InputN || (NodeN < '2' && InputN < '2'))
+		{
+			AddAttributeFromValueOrInterface(Input, InputChannelName, ParentShaderNode);
+			return;
+		}
+		std::string Category;
+		TArray<FInputToCopy> InputsToCopy;
+		TArray<FInputToCreate> InputsToCreate;
+
+		// float to N
+		if(InputN < '2')
+		{
+			if(NodeN == '2')
+			{
+				Category = mx::Category::Combine2;
+				InputsToCopy.Add({ "in1", Input });
+				InputsToCopy.Add({ "in2", Input });
+			}
+			else if(NodeN == '3')
+			{
+				Category = mx::Category::Combine3;
+				InputsToCopy.Add({ "in1", Input });
+				InputsToCopy.Add({ "in2", Input });
+				InputsToCopy.Add({ "in3", Input });
+			}
+			else if(NodeN == '4')
+			{
+				Category = mx::Category::Combine4;
+				InputsToCopy.Add({ "in1", Input });
+				InputsToCopy.Add({ "in2", Input });
+				InputsToCopy.Add({ "in3", Input });
+				InputsToCopy.Add({ "in4", Input });
+			}
+		}
+		else if((InputN == '2' && NodeN == '3') || (InputN == '3' && NodeN == '4'))
+		{
+			Category = mx::Category::Combine2;
+			InputsToCopy.Add({ "in1", Input });
+			InputsToCreate.Add({ "in2", FAttributeValueArray{{"type","float"}, {"value","1"}} });
+		}
+
+		//copy "in" input into "in1" and create "amount" input as "in2" with the value 360 because UE's angle is a value between [0,1]
+		mx::NodePtr CombineNode = CreateNode(UpstreamNode->getParent()->asA<mx::NodeGraph>(),
+											 UpstreamNode->getName().c_str(),
+											 Category.c_str(),
+											 InputsToCopy, // rename input to match <divide> input "in1"
+											 InputsToCreate);
+
+		// Input now points to the new node
+		Input->setNodeName(CombineNode->getName());
+	}
+}
+
+void FMaterialXSurfaceShaderAbstract::ConnectIfGreaterInputToOutput(MaterialX::NodePtr UpstreamNode, UInterchangeShaderNode* ParentShaderNode, const FString& InputChannelName)
+{
+	UInterchangeShaderNode* NodeIf = CreateShaderNode<UInterchangeShaderNode>(UpstreamNode->getName().c_str(), TEXT("If"));
+	UInterchangeShaderPortsAPI::ConnectDefaultOuputToInput(ParentShaderNode, InputChannelName, NodeIf->GetUniqueID());
+
+	if(mx::InputPtr Input = UpstreamNode->getInput("value1"))
+	{
+		AddAttributeFromValueOrInterface(Input, GetInputName(Input), NodeIf);
+	}
+
+	if(mx::InputPtr Input = UpstreamNode->getInput("value2"))
+	{
+		AddAttributeFromValueOrInterface(Input, GetInputName(Input), NodeIf);
+	}
+
+	if(mx::InputPtr Input = UpstreamNode->getInput("in1"))
+	{
+		AddAttributeFromValueOrInterface(Input, GetInputName(Input), NodeIf);
+	}
+
+	//In that case we also need to add an attribute to AEqualsB
+	mx::InputPtr Input = UpstreamNode->getInput("in2");
+	if(Input)
+	{
+		bool bHasValue = AddAttributeFromValueOrInterface(Input, GetInputName(Input), NodeIf);
+		AddAttributeFromValueOrInterface(Input, TEXT("AEqualsB"), NodeIf);
+
+		if(bHasValue)
+		{
+			//Let's add a new input that is a copy of in2 to connect it to the equal input
+			mx::InputPtr Input3 = UpstreamNode->addInput("in3");
+			Input3->copyContentFrom(Input);
+			SetAttributeNewName(Input3, "AEqualsB");
+		}
+	}
+}
+
+void FMaterialXSurfaceShaderAbstract::ConnectIfGreaterEqInputToOutput(MaterialX::NodePtr UpstreamNode, UInterchangeShaderNode* ParentShaderNode, const FString& InputChannelName)
+{
+	UInterchangeShaderNode* NodeIf = CreateShaderNode<UInterchangeShaderNode>(UpstreamNode->getName().c_str(), TEXT("If"));
+	UInterchangeShaderPortsAPI::ConnectDefaultOuputToInput(ParentShaderNode, InputChannelName, NodeIf->GetUniqueID());
+
+	if(mx::InputPtr Input = UpstreamNode->getInput("value1"))
+	{
+		AddAttributeFromValueOrInterface(Input, GetInputName(Input), NodeIf);
+	}
+
+	if(mx::InputPtr Input = UpstreamNode->getInput("value2"))
+	{
+		AddAttributeFromValueOrInterface(Input, GetInputName(Input), NodeIf);
+	}
+
+	if(mx::InputPtr Input = UpstreamNode->getInput("in2"))
+	{
+		AddAttributeFromValueOrInterface(Input, GetInputName(Input), NodeIf);
+	}
+
+	//In that case we also need to add an attribute to AEqualsB
+	mx::InputPtr Input = UpstreamNode->getInput("in1");
+	if(Input)
+	{
+		bool bHasValue = AddAttributeFromValueOrInterface(Input, GetInputName(Input), NodeIf);
+		AddAttributeFromValueOrInterface(Input, TEXT("AEqualsB"), NodeIf);
+
+		if(bHasValue)
+		{
+			//Let's add a new input that is a copy of in2 to connect it to the equal input
+			mx::InputPtr Input3 = UpstreamNode->addInput("in3");
+			Input3->copyContentFrom(Input);
+			SetAttributeNewName(Input3, "AEqualsB");
+		}
+	}
+}
+
+void FMaterialXSurfaceShaderAbstract::ConnectIfEqualInputToOutput(MaterialX::NodePtr UpstreamNode, UInterchangeShaderNode* ParentShaderNode, const FString& InputChannelName)
+{
+	UInterchangeShaderNode* NodeIf = CreateShaderNode<UInterchangeShaderNode>(UpstreamNode->getName().c_str(), TEXT("If"));
+	UInterchangeShaderPortsAPI::ConnectDefaultOuputToInput(ParentShaderNode, InputChannelName, NodeIf->GetUniqueID());
+
+	if(mx::InputPtr Input = UpstreamNode->getInput("value1"); Input && Input->hasValue())
+	{
+		AddAttributeFromValueOrInterface(Input, GetInputName(Input), NodeIf);
+	}
+
+	if(mx::InputPtr Input = UpstreamNode->getInput("value2"); Input && Input->hasValue())
+	{
+		AddAttributeFromValueOrInterface(Input, GetInputName(Input), NodeIf);
+	}
+
+	if(mx::InputPtr Input = UpstreamNode->getInput("in1"); Input && Input->hasValue())
+	{
+		AddAttributeFromValueOrInterface(Input, GetInputName(Input), NodeIf);
+	}
+
+	//In that case we also need to add an attribute to AGreaterThanB
+	mx::InputPtr Input = UpstreamNode->getInput("in2");
+	if(Input)
+	{
+		bool bHasValue = AddAttributeFromValueOrInterface(Input, GetInputName(Input), NodeIf);
+		AddAttributeFromValueOrInterface(Input, TEXT("AGreaterThanB"), NodeIf);
+
+		if(bHasValue)
+		{
+			//Let's add a new input that is a copy of in2 to connect it to the equal input
+			mx::InputPtr Input3 = UpstreamNode->addInput("in3");
+			Input3->copyContentFrom(Input);
+			SetAttributeNewName(Input3, "AGreaterThanB");
+		}
+	}
+}
+
+void FMaterialXSurfaceShaderAbstract::ConnectOutsideInputToOutput(MaterialX::NodePtr UpstreamNode, UInterchangeShaderNode* ParentShaderNode, const FString& InputChannelName)
+{
+	//in * (1 - mask)
+	UInterchangeShaderNode* NodeMultiply = CreateShaderNode<UInterchangeShaderNode>(UpstreamNode->getName().c_str(), TEXT("Multiply"));
+	AddAttributeFromValueOrInterface(UpstreamNode->getInput("in"), TEXT("A"), NodeMultiply);
+	UInterchangeShaderPortsAPI::ConnectDefaultOuputToInput(ParentShaderNode, InputChannelName, NodeMultiply->GetUniqueID());
+
+	UInterchangeShaderNode* NodeOneMinus = CreateShaderNode<UInterchangeShaderNode>(UpstreamNode->getName().c_str() + FString(TEXT("_OneMinus")), TEXT("OneMinus"));
+	AddAttributeFromValueOrInterface(UpstreamNode->getInput("mask"), TEXT("Input"), NodeOneMinus);
+	UInterchangeShaderPortsAPI::ConnectDefaultOuputToInput(NodeMultiply, TEXT("B"), NodeOneMinus->GetUniqueID());
+}
+
+void FMaterialXSurfaceShaderAbstract::ConnectPositionInputToOutput(MaterialX::NodePtr UpstreamNode, UInterchangeShaderNode* ParentShaderNode, const FString& InputChannelName)
+{
+	// MaterialX defines the space as: object, model, world
+	// model: The local coordinate space of the geometry, before any local deformations or global transforms have been applied.
+	// object: The local coordinate space of the geometry, after local deformations have been applied, but before any global transforms.
+	// world : The global coordinate space of the geometry, after local deformationsand global transforms have been applied.
+
+	// For the moment we don't have the distinction between model/object, so let's just create an UMaterialExpressionWorldPosition
+	// In case of model/object we need to add a TransformPoint from world to local space
+	UInterchangeShaderNode* PositionNode = CreateShaderNode<UInterchangeShaderNode>(UpstreamNode->getName().c_str() + FString(TEXT("_Position")), TEXT("WorldPosition"));
+
+	// In case of the position node, it seems that the unit is different, we assume for now a conversion from mm -> m, even if UE by default is cm
+	// See standard_surface_marble_solid file, especially on the fractal3d node
+	UInterchangeShaderNode* UnitNode = CreateShaderNode<UInterchangeShaderNode>(UpstreamNode->getName().c_str(), TEXT("Multiply"));
+	UnitNode->AddFloatAttribute(UInterchangeShaderPortsAPI::MakeInputValueKey(TEXT("B")), 0.001f);
+	UInterchangeShaderPortsAPI::ConnectDefaultOuputToInput(UnitNode, TEXT("A"), PositionNode->GetUniqueID());
+
+	UInterchangeShaderNode* NodeToConnectTo = ParentShaderNode;
+	FString InputToConnectTo = InputChannelName;
+
+	std::string Space = "world";
+	mx::InputPtr InputSpace = UpstreamNode->getInput("space");
+
+	if(InputSpace)
+	{
+		Space = InputSpace->getValueString();
+	}
+
+	//the default space defined by the nodedef is "object"
+	if(Space != "world" || !InputSpace)
+	{
+		using namespace UE::Interchange::Materials::Standard::Nodes;
+
+		UInterchangeShaderNode* TransformNode = CreateShaderNode<UInterchangeShaderNode>(UpstreamNode->getName().c_str() + FString(TEXT("_Transform")), TransformPosition::Name.ToString());
+		NodeToConnectTo = TransformNode;
+		InputToConnectTo = TransformPosition::Inputs::Input.ToString();
+		TransformNode->AddInt32Attribute(TransformPosition::Attributes::TransformSourceType, EMaterialPositionTransformSource::TRANSFORMPOSSOURCE_World);
+		TransformNode->AddInt32Attribute(TransformPosition::Attributes::TransformType, EMaterialPositionTransformSource::TRANSFORMPOSSOURCE_Local);
+		UInterchangeShaderPortsAPI::ConnectDefaultOuputToInput(ParentShaderNode, InputChannelName, TransformNode->GetUniqueID());
+	}
+
+	UInterchangeShaderPortsAPI::ConnectDefaultOuputToInput(NodeToConnectTo, InputToConnectTo, UnitNode->GetUniqueID());
+}
+
+void FMaterialXSurfaceShaderAbstract::ConnectNormalInputToOutput(MaterialX::NodePtr UpstreamNode, UInterchangeShaderNode* ParentShaderNode, const FString& InputChannelName)
+{
+	// MaterialX defines the space as: object, model, world
+	// model: The local coordinate space of the geometry, before any local deformations or global transforms have been applied.
+	// object: The local coordinate space of the geometry, after local deformations have been applied, but before any global transforms.
+	// world : The global coordinate space of the geometry, after local deformationsand global transforms have been applied.
+
+	// For the moment we don't have the distinction between model/object, so let's just create an UMaterialExpressionVertexNormalWS
+	// In case of model/object we need to add a TransformVector from world to local space
+	UInterchangeShaderNode* NormalNode = CreateShaderNode<UInterchangeShaderNode>(UpstreamNode->getName().c_str(), TEXT("VertexNormalWS"));
+	UInterchangeShaderNode* NodeToConnectTo = ParentShaderNode;
+	FString InputToConnectTo = InputChannelName;
+
+	std::string Space = "world";
+	mx::InputPtr InputSpace = UpstreamNode->getInput("space");
+
+	if(InputSpace)
+	{
+		Space = InputSpace->getValueString();
+	}
+
+	//the default space defined by the nodedef is "object"
+	if(Space != "world" || !InputSpace)
+	{
+		using namespace UE::Interchange::Materials::Standard::Nodes;
+
+		UInterchangeShaderNode* TransformNode = CreateShaderNode<UInterchangeShaderNode>(UpstreamNode->getName().c_str() + FString(TEXT("_Transform")), TransformVector::Name.ToString());
+		NodeToConnectTo = TransformNode;
+		InputToConnectTo = TransformVector::Inputs::Input.ToString();
+		TransformNode->AddInt32Attribute(TransformVector::Attributes::TransformSourceType, EMaterialVectorCoordTransformSource::TRANSFORMSOURCE_World);
+		TransformNode->AddInt32Attribute(TransformVector::Attributes::TransformType, EMaterialVectorCoordTransform::TRANSFORM_Local);
+		UInterchangeShaderPortsAPI::ConnectDefaultOuputToInput(ParentShaderNode, InputChannelName, TransformNode->GetUniqueID());
+	}
+
+	UInterchangeShaderPortsAPI::ConnectDefaultOuputToInput(NodeToConnectTo, InputToConnectTo, NormalNode->GetUniqueID());
+}
+
+void FMaterialXSurfaceShaderAbstract::ConnectTangentInputToOutput(MaterialX::NodePtr UpstreamNode, UInterchangeShaderNode* ParentShaderNode, const FString& InputChannelName)
+{
+	// MaterialX defines the space as: object, model, world
+	// model: The local coordinate space of the geometry, before any local deformations or global transforms have been applied.
+	// object: The local coordinate space of the geometry, after local deformations have been applied, but before any global transforms.
+	// world : The global coordinate space of the geometry, after local deformationsand global transforms have been applied.
+
+	// For the moment we don't have the distinction between model/object, so let's just create an UMaterialExpressionVertexTangentWS
+	// In case of model/object we need to add a TransformVector from world to local space
+	UInterchangeShaderNode* TangentNode = CreateShaderNode<UInterchangeShaderNode>(UpstreamNode->getName().c_str(), TEXT("VertexTangentWS"));
+	UInterchangeShaderNode* NodeToConnectTo = ParentShaderNode;
+	FString InputToConnectTo = InputChannelName;
+
+	std::string Space = "world";
+	mx::InputPtr InputSpace = UpstreamNode->getInput("space");
+
+	if(InputSpace)
+	{
+		Space = InputSpace->getValueString();
+	}
+
+	//the default space defined by the nodedef is "object"
+	if(Space != "world" || !InputSpace)
+	{
+		using namespace UE::Interchange::Materials::Standard::Nodes;
+
+		UInterchangeShaderNode* TransformNode = CreateShaderNode<UInterchangeShaderNode>(UpstreamNode->getName().c_str() + FString(TEXT("_Transform")), TransformVector::Name.ToString());
+		NodeToConnectTo = TransformNode;
+		InputToConnectTo = TransformVector::Inputs::Input.ToString();
+		TransformNode->AddInt32Attribute(TransformVector::Attributes::TransformSourceType, EMaterialVectorCoordTransformSource::TRANSFORMSOURCE_World);
+		TransformNode->AddInt32Attribute(TransformVector::Attributes::TransformType, EMaterialVectorCoordTransform::TRANSFORM_Local);
+		UInterchangeShaderPortsAPI::ConnectDefaultOuputToInput(ParentShaderNode, InputChannelName, TransformNode->GetUniqueID());
+	}
+
+	UInterchangeShaderPortsAPI::ConnectDefaultOuputToInput(NodeToConnectTo, InputToConnectTo, TangentNode->GetUniqueID());
+}
+
+void FMaterialXSurfaceShaderAbstract::ConnectBitangentInputToOutput(MaterialX::NodePtr UpstreamNode, UInterchangeShaderNode* ParentShaderNode, const FString& InputChannelName)
+{
+	// MaterialX defines the space as: object, model, world
+	// model: The local coordinate space of the geometry, before any local deformations or global transforms have been applied.
+	// object: The local coordinate space of the geometry, after local deformations have been applied, but before any global transforms.
+	// world : The global coordinate space of the geometry, after local deformationsand global transforms have been applied.
+
+	// For the moment we don't have the distinction between model/object, so let's just do the cross product between the normal and the tangent
+	// In case of model/object we need to add a TransformVector from world to local space
+	UInterchangeShaderNode* NormalNode = CreateShaderNode<UInterchangeShaderNode>(UpstreamNode->getName().c_str() + FString(TEXT("_Normal")), TEXT("VertexNormalWS"));
+	UInterchangeShaderNode* TangentNode = CreateShaderNode<UInterchangeShaderNode>(UpstreamNode->getName().c_str() + FString(TEXT("_Tangent")), TEXT("VertexTangentWS"));
+	UInterchangeShaderNode* BitangentNode = CreateShaderNode<UInterchangeShaderNode>(UpstreamNode->getName().c_str(), TEXT("CrossProduct"));
+
+	UInterchangeShaderPortsAPI::ConnectDefaultOuputToInput(BitangentNode, TEXT("A"), NormalNode->GetUniqueID());
+	UInterchangeShaderPortsAPI::ConnectDefaultOuputToInput(BitangentNode, TEXT("B"), TangentNode->GetUniqueID());
+
+
+	UInterchangeShaderNode* NodeToConnectTo = ParentShaderNode;
+	FString InputToConnectTo = InputChannelName;
+
+	std::string Space = "world";
+	mx::InputPtr InputSpace = UpstreamNode->getInput("space");
+
+	if(InputSpace)
+	{
+		Space = InputSpace->getValueString();
+	}
+
+	//the default space defined by the nodedef is "object"
+	if(Space != "world" || !InputSpace)
+	{
+		using namespace UE::Interchange::Materials::Standard::Nodes;
+
+		UInterchangeShaderNode* TransformNode = CreateShaderNode<UInterchangeShaderNode>(UpstreamNode->getName().c_str() + FString(TEXT("_Transform")), TransformVector::Name.ToString());
+		NodeToConnectTo = TransformNode;
+		InputToConnectTo = TransformVector::Inputs::Input.ToString();
+		TransformNode->AddInt32Attribute(TransformVector::Attributes::TransformSourceType, EMaterialVectorCoordTransformSource::TRANSFORMSOURCE_World);
+		TransformNode->AddInt32Attribute(TransformVector::Attributes::TransformType, EMaterialVectorCoordTransform::TRANSFORM_Local);
+		UInterchangeShaderPortsAPI::ConnectDefaultOuputToInput(ParentShaderNode, InputChannelName, TransformNode->GetUniqueID());
+	}
+
+	UInterchangeShaderPortsAPI::ConnectDefaultOuputToInput(NodeToConnectTo, InputToConnectTo, BitangentNode->GetUniqueID());
+}
+
+void FMaterialXSurfaceShaderAbstract::ConnectTimeInputToOutput(MaterialX::NodePtr UpstreamNode, UInterchangeShaderNode* ParentShaderNode, const FString& InputChannelName)
+{
+	using namespace UE::Interchange::Materials::Standard::Nodes;
+	UInterchangeShaderNode* TimeNode = CreateShaderNode<UInterchangeShaderNode>(UpstreamNode->getName().c_str(), TEXT("Time"));
+	TimeNode->AddBooleanAttribute(Time::Attributes::OverridePeriod, true);
+
+	float FPS;
+	mx::InputPtr Input = UpstreamNode->getInput("fps");
+
+	//Take the default value from the node definition
+	if(!Input)
+	{
+		Input = UpstreamNode->getNodeDef()->getInput("fps");
+	}
+
+	FPS = mx::fromValueString<float>(Input->getValueString());
+
+	//UE is a period
+	TimeNode->AddFloatAttribute(Time::Attributes::Period, 1.f / FPS);
+
+	UInterchangeShaderPortsAPI::ConnectDefaultOuputToInput(ParentShaderNode, InputChannelName, TimeNode->GetUniqueID());
+}
+
+void FMaterialXSurfaceShaderAbstract::ConnectNoise3DInputToOutput(MaterialX::NodePtr UpstreamNode, UInterchangeShaderNode* ParentShaderNode, const FString& InputChannelName)
+{
+	using namespace UE::Interchange::Materials::Standard::Nodes;
+
+	// MaterialX defines the Noise3d as Perlin Noise which is multiplied by the Amplitude then Added to Pivot
+	UInterchangeShaderNode* NoiseNode = CreateShaderNode<UInterchangeShaderNode>(UpstreamNode->getName().c_str(), Noise::Name.ToString());
+	NoiseNode->AddInt32Attribute(Noise::Attributes::Function, ENoiseFunction::NOISEFUNCTION_GradientTex);
+	NoiseNode->AddBooleanAttribute(Noise::Attributes::Turbulence, false);
+	NoiseNode->AddFloatAttribute(Noise::Attributes::OutputMin, 0);
+
+	UInterchangeShaderNode* NodeToConnect = NoiseNode;
+
+	// Multiply Node
+	auto ConnectNodeToInput = [&](mx::InputPtr Input, UInterchangeShaderNode* NodeToConnectTo, const FString& ShaderType, int32 IndexAttrib) -> UInterchangeShaderNode*
+	{
+		if(!Input)
+		{
+			return nullptr;
+		}
+
+		const FString ShaderNodeName = UpstreamNode->getName().c_str() + FString{ TEXT("_") } + ShaderType;
+		UInterchangeShaderNode* ShaderNode = CreateShaderNode<UInterchangeShaderNode>(ShaderNodeName, ShaderType);
+
+		UInterchangeShaderPortsAPI::ConnectDefaultOuputToInput(ShaderNode, TEXT("A"), NodeToConnectTo->GetUniqueID());
+
+		// Connect the amplitude node to the shader node not the noise
+		// it will be handle during the upstream-downstream connection phase
+		// The index is here for the unicity of the attribute
+		UpstreamNode->setAttribute(mx::Attributes::ParentName + std::to_string(IndexAttrib), TCHAR_TO_ANSI(*ShaderNodeName));
+		AddAttributeFromValueOrInterface(Input, TEXT("B"), ShaderNode);
+
+		return ShaderNode;
+	};
+
+	if(UInterchangeShaderNode* MultiplyNode = ConnectNodeToInput(UpstreamNode->getInput("amplitude"), NoiseNode, TEXT("Multiply"), 0))
+	{
+		NodeToConnect = MultiplyNode;
+	}
+
+
+	if(UInterchangeShaderNode* AddNode = ConnectNodeToInput(UpstreamNode->getInput("pivot"), NodeToConnect, TEXT("Add"), 1))
+	{
+		NodeToConnect = AddNode;
+	}
+
+	UInterchangeShaderPortsAPI::ConnectDefaultOuputToInput(ParentShaderNode, InputChannelName, NodeToConnect->GetUniqueID());
+}
+
+void FMaterialXSurfaceShaderAbstract::ConnectCellNoise3DInputToOutput(MaterialX::NodePtr UpstreamNode, UInterchangeShaderNode* ParentShaderNode, const FString& InputChannelName)
+{
+	using namespace UE::Interchange::Materials::Standard::Nodes;
+
+	// Let's use a vector noise for this one, the only one that is close to MaterialX implementation
+	UInterchangeShaderNode* NoiseNode = CreateShaderNode<UInterchangeShaderNode>(UpstreamNode->getName().c_str(), VectorNoise::Name.ToString());
+	NoiseNode->AddInt32Attribute(VectorNoise::Attributes::Function, EVectorNoiseFunction::VNF_CellnoiseALU);
+
+	UInterchangeShaderPortsAPI::ConnectDefaultOuputToInput(ParentShaderNode, InputChannelName, NoiseNode->GetUniqueID());
+}
+
+void FMaterialXSurfaceShaderAbstract::ConnectWorleyNoise3DInputToOutput(MaterialX::NodePtr UpstreamNode, UInterchangeShaderNode* ParentShaderNode, const FString& InputChannelName)
+{
+	using namespace UE::Interchange::Materials::Standard::Nodes;
+
+	//Also called Voronoi, the implementation is a bit different in UE, especially we don't have access to the jitter
+	UInterchangeShaderNode* NoiseNode = CreateShaderNode<UInterchangeShaderNode>(UpstreamNode->getName().c_str(), Noise::Name.ToString());
+	NoiseNode->AddInt32Attribute(Noise::Attributes::Function, ENoiseFunction::NOISEFUNCTION_VoronoiALU);
+
+	UInterchangeShaderPortsAPI::ConnectDefaultOuputToInput(ParentShaderNode, InputChannelName, NoiseNode->GetUniqueID());
+}
+
+void FMaterialXSurfaceShaderAbstract::ConnectHeightToNormalInputToOutput(MaterialX::NodePtr UpstreamNode, UInterchangeShaderNode* ParentShaderNode, const FString& InputChannelName)
+{
+	if(mx::InputPtr Input = UpstreamNode->getInput("in"))
+	{
+		// Image node will become this node
+		if(mx::NodePtr ConnectedNode = Input->getConnectedNode();
+		   ConnectedNode && ConnectedNode->getCategory() == mx::Category::Image)
+		{
+			//we need to copy the content of the image node to this node
+			UpstreamNode->copyContentFrom(ConnectedNode);
+
+			SetMatchingInputsNames(UpstreamNode);
+
+			//the copy overwrite every attribute of the node, so we need to get them back, essentially the type and the renaming
+			// the output is always a vec3
+			UpstreamNode->setType(mx::Type::Vector3);
+
+			mx::NodeGraphPtr Graph = UpstreamNode->getParent()->asA<mx::NodeGraph>();
+			Graph->removeNode(ConnectedNode->getName());
+
+			using namespace UE::Interchange::Materials::Standard::Nodes;
+
+			if(UInterchangeTextureNode* TextureNode = CreateTextureNode<UInterchangeTexture2DNode>(UpstreamNode))
+			{
+				UInterchangeShaderNode* HeightMapNode = CreateShaderNode<UInterchangeShaderNode>(UpstreamNode->getName().c_str(), NormalFromHeightMap::Name.ToString());
+				UInterchangeShaderPortsAPI::ConnectDefaultOuputToInput(ParentShaderNode, InputChannelName, HeightMapNode->GetUniqueID());
+
+				const FString TextureNodeName = UpstreamNode->getName().c_str() + FString{ "_texture" };
+				UInterchangeShaderNode* TextureShaderNode = CreateShaderNode<UInterchangeShaderNode>(TextureNodeName, TextureObject::Name.ToString());
+				TextureShaderNode->AddStringAttribute(UInterchangeShaderPortsAPI::MakeInputValueKey(TextureObject::Inputs::Texture.ToString()), TextureNode->GetUniqueID());
+				UInterchangeShaderPortsAPI::ConnectDefaultOuputToInput(HeightMapNode, NormalFromHeightMap::Inputs::HeightMap.ToString(), TextureShaderNode->GetUniqueID());
+
+				AddAttributeFromValueOrInterface(UpstreamNode->getInput("scale"), NormalFromHeightMap::Inputs::Intensity.ToString(), HeightMapNode);
+			}
+			else
+			{
+				AddAttributeFromValueOrInterface(UpstreamNode->getInput(mx::NodeGroup::Texture2D::Inputs::Default), InputChannelName, ParentShaderNode);
+			}
+		}
+		else
+		{
+			// For the moment it doesn't make sense to plug a value to it, so let's plug directly the child to the parent, in the future we could implement a Sobel and handle a multi output
+			SetAttributeNewName(Input, TCHAR_TO_UTF8(*InputChannelName)); //let's take the parent node's input name
+			ShaderNodes.Add(UpstreamNode->getName().c_str(), ParentShaderNode);
+		}
+	}
+}
+
+void FMaterialXSurfaceShaderAbstract::ConnectBlurInputToOutput(MaterialX::NodePtr UpstreamNode, UInterchangeShaderNode* ParentShaderNode, const FString& InputChannelName)
+{
+	if(mx::InputPtr Input = UpstreamNode->getInput("in"))
+	{
+		// Image node will become this node
+		if(mx::NodePtr ConnectedNode = Input->getConnectedNode();
+		   ConnectedNode && ConnectedNode->getCategory() == mx::Category::Image)
+		{
+			std::string NodeType = UpstreamNode->getType();
+
+			//we need to copy the content of the image node to this node
+			UpstreamNode->copyContentFrom(ConnectedNode);
+
+			SetMatchingInputsNames(UpstreamNode);
+
+			//the copy overwrites every attribute of the node, so we need to get them back, essentially the type and the renaming
+			UpstreamNode->setType(NodeType);
+
+			mx::NodeGraphPtr Graph = UpstreamNode->getParent()->asA<mx::NodeGraph>();
+			Graph->removeNode(ConnectedNode->getName());
+
+			using namespace UE::Interchange::Materials::Standard::Nodes;
+
+			if(UInterchangeTextureNode* TextureNode = CreateTextureNode<UInterchangeTextureBlurNode>(UpstreamNode))
+			{
+				FString OutputChannel{ TEXT("RGB") };
+
+				if(NodeType == mx::Type::Vector4 || NodeType == mx::Type::Color4)
+				{
+					OutputChannel = TEXT("RGBA");
+				}
+				else if(NodeType == mx::Type::Float)
+				{
+					OutputChannel = TEXT("R");
+				}
+
+				UInterchangeShaderNode* TextureShaderNode = CreateShaderNode<UInterchangeShaderNode>(UpstreamNode->getName().c_str(), TextureSampleBlur::Name.ToString());
+				TextureShaderNode->AddStringAttribute(UInterchangeShaderPortsAPI::MakeInputValueKey(TextureSampleBlur::Inputs::Texture.ToString()), TextureNode->GetUniqueID());
+				UInterchangeShaderPortsAPI::ConnectOuputToInput(ParentShaderNode, InputChannelName, TextureShaderNode->GetUniqueID(), OutputChannel);
+
+				if(mx::InputPtr InputKernel = UpstreamNode->getInput("filtertype"))
+				{
+					// By default TextureSampleBox uses gaussian filter
+					if(InputKernel->getValueString() == "box")
+					{
+						TextureShaderNode->AddInt32Attribute(TextureSampleBlur::Attributes::Filter, 0);
+					}
+				}
+
+				if(mx::InputPtr InputKernel = UpstreamNode->getInput("size"))
+				{
+					float KernelSize = mx::fromValueString<float>(InputKernel->getValueString());
+					TextureShaderNode->AddFloatAttribute(TextureSampleBlur::Attributes::KernelSize, KernelSize);
+				}
+			}
+			else
+			{
+				AddAttributeFromValueOrInterface(UpstreamNode->getInput(mx::NodeGroup::Texture2D::Inputs::Default), InputChannelName, ParentShaderNode);
+			}
+		}
+		else
+		{
+			// For a blur it doesn't make sense if there's no image input
+			SetAttributeNewName(Input, TCHAR_TO_UTF8(*InputChannelName)); //let's take the parent node's input name
+			ShaderNodes.Add(UpstreamNode->getName().c_str(), ParentShaderNode);
+		}
+	}
+}
+
+UInterchangeShaderNode* FMaterialXSurfaceShaderAbstract::CreateMaskShaderNode(uint8 RGBA, const FString& NodeName)
+{
+	bool bR = (0b1000 & RGBA) >> 3;
+	bool bG = (0b0100 & RGBA) >> 2;
+	bool bB = (0b0010 & RGBA) >> 1;
+	bool bA = (0b0001 & RGBA) >> 0;
+	using namespace UE::Interchange::Materials::Standard::Nodes;
+	UInterchangeShaderNode* MaskShaderNode = CreateShaderNode<UInterchangeShaderNode>(NodeName, Mask::Name.ToString());
+	MaskShaderNode->AddBooleanAttribute(Mask::Attributes::R, bR);
+	MaskShaderNode->AddBooleanAttribute(Mask::Attributes::G, bG);
+	MaskShaderNode->AddBooleanAttribute(Mask::Attributes::B, bB);
+	MaskShaderNode->AddBooleanAttribute(Mask::Attributes::A, bA);
+
+	return MaskShaderNode;
+}
+
+const FString& FMaterialXSurfaceShaderAbstract::GetMatchedInputName(MaterialX::NodePtr Node, MaterialX::InputPtr Input) const
+{
+	static FString EmptyString;
+
+	FMaterialXManager& Manager = FMaterialXManager::GetInstance();
+
+	if(Input)
+	{
+		const FString NodeCategory{ Node->getCategory().c_str() };
+		const FString InputName{ GetInputName(Input) };
+		
+		if(const FString* Result = Manager.FindMatchingInput({ NodeCategory, InputName }))
+		{
+			return *Result;
+		}
+		else if((Result = Manager.FindMatchingInput({ EmptyString, InputName })))
+		{
+			return *Result;
+		}
+	}
+
+	return EmptyString;
+}
+
+FString FMaterialXSurfaceShaderAbstract::GetInputName(MaterialX::InputPtr Input) const
+{
+	if(Input->hasAttribute(mx::Attributes::NewName))
+	{
+		return Input->getAttribute(mx::Attributes::NewName).c_str();
+	}
+	else
+	{
+		return Input->getName().c_str();
+	}
+}
+
+FString FMaterialXSurfaceShaderAbstract::GetFilePrefix(MaterialX::ElementPtr Element) const
+{
+	FString FilePrefix;
+
+	if(Element)
+	{
+		if(Element->hasFilePrefix())
+		{
+			return FString(Element->getFilePrefix().c_str());
+		}
+		else
+		{
+			return GetFilePrefix(Element->getParent());
+		}
+	}
+
+	return FilePrefix;
+}
+
+FLinearColor FMaterialXSurfaceShaderAbstract::GetVector(MaterialX::InputPtr Input) const
+{
+	FLinearColor LinearColor;
+
+	if(Input->getType() == mx::Type::Vector2)
+	{
+		mx::Vector2 Color = mx::fromValueString<mx::Vector2>(Input->getValueString());
+		LinearColor = FLinearColor{ Color[0], Color[1], 0 };
+	}
+	else if(Input->getType() == mx::Type::Vector3)
+	{
+		mx::Vector3 Color = mx::fromValueString<mx::Vector3>(Input->getValueString());
+		LinearColor = FLinearColor{ Color[0], Color[1], Color[2] };
+	}
+	else if(Input->getType() == mx::Type::Vector4)
+	{
+		mx::Vector4 Color = mx::fromValueString<mx::Vector4>(Input->getValueString());
+		LinearColor = FLinearColor{ Color[0], Color[1], Color[2], Color[3] };
+	}
+	else
+	{
+		ensureMsgf(false, TEXT("input type can only be a vectorN"));
+	}
+
+	return LinearColor;
+}
+
+FString FMaterialXSurfaceShaderAbstract::GetAttributeParentName(MaterialX::NodePtr Node) const
+{
+	FString ParentName;
+	const mx::StringVec& Attributes = Node->getAttributeNames();
+
+	// For consistency the parent attribute has an index attach to it to ensure unicity
+	// Attributes are set in order, we only need to take the first one (it will be remove after that)
+	for(const std::string& Attrib : Attributes)
+	{
+		if(Attrib.find(mx::Attributes::ParentName) != std::string::npos)
+		{
+			ParentName = Node->getAttribute(Attrib).c_str();
+			Node->removeAttribute(Attrib);
+			break;
+		}
+	}
+
+	return ParentName.IsEmpty() ? Node->getName().c_str() : ParentName;
+}
+
+void FMaterialXSurfaceShaderAbstract::RegisterConnectNodeOutputToInputDelegates()
+{
+	MatchingConnectNodeDelegates.Add(mx::Category::Constant, FOnConnectNodeOutputToInput::CreateSP(this, &FMaterialXSurfaceShaderAbstract::ConnectConstantInputToOutput));
+	MatchingConnectNodeDelegates.Add(mx::Category::Extract, FOnConnectNodeOutputToInput::CreateSP(this, &FMaterialXSurfaceShaderAbstract::ConnectExtractInputToOutput));
+	MatchingConnectNodeDelegates.Add(mx::Category::Dot, FOnConnectNodeOutputToInput::CreateSP(this, &FMaterialXSurfaceShaderAbstract::ConnectDotInputToOutput));
+	MatchingConnectNodeDelegates.Add(mx::Category::NormalMap, FOnConnectNodeOutputToInput::CreateSP(this, &FMaterialXSurfaceShaderAbstract::ConnectDotInputToOutput));
+	MatchingConnectNodeDelegates.Add(mx::Category::TransformPoint, FOnConnectNodeOutputToInput::CreateSP(this, &FMaterialXSurfaceShaderAbstract::ConnectTransformPositionInputToOutput));
+	MatchingConnectNodeDelegates.Add(mx::Category::TransformVector, FOnConnectNodeOutputToInput::CreateSP(this, &FMaterialXSurfaceShaderAbstract::ConnectTransformVectorInputToOutput));
+	MatchingConnectNodeDelegates.Add(mx::Category::TransformNormal, FOnConnectNodeOutputToInput::CreateSP(this, &FMaterialXSurfaceShaderAbstract::ConnectTransformVectorInputToOutput));
+	MatchingConnectNodeDelegates.Add(mx::Category::Rotate2D, FOnConnectNodeOutputToInput::CreateSP(this, &FMaterialXSurfaceShaderAbstract::ConnectRotate2DInputToOutput));
+	MatchingConnectNodeDelegates.Add(mx::Category::Rotate3D, FOnConnectNodeOutputToInput::CreateSP(this, &FMaterialXSurfaceShaderAbstract::ConnectRotate3DInputToOutput));
+	MatchingConnectNodeDelegates.Add(mx::Category::Image, FOnConnectNodeOutputToInput::CreateSP(this, &FMaterialXSurfaceShaderAbstract::ConnectImageInputToOutput));
+	MatchingConnectNodeDelegates.Add(mx::Category::Convert, FOnConnectNodeOutputToInput::CreateSP(this, &FMaterialXSurfaceShaderAbstract::ConnectConvertInputToOutput));
+	MatchingConnectNodeDelegates.Add(mx::Category::IfGreater, FOnConnectNodeOutputToInput::CreateSP(this, &FMaterialXSurfaceShaderAbstract::ConnectIfGreaterInputToOutput));
+	MatchingConnectNodeDelegates.Add(mx::Category::IfGreaterEq, FOnConnectNodeOutputToInput::CreateSP(this, &FMaterialXSurfaceShaderAbstract::ConnectIfGreaterEqInputToOutput));
+	MatchingConnectNodeDelegates.Add(mx::Category::IfEqual, FOnConnectNodeOutputToInput::CreateSP(this, &FMaterialXSurfaceShaderAbstract::ConnectIfEqualInputToOutput));
+	MatchingConnectNodeDelegates.Add(mx::Category::Outside, FOnConnectNodeOutputToInput::CreateSP(this, &FMaterialXSurfaceShaderAbstract::ConnectOutsideInputToOutput));
+	MatchingConnectNodeDelegates.Add(mx::Category::Position, FOnConnectNodeOutputToInput::CreateSP(this, &FMaterialXSurfaceShaderAbstract::ConnectPositionInputToOutput));
+	MatchingConnectNodeDelegates.Add(mx::Category::Normal, FOnConnectNodeOutputToInput::CreateSP(this, &FMaterialXSurfaceShaderAbstract::ConnectNormalInputToOutput));
+	MatchingConnectNodeDelegates.Add(mx::Category::Tangent, FOnConnectNodeOutputToInput::CreateSP(this, &FMaterialXSurfaceShaderAbstract::ConnectTangentInputToOutput));
+	MatchingConnectNodeDelegates.Add(mx::Category::Bitangent, FOnConnectNodeOutputToInput::CreateSP(this, &FMaterialXSurfaceShaderAbstract::ConnectBitangentInputToOutput));
+	MatchingConnectNodeDelegates.Add(mx::Category::Time, FOnConnectNodeOutputToInput::CreateSP(this, &FMaterialXSurfaceShaderAbstract::ConnectTimeInputToOutput));
+	MatchingConnectNodeDelegates.Add(mx::Category::Noise3D, FOnConnectNodeOutputToInput::CreateSP(this, &FMaterialXSurfaceShaderAbstract::ConnectNoise3DInputToOutput));
+	MatchingConnectNodeDelegates.Add(mx::Category::CellNoise3D, FOnConnectNodeOutputToInput::CreateSP(this, &FMaterialXSurfaceShaderAbstract::ConnectCellNoise3DInputToOutput));
+	MatchingConnectNodeDelegates.Add(mx::Category::WorleyNoise3D, FOnConnectNodeOutputToInput::CreateSP(this, &FMaterialXSurfaceShaderAbstract::ConnectWorleyNoise3DInputToOutput));
+	MatchingConnectNodeDelegates.Add(mx::Category::Blur, FOnConnectNodeOutputToInput::CreateSP(this, &FMaterialXSurfaceShaderAbstract::ConnectBlurInputToOutput));
+	MatchingConnectNodeDelegates.Add(mx::Category::HeightToNormal, FOnConnectNodeOutputToInput::CreateSP(this, &FMaterialXSurfaceShaderAbstract::ConnectHeightToNormalInputToOutput));
+}
+
+void FMaterialXSurfaceShaderAbstract::SetMatchingInputsNames(MaterialX::NodePtr Node) const
+{
+	if(Node)
+	{
+		if(const std::string& IsVisited = Node->getAttribute(mx::Attributes::IsVisited); IsVisited.empty())
+		{
+			Node->setAttribute(mx::Attributes::IsVisited, "true");
+
+			for(mx::InputPtr Input : Node->getInputs())
+			{
+				if(const FString& Name = GetMatchedInputName(Node, Input); !Name.IsEmpty())
+				{
+					SetAttributeNewName(Input, TCHAR_TO_UTF8(*Name));
+				}
+			}
+		}
+	}
+}
+
+void FMaterialXSurfaceShaderAbstract::SetAttributeNewName(MaterialX::InputPtr Input, const char* NewName) const
+{
+	Input->setAttribute(mx::Attributes::NewName, NewName);
+}
+
+#undef LOCTEXT_NAMESPACE
+#endif
\ No newline at end of file
diff --git a/Engine/Plugins/Interchange/Runtime/Source/Import/Private/MaterialX/MaterialXUtils/MaterialXSurfaceUnlitShader.cpp b/Engine/Plugins/Interchange/Runtime/Source/Import/Private/MaterialX/MaterialXUtils/MaterialXSurfaceUnlitShader.cpp
new file mode 100644
--- /dev/null
+++ b/Engine/Plugins/Interchange/Runtime/Source/Import/Private/MaterialX/MaterialXUtils/MaterialXSurfaceUnlitShader.cpp
@@ -0,0 +1,44 @@
+// Copyright Epic Games, Inc. All Rights Reserved.
+
+#if WITH_EDITOR
+#include "MaterialXSurfaceUnlitShader.h"
+
+namespace mx = MaterialX;
+
+FMaterialXSurfaceUnlitShader::FMaterialXSurfaceUnlitShader(UInterchangeBaseNodeContainer& BaseNodeContainer)
+	: FMaterialXSurfaceShaderAbstract(BaseNodeContainer)
+{
+	NodeDefinition = mx::NodeDefinition::SurfaceUnlit;
+}
+
+TSharedRef<FMaterialXBase> FMaterialXSurfaceUnlitShader::MakeInstance(UInterchangeBaseNodeContainer& BaseNodeContainer)
+{
+	TSharedRef<FMaterialXSurfaceUnlitShader> Result = MakeShared<FMaterialXSurfaceUnlitShader>(BaseNodeContainer);
+	Result->RegisterConnectNodeOutputToInputDelegates();
+	return Result;
+}
+
+void FMaterialXSurfaceUnlitShader::Translate(mx::NodePtr SurfaceUnlitNode)
+{
+	this->SurfaceShaderNode = SurfaceUnlitNode;
+
+	using namespace UE::Interchange::Materials;
+
+	UInterchangeShaderGraphNode* ShaderGraphNode = CreateShaderNode<UInterchangeShaderGraphNode>(SurfaceShaderNode->getName().c_str(), SurfaceUnlit::Name.ToString());
+
+	//Emission
+	ConnectNodeOutputToInput(mx::SurfaceUnlit::Input::Emission, ShaderGraphNode, SurfaceUnlit::Parameters::Emission.ToString(), mx::SurfaceUnlit::DefaultValue::Float::Emission);
+
+	//Emission Color
+	ConnectNodeOutputToInput(mx::SurfaceUnlit::Input::EmissionColor, ShaderGraphNode, SurfaceUnlit::Parameters::EmissionColor.ToString(), mx::SurfaceUnlit::DefaultValue::Color3::EmissionColor);
+
+	//Opacity
+	ConnectNodeOutputToInput(mx::SurfaceUnlit::Input::Opacity, ShaderGraphNode, SurfaceUnlit::Parameters::Opacity.ToString(), mx::SurfaceUnlit::DefaultValue::Float::Opacity);
+
+	//Transmission
+	ConnectNodeOutputToInput(mx::SurfaceUnlit::Input::Transmission, ShaderGraphNode, SurfaceUnlit::Parameters::Transmission.ToString(), mx::SurfaceUnlit::DefaultValue::Float::Transmission);
+
+	//Transmission Color
+	ConnectNodeOutputToInput(mx::SurfaceUnlit::Input::TransmissionColor, ShaderGraphNode, SurfaceUnlit::Parameters::TransmissionColor.ToString(), mx::SurfaceUnlit::DefaultValue::Color3::TransmissionColor);
+}
+#endif
\ No newline at end of file
diff --git a/Engine/Plugins/Interchange/Runtime/Source/Import/Private/MaterialX/MaterialXUtils/MaterialXSurfaceUnlitShader.h b/Engine/Plugins/Interchange/Runtime/Source/Import/Private/MaterialX/MaterialXUtils/MaterialXSurfaceUnlitShader.h
new file mode 100644
--- /dev/null
+++ b/Engine/Plugins/Interchange/Runtime/Source/Import/Private/MaterialX/MaterialXUtils/MaterialXSurfaceUnlitShader.h
@@ -0,0 +1,23 @@
+// Copyright Epic Games, Inc. All Rights Reserved.
+
+#pragma once
+
+#if WITH_EDITOR
+#include "MaterialX/MaterialXUtils/MaterialXSurfaceShaderAbstract.h"
+
+class FMaterialXSurfaceUnlitShader : public FMaterialXSurfaceShaderAbstract
+{
+protected:
+	
+	template<typename T, ESPMode>
+	friend class SharedPointerInternals::TIntrusiveReferenceController;
+
+	FMaterialXSurfaceUnlitShader(UInterchangeBaseNodeContainer & BaseNodeContainer);
+
+public:
+
+	static TSharedRef<FMaterialXBase> MakeInstance(UInterchangeBaseNodeContainer& BaseNodeContainer);
+
+	virtual void Translate(MaterialX::NodePtr SurfaceUnlitNode) override;
+};
+#endif
\ No newline at end of file
diff --git a/Engine/Plugins/Interchange/Runtime/Source/Import/Public/MaterialX/InterchangeMaterialXTranslator.h b/Engine/Plugins/Interchange/Runtime/Source/Import/Public/MaterialX/InterchangeMaterialXTranslator.h
--- a/Engine/Plugins/Interchange/Runtime/Source/Import/Public/MaterialX/InterchangeMaterialXTranslator.h
+++ b/Engine/Plugins/Interchange/Runtime/Source/Import/Public/MaterialX/InterchangeMaterialXTranslator.h
@@ -6,16 +6,6 @@
 #include "InterchangeTextureNode.h"
 #include "Texture/InterchangeTexturePayloadInterface.h"
 
-#if WITH_EDITOR
-#if defined(PRAGMA_DISABLE_MISSING_BRACES_WARNINGS)
-PRAGMA_DISABLE_MISSING_BRACES_WARNINGS
-#endif
-#include "MaterialXCore/Document.h"
-#if defined(PRAGMA_ENABLE_MISSING_BRACES_WARNINGS)
-PRAGMA_ENABLE_MISSING_BRACES_WARNINGS
-#endif
-#endif
-
 #include "InterchangeMaterialXTranslator.generated.h"
 
 class UInterchangeTextureNode;
@@ -29,8 +19,6 @@ class UInterchangeMaterialXTranslator : public UInterchangeTranslatorBase, publi
 
 public:
 
-	UInterchangeMaterialXTranslator();
-
 	/** Begin UInterchangeTranslatorBase API*/
 
 	virtual EInterchangeTranslatorType GetTranslatorType() const override;
@@ -49,537 +37,4 @@ public:
 	/** End UInterchangeTranslatorBase API*/
 
 	virtual TOptional<UE::Interchange::FImportImage> GetTexturePayloadData(const UInterchangeSourceData* PayloadSourceData, const FString& PayloadKey) const override;
-
-#if WITH_EDITOR
-protected:
-
-	/**
-	 * Process Autodesk's standard surface shader
-	 * 
-	 * @param NodeContainer - The container where to add the translated Interchange nodes.
-	 * @param StandardSurfaceNode - The <standard_surface> in the MaterialX file
-	 * @param Document - The MaterialX Document that contains all the definitions of the loaded libraries
-	 */
-	void ProcessStandardSurface(UInterchangeBaseNodeContainer & NodeContainer, MaterialX::NodePtr StandardSurfaceNode, MaterialX::DocumentPtr Document) const;
-
-	/**
-	 * Process surface_unlit shader
-	 *
-	 * @param NodeContainer - The container where to add the translated Interchange nodes.
-	 * @param SurfaceUnlitNode - The <surface_unlit> in the MaterialX file
-	 * @param Document - The MaterialX Document that contains all the definitions of the loaded libraries
-	 */
-	void ProcessSurfaceUnlit(UInterchangeBaseNodeContainer& NodeContainer, MaterialX::NodePtr SurfaceUnlitNode, MaterialX::DocumentPtr Document) const;
-
-	/**
-	 * Process light shader, MaterialX doesn't standardized lights, but defines the 3 common ones, directional, point and spot
-	 * 
-	 * @param The container where to add the translated Interchange nodes.
-	 * @param LightShaderNode - a MaterialX light shader node
-	 * @param Document - The MaterialX Document that contains all the definitions of the loaded libraries
-	 */
-	void ProcessLightShader(UInterchangeBaseNodeContainer& NodeContainer, MaterialX::NodePtr LightShaderNode, MaterialX::DocumentPtr Document) const;
-
-	/**
-	 * Create a directional light node and set the proper transform in the scene node
-	 *
-	 * @param DirectionalLightShaderNode - The MaterialX node related to a <directional_light>
-	 * @param SceneNode - The interchange scene node, to which the node will be attached to
-	 * @param The container where to add the translated Interchange nodes.
-	 * @param Document - The MaterialX Document that contains all the definitions of the loaded libraries
-	 */
-	UInterchangeBaseLightNode* CreateDirectionalLightNode(MaterialX::NodePtr DirectionalLightShaderNode, UInterchangeSceneNode* SceneNode, UInterchangeBaseNodeContainer& NodeContainer, MaterialX::DocumentPtr Document) const;
-
-	/**
-	 * Create a point  light node and set the proper transform in the scene node
-	 *
-	 * @param PointLightShaderNode - The MaterialX node related to a <point_light>
-	 * @param SceneNode - The interchange scene node, to which the node will be attached to
-	 * @param The container where to add the translated Interchange nodes.
-	 * @param Document - The MaterialX Document that contains all the definitions of the loaded libraries
-	 */
-	UInterchangeBaseLightNode* CreatePointLightNode(MaterialX::NodePtr PointLightShaderNode, UInterchangeSceneNode* SceneNode, UInterchangeBaseNodeContainer& NodeContainer, MaterialX::DocumentPtr Document) const;
-	
-	/**
-	 * Create a spot light node and set the proper transform in the scene node
-	 *
-	 * @param SpotLightShaderNode - The MaterialX node related to a <spot_light>
-	 * @param SceneNode - The interchange scene node, to which the node will be attached to
-	 * @param The container where to add the translated Interchange nodes.
-	 * @param Document - The MaterialX Document that contains all the definitions of the loaded libraries
-	 */
-	UInterchangeBaseLightNode* CreateSpotLightNode(MaterialX::NodePtr SpotLightShaderNode, UInterchangeSceneNode* SceneNode, UInterchangeBaseNodeContainer& NodeContainer, MaterialX::DocumentPtr Document) const;
-	/**
-	 * Connect an ouput in the NodeGraph to the ShaderGraph
-	 * 
-	 * @param InputToNodeGraph - The input from the standard surface to retrieve the output in the NodeGraph
-	 * @param ShaderNode - The Interchange shader node to connect the MaterialX's node graph to
-	 * @param ParentInputName - The name of the input of the shader node to which we want the node graph to be connected to
-	 * @param NamesToShaderNodes - Map of the shader nodes already created
-	 * @param The container where to add the translated Interchange nodes.
-	 * 
-	 * @return true if the given input is attached to one of the outputs of a node graph
-	 */
-	bool ConnectNodeGraphOutputToInput(MaterialX::InputPtr InputToNodeGraph, UInterchangeShaderNode * ShaderNode, const FString& ParentInputName, TMap<FString, UInterchangeShaderNode*> & NamesToShaderNodes, UInterchangeBaseNodeContainer & NodeContainer) const;
-
-	/**
-	 * Create and Connect the output of a MaterialX node that has already a matching in UE to a shader node
-	 * 
-	 * @param Node - The MaterialX node of a given type used to create the appropriate shader node
-	 * @param ParentShaderNode - The shader node to connect to
-	 * @param InputChannelName - The input of the ParentShaderNode to connect to
-	 * @param NamesToShaderNodes - Map of the shader nodes already created
-	 * @param The container where to add the translated Interchange nodes.
-	 * 
-	 * @return true if a shader node has been successfully created and is connected to the given input
-	 */
-	bool ConnectNodeOutputToInput(MaterialX::NodePtr Node, UInterchangeShaderNode* ParentShaderNode, const FString& InputChannelName, TMap<FString, UInterchangeShaderNode*>& NamesToShaderNodes, UInterchangeBaseNodeContainer& NodeContainer) const;
-
-	/**
-	 * Create and Connect manually the output of a MaterialX node to a shader node
-	 *
-	 * @param Edge - The MaterialX edge that has the current node, its parent and the input bridge between the two
-	 * @param ParentShaderNode - The shader node to connect to
-	 * @param InputChannelName - The input of the ParentShaderNode to connect to
-	 * @param NamesToShaderNodes - Map of the shader nodes already created
-	 * @param The container where to add the translated Interchange nodes.
-	 *
-	 * @return true if a shader node has been successfully created and is connected to the given input
-	 */
-	void ConnectNodeCategoryOutputToInput(const MaterialX::Edge & Edge, UInterchangeShaderNode* ParentShaderNode, const FString& InputChannelName, TMap<FString, UInterchangeShaderNode*>& NamesToShaderNodes, UInterchangeBaseNodeContainer& NodeContainer) const;
-
-	/**
-     * Create and Connect a node name directly connected from an input to a shader node
-     *
-     * @param Node - The MaterialX node of a given type used to create the appropriate shader node
-     * @param ParentShaderNode - The shader node to connect to
-     * @param InputChannelName - The input of the ParentShaderNode to connect to
-     * @param NamesToShaderNodes - Map of the shader nodes already created
-     * @param The container where to add the translated Interchange nodes.
-     *
-     * @return true if a shader node has been successfully created and is connected to the given input
-     */
-	bool ConnectNodeNameOutputToInput(MaterialX::InputPtr Input, UInterchangeShaderNode* ShaderNode, const FString& ParentInputName, TMap<FString, UInterchangeShaderNode*>& NamesToShaderNodes, UInterchangeBaseNodeContainer& NodeContainer) const;
-	
-	/**
-	 * Helper template function to create an InterchangeShaderNode
-	 * 
-	 * @param NodeName - The name of the shader node
-	 * @param ShaderType - The shader node's type we want to create
-	 * @param NamesToShaderNodes - Map of the shader nodes already created
-	 * @param The container where to add the translated Interchange nodes.
-	 * 
-	 * @return The shader node that was created
-	 */
-	template<typename ShaderNodeType>
-	ShaderNodeType* CreateShaderNode(const FString& NodeName, const FString& ShaderType, TMap<FString, UInterchangeShaderNode*>& NamesToShaderNodes, UInterchangeBaseNodeContainer& NodeContainer) const
-	{
-		static_assert(std::is_convertible_v<ShaderNodeType*, UInterchangeShaderNode*>, "CreateShaderNode only accepts type that derived from UInterchangeShaderNode");
-
-		ShaderNodeType* Node;
-
-		const FString NodeUID = UInterchangeShaderNode::MakeNodeUid(NodeName, FStringView{});
-		
-		//Test directly in the NodeContainer, because the NamesToShaderNodes can be altered during the node graph either by the parent (dot/normalmap),
-		//or by putting an intermediary node between the child and the parent (tiledimage)
-		if(Node = const_cast<ShaderNodeType*>(Cast<ShaderNodeType>(NodeContainer.GetNode(NodeUID))); !Node)
-		{
-			Node = NewObject<ShaderNodeType>(&NodeContainer);
-			Node->InitializeNode(NodeUID, NodeName, EInterchangeNodeContainerType::TranslatedAsset);
-			NodeContainer.AddNode(Node);
-			if constexpr(std::is_same_v<ShaderNodeType, UInterchangeShaderGraphNode>)
-			{
-				NodeContainer.SetNodeParentUid(NodeUID, FString{});
-			}
-			Node->SetCustomShaderType(ShaderType);
-
-			NamesToShaderNodes.Add(NodeName, Node);
-		}
-
-		return Node;
-	}
-
-	/**
-	 * Helper function to create an InterchangeTextureNode
-	 *
-	 * @param Node - The MaterialX node, it should be of the category <image> no test is done on it
-	 * @param The container where to add the translated Interchange nodes.
-	 * 
-	 * @return The texture node that was created
-	 */
-	template<typename TextureTypeNode>
-	UInterchangeTextureNode* CreateTextureNode(MaterialX::NodePtr Node, UInterchangeBaseNodeContainer& NodeContainer) const
-	{
-		static_assert(std::is_convertible_v<TextureTypeNode*, UInterchangeTextureNode*>, "CreateTextureNode only accepts type that derived from UInterchangeTextureNode");
-		UInterchangeTextureNode* TextureNode = nullptr;
-
-		//A node image should have an input file otherwise the user should check its default value
-		if(Node)
-		{
-			if(MaterialX::InputPtr InputFile = Node->getInput("file"); InputFile && InputFile->hasValue())
-			{
-				FString Filepath{ InputFile->getValueString().c_str() };
-				const FString FilePrefix = GetFilePrefix(InputFile);
-				Filepath = FPaths::Combine(FilePrefix, Filepath);
-				const FString Filename = FPaths::GetCleanFilename(Filepath);
-				const FString TextureNodeUID = TEXT("\\Texture\\") + Filename;
-
-				//Only add the TextureNode once
-				TextureNode = const_cast<UInterchangeTextureNode*>(Cast<UInterchangeTextureNode>(NodeContainer.GetNode(TextureNodeUID)));
-				if(TextureNode == nullptr)
-				{
-					TextureNode = NewObject<TextureTypeNode>(&NodeContainer);
-					TextureNode->InitializeNode(TextureNodeUID, Filename, EInterchangeNodeContainerType::TranslatedAsset);
-					NodeContainer.AddNode(TextureNode);
-
-					if(FPaths::IsRelative(Filepath))
-					{
-						Filepath = FPaths::ConvertRelativePathToFull(FPaths::GetPath(Node->getActiveSourceUri().c_str()), Filepath);
-					}
-
-					TextureNode->SetPayLoadKey(Filepath);
-
-					const FString ColorSpace = GetColorSpace(InputFile);
-					const bool bIsSRGB = ColorSpace == TEXT("srgb_texture");
-					TextureNode->SetCustomSRGB(bIsSRGB);
-				}
-			}
-		}
-
-		return TextureNode;
-	}
-
-	/**
-	 * Get the UE corresponding name of a MaterialX Node category and input for a material
-	 * 
-	 * @param Input - MaterialX input
-	 * 
-	 * @return The matched name of the Node/Input else empty string
-	 */
-	const FString& GetMatchedInputName(MaterialX::NodePtr Node, MaterialX::InputPtr Input) const;
-	
-	/**
-	 * Rename the inputs names of a node to correspond to the one used by UE, it will keep the old inputs names under the attribute "oldname"
-	 * 
-	 * @param Node - Look up to all inputs of Node and rename them to match UE names
-	 */
-	void RenameNodeInputs(MaterialX::NodePtr Node) const;
-
-	/**
-	 * Rename the input name, it will keep the original input name under the attribute MaterialX::Attributes::OriginalName.
-	 * It will keep the uniqueness of the name inside the MaterialX Document
-	 *
-	 * @param Input - The input to rename
-	 * @param NewName - the new name of the input
-	 */
-	void RenameInput(MaterialX::InputPtr Input, const char * NewName) const;
-
-	/**
-	 * Get the input name, use this function instead of getName, because a renaming may have occured and we ensure to have the proper name that will be used by UE inputs
-	 * 
-	 * @param Input - The input to retrieve the name from
-	 * 
-	 * @return The input name
-	 */
-	FString GetInputName(MaterialX::InputPtr Input) const;
-
-	/**
-	 * Retrieve the input from a standard_surface node, or take the default input from the library, 
-	 * this function should only be called after testing the MaterialX libraries have been successfully imported, meaning the node definition of the standard_surface 
-	 * should always be valid
-	 * 
-	 * @param StandardSurface - the <standard_surface> node
-	 * @param InputName - the input name to retrieve
-	 * @param Document - the MaterialX library which has the node definition of the standard_surface
-	 * 
-	 * @return the input from the given name
-	 */
-	MaterialX::InputPtr GetStandardSurfaceInput(MaterialX::NodePtr StandardSurface, const char* InputName, MaterialX::DocumentPtr Document) const;
-
-	/**
-	 * Retrieve the input from a surface_unlit node, or take the default input from the library,
-	 * this function should only be called after testing the MaterialX libraries have been successfully imported, meaning the node definition of the surface_unlit
-	 * should always be valid
-	 *
-	 * @param SurfaceUnlit - the <surface_unlit> node
-	 * @param InputName - the input name to retrieve
-	 * @param Document - the MaterialX library which has the node definition of the surface_unlit
-	 *
-	 * @return the input from the given name
-	 */
-	MaterialX::InputPtr GetSurfaceUnlitInput(MaterialX::NodePtr SurfaceUnlit, const char* InputName, MaterialX::DocumentPtr Document) const;
-
-	/**
-	 * Retrieve the input from a point_light node, or take the default input from the library,
-	 * this function should only be called after testing the MaterialX libraries have been successfully imported, meaning the node definition of the point_light
-	 * should always be valid
-	 *
-	 * @param PointLight - the <point_light> node
-	 * @param InputName - the input name to retrieve
-	 * @param Document - the MaterialX library which has the node definition of the standard_surface
-	 *
-	 * @return the input from the given name
-	 */
-	MaterialX::InputPtr GetPointLightInput(MaterialX::NodePtr PointLight, const char* InputName, MaterialX::DocumentPtr Document) const;
-
-	/**
-	 * Retrieve the input from a directional_light node, or take the default input from the library,
-	 * this function should only be called after testing the MaterialX libraries have been successfully imported, meaning the node definition of the directional_light
-	 * should always be valid
-	 *
-	 * @param DirectionalLight - the <directional_light> node
-	 * @param InputName - the input name to retrieve
-	 * @param Document - the MaterialX library which has the node definition of the standard_surface
-	 *
-	 * @return the input from the given name
-	 */
-	MaterialX::InputPtr GetDirectionalLightInput(MaterialX::NodePtr DirectionalLight, const char* InputName, MaterialX::DocumentPtr Document) const;
-
-	/**
-	 * Retrieve the input from a spot_light node, or take the default input from the library,
-	 * this function should only be called after testing the MaterialX libraries have been successfully imported, meaning the node definition of the spot_light
-	 * should always be valid
-	 *
-	 * @param SpotLight - the <spot_light > node
-	 * @param InputName - the input name to retrieve
-	 * @param Document - the MaterialX library which has the node definition of the standard_surface
-	 *
-	 * @return the input from the given name
-	 */
-	MaterialX::InputPtr GetSpotLightInput(MaterialX::NodePtr SpotLight, const char* InputName, MaterialX::DocumentPtr Document) const;
-
-	/**
-	 * Add an attribute to a shader node from the given MaterialX input if that input has either a value or an interface name
-	 *
-	 * @param Input - The MaterialX input to retrieve and add the value from, must be of type float/color/vector
-	 * @param InputChannelName - The name of the shader node's input to add the attribute
-	 * @param ShaderNode - The shader node to which we want to add the attribute
-	 *
-	 * @return true if the attribute was successfully added
-	 */
-	bool AddAttributeFromValueOrInterface(MaterialX::InputPtr Input, const FString& InputChannelName, UInterchangeShaderNode* ShaderNode) const;
-
-	/**
-	 * Add an attribute to a shader node from the given MaterialX input, only floats and linear colors are supported for the moment
-	 * 
-	 * @param Input - The MaterialX input to retrieve and add the value from, must be of type float/color/vector
-	 * @param InputChannelName - The name of the shader node's input to add the attribute
-	 * @param ShaderNode - The shader node to which we want to add the attribute
-	 * 
-	 * @return true if the attribute was successfully added
-	 */
-	bool AddAttribute(MaterialX::InputPtr Input, const FString& InputChannelName, UInterchangeShaderNode* ShaderNode) const;
-
-	/**
-	 * Add a float attribute to a shader node only if its value taken from the input is not equal to its default value. Return false if the attribute does not exist or if we cannot add it
-	 *
-	 * @param Input - The MaterialX input to retrieve and add the value from, must be of type float/color/vector
-	 * @param InputChannelName - The name of the shader node's input to add the attribute
-	 * @param ShaderNode - The shader node to which we want to add the attribute
-	 * @param DefaultValue - the default value to test the input against
-	 *
-	 * @return true if the attribute was successfully added
-	 */
-	bool AddFloatAttribute(MaterialX::InputPtr Input, const FString& InputChannelName, UInterchangeShaderNode* ShaderNode, float DefaultValue) const;
-
-	/**
-	 * Add a FLinearColor attribute to a shader node only if its value taken from the input is not equal to its default value. Return false if the attribute does not exist or if we cannot add it
-	 *
-	 * @param Input - The MaterialX input to retrieve and add the value from, must be of type float/color/vector
-	 * @param InputChannelName - The name of the shader node's input to add the attribute
-	 * @param ShaderNode - The shader node to which we want to add the attribute
-	 * @param DefaultValue - the default value to test the input against
-	 *
-	 * @return true if the attribute was successfully added
-	 */
-	bool AddLinearColorAttribute(MaterialX::InputPtr Input, const FString& InputChannelName, UInterchangeShaderNode* ShaderNode, const FLinearColor& DefaultValue) const;
-
-	/**
-	 * Return the innermost file prefix of an element in the current scope, if it has none, it will take the one from its parents
-	 * 
-	 * @param Element - the Element to retrieve the file prefix from (can be anything, an input, a node, a nodegraph, etc.)
-	 * 
-	 * @return a file prefix or an empty string
-	 * 
-	 */
-	FString GetFilePrefix(MaterialX::ElementPtr Element) const;
-
-	/**
-	 * Return the innermost color space of an element in the current scope, if it has none, it will take the one from its parents
-	 *
-	 * @param Element - the Element to retrieve the color space from (can be anything, an input, a node, a nodegraph, etc.)
-	 *
-	 * @return a color space or an empty string
-	 *
-	 */
-	FString GetColorSpace(MaterialX::ElementPtr Element) const;
-
-	/**
-	 * Helper function that returns a linear color after a color space conversion, the function makes no assumption on the input, and it should have a value of Color3-4 type
-	 * 
-	 * @param Input - The input that has a Color3-4 value in it
-	 * 
-	 * @return The linear color after color space conversion
-	 */
-	FLinearColor GetLinearColor(MaterialX::InputPtr Input) const;
-
-	/** Begin Connect MaterialX nodes*/
-	
-	/** <constant> */
-	void ConnectConstantInputToOutput(MaterialX::NodePtr UpstreamNode, UInterchangeShaderNode* ParentShaderNode, const FString& InputChannelName, TMap<FString, UInterchangeShaderNode*>& NamesToShaderNodes, UInterchangeBaseNodeContainer& NodeContainer) const;
-	
-	/** <extract> */
-	void ConnectExtractInputToOutput(MaterialX::NodePtr UpstreamNode, UInterchangeShaderNode* ParentShaderNode, const FString& InputChannelName, TMap<FString, UInterchangeShaderNode*>& NamesToShaderNodes, UInterchangeBaseNodeContainer& NodeContainer) const;
-	
-	/** <dot> */
-	void ConnectDotInputToOutput(MaterialX::NodePtr UpstreamNode, UInterchangeShaderNode* ParentShaderNode, const FString& InputChannelName, TMap<FString, UInterchangeShaderNode*>& NamesToShaderNodes, UInterchangeBaseNodeContainer& NodeContainer) const;
-	
-	/** <transformpoint> */
-	void ConnectTransformPositionInputToOutput(MaterialX::NodePtr UpstreamNode, UInterchangeShaderNode* ParentShaderNode, const FString& InputChannelName, TMap<FString, UInterchangeShaderNode*>& NamesToShaderNodes, UInterchangeBaseNodeContainer& NodeContainer) const;
-	
-	/** <transformvector> */
-	void ConnectTransformVectorInputToOutput(MaterialX::NodePtr UpstreamNode, UInterchangeShaderNode* ParentShaderNode, const FString& InputChannelName, TMap<FString, UInterchangeShaderNode*>& NamesToShaderNodes, UInterchangeBaseNodeContainer& NodeContainer) const;
-	
-	/** <rotate2d> */
-	void ConnectRotate2DInputToOutput(MaterialX::NodePtr UpstreamNode, UInterchangeShaderNode* ParentShaderNode, const FString& InputChannelName, TMap<FString, UInterchangeShaderNode*>& NamesToShaderNodes, UInterchangeBaseNodeContainer& NodeContainer) const;
-	
-	/** <rotate3d> */
-	void ConnectRotate3DInputToOutput(MaterialX::NodePtr UpstreamNode, UInterchangeShaderNode* ParentShaderNode, const FString& InputChannelName, TMap<FString, UInterchangeShaderNode*>& NamesToShaderNodes, UInterchangeBaseNodeContainer& NodeContainer) const;
-	
-	/** <image> */
-	void ConnectImageInputToOutput(MaterialX::NodePtr UpstreamNode, UInterchangeShaderNode* ParentShaderNode, const FString& InputChannelName, TMap<FString, UInterchangeShaderNode*>& NamesToShaderNodes, UInterchangeBaseNodeContainer& NodeContainer) const;
-
-	/** <convert> */
-	void ConnectConvertInputToOutput(MaterialX::NodePtr UpstreamNode, UInterchangeShaderNode* ParentShaderNode, const FString& InputChannelName, TMap<FString, UInterchangeShaderNode*>& NamesToShaderNodes, UInterchangeBaseNodeContainer& NodeContainer) const;
-	
-	/** <ifgreater> */
-	void ConnectIfGreaterInputToOutput(MaterialX::NodePtr UpstreamNode, UInterchangeShaderNode* ParentShaderNode, const FString& InputChannelName, TMap<FString, UInterchangeShaderNode*>& NamesToShaderNodes, UInterchangeBaseNodeContainer& NodeContainer) const;
-	
-	/** <ifgreatereq> */
-	void ConnectIfGreaterEqInputToOutput(MaterialX::NodePtr UpstreamNode, UInterchangeShaderNode* ParentShaderNode, const FString& InputChannelName, TMap<FString, UInterchangeShaderNode*>& NamesToShaderNodes, UInterchangeBaseNodeContainer& NodeContainer) const;
-	
-	/** <ifequal> */
-	void ConnectIfEqualInputToOutput(MaterialX::NodePtr UpstreamNode, UInterchangeShaderNode* ParentShaderNode, const FString& InputChannelName, TMap<FString, UInterchangeShaderNode*>& NamesToShaderNodes, UInterchangeBaseNodeContainer& NodeContainer) const;
-	
-	/** <outside> */
-	void ConnectOutsideInputToOutput(MaterialX::NodePtr UpstreamNode, UInterchangeShaderNode* ParentShaderNode, const FString& InputChannelName, TMap<FString, UInterchangeShaderNode*>& NamesToShaderNodes, UInterchangeBaseNodeContainer& NodeContainer) const;
-	
-	/** <position> */
-	void ConnectPositionInputToOutput(MaterialX::NodePtr UpstreamNode, UInterchangeShaderNode* ParentShaderNode, const FString& InputChannelName, TMap<FString, UInterchangeShaderNode*>& NamesToShaderNodes, UInterchangeBaseNodeContainer& NodeContainer) const;
-	
-	/** <normal> */
-	void ConnectNormalInputToOutput(MaterialX::NodePtr UpstreamNode, UInterchangeShaderNode* ParentShaderNode, const FString& InputChannelName, TMap<FString, UInterchangeShaderNode*>& NamesToShaderNodes, UInterchangeBaseNodeContainer& NodeContainer) const;
-	
-	/** <tangent> */
-	void ConnectTangentInputToOutput(MaterialX::NodePtr UpstreamNode, UInterchangeShaderNode* ParentShaderNode, const FString& InputChannelName, TMap<FString, UInterchangeShaderNode*>& NamesToShaderNodes, UInterchangeBaseNodeContainer& NodeContainer) const;
-	
-	/** <bitangent> */
-	void ConnectBitangentInputToOutput(MaterialX::NodePtr UpstreamNode, UInterchangeShaderNode* ParentShaderNode, const FString& InputChannelName, TMap<FString, UInterchangeShaderNode*>& NamesToShaderNodes, UInterchangeBaseNodeContainer& NodeContainer) const;
-	
-	/** <time> */
-	void ConnectTimeInputToOutput(MaterialX::NodePtr UpstreamNode, UInterchangeShaderNode* ParentShaderNode, const FString& InputChannelName, TMap<FString, UInterchangeShaderNode*>& NamesToShaderNodes, UInterchangeBaseNodeContainer& NodeContainer) const;
-	
-	/** <noise3d> */
-	void ConnectNoise3DInputToOutput(MaterialX::NodePtr UpstreamNode, UInterchangeShaderNode* ParentShaderNode, const FString& InputChannelName, TMap<FString, UInterchangeShaderNode*>& NamesToShaderNodes, UInterchangeBaseNodeContainer& NodeContainer) const;
-	
-	/** <cellnoise3d> */
-	void ConnectCellNoise3DInputToOutput(MaterialX::NodePtr UpstreamNode, UInterchangeShaderNode* ParentShaderNode, const FString& InputChannelName, TMap<FString, UInterchangeShaderNode*>& NamesToShaderNodes, UInterchangeBaseNodeContainer& NodeContainer) const;
-	
-	/** <worleynoise3d> */
-	void ConnectWorleyNoise3DInputToOutput(MaterialX::NodePtr UpstreamNode, UInterchangeShaderNode* ParentShaderNode, const FString& InputChannelName, TMap<FString, UInterchangeShaderNode*>& NamesToShaderNodes, UInterchangeBaseNodeContainer& NodeContainer) const;
-	
-	/** <heighttonormal> */
-	void ConnectHeightToNormalInputToOutput(MaterialX::NodePtr UpstreamNode, UInterchangeShaderNode* ParentShaderNode, const FString& InputChannelName, TMap<FString, UInterchangeShaderNode*>& NamesToShaderNodes, UInterchangeBaseNodeContainer& NodeContainer) const;
-	
-	/** <blur> */
-	void ConnectBlurInputToOutput(MaterialX::NodePtr UpstreamNode, UInterchangeShaderNode* ParentShaderNode, const FString& InputChannelName, TMap<FString, UInterchangeShaderNode*>& NamesToShaderNodes, UInterchangeBaseNodeContainer& NodeContainer) const;
-
-	/** End Connect MaterialX nodes*/
-
-	UInterchangeShaderNode* CreateMaskShaderNode(uint8 RGBA, const FString & NodeName, TMap<FString, UInterchangeShaderNode*>& NamesToShaderNodes, UInterchangeBaseNodeContainer& NodeContainer) const;
-
-	/**
-	 * Helper function that returns a vector, the function makes no assumption on the input, and it should have a value of vectorN type
-	 *
-	 * @param Input - The input that has a vectorN value in it
-	 *
-	 * @return The vector
-	 */
-	FLinearColor GetVector(MaterialX::InputPtr Input) const;
-
-	/** 
-	 * @param const char * - New name of the Input to copy
-	 * @param MaterialX::InputPtr - Input to copy
-	 */
-	using FInputToCopy = TPair<const char*, MaterialX::InputPtr>; 
-
-	/**
-	 * @param const char * - Attribute
-	 * @param const char * - Value
-	 */
-	using FAttributeValue = TPair<const char*, const char*>;
-
-	/**
-	 * @param TArray<FAttributeValueType> - Array of attribute/value
-	 */
-	using FAttributeValueArray = TArray<FAttributeValue>;
-
-	/**
-	 * @param TArray<FAttributeValueType> - Array of attribute/value
-	 */
-	using FInputToCreate = TPair<const char*, FAttributeValueArray>; // Name - Attributes
-
-	/**
-	 * Helper function that create a MaterialX node, useful to insert a node between 2 nodes especially when an operation cannot be done in one pass
-	 * If the node already exists in the graph, this one is returned
-	 *
-	 * @param NodeGraph - The node graph to insert the new node into
-	 * @param NodeName - The name of the node
-	 * @param Category - The category of the node to create
-	 * @param InputsToCopy - The inputs from another node/nodegraph to copy in the new node, can be empty
-	 * @param InputsToCreate - The inputs to create from scratch and insert in the new node, can be empty
-	 * 
-	 * @return The newly created node
-	 */
-	MaterialX::NodePtr CreateNode(MaterialX::ElementPtr NodeGraph, const char * NodeName, const char * Category, TArray<FInputToCopy> InputsToCopy, TArray<FInputToCreate> InputsToCreate) const;
-
-	/**
-	 * Convert math nodes (add, sub, div, mul) that have a neutral input, to the MaterialX <dot> node which is simply the identity or a noop
-	 * This allows us to optimize the final shader by avoiding unnecessary operation on a material (for example a*1 = a)
-	 * This function should be called after the subgraphs substitution in order to search for all the nodes even in the subgraphs
-	 * 
-	 * @param NodeGraph - The node graph to retrieve the different math nodes
-	 */
-	void ConvertNeutralNodesToDot(MaterialX::ElementPtr NodeGraph) const;
-
-	/**
-	 * Add a texcoord input to the tiledimage nodes if the input is missing, and connect it to a texcoord node in the nodegraph
-	 * This function should be called before the subgraphs substitution to avoid missing UVs input on a texture
-	 * See for example standard_surface_brick_procedural where the tiledimage nodes have no texcoord input
-	 * 
-	 * @param Node - The nodegraph to retrieve the different tiledimage nodes
-	 */
-	void AddTexCoordToTiledImageNodes(MaterialX::ElementPtr Graph) const;
-
-	/**
-	 * Update a MaterialX Document recursively by either initializing nodes (like <tiledimage>) or updating some nodes (e.g: math nodes with neutral elements)
-	 * 
-	 * @param Graph - The graph to update and all of its sub graphs
-	 */
-	void UpdateDocumentRecursively(MaterialX::GraphElementPtr Graph) const;
-
-	/**
-	 * Retrieve the interchange parent name of a MaterialX node, useful when a node is a combination of several nodes connected to different inputs (look for example Noise3D)
-	 *
-	 * @param Node - The node to retrieve the parent name 
-	 * @return The node of the parent
-	 */
-	FString GetAttributeParentName(MaterialX::NodePtr Node) const;
-
-private:
-
-	TMap<TPair<FString, FString>, FString> InputNamesMaterialX2UE; //given a MaterialX node (category - input), return the UE/Interchange input name.
-	TMap<FString, FString> NodeNamesMaterialX2UE; //given a MaterialX node category, return the UE category
-	TSet<FString> UEInputs;
-#endif // WITH_EDITOR
 };
\ No newline at end of file
diff --git a/Engine/Plugins/Interchange/Runtime/Source/Import/Public/MaterialX/MaterialXUtils/MaterialXBase.h b/Engine/Plugins/Interchange/Runtime/Source/Import/Public/MaterialX/MaterialXUtils/MaterialXBase.h
new file mode 100644
--- /dev/null
+++ b/Engine/Plugins/Interchange/Runtime/Source/Import/Public/MaterialX/MaterialXUtils/MaterialXBase.h
@@ -0,0 +1,127 @@
+// Copyright Epic Games, Inc. All Rights Reserved.
+
+#pragma once
+
+#if WITH_EDITOR
+
+#include "InterchangeMaterialDefinitions.h"
+#include "MaterialXCore/Document.h"
+#include "MaterialX/InterchangeMaterialXDefinitions.h"
+#include "Nodes/InterchangeBaseNodeContainer.h"
+
+class FMaterialXBase : public TSharedFromThis<FMaterialXBase>
+{
+protected:
+
+	FMaterialXBase(UInterchangeBaseNodeContainer& BaseNodeContainer);
+
+public:
+
+	virtual void Translate(MaterialX::NodePtr Node) = 0;
+
+	virtual ~FMaterialXBase() = default;
+	
+	/**
+	 * Update a MaterialX Document recursively by either initializing nodes (like <tiledimage>) or updating some nodes (e.g: math nodes with neutral elements)
+	 *
+	 * @param Graph - The graph to update and all of its sub graphs
+	 */
+	static void UpdateDocumentRecursively(MaterialX::GraphElementPtr Graph);
+
+protected:
+
+	/**
+	 * @param const char * - New name of the Input to copy
+	 * @param MaterialX::InputPtr - Input to copy
+	 */
+	using FInputToCopy = TPair<const char*, MaterialX::InputPtr>;
+
+	/**
+	 * @param const char * - Attribute
+	 * @param const char * - Value
+	 */
+	using FAttributeValue = TPair<const char*, const char*>;
+
+	/**
+	 * @param TArray<FAttributeValueType> - Array of attribute/value
+	 */
+	using FAttributeValueArray = TArray<FAttributeValue>;
+
+	/**
+	 * @param TArray<FAttributeValueType> - Array of attribute/value
+	 */
+	using FInputToCreate = TPair<const char*, FAttributeValueArray>; // Name - Attributes
+
+	/**
+	 * Add a texcoord input to the tiledimage nodes if the input is missing, and connect it to a texcoord node in the nodegraph
+	 * This function should be called before the subgraphs substitution to avoid missing UVs input on a texture
+	 * See for example standard_surface_brick_procedural where the tiledimage nodes have no texcoord input
+	 *
+	 * @param Node - The nodegraph to retrieve the different tiledimage nodes
+	 */
+	static void AddTexCoordToTiledImageNodes(MaterialX::ElementPtr Graph);
+
+	/**
+	 * Convert math nodes (add, sub, div, mul) that have a neutral input, to the MaterialX <dot> node which is simply the identity or a noop
+	 * This allows us to optimize the final shader by avoiding unnecessary operation on a material (for example a*1 = a)
+	 * This function should be called after the subgraphs substitution in order to search for all the nodes even in the subgraphs
+	 *
+	 * @param NodeGraph - The node graph to retrieve the different math nodes
+	 */
+	static void ConvertNeutralNodesToDot(MaterialX::ElementPtr NodeGraph);
+
+	/**
+	 * Helper function that create a MaterialX node, useful to insert a node between 2 nodes especially when an operation cannot be done in one pass
+	 * If the node already exists in the graph, this one is returned
+	 *
+	 * @param NodeGraph - The node graph to insert the new node into
+	 * @param NodeName - The name of the node
+	 * @param Category - The category of the node to create
+	 * @param InputsToCopy - The inputs from another node/nodegraph to copy in the new node, can be empty
+	 * @param InputsToCreate - The inputs to create from scratch and insert in the new node, can be empty
+	 *
+	 * @return The newly created node
+	 */
+	static MaterialX::NodePtr CreateNode(MaterialX::ElementPtr NodeGraph, const char* NodeName, const char* Category, const TArray<FInputToCopy>& InputsToCopy, const TArray<FInputToCreate>& InputsToCreate);
+
+
+	/**
+	 * Return the innermost color space of an element in the current scope, if it has none, it will take the one from its parents
+	 *
+	 * @param Element - the Element to retrieve the color space from (can be anything, an input, a node, a nodegraph, etc.)
+	 *
+	 * @return a color space or an empty string
+	 *
+	 */
+	static FString GetColorSpace(MaterialX::ElementPtr Element);
+
+
+	/**
+	 * Retrieve the input from a surfaceshader node, or take the default input from the library,
+	 * this function should only be called after testing the MaterialX libraries have been successfully imported, meaning the node definition of the surfaceshader
+	 * should always be valid
+	 *
+	 * @param StandardSurface - the <standard_surface> node
+	 * @param InputName - the input name to retrieve
+	 *
+	 * @return the input from the given name
+	 */
+	MaterialX::InputPtr GetInput(MaterialX::NodePtr Node, const char* InputName);
+
+	/**
+	 * Helper function that returns a linear color after a color space conversion, the function makes no assumption on the input, and it should have a value of Color3-4 type
+	 *
+	 * @param Input - The input that has a Color3-4 value in it
+	 *
+	 * @return The linear color after color space conversion
+	 */
+	static FLinearColor GetLinearColor(MaterialX::InputPtr Input);
+
+protected:
+
+	UInterchangeBaseNodeContainer& NodeContainer;
+
+	/** It's up to the derive class to initialize it with the proper node definition*/
+	const char* NodeDefinition; 
+};
+#endif
diff --git a/Engine/Plugins/Interchange/Runtime/Source/Import/Public/MaterialX/MaterialXUtils/MaterialXLightShaderAbstract.h b/Engine/Plugins/Interchange/Runtime/Source/Import/Public/MaterialX/MaterialXUtils/MaterialXLightShaderAbstract.h
new file mode 100644
--- /dev/null
+++ b/Engine/Plugins/Interchange/Runtime/Source/Import/Public/MaterialX/MaterialXUtils/MaterialXLightShaderAbstract.h
@@ -0,0 +1,27 @@
+// Copyright Epic Games, Inc. All Rights Reserved. 
+
+#pragma once
+
+#if WITH_EDITOR
+
+#include "MaterialXBase.h"
+#include "Nodes/InterchangeBaseNodeContainer.h"
+#include "InterchangeSceneNode.h"
+#include "InterchangeLightNode.h"
+
+class FMaterialXLightShaderAbstract : public FMaterialXBase
+{
+protected:
+
+	FMaterialXLightShaderAbstract(UInterchangeBaseNodeContainer& BaseNodeContainer);
+
+	virtual void PreTranslate();
+
+	virtual void PostTranslate();
+
+	FTransform Transform;
+	MaterialX::NodePtr LightShaderNode;
+	UInterchangeSceneNode* SceneNode;
+	UInterchangeBaseLightNode* LightNode;
+};
+#endif
\ No newline at end of file
diff --git a/Engine/Plugins/Interchange/Runtime/Source/Import/Public/MaterialX/MaterialXUtils/MaterialXManager.h b/Engine/Plugins/Interchange/Runtime/Source/Import/Public/MaterialX/MaterialXUtils/MaterialXManager.h
new file mode 100644
--- /dev/null
+++ b/Engine/Plugins/Interchange/Runtime/Source/Import/Public/MaterialX/MaterialXUtils/MaterialXManager.h
@@ -0,0 +1,55 @@
+// Copyright Epic Games, Inc. All Rights Reserved.
+
+#pragma once
+
+#if WITH_EDITOR
+
+#include "CoreMinimal.h"
+#include "MaterialXCore/Node.h"
+
+class FMaterialXBase;
+class UInterchangeBaseNodeContainer;
+class UInterchangeShaderNode;
+class FMaterialXSurfaceShaderAbstract;
+
+class FMaterialXManager
+{
+public:
+
+	DECLARE_DELEGATE_RetVal_OneParam(TSharedRef<FMaterialXBase>, FOnGetMaterialXInstance, UInterchangeBaseNodeContainer&);
+
+	static FMaterialXManager& GetInstance();
+
+	/** Find a corresponding Material Expression Input given a (Category, Input) pair*/
+	const FString* FindMatchingInput(const TPair<FString, FString>& CategoryInputKey) const;
+
+	/** Find a stored Material Expression Input */
+	const FString* FindMaterialExpressionInput(const FString& InputKey) const;
+
+	/** Find a matching Material Expression given a MaterialX category*/
+	const FString* FindMatchingMaterialExpression(const FString& CategoryKey) const;
+
+	TSharedPtr<FMaterialXBase> GetShaderTranslator(const FString& CategoryShader, UInterchangeBaseNodeContainer& NodeContainer);
+
+	void RegisterMaterialXInstance(const FString& CategoryShader, FOnGetMaterialXInstance MaterialXInstanceDelegate);
+
+	FMaterialXManager(const FMaterialXManager&) = delete;
+	FMaterialXManager& operator=(const FMaterialXManager&) = delete;
+
+private:
+
+	FMaterialXManager();
+
+	/** Given a MaterialX node (category - input), return the UE/Interchange input name.*/
+	TMap<TPair<FString, FString>, FString> MatchingInputNames;
+
+	/** Given a MaterialX node category, return the UE material expression class name*/
+	TMap<FString, FString> MatchingMaterialExpressions;
+
+	/** The different inputs of material expression that we may encounter, the MaterialX Document is modified consequently regarding those*/
+	TSet<FString> MaterialExpressionInputs;
+
+	/** Container of a MaterialX document, to translate the different nodes based on their category, same reason as above concerning the use of STL*/
+	TMap<FString, FOnGetMaterialXInstance> MaterialXContainerDelegates;
+};
+#endif
diff --git a/Engine/Plugins/Interchange/Runtime/Source/Import/Public/MaterialX/MaterialXUtils/MaterialXSurfaceShaderAbstract.h b/Engine/Plugins/Interchange/Runtime/Source/Import/Public/MaterialX/MaterialXUtils/MaterialXSurfaceShaderAbstract.h
new file mode 100644
--- /dev/null
+++ b/Engine/Plugins/Interchange/Runtime/Source/Import/Public/MaterialX/MaterialXUtils/MaterialXSurfaceShaderAbstract.h
@@ -0,0 +1,383 @@
+// Copyright Epic Games, Inc. All Rights Reserved.
+
+#pragma once
+
+#if WITH_EDITOR
+#include "InterchangeShaderGraphNode.h"
+#include "InterchangeTranslatorBase.h"
+#include "InterchangeTextureNode.h"
+#include "MaterialXBase.h"
+
+class FMaterialXSurfaceShaderAbstract : public FMaterialXBase
+{
+protected:
+
+	DECLARE_DELEGATE_ThreeParams(FOnConnectNodeOutputToInput, MaterialX::NodePtr, UInterchangeShaderNode*, const FString&);
+
+	FMaterialXSurfaceShaderAbstract(UInterchangeBaseNodeContainer & BaseNodeContainer);
+
+	/**
+	 * Add an attribute to a shader node from the given MaterialX input, only floats and linear colors are supported for the moment
+	 *
+	 * @param Input - The MaterialX input to retrieve and add the value from, must be of type float/color/vector
+	 * @param InputChannelName - The name of the shader node's input to add the attribute
+	 * @param ShaderNode - The shader node to which we want to add the attribute
+	 *
+	 * @return true if the attribute was successfully added
+	 */
+	bool AddAttribute(MaterialX::InputPtr Input, const FString& InputChannelName, UInterchangeShaderNode* ShaderNode) const;
+
+	/**
+	 * Add an attribute to a shader node from the given MaterialX input if that input has either a value or an interface name
+	 *
+	 * @param Input - The MaterialX input to retrieve and add the value from, must be of type float/color/vector
+	 * @param InputChannelName - The name of the shader node's input to add the attribute
+	 * @param ShaderNode - The shader node to which we want to add the attribute
+	 *
+	 * @return true if the attribute was successfully added
+	 */
+	bool AddAttributeFromValueOrInterface(MaterialX::InputPtr Input, const FString& InputChannelName, UInterchangeShaderNode* ShaderNode) const;
+
+	/**
+	 * Add a float attribute to a shader node only if its value taken from the input is not equal to its default value. Return false if the attribute does not exist or if we cannot add it
+	 *
+	 * @param Input - The MaterialX input to retrieve and add the value from, must be of type float/color/vector
+	 * @param InputChannelName - The name of the shader node's input to add the attribute
+	 * @param ShaderNode - The shader node to which we want to add the attribute
+	 * @param DefaultValue - the default value to test the input against
+	 *
+	 * @return true if the attribute was successfully added
+	 */
+	bool AddFloatAttribute(MaterialX::InputPtr Input, const FString& InputChannelName, UInterchangeShaderNode* ShaderNode, float DefaultValue) const;
+
+	/**
+	 * Add a FLinearColor attribute to a shader node only if its value taken from the input is not equal to its default value. Return false if the attribute does not exist or if we cannot add it
+	 *
+	 * @param Input - The MaterialX input to retrieve and add the value from, must be of type float/color/vector
+	 * @param InputChannelName - The name of the shader node's input to add the attribute
+	 * @param ShaderNode - The shader node to which we want to add the attribute
+	 * @param DefaultValue - the default value to test the input against
+	 *
+	 * @return true if the attribute was successfully added
+	 */
+	bool AddLinearColorAttribute(MaterialX::InputPtr Input, const FString& InputChannelName, UInterchangeShaderNode* ShaderNode, const FLinearColor& DefaultValue) const;
+
+	/**
+	 * Connect an output either from a node name or a node graph from a MaterialX input to the ShaderNode
+	 * 
+	 * @param InputName - The name of the input of the SurfaceShaderNode to retrieve
+	 * @param ShaderNode - The Interchange shader node to connect the MaterialX's node or node graph to
+	 * @param InputShaderName - The name of the input of the shader node to connect to
+	 * @param DefaultValue - The default value of the MaterialX input
+	 */
+	template<typename T>
+	void ConnectNodeOutputToInput(const char* InputName, UInterchangeShaderNode* ShaderNode, const FString& InputShaderName, T DefaultValue)
+	{
+		MaterialX::DocumentPtr Document = SurfaceShaderNode->getDocument();
+		MaterialX::InputPtr Input = GetInput(SurfaceShaderNode, InputName);
+
+		if(!ConnectNodeGraphOutputToInput(Input, ShaderNode, InputShaderName))
+		{
+			if(!ConnectNodeNameOutputToInput(Input, ShaderNode, InputShaderName))
+			{
+				//the different inputs from a standard surface are either float, linear color, or no attribute at all
+				if constexpr(std::is_same_v<decltype(DefaultValue), float>)
+				{
+					AddFloatAttribute(Input, InputShaderName, ShaderNode, DefaultValue);
+				}
+				else if constexpr(std::is_same_v<decltype(DefaultValue), FLinearColor>)
+				{
+					AddLinearColorAttribute(Input, InputShaderName, ShaderNode, DefaultValue);
+				}
+			}
+		}
+	};
+
+	/**
+	 * Connect an ouput in the NodeGraph to the ShaderGraph
+	 *
+	 * @param InputToNodeGraph - The input from the standard surface to retrieve the output in the NodeGraph
+	 * @param ShaderNode - The Interchange shader node to connect the MaterialX's node graph to
+	 * @param ParentInputName - The name of the input of the shader node to which we want the node graph to be connected to
+	 *
+	 * @return true if the given input is attached to one of the outputs of a node graph
+	 */
+	bool ConnectNodeGraphOutputToInput(MaterialX::InputPtr InputToNodeGraph, UInterchangeShaderNode* ShaderNode, const FString& ParentInputName);
+
+	/**
+	 * Create and Connect the output of a MaterialX node that has already a matching in UE to a shader node,
+	 * If not search for a registered delegate
+	 *
+	 * @param Node - The MaterialX node of a given type used to create the appropriate shader node
+	 * @param ParentShaderNode - The shader node to connect to
+	 * @param InputChannelName - The input of the ParentShaderNode to connect to
+	 *
+	 * @return true if a shader node has been successfully created and is connected to the given input
+	 */
+	bool ConnectMatchingNodeOutputToInput(MaterialX::NodePtr Node, UInterchangeShaderNode* ParentShaderNode, const FString& InputChannelName);
+
+	/**
+	 * Create and Connect manually the output of a MaterialX node to a shader node
+	 *
+	 * @param Edge - The MaterialX edge that has the current node, its parent and the input bridge between the two
+	 * @param ParentShaderNode - The shader node to connect to
+	 * @param InputChannelName - The input of the ParentShaderNode to connect to
+	 *
+	 * @return true if a shader node has been successfully created and is connected to the given input
+	 */
+	void ConnectNodeCategoryOutputToInput(const MaterialX::Edge& Edge, UInterchangeShaderNode* ParentShaderNode, const FString& InputChannelName);
+
+	/**
+	 * Create and Connect a node name directly connected from an input to a shader node
+	 *
+	 * @param Node - The MaterialX node of a given type used to create the appropriate shader node
+	 * @param ParentShaderNode - The shader node to connect to
+	 * @param InputChannelName - The input of the ParentShaderNode to connect to
+	 *
+	 * @return true if a shader node has been successfully created and is connected to the given input
+	 */
+	bool ConnectNodeNameOutputToInput(MaterialX::InputPtr Input, UInterchangeShaderNode* ShaderNode, const FString& ParentInputName);
+
+	/** Begin Connect MaterialX nodes*/
+
+	/** <constant> */
+	void ConnectConstantInputToOutput(MaterialX::NodePtr UpstreamNode, UInterchangeShaderNode* ParentShaderNode, const FString& InputChannelName);
+
+	/** <extract> */
+	void ConnectExtractInputToOutput(MaterialX::NodePtr UpstreamNode, UInterchangeShaderNode* ParentShaderNode, const FString& InputChannelName);
+
+	/** <dot> */
+	void ConnectDotInputToOutput(MaterialX::NodePtr UpstreamNode, UInterchangeShaderNode* ParentShaderNode, const FString& InputChannelName);
+
+	/** <transformpoint> */
+	void ConnectTransformPositionInputToOutput(MaterialX::NodePtr UpstreamNode, UInterchangeShaderNode* ParentShaderNode, const FString& InputChannelName);
+
+	/** <transformvector> */
+	void ConnectTransformVectorInputToOutput(MaterialX::NodePtr UpstreamNode, UInterchangeShaderNode* ParentShaderNode, const FString& InputChannelName);
+
+	/** <rotate2d> */
+	void ConnectRotate2DInputToOutput(MaterialX::NodePtr UpstreamNode, UInterchangeShaderNode* ParentShaderNode, const FString& InputChannelName);
+
+	/** <rotate3d> */
+	void ConnectRotate3DInputToOutput(MaterialX::NodePtr UpstreamNode, UInterchangeShaderNode* ParentShaderNode, const FString& InputChannelName);
+
+	/** <image> */
+	void ConnectImageInputToOutput(MaterialX::NodePtr UpstreamNode, UInterchangeShaderNode* ParentShaderNode, const FString& InputChannelName);
+
+	/** <convert> */
+	void ConnectConvertInputToOutput(MaterialX::NodePtr UpstreamNode, UInterchangeShaderNode* ParentShaderNode, const FString& InputChannelName);
+
+	/** <ifgreater> */
+	void ConnectIfGreaterInputToOutput(MaterialX::NodePtr UpstreamNode, UInterchangeShaderNode* ParentShaderNode, const FString& InputChannelName);
+
+	/** <ifgreatereq> */
+	void ConnectIfGreaterEqInputToOutput(MaterialX::NodePtr UpstreamNode, UInterchangeShaderNode* ParentShaderNode, const FString& InputChannelName);
+
+	/** <ifequal> */
+	void ConnectIfEqualInputToOutput(MaterialX::NodePtr UpstreamNode, UInterchangeShaderNode* ParentShaderNode, const FString& InputChannelName);
+
+	/** <outside> */
+	void ConnectOutsideInputToOutput(MaterialX::NodePtr UpstreamNode, UInterchangeShaderNode* ParentShaderNode, const FString& InputChannelName);
+
+	/** <position> */
+	void ConnectPositionInputToOutput(MaterialX::NodePtr UpstreamNode, UInterchangeShaderNode* ParentShaderNode, const FString& InputChannelName);
+
+	/** <normal> */
+	void ConnectNormalInputToOutput(MaterialX::NodePtr UpstreamNode, UInterchangeShaderNode* ParentShaderNode, const FString& InputChannelName);
+
+	/** <tangent> */
+	void ConnectTangentInputToOutput(MaterialX::NodePtr UpstreamNode, UInterchangeShaderNode* ParentShaderNode, const FString& InputChannelName);
+
+	/** <bitangent> */
+	void ConnectBitangentInputToOutput(MaterialX::NodePtr UpstreamNode, UInterchangeShaderNode* ParentShaderNode, const FString& InputChannelName);
+
+	/** <time> */
+	void ConnectTimeInputToOutput(MaterialX::NodePtr UpstreamNode, UInterchangeShaderNode* ParentShaderNode, const FString& InputChannelName);
+
+	/** <noise3d> */
+	void ConnectNoise3DInputToOutput(MaterialX::NodePtr UpstreamNode, UInterchangeShaderNode* ParentShaderNode, const FString& InputChannelName);
+
+	/** <cellnoise3d> */
+	void ConnectCellNoise3DInputToOutput(MaterialX::NodePtr UpstreamNode, UInterchangeShaderNode* ParentShaderNode, const FString& InputChannelName);
+
+	/** <worleynoise3d> */
+	void ConnectWorleyNoise3DInputToOutput(MaterialX::NodePtr UpstreamNode, UInterchangeShaderNode* ParentShaderNode, const FString& InputChannelName);
+
+	/** <heighttonormal> */
+	void ConnectHeightToNormalInputToOutput(MaterialX::NodePtr UpstreamNode, UInterchangeShaderNode* ParentShaderNode, const FString& InputChannelName);
+
+	/** <blur> */
+	void ConnectBlurInputToOutput(MaterialX::NodePtr UpstreamNode, UInterchangeShaderNode* ParentShaderNode, const FString& InputChannelName);
+
+	/** End Connect MaterialX nodes*/
+
+	/**
+	 * Create a ComponentMask shader node.
+	 *
+	 * @param RGBA - The mask component, e.g: 0b1011 -> Only RBA are toggled
+	 * @param NodeName - the name of the shader node
+	 *
+	 * @return The ComponentMask node
+	 */
+	UInterchangeShaderNode* CreateMaskShaderNode(uint8 RGBA, const FString& NodeName);
+
+	/**
+	 * Helper template function to create an InterchangeShaderNode
+	 *
+	 * @param NodeName - The name of the shader node
+	 * @param ShaderType - The shader node's type we want to create
+	 *
+	 * @return The shader node that was created
+	 */
+	template<typename ShaderNodeType>
+	ShaderNodeType* CreateShaderNode(const FString& NodeName, const FString& ShaderType)
+	{
+		static_assert(std::is_convertible_v<ShaderNodeType*, UInterchangeShaderNode*>, "CreateShaderNode only accepts type that derived from UInterchangeShaderNode");
+
+		ShaderNodeType* Node;
+
+		const FString NodeUID = UInterchangeShaderNode::MakeNodeUid(NodeName, FStringView{});
+
+		//Test directly in the NodeContainer, because the ShaderNodes can be altered during the node graph either by the parent (dot/normalmap),
+		//or by putting an intermediary node between the child and the parent (tiledimage)
+		if(Node = const_cast<ShaderNodeType*>(Cast<ShaderNodeType>(NodeContainer.GetNode(NodeUID))); !Node)
+		{
+			Node = NewObject<ShaderNodeType>(&NodeContainer);
+			Node->InitializeNode(NodeUID, NodeName, EInterchangeNodeContainerType::TranslatedAsset);
+			NodeContainer.AddNode(Node);
+			if constexpr(std::is_same_v<ShaderNodeType, UInterchangeShaderGraphNode>)
+			{
+				NodeContainer.SetNodeParentUid(NodeUID, FString{});
+			}
+			Node->SetCustomShaderType(ShaderType);
+
+			ShaderNodes.Add(NodeName, Node);
+		}
+
+		return Node;
+	}
+
+	/**
+	 * Helper function to create an InterchangeTextureNode
+	 *
+	 * @param Node - The MaterialX node, it should be of the category <image> no test is done on it
+	 *
+	 * @return The texture node that was created
+	 */
+	template<typename TextureTypeNode>
+	UInterchangeTextureNode* CreateTextureNode(MaterialX::NodePtr Node) const
+	{
+		static_assert(std::is_convertible_v<TextureTypeNode*, UInterchangeTextureNode*>, "CreateTextureNode only accepts type that derived from UInterchangeTextureNode");
+		UInterchangeTextureNode* TextureNode = nullptr;
+
+		//A node image should have an input file otherwise the user should check its default value
+		if(Node)
+		{
+			if(MaterialX::InputPtr InputFile = Node->getInput("file"); InputFile && InputFile->hasValue())
+			{
+				FString Filepath{ InputFile->getValueString().c_str() };
+				const FString FilePrefix = GetFilePrefix(InputFile);
+				Filepath = FPaths::Combine(FilePrefix, Filepath);
+				const FString Filename = FPaths::GetCleanFilename(Filepath);
+				const FString TextureNodeUID = TEXT("\\Texture\\") + Filename;
+
+				//Only add the TextureNode once
+				TextureNode = const_cast<UInterchangeTextureNode*>(Cast<UInterchangeTextureNode>(NodeContainer.GetNode(TextureNodeUID)));
+				if(TextureNode == nullptr)
+				{
+					TextureNode = NewObject<TextureTypeNode>(&NodeContainer);
+					TextureNode->InitializeNode(TextureNodeUID, Filename, EInterchangeNodeContainerType::TranslatedAsset);
+					NodeContainer.AddNode(TextureNode);
+
+					if(FPaths::IsRelative(Filepath))
+					{
+						Filepath = FPaths::ConvertRelativePathToFull(FPaths::GetPath(Node->getActiveSourceUri().c_str()), Filepath);
+					}
+
+					TextureNode->SetPayLoadKey(Filepath);
+
+					const FString ColorSpace = GetColorSpace(InputFile);
+					const bool bIsSRGB = ColorSpace == TEXT("srgb_texture");
+					TextureNode->SetCustomSRGB(bIsSRGB);
+				}
+			}
+		}
+
+		return TextureNode;
+	}
+
+	/**
+	 * Get the UE corresponding name of a MaterialX Node category and input for a material
+	 *
+	 * @param Input - MaterialX input
+	 *
+	 * @return The matched name of the Node/Input else empty string
+	 */
+	const FString& GetMatchedInputName(MaterialX::NodePtr Node, MaterialX::InputPtr Input) const;
+
+	/**
+	 * Get the input name, use this function instead of getName, because a renaming may have occured and we ensure to have the proper name that will be used by UE inputs
+	 *
+	 * @param Input - The input to retrieve the name from
+	 *
+	 * @return The input name
+	 */
+	FString GetInputName(MaterialX::InputPtr Input) const;
+
+	/**
+	 * Return the innermost file prefix of an element in the current scope, if it has none, it will take the one from its parents
+	 *
+	 * @param Element - the Element to retrieve the file prefix from (can be anything, an input, a node, a nodegraph, etc.)
+	 *
+	 * @return a file prefix or an empty string
+	 *
+	 */
+	FString GetFilePrefix(MaterialX::ElementPtr Element) const;
+
+	/**
+	 * Helper function that returns a vector, the function makes no assumption on the input, and it should have a value of vectorN type
+	 *
+	 * @param Input - The input that has a vectorN value in it
+	 *
+	 * @return The vector
+	 */
+	FLinearColor GetVector(MaterialX::InputPtr Input) const;
+
+	/**
+	 * Retrieve the interchange parent name of a MaterialX node, useful when a node is a combination of several nodes connected to different inputs (look for example Noise3D)
+	 *
+	 * @param Node - The node to retrieve the parent name
+	 * @return The node of the parent
+	 */
+	FString GetAttributeParentName(MaterialX::NodePtr Node) const;
+
+	virtual void RegisterConnectNodeOutputToInputDelegates();
+
+	/**
+	 * Set the matching inputs names of a node to correspond to the one used by UE, the matching name is stored under the attribute UE::NewName
+	 *
+	 * @param Node - Look up to all inputs of Node and set the matching name attribute
+	 */
+	void SetMatchingInputsNames(MaterialX::NodePtr Node) const;
+
+	/**
+	 * Add the input new name under the attribute UE::NewName.
+	 *
+	 * @param Input - The input to add the proper attribute
+	 * @param NewName - the new name of the input
+	 */
+	void SetAttributeNewName(MaterialX::InputPtr Input, const char* NewName) const;
+
+
+protected:
+
+	/** Store the shader nodes only when we create the shader graph node*/
+	TMap<FString, UInterchangeShaderNode*> ShaderNodes;
+
+	/** The surface shader node processed during the Translate, up to the derived class to initialize it*/
+	MaterialX::NodePtr SurfaceShaderNode;
+
+	/** Matching MaterialX category and Connect function*/
+	TMap<FString, FOnConnectNodeOutputToInput> MatchingConnectNodeDelegates;
+};
+#endif
\ No newline at end of file
