[Interchange - glTF] Support KHR_materials_emissive_strength in glTF importer

diff --git a/Engine/Plugins/Editor/GLTFImporter/Source/GLTFCore/Private/GLTF/ExtensionsHandler.cpp b/Engine/Plugins/Editor/GLTFImporter/Source/GLTFCore/Private/GLTF/ExtensionsHandler.cpp
--- a/Engine/Plugins/Editor/GLTFImporter/Source/GLTFCore/Private/GLTF/ExtensionsHandler.cpp
+++ b/Engine/Plugins/Editor/GLTFImporter/Source/GLTFCore/Private/GLTF/ExtensionsHandler.cpp
@@ -116,6 +116,7 @@ namespace GLTF
 			EExtension::KHR_MaterialsSheen,
 			EExtension::KHR_MaterialsIOR,
 			EExtension::KHR_MaterialsSpecular,
+			EExtension::KHR_MaterialsEmissiveStrength,
 			EExtension::MSFT_PackingOcclusionRoughnessMetallic,
 			EExtension::MSFT_PackingNormalRoughnessMetallic 
 		};
@@ -229,6 +230,16 @@ namespace GLTF
 					Asset->ProcessedExtensions.Add(EExtension::KHR_MaterialsSpecular);
 				}
 				break;
+				case EExtension::KHR_MaterialsEmissiveStrength:
+				{
+					const FJsonObject& Emissive = ExtObj;
+
+					Material.bHasEmissiveStrength = true;
+					Material.EmissiveStrength = GetScalar(Emissive, TEXT("emissiveStrength"), 1.0f);
+
+					Asset->ProcessedExtensions.Add(EExtension::KHR_MaterialsEmissiveStrength);
+				}
+				break;
 				case EExtension::MSFT_PackingOcclusionRoughnessMetallic:
 				{
 					const FJsonObject& Packing = ExtObj;
diff --git a/Engine/Plugins/Editor/GLTFImporter/Source/GLTFCore/Private/GLTF/GLTFAsset.cpp b/Engine/Plugins/Editor/GLTFImporter/Source/GLTFCore/Private/GLTF/GLTFAsset.cpp
--- a/Engine/Plugins/Editor/GLTFImporter/Source/GLTFCore/Private/GLTF/GLTFAsset.cpp
+++ b/Engine/Plugins/Editor/GLTFImporter/Source/GLTFCore/Private/GLTF/GLTFAsset.cpp
@@ -382,6 +382,8 @@ namespace GLTF
 				return TEXT("KHR_materials_ior");
 			case GLTF::EExtension::KHR_MaterialsSpecular:
 				return TEXT("KHR_materials_specular");
+			case GLTF::EExtension::KHR_MaterialsEmissiveStrength:
+				return TEXT("KHR_materials_emissive_strength");
 			case GLTF::EExtension::KHR_TextureTransform:
 				return TEXT("KHR_texture_transform");
 			case GLTF::EExtension::KHR_DracoMeshCompression:
diff --git a/Engine/Plugins/Editor/GLTFImporter/Source/GLTFCore/Public/GLTFAsset.h b/Engine/Plugins/Editor/GLTFImporter/Source/GLTFCore/Public/GLTFAsset.h
--- a/Engine/Plugins/Editor/GLTFImporter/Source/GLTFCore/Public/GLTFAsset.h
+++ b/Engine/Plugins/Editor/GLTFImporter/Source/GLTFCore/Public/GLTFAsset.h
@@ -32,6 +32,7 @@ namespace GLTF
 		KHR_MaterialsVariants,
 		KHR_MaterialsIOR,
 		KHR_MaterialsSpecular,
+		KHR_MaterialsEmissiveStrength,
 		KHR_TextureTransform,
 		KHR_DracoMeshCompression,
 		KHR_LightsPunctual,
diff --git a/Engine/Plugins/Editor/GLTFImporter/Source/GLTFCore/Public/GLTFMaterial.h b/Engine/Plugins/Editor/GLTFImporter/Source/GLTFCore/Public/GLTFMaterial.h
--- a/Engine/Plugins/Editor/GLTFImporter/Source/GLTFCore/Public/GLTFMaterial.h
+++ b/Engine/Plugins/Editor/GLTFImporter/Source/GLTFCore/Public/GLTFMaterial.h
@@ -141,6 +141,7 @@ namespace GLTF
 		FSheen              Sheen;
 		FSpecular           Specular;
 		float               IOR;
+		float               EmissiveStrength;
 
 		// base properties
 		FTextureMap Normal;
@@ -163,6 +164,7 @@ namespace GLTF
 		bool     bHasTransmission;
 		bool     bHasIOR;
 		bool     bHasSpecular;
+		bool     bHasEmissiveStrength;
 
 		FString UniqueId; //will be generated in FAsset::GenerateNames
 
@@ -171,6 +173,7 @@ namespace GLTF
 		    , BaseColorFactor {1.0f, 1.0f, 1.0f, 1.0f}
 		    , ShadingModel(EShadingModel::MetallicRoughness)
 			, IOR(1.0f)
+			, EmissiveStrength(1.0f)
 		    , NormalScale(1.f)
 		    , OcclusionStrength(1.f)
 		    , EmissiveFactor(FVector::ZeroVector)
@@ -183,6 +186,7 @@ namespace GLTF
 			, bHasTransmission(false)
 			, bHasIOR(false)
 			, bHasSpecular(false)
+			, bHasEmissiveStrength(false)
 		{
 		}
 
diff --git a/Engine/Plugins/Interchange/Runtime/Source/Import/Private/Gltf/InterchangeGltfTranslator.cpp b/Engine/Plugins/Interchange/Runtime/Source/Import/Private/Gltf/InterchangeGltfTranslator.cpp
--- a/Engine/Plugins/Interchange/Runtime/Source/Import/Private/Gltf/InterchangeGltfTranslator.cpp
+++ b/Engine/Plugins/Interchange/Runtime/Source/Import/Private/Gltf/InterchangeGltfTranslator.cpp
@@ -66,6 +66,7 @@ static const TArray<FString> ImporterSupportedExtensions = {
 	GLTF::ToString(GLTF::EExtension::KHR_MaterialsSheen),
 	GLTF::ToString(GLTF::EExtension::KHR_MaterialsSpecular),
 	GLTF::ToString(GLTF::EExtension::KHR_MaterialsPbrSpecularGlossiness),
+	GLTF::ToString(GLTF::EExtension::KHR_MaterialsEmissiveStrength),
 	GLTF::ToString(GLTF::EExtension::MSFT_PackingOcclusionRoughnessMetallic),
 	GLTF::ToString(GLTF::EExtension::MSFT_PackingNormalRoughnessMetallic),
 	/* Textures */
@@ -617,10 +618,10 @@ void UInterchangeGltfTranslator::HandleGltfMaterial( UInterchangeBaseNodeContain
 		}
 
 		// Emissive
-		if ( GltfMaterial.Emissive.TextureIndex != INDEX_NONE || !GltfMaterial.EmissiveFactor.IsNearlyZero() )
+		if ( GltfMaterial.Emissive.TextureIndex != INDEX_NONE || !GltfMaterial.EmissiveFactor.IsNearlyZero() || GltfMaterial.bHasEmissiveStrength)
 		{
 			TVariant< FLinearColor, float > EmissiveFactor;
-			EmissiveFactor.Set< FLinearColor >( FLinearColor( GltfMaterial.EmissiveFactor ) );
+			EmissiveFactor.Set< FLinearColor >(GltfMaterial.bHasEmissiveStrength ? FLinearColor(GltfMaterial.EmissiveFactor) * GltfMaterial.EmissiveStrength : FLinearColor(GltfMaterial.EmissiveFactor));
 
 			HandleGltfMaterialParameter( NodeContainer, GltfMaterial.Emissive, ShaderGraphNode, Common::Parameters::EmissiveColor.ToString(),
 				EmissiveFactor, Standard::Nodes::TextureSample::Outputs::RGB.ToString() );
