[Interchange - glTF] Add support for Vertex color.
Added VertexColor InterchangeMaterialDefinition and usage of it to GLTF Translator.

diff --git a/Engine/Plugins/Interchange/Runtime/Source/Import/Private/Gltf/InterchangeGltfTranslator.cpp b/Engine/Plugins/Interchange/Runtime/Source/Import/Private/Gltf/InterchangeGltfTranslator.cpp
--- a/Engine/Plugins/Interchange/Runtime/Source/Import/Private/Gltf/InterchangeGltfTranslator.cpp
+++ b/Engine/Plugins/Interchange/Runtime/Source/Import/Private/Gltf/InterchangeGltfTranslator.cpp
@@ -357,7 +357,7 @@ void UInterchangeGltfTranslator::HandleGltfNode( UInterchangeBaseNodeContainer&
 }
 
 void UInterchangeGltfTranslator::HandleGltfMaterialParameter( UInterchangeBaseNodeContainer& NodeContainer, const GLTF::FTextureMap& TextureMap, UInterchangeShaderNode& ShaderNode,
-		const FString& MapName, const TVariant< FLinearColor, float >& MapFactor, const FString& OutputChannel, const bool bInverse, const bool bIsNormal ) const
+		const FString& MapName, const TVariant< FLinearColor, float >& MapFactor, const FString& OutputChannel, const bool bInverse, const bool bIsNormal, const bool bUseVertexColor) const
 {
 	using namespace UE::Interchange::Materials;
 	using namespace UE::Interchange::Gltf::Private;
@@ -365,6 +365,26 @@ void UInterchangeGltfTranslator::HandleGltfMaterialParameter( UInterchangeBaseNo
 	UInterchangeShaderNode* NodeToConnectTo = &ShaderNode;
 	FString InputToConnectTo = MapName;
 
+	if (bUseVertexColor)
+	{
+		//From GLTF Specification:
+		// "if a primitive specifies a vertex color using the attribute semantic property COLOR_0, then this value acts as an additional linear multiplier to base color."
+		const FString MultiplierNodeName = MapName + TEXT("VertexColorMultiply");
+		UInterchangeShaderNode* MultiplierNode = UInterchangeShaderNode::Create(&NodeContainer, MultiplierNodeName, ShaderNode.GetUniqueID());
+		MultiplierNode->SetCustomShaderType(Standard::Nodes::Multiply::Name.ToString());
+
+		UInterchangeShaderPortsAPI::ConnectDefaultOuputToInput(NodeToConnectTo, InputToConnectTo, MultiplierNode->GetUniqueID());
+		NodeToConnectTo = MultiplierNode;
+		InputToConnectTo = Standard::Nodes::Multiply::Inputs::B.ToString();
+
+
+		const FString VertexColorNodeName = MapName + TEXT("VertexColor");
+		UInterchangeShaderNode* VertexColorNode = UInterchangeShaderNode::Create(&NodeContainer, VertexColorNodeName, ShaderNode.GetUniqueID());
+		VertexColorNode->SetCustomShaderType(Standard::Nodes::VertexColor::Name.ToString());
+
+		UInterchangeShaderPortsAPI::ConnectDefaultOuputToInput(MultiplierNode, Standard::Nodes::Multiply::Inputs::A.ToString(), VertexColorNode->GetUniqueID());
+	}
+
 	if (bInverse)
 	{
 		const FString OneMinusNodeName = MapName + TEXT("OneMinus");
@@ -488,7 +508,7 @@ void UInterchangeGltfTranslator::HandleGltfMaterialParameter( UInterchangeBaseNo
 	}
 }
 
-void UInterchangeGltfTranslator::HandleGltfMaterial( UInterchangeBaseNodeContainer& NodeContainer, const GLTF::FMaterial& GltfMaterial, UInterchangeShaderGraphNode& ShaderGraphNode ) const
+void UInterchangeGltfTranslator::HandleGltfMaterial( UInterchangeBaseNodeContainer& NodeContainer, const GLTF::FMaterial& GltfMaterial, UInterchangeShaderGraphNode& ShaderGraphNode, bool bUseVertexColor) const
 {
 	using namespace UE::Interchange::Materials;
 
@@ -511,7 +531,7 @@ void UInterchangeGltfTranslator::HandleGltfMaterial( UInterchangeBaseNodeContain
 			BaseColorFactor.Set< FLinearColor >(FLinearColor(GltfMaterial.BaseColorFactor));
 
 			HandleGltfMaterialParameter(NodeContainer, GltfMaterial.BaseColor, ShaderGraphNode, Unlit::Parameters::UnlitColor.ToString(),
-				BaseColorFactor, Standard::Nodes::TextureSample::Outputs::RGB.ToString());
+				BaseColorFactor, Standard::Nodes::TextureSample::Outputs::RGB.ToString(), false, false, bUseVertexColor);
 		}
 
 		// Opacity (use the base color alpha channel)
@@ -540,7 +560,7 @@ void UInterchangeGltfTranslator::HandleGltfMaterial( UInterchangeBaseNodeContain
 			BaseColorFactor.Set< FLinearColor >( FLinearColor( GltfMaterial.BaseColorFactor ) );
 
 			HandleGltfMaterialParameter( NodeContainer, GltfMaterial.BaseColor, ShaderGraphNode, PBR::Parameters::BaseColor.ToString(),
-				BaseColorFactor, Standard::Nodes::TextureSample::Outputs::RGB.ToString() );
+				BaseColorFactor, Standard::Nodes::TextureSample::Outputs::RGB.ToString(), false, false, bUseVertexColor );
 		}
 
 		// Metallic
@@ -579,7 +599,7 @@ void UInterchangeGltfTranslator::HandleGltfMaterial( UInterchangeBaseNodeContain
 			DiffuseColorFactor.Set< FLinearColor >( FLinearColor( GltfMaterial.BaseColorFactor ) );
 
 			HandleGltfMaterialParameter( NodeContainer, GltfMaterial.BaseColor, ShaderGraphNode, Phong::Parameters::DiffuseColor.ToString(),
-				DiffuseColorFactor, Standard::Nodes::TextureSample::Outputs::RGB.ToString() );
+				DiffuseColorFactor, Standard::Nodes::TextureSample::Outputs::RGB.ToString(), false, false, bUseVertexColor );
 		}
 
 		// Specular Color
@@ -1069,6 +1089,31 @@ bool UInterchangeGltfTranslator::Translate( UInterchangeBaseNodeContainer& NodeC
 		}
 	}
 
+	// Meshes
+	TSet<FString> MaterialsUsedOnMeshesWithVertexColor;
+	TSet<int> UnusedGltfMeshIndices;
+	{
+		int32 MeshIndex = 0;
+		for (const GLTF::FMesh& GltfMesh : GltfAsset.Meshes)
+		{
+			UnusedGltfMeshIndices.Add(MeshIndex++);
+
+			if (GltfMesh.HasColors())
+			{
+				for (int32 PrimitiveCounter = 0; PrimitiveCounter < GltfMesh.Primitives.Num(); PrimitiveCounter++)
+				{
+					const GLTF::FPrimitive& Primitive = GltfMesh.Primitives[PrimitiveCounter];
+
+					if (GltfAsset.Materials.IsValidIndex(Primitive.MaterialIndex))
+					{
+						const FString ShaderGraphNodeUid = UInterchangeShaderGraphNode::MakeNodeUid(GltfAsset.Materials[Primitive.MaterialIndex].UniqueId);
+						MaterialsUsedOnMeshesWithVertexColor.Add(ShaderGraphNodeUid);
+					}
+				}
+			}
+		}
+	}
+
 	// Materials
 	{
 		int32 MaterialIndex = 0;
@@ -1077,21 +1122,12 @@ bool UInterchangeGltfTranslator::Translate( UInterchangeBaseNodeContainer& NodeC
 			UInterchangeShaderGraphNode* ShaderGraphNode = UInterchangeShaderGraphNode::Create(&NodeContainer, GltfMaterial.UniqueId);
 			ShaderGraphNode->SetDisplayLabel(GltfMaterial.Name);
 
-			HandleGltfMaterial( NodeContainer, GltfMaterial, *ShaderGraphNode );
+			bool bUseVertexColor = MaterialsUsedOnMeshesWithVertexColor.Contains(ShaderGraphNode->GetUniqueID());
+			HandleGltfMaterial( NodeContainer, GltfMaterial, *ShaderGraphNode, bUseVertexColor );
 			++MaterialIndex;
 		}
 	}
 
-	TSet<int> UnusedGltfMeshIndices;
-	// Meshes
-	{
-		int32 MeshIndex = 0;
-		for (const GLTF::FMesh& GltfMesh : GltfAsset.Meshes)
-		{
-			UnusedGltfMeshIndices.Add(MeshIndex++);
-		}
-	}
-
 	// Cameras
 	{
 		int32 CameraIndex = 0;
diff --git a/Engine/Plugins/Interchange/Runtime/Source/Import/Public/Gltf/InterchangeGltfTranslator.h b/Engine/Plugins/Interchange/Runtime/Source/Import/Public/Gltf/InterchangeGltfTranslator.h
--- a/Engine/Plugins/Interchange/Runtime/Source/Import/Public/Gltf/InterchangeGltfTranslator.h
+++ b/Engine/Plugins/Interchange/Runtime/Source/Import/Public/Gltf/InterchangeGltfTranslator.h
@@ -65,9 +65,9 @@ protected:
 	void HandleGltfSkeletons( UInterchangeBaseNodeContainer& NodeContainer, const FString& SceneNodeUid, const TArray<int32>& SkinnedMeshNodes, TSet<int>& UnusedMeshIndices ) const;
 	void HandleGltfNode( UInterchangeBaseNodeContainer& NodeContainer, const GLTF::FNode& GltfNode, const FString& ParentNodeUid, const int32 NodeIndex, 
 		bool &bHasVariants, TArray<int32>& SkinnedMeshNodes, TSet<int>& UnusedMeshIndices ) const;
-	void HandleGltfMaterial( UInterchangeBaseNodeContainer& NodeContainer, const GLTF::FMaterial& GltfMaterial, UInterchangeShaderGraphNode& ShaderGraphNode ) const;
+	void HandleGltfMaterial( UInterchangeBaseNodeContainer& NodeContainer, const GLTF::FMaterial& GltfMaterial, UInterchangeShaderGraphNode& ShaderGraphNode, bool bUseVertexColor ) const;
 	void HandleGltfMaterialParameter( UInterchangeBaseNodeContainer& NodeContainer, const GLTF::FTextureMap& TextureMap, UInterchangeShaderNode& ShaderNode,
-		const FString& MapName, const TVariant< FLinearColor, float >& MapFactor, const FString& OutputChannel, const bool bInverse = false, const bool bIsNormal = false ) const;
+		const FString& MapName, const TVariant< FLinearColor, float >& MapFactor, const FString& OutputChannel, const bool bInverse = false, const bool bIsNormal = false, const bool bUseVertexColor = false ) const;
 	void HandleGltfAnimation(UInterchangeBaseNodeContainer& NodeContainer, int32 AnimationIndex) const;
 	void HandleGltfVariants(UInterchangeBaseNodeContainer& NodeContainer, const FString& FileName) const;
 	UInterchangeMeshNode* HandleGltfMesh(UInterchangeBaseNodeContainer& NodeContainer, const GLTF::FMesh& GltfMesh, int MeshIndex,
diff --git a/Engine/Plugins/Interchange/Runtime/Source/Nodes/Public/InterchangeMaterialDefinitions.h b/Engine/Plugins/Interchange/Runtime/Source/Nodes/Public/InterchangeMaterialDefinitions.h
--- a/Engine/Plugins/Interchange/Runtime/Source/Nodes/Public/InterchangeMaterialDefinitions.h
+++ b/Engine/Plugins/Interchange/Runtime/Source/Nodes/Public/InterchangeMaterialDefinitions.h
@@ -259,6 +259,11 @@ namespace Interchange
 						const FName Position = TEXT("Position");
 					}
 				}
+
+				namespace VertexColor
+				{
+					const FName Name = TEXT("VertexColor");
+				}
 			}
 		}
 
diff --git a/Engine/Plugins/Interchange/Runtime/Source/Pipelines/Private/InterchangeGenericMaterialPipeline.cpp b/Engine/Plugins/Interchange/Runtime/Source/Pipelines/Private/InterchangeGenericMaterialPipeline.cpp
--- a/Engine/Plugins/Interchange/Runtime/Source/Pipelines/Private/InterchangeGenericMaterialPipeline.cpp
+++ b/Engine/Plugins/Interchange/Runtime/Source/Pipelines/Private/InterchangeGenericMaterialPipeline.cpp
@@ -41,6 +41,7 @@
 #include "Materials/MaterialExpressionTransform.h"
 #include "Materials/MaterialExpressionVectorNoise.h"
 #include "Materials/MaterialExpressionVectorParameter.h"
+#include "Materials/MaterialExpressionVertexColor.h"
 #include "Materials/MaterialFunction.h"
 #include "Materials/MaterialInstanceConstant.h"
 #include "Materials/MaterialInstanceDynamic.h"
@@ -2319,6 +2320,11 @@ UInterchangeMaterialExpressionFactoryNode* UInterchangeGenericMaterialPipeline::
 	{
 		HandleVectorNoiseNode(ShaderNode, MaterialFactoryNode, MaterialExpression);
 	}
+	else if(*ShaderType == Nodes::VertexColor::Name)
+	{
+		MaterialExpression->SetCustomExpressionClassName(UMaterialExpressionVertexColor::StaticClass()->GetName());
+		
+	}
 	else
 	{
 		const FString ExpressionClassName = TEXT("MaterialExpression") + ShaderType;
