Reverting InterchangeGenericMaterialPipeline related changes for VertexSupport as it's not needed
(VertexSupport shader node is automatically generated and needs no special treatment)

diff --git a/Engine/Plugins/Interchange/Runtime/Source/Pipelines/Private/InterchangeGenericMaterialPipeline.cpp b/Engine/Plugins/Interchange/Runtime/Source/Pipelines/Private/InterchangeGenericMaterialPipeline.cpp
--- a/Engine/Plugins/Interchange/Runtime/Source/Pipelines/Private/InterchangeGenericMaterialPipeline.cpp
+++ b/Engine/Plugins/Interchange/Runtime/Source/Pipelines/Private/InterchangeGenericMaterialPipeline.cpp
@@ -41,7 +41,6 @@
 #include "Materials/MaterialExpressionTransform.h"
 #include "Materials/MaterialExpressionVectorNoise.h"
 #include "Materials/MaterialExpressionVectorParameter.h"
-#include "Materials/MaterialExpressionVertexColor.h"
 #include "Materials/MaterialFunction.h"
 #include "Materials/MaterialInstanceConstant.h"
 #include "Materials/MaterialInstanceDynamic.h"
@@ -2320,11 +2319,6 @@ UInterchangeMaterialExpressionFactoryNode* UInterchangeGenericMaterialPipeline::
 	{
 		HandleVectorNoiseNode(ShaderNode, MaterialFactoryNode, MaterialExpression);
 	}
-	else if(*ShaderType == Nodes::VertexColor::Name)
-	{
-		MaterialExpression->SetCustomExpressionClassName(UMaterialExpressionVertexColor::StaticClass()->GetName());
-		
-	}
 	else
 	{
 		const FString ExpressionClassName = TEXT("MaterialExpression") + ShaderType;
