MaterialX:
- Using the name of the <surfacematerial> instead of the surfaceshader
- Using Material Expressions for <place2d>, <rotate2d> and <ramp4>
- Adding support of clamp and mirror modes for textures
- Modifying Order of Lerp for SplitTB, RampTB Material Expression

Test files:
- Renaming imported tests cases
-changing ground truth for python test

diff --git a/Engine/Plugins/Interchange/Runtime/Source/Import/Private/MaterialX/InterchangeMaterialXDefinitions.h b/Engine/Plugins/Interchange/Runtime/Source/Import/Private/MaterialX/InterchangeMaterialXDefinitions.h
--- a/Engine/Plugins/Interchange/Runtime/Source/Import/Private/MaterialX/InterchangeMaterialXDefinitions.h
+++ b/Engine/Plugins/Interchange/Runtime/Source/Import/Private/MaterialX/InterchangeMaterialXDefinitions.h
@@ -271,6 +271,7 @@ MATERIALX_NAMESPACE_BEGIN
 		// Procedural nodes 
 		static constexpr const char* Constant = "constant";
 		// Procedural2D nodes 
+		static constexpr const char* Ramp4 = "ramp4";
 		static constexpr const char* RampLR = "ramplr";
 		static constexpr const char* RampTB = "ramptb";
 		static constexpr const char* SplitLR = "splitlr";
@@ -306,7 +307,11 @@ MATERIALX_NAMESPACE_BEGIN
 		// Application
 		static constexpr const char* Time= "time";
 		// PBR
+		static constexpr const char* GltfPbr = "gltf_pbr";
+		static constexpr const char* DisneyBSDF2012 = "disney_brdf_2012";
+		static constexpr const char* DisneyBSDF2015 = "disney_bsdf_2015";
 		static constexpr const char* StandardSurface = "standard_surface";
+		static constexpr const char* UsdPreviewSurface = "UsdPreviewSurface";
 		// Shader
 		static constexpr const char* SurfaceUnlit = "surface_unlit";
 		// Convolution
@@ -318,6 +323,7 @@ MATERIALX_NAMESPACE_BEGIN
 	{
 		static constexpr const char* StandardSurface = "ND_standard_surface_surfaceshader";
 		static constexpr const char* SurfaceUnlit = "ND_surface_unlit";
+		static constexpr const char* UsdPreviewSurface = "ND_UsdPreviewSurface_surfaceshader";
 		static constexpr const char* PointLight = "ND_point_light";
 		static constexpr const char* DirectionalLight = "ND_directional_light";
 		static constexpr const char* SpotLight = "ND_spot_light";
diff --git a/Engine/Plugins/Interchange/Runtime/Source/Import/Private/MaterialX/MaterialExpressions/MaterialExpressionPlace2D.cpp b/Engine/Plugins/Interchange/Runtime/Source/Import/Private/MaterialX/MaterialExpressions/MaterialExpressionPlace2D.cpp
new file mode 100644
--- /dev/null
+++ b/Engine/Plugins/Interchange/Runtime/Source/Import/Private/MaterialX/MaterialExpressions/MaterialExpressionPlace2D.cpp
@@ -0,0 +1,160 @@
+// Copyright Epic Games, Inc. All Rights Reserved.
+
+#include "MaterialExpressionPlace2D.h"
+#include "MaterialExpressionRotate2D.h"
+#include "MaterialCompiler.h"
+#include "Materials/MaterialExpressionAdd.h"
+#include "Materials/MaterialExpressionAppendVector.h"
+#include "Materials/MaterialExpressionComponentMask.h"
+#include "Materials/MaterialExpressionConstant.h"
+#include "Materials/MaterialExpressionConstant2Vector.h"
+#include "Materials/MaterialExpressionDivide.h"
+#include "Materials/MaterialExpressionMultiply.h"
+#include "Materials/MaterialExpressionOneMinus.h"
+#include "Materials/MaterialExpressionSubtract.h"
+#include "Materials/MaterialExpressionTextureCoordinate.h"
+
+#include UE_INLINE_GENERATED_CPP_BY_NAME(MaterialExpressionPlace2D)
+
+#define LOCTEXT_NAMESPACE "MaterialExpressionPlace2D"
+
+UMaterialExpressionPlace2D::UMaterialExpressionPlace2D(const FObjectInitializer& ObjectInitializer)
+	: Super(ObjectInitializer)
+{
+	// Structure to hold one-time initialization
+	struct FConstructorStatics
+	{
+		FText NAME_MaterialX;
+		FText NAME_Compositing;
+		FConstructorStatics()
+			: NAME_MaterialX(LOCTEXT("MaterialX", "MaterialX"))
+			, NAME_Compositing(LOCTEXT("Math", "Math"))
+		{}
+	};
+	static FConstructorStatics ConstructorStatics;
+
+#if WITH_EDITORONLY_DATA
+	MenuCategories.Add(ConstructorStatics.NAME_MaterialX);
+	MenuCategories.Add(ConstructorStatics.NAME_Compositing);
+#endif
+}
+
+#if WITH_EDITOR
+namespace
+{
+	template<typename ExpressionT>
+	FExpressionInput NewExpression(const FExpressionInput& A, const FExpressionInput& B)
+	{
+		FExpressionInput ExpressionInput;
+
+		ExpressionT* Expression = NewObject<ExpressionT>();
+		Expression->A = A;
+		Expression->B = B;
+
+		ExpressionInput.Connect(0, Expression);
+
+		return ExpressionInput;
+	}
+}
+
+int32 UMaterialExpressionPlace2D::Compile(FMaterialCompiler* Compiler, int32 OutputIndex)
+{
+	auto Constant2Expression = [](FExpressionInput& Input, float X, float Y)
+	{
+		if(Input.GetTracedInput().Expression == nullptr)
+		{
+			UMaterialExpressionConstant2Vector* Expression = NewObject<UMaterialExpressionConstant2Vector>();
+			Expression->R = X;
+			Expression->G = Y;
+			Input.Connect(0, Expression);
+		}
+	};
+
+	auto InputMask = [](const FExpressionInput& Input, uint32 R, uint32 G, uint32 B, uint32 A)
+	{
+		UMaterialExpressionComponentMask* Mask = NewObject<UMaterialExpressionComponentMask>();
+		Mask->R = R;
+		Mask->G = G;
+		Mask->B = B;
+		Mask->A = A;
+		Mask->Input = Input;
+
+		FExpressionInput ExpressionInput;
+		ExpressionInput.Connect(0, Mask);
+		return ExpressionInput;
+	};
+
+	if(Coordinates.GetTracedInput().Expression == nullptr)
+	{
+		UMaterialExpressionTextureCoordinate* TextureCoordinate = NewObject<UMaterialExpressionTextureCoordinate>();
+		TextureCoordinate->CoordinateIndex = ConstCoordinate;
+		Coordinates.Connect(0, TextureCoordinate);
+	}
+
+	if(RotationAngle.GetTracedInput().Expression == nullptr)
+	{
+		UMaterialExpressionConstant* Expression = NewObject<UMaterialExpressionConstant>();
+		Expression->R = ConstRotationAngle;
+		RotationAngle.Connect(0, Expression);
+	}
+
+	Constant2Expression(Pivot, 0, 0);
+	Constant2Expression(Scale, 1, 1);
+	Constant2Expression(Offset, 0, 0);
+
+	// We need to adjust the UVs to MaterialX UVs, Y is inverted
+	FExpressionInput ScaleX = InputMask(Scale, 1, 0, 0, 0);
+	FExpressionInput ScaleY = InputMask(Scale, 0, 1, 0, 0);
+	FExpressionInput NegateScaleY;
+	{
+		UMaterialExpressionMultiply* NegateScaleYExp = NewObject<UMaterialExpressionMultiply>();
+		NegateScaleYExp->ConstA = -1.f;
+		NegateScaleYExp->B = ScaleY;
+		NegateScaleY.Connect(0, NegateScaleYExp);
+	}
+
+	FExpressionInput AppendVectorScale = NewExpression<UMaterialExpressionAppendVector>(ScaleX, NegateScaleY);
+
+	FExpressionInput SubPivot = NewExpression<UMaterialExpressionSubtract>(Coordinates, Pivot);
+	FExpressionInput ApplyScale = NewExpression<UMaterialExpressionDivide>(SubPivot, AppendVectorScale);
+
+	FExpressionInput ApplyRot;
+	{
+		UMaterialExpressionRotate2D* Rotate2D = NewObject<UMaterialExpressionRotate2D>();
+		Rotate2D->Input = ApplyScale;
+		Rotate2D->RotationAngle = RotationAngle;
+		ApplyRot.Connect(0, Rotate2D);
+	}
+
+	FExpressionInput Applyoffset = NewExpression<UMaterialExpressionSubtract>(ApplyRot, Offset);
+	FExpressionInput AddPivot = NewExpression<UMaterialExpressionAdd>(Applyoffset, Pivot);
+
+	// Readjust the UVs back
+	FExpressionInput AddPivotX = InputMask(AddPivot, 1, 0, 0, 0);
+	FExpressionInput AddPivotY = InputMask(AddPivot, 0, 1, 0, 0);
+	FExpressionInput NegateAddPivotY;
+	{
+		UMaterialExpressionOneMinus* NegateAddPivotYExp = NewObject<UMaterialExpressionOneMinus>();
+		NegateAddPivotYExp->Input = AddPivotY;
+		NegateAddPivotY.Connect(0, NegateAddPivotYExp);
+	}
+
+	FExpressionInput AppendVectorAddPivot = NewExpression<UMaterialExpressionAppendVector>(AddPivotX, NegateAddPivotY);
+
+	return AppendVectorAddPivot.Compile(Compiler);
+}
+
+void UMaterialExpressionPlace2D::GetCaption(TArray<FString>& OutCaptions) const
+{
+	OutCaptions.Add(TEXT("Place2D"));
+}
+
+void UMaterialExpressionPlace2D::GetExpressionToolTip(TArray<FString>& OutToolTip)
+{
+	ConvertToMultilineToolTip(TEXT("Transform incoming UV texture coordinates from one 2D frame of reference to another."
+								   "The order of operation is: -pivot, scale, rotate, translate, +pivot."),
+							  75, OutToolTip);
+}
+#endif
+
+#undef LOCTEXT_NAMESPACE 
\ No newline at end of file
diff --git a/Engine/Plugins/Interchange/Runtime/Source/Import/Private/MaterialX/MaterialExpressions/MaterialExpressionPlace2D.h b/Engine/Plugins/Interchange/Runtime/Source/Import/Private/MaterialX/MaterialExpressions/MaterialExpressionPlace2D.h
new file mode 100644
--- /dev/null
+++ b/Engine/Plugins/Interchange/Runtime/Source/Import/Private/MaterialX/MaterialExpressions/MaterialExpressionPlace2D.h
@@ -0,0 +1,62 @@
+// Copyright Epic Games, Inc. All Rights Reserved.
+#pragma once
+
+#include "CoreMinimal.h"
+#include "Materials/MaterialExpression.h"
+
+#include "MaterialExpressionPlace2D.generated.h"
+
+/**
+ * Transform incoming UV texture coordinates from one 2D frame of reference to another. 
+ * operationorder (integer enum): the order in which to perform the transform operations.
+ * "0" or "SRT" performs -pivot, scale, rotate, translate, +pivot as per the original
+ * implementation matching the behavior of certain DCC packages, and "1" or "TRS" performs
+ * -pivot, translate, rotate, scale, +pivot which does not introduce texture shear.
+ * Default is 0 "SRT" for backward compatibility.*/
+UCLASS(MinimalAPI, collapsecategories, hidecategories = Object)
+class UMaterialExpressionPlace2D : public UMaterialExpression
+{
+	GENERATED_UCLASS_BODY()
+
+	UPROPERTY(meta = (RequiredInput = "false", ToolTip = "Defaults to 'ConstCoordinate' if not specified"))
+	FExpressionInput Coordinates;
+
+	/** The pivot coordinate for scale and rotate: this is subtracted from u,v before
+	 *  applying scale/rotate, then added back after. Default is (0,0). */
+	UPROPERTY()
+	FExpressionInput Pivot;
+
+	/** Divide the u,v coord (after subtracting pivot ) by this, so a scale (2,2)
+	 *  makes the texture image appear twice as big. Negative values can be used to flip or flop the
+	 *  texture space. Default is (1,1). */
+	UPROPERTY()
+	FExpressionInput Scale;
+
+	/** Subtract this amount from the scaled/rotated/“pivot added back” UV
+	 *  coordinate; since U0,V0 is typically the lower left corner, a positive offset moves the texture
+	 *  image up and right. Default is (0,0).*/
+	UPROPERTY()
+	FExpressionInput Offset;
+
+	/** Rotate u,v coord (after subtracting pivot) by this amount in degrees, so a
+	 *  positive value rotates UV coords counter-clockwise, and the image clockwise. */
+	UPROPERTY(meta = (RequiredInput = "false", ToolTip = "Defaults to 'ConstRotationAngle' if not specified"))
+	FExpressionInput RotationAngle;
+
+	/** only used if RotationAngle is not hooked up */
+	UPROPERTY(EditAnywhere, Category = MaterialExpression, meta = (OverridingInputProperty = "RotationAngle"))
+	float ConstRotationAngle = 0.f;
+
+	/** only used if Coordinates is not hooked up */
+	UPROPERTY(EditAnywhere, Category = MaterialExpression, meta = (OverridingInputProperty = "Coordinates"))
+	uint8 ConstCoordinate;
+
+	//~ Begin UMaterialExpression Interface
+#if WITH_EDITOR
+	virtual int32 Compile(class FMaterialCompiler* Compiler, int32 OutputIndex) override;
+	virtual void GetCaption(TArray<FString>& OutCaptions) const override;
+	virtual void GetExpressionToolTip(TArray<FString>& OutToolTip) override;
+#endif
+	//~ End UMaterialExpression Interface
+};
+
diff --git a/Engine/Plugins/Interchange/Runtime/Source/Import/Private/MaterialX/MaterialExpressions/MaterialExpressionRamp4.cpp b/Engine/Plugins/Interchange/Runtime/Source/Import/Private/MaterialX/MaterialExpressions/MaterialExpressionRamp4.cpp
new file mode 100644
--- /dev/null
+++ b/Engine/Plugins/Interchange/Runtime/Source/Import/Private/MaterialX/MaterialExpressions/MaterialExpressionRamp4.cpp
@@ -0,0 +1,78 @@
+// Copyright Epic Games, Inc. All Rights Reserved.
+#include "MaterialExpressionRamp4.h"
+#include "MaterialCompiler.h"
+
+#define LOCTEXT_NAMESPACE "MaterialExpressionRamp4"
+
+UMaterialExpressionRamp4::UMaterialExpressionRamp4(const FObjectInitializer& ObjectInitializer)
+	: Super(ObjectInitializer)
+{
+	// Structure to hold one-time initialization
+	struct FConstructorStatics
+	{
+		FText NAME_MaterialX;
+		FText NAME_Utility;
+		FConstructorStatics()
+			: NAME_MaterialX(LOCTEXT("MaterialX", "MaterialX"))
+			, NAME_Utility(LOCTEXT("Procedural", "Procedural"))
+		{}
+	};
+	static FConstructorStatics ConstructorStatics;
+
+#if WITH_EDITORONLY_DATA
+	MenuCategories.Add(ConstructorStatics.NAME_MaterialX);
+	MenuCategories.Add(ConstructorStatics.NAME_Utility);
+#endif
+}
+
+#if WITH_EDITOR
+int32 UMaterialExpressionRamp4::Compile(FMaterialCompiler* Compiler, int32 OutputIndex)
+{
+	if(!A.GetTracedInput().Expression)
+	{
+		return Compiler->Errorf(TEXT("Missing A input"));
+	}
+
+	if(!B.GetTracedInput().Expression)
+	{
+		return Compiler->Errorf(TEXT("Missing B input"));
+	}
+	
+	if(!C.GetTracedInput().Expression)
+	{
+		return Compiler->Errorf(TEXT("Missing C input"));
+	}
+
+	if(!D.GetTracedInput().Expression)
+	{
+		return Compiler->Errorf(TEXT("Missing D input"));
+	}
+
+	int32 IndexCoordinates = Coordinates.GetTracedInput().Expression ? Coordinates.Compile(Compiler) : Compiler->TextureCoordinate(ConstCoordinate, false, false);
+	int32 IndexValueTL = A.Compile(Compiler);
+	int32 IndexValueTR = B.Compile(Compiler);
+	int32 IndexValueBL = C.Compile(Compiler);
+	int32 IndexValueBR = D.Compile(Compiler);
+
+	int32 TexClamp = Compiler->Saturate(IndexCoordinates);
+
+	int32 S = Compiler->ComponentMask(TexClamp, 1, 0, 0, 0);
+	int32 T = Compiler->ComponentMask(TexClamp, 0, 1, 0, 0);
+	int32 MixTop = Compiler->Lerp(IndexValueTL, IndexValueTR, S);
+	int32 MixBot = Compiler->Lerp(IndexValueBL, IndexValueBR, S);
+
+	return Compiler->Lerp(MixBot, MixTop, T);
+}
+
+void UMaterialExpressionRamp4::GetCaption(TArray<FString>& OutCaptions) const
+{
+	OutCaptions.Add(TEXT("Ramp4"));
+}
+
+void UMaterialExpressionRamp4::GetExpressionToolTip(TArray<FString>& OutToolTip)
+{
+	ConvertToMultilineToolTip(TEXT("A 4-corner bilinear value ramp."), 40, OutToolTip);
+}
+#endif
+
+#undef LOCTEXT_NAMESPACE 
\ No newline at end of file
diff --git a/Engine/Plugins/Interchange/Runtime/Source/Import/Private/MaterialX/MaterialExpressions/MaterialExpressionRamp4.h b/Engine/Plugins/Interchange/Runtime/Source/Import/Private/MaterialX/MaterialExpressions/MaterialExpressionRamp4.h
new file mode 100644
--- /dev/null
+++ b/Engine/Plugins/Interchange/Runtime/Source/Import/Private/MaterialX/MaterialExpressions/MaterialExpressionRamp4.h
@@ -0,0 +1,48 @@
+// Copyright Epic Games, Inc. All Rights Reserved.
+#pragma once
+
+#include "CoreMinimal.h"
+#include "Materials/MaterialExpression.h"
+
+#include "MaterialExpressionRamp4.generated.h"
+
+/**
+ * A material expression that computes a 4-corner bilinear value ramp..
+ */
+UCLASS(collapsecategories, hidecategories = Object, MinimalAPI)
+class UMaterialExpressionRamp4: public UMaterialExpression
+{
+	GENERATED_UCLASS_BODY()
+
+	UPROPERTY(meta = (RequiredInput = "false", ToolTip = "Defaults to 'ConstCoordinate' if not specified"))
+	FExpressionInput Coordinates;
+
+	/** The value at the top-left (U0V1) corner */
+	UPROPERTY()
+	FExpressionInput A;
+
+	/** The value at the top-right (U1V1) corner */
+	UPROPERTY()
+	FExpressionInput B;
+
+	/** The value at the bottom-left (U0V0) corner */
+	UPROPERTY()
+	FExpressionInput C;
+
+	/** The value at the bottom-right (U1V0) corner */
+	UPROPERTY()
+	FExpressionInput D;
+
+	/** only used if Coordinates is not hooked up */
+	UPROPERTY(EditAnywhere, Category = MaterialExpression, meta = (OverridingInputProperty = "Coordinates"))
+	uint8 ConstCoordinate;
+
+	//~ Begin UMaterialExpression Interface
+#if WITH_EDITOR
+	virtual int32 Compile(class FMaterialCompiler* Compiler, int32 OutputIndex) override;
+	virtual void GetCaption(TArray<FString>& OutCaptions) const override;
+	virtual void GetExpressionToolTip(TArray<FString>& OutToolTip) override;
+#endif
+	//~ End UMaterialExpression Interface
+};
+
diff --git a/Engine/Plugins/Interchange/Runtime/Source/Import/Private/MaterialX/MaterialExpressions/MaterialExpressionRampTopBottom.cpp b/Engine/Plugins/Interchange/Runtime/Source/Import/Private/MaterialX/MaterialExpressions/MaterialExpressionRampTopBottom.cpp
--- a/Engine/Plugins/Interchange/Runtime/Source/Import/Private/MaterialX/MaterialExpressions/MaterialExpressionRampTopBottom.cpp
+++ b/Engine/Plugins/Interchange/Runtime/Source/Import/Private/MaterialX/MaterialExpressions/MaterialExpressionRampTopBottom.cpp
@@ -42,7 +42,7 @@ int32 UMaterialExpressionRampTopBottom::Compile(FMaterialCompiler* Compiler, int
 
 	int32 CoordinateIndex = Coordinates.GetTracedInput().Expression ? Coordinates.Compile(Compiler) : Compiler->TextureCoordinate(ConstCoordinate, false, false);
 
-	return Compiler->Lerp(A.Compile(Compiler), B.Compile(Compiler), Compiler->Saturate(Compiler->ComponentMask(CoordinateIndex, false, true, false, false)));
+	return Compiler->Lerp(B.Compile(Compiler), A.Compile(Compiler), Compiler->Saturate(Compiler->ComponentMask(CoordinateIndex, false, true, false, false)));
 }
 
 void UMaterialExpressionRampTopBottom::GetCaption(TArray<FString>& OutCaptions) const
diff --git a/Engine/Plugins/Interchange/Runtime/Source/Import/Private/MaterialX/MaterialExpressions/MaterialExpressionRotate2D.cpp b/Engine/Plugins/Interchange/Runtime/Source/Import/Private/MaterialX/MaterialExpressions/MaterialExpressionRotate2D.cpp
new file mode 100644
--- /dev/null
+++ b/Engine/Plugins/Interchange/Runtime/Source/Import/Private/MaterialX/MaterialExpressions/MaterialExpressionRotate2D.cpp
@@ -0,0 +1,70 @@
+// Copyright Epic Games, Inc. All Rights Reserved.
+
+#include "MaterialExpressionRotate2D.h"
+#include "MaterialCompiler.h"
+#include "Materials/MaterialExpressionCustom.h"
+
+#include UE_INLINE_GENERATED_CPP_BY_NAME(MaterialExpressionRotate2D)
+
+#define LOCTEXT_NAMESPACE "MaterialExpressionRotate2D"
+
+UMaterialExpressionRotate2D::UMaterialExpressionRotate2D(const FObjectInitializer& ObjectInitializer)
+	: Super(ObjectInitializer)
+{
+	// Structure to hold one-time initialization
+	struct FConstructorStatics
+	{
+		FText NAME_MaterialX;
+		FText NAME_Math;
+		FConstructorStatics()
+			: NAME_MaterialX(LOCTEXT("MaterialX", "MaterialX"))
+			, NAME_Math(LOCTEXT("Math", "Math"))
+		{}
+	};
+	static FConstructorStatics ConstructorStatics;
+
+#if WITH_EDITORONLY_DATA
+	MenuCategories.Add(ConstructorStatics.NAME_MaterialX);
+	MenuCategories.Add(ConstructorStatics.NAME_Math);
+#endif
+}
+
+#if WITH_EDITOR
+int32 UMaterialExpressionRotate2D::Compile(FMaterialCompiler* Compiler, int32 OutputIndex)
+{
+	if(!Input.GetTracedInput().Expression)
+	{
+		return Compiler->Errorf(TEXT("Missing Rotate2D input"));
+	}
+
+	int32 IndexInput = Input.Compile(Compiler);
+	int32 IndexAngle = RotationAngle.GetTracedInput().Expression ? RotationAngle.Compile(Compiler) : Compiler->Constant(ConstRotationAngle);
+
+	int32 RotationRadians = Compiler->Mul(IndexAngle, Compiler->Constant(float(UE_PI) / 180.f));
+	int32 SA = Compiler->Sine(RotationRadians);
+	int32 CA = Compiler->Cosine(RotationRadians);
+
+	int32 InputX = Compiler->ComponentMask(IndexInput, true, false, false, false);
+	int32 InputY = Compiler->ComponentMask(IndexInput, false, true, false, false);
+	
+	return Compiler->AppendVector(
+		Compiler->Add(Compiler->Mul(CA, InputX),
+					  Compiler->Mul(SA, InputY)),
+
+		Compiler->Sub(Compiler->Mul(CA, InputY),
+					  Compiler->Mul(SA, InputX))
+	);
+}
+
+void UMaterialExpressionRotate2D::GetCaption(TArray<FString>& OutCaptions) const
+{
+	OutCaptions.Add(TEXT("Rotate2D"));
+}
+
+void UMaterialExpressionRotate2D::GetExpressionToolTip(TArray<FString>& OutToolTip)
+{
+	OutToolTip.Add(TEXT("rotates a vector2 value about the origin in 2D."));
+}
+#endif
+
+#undef LOCTEXT_NAMESPACE 
\ No newline at end of file
diff --git a/Engine/Plugins/Interchange/Runtime/Source/Import/Private/MaterialX/MaterialExpressions/MaterialExpressionRotate2D.h b/Engine/Plugins/Interchange/Runtime/Source/Import/Private/MaterialX/MaterialExpressions/MaterialExpressionRotate2D.h
new file mode 100644
--- /dev/null
+++ b/Engine/Plugins/Interchange/Runtime/Source/Import/Private/MaterialX/MaterialExpressions/MaterialExpressionRotate2D.h
@@ -0,0 +1,36 @@
+// Copyright Epic Games, Inc. All Rights Reserved.
+#pragma once
+
+#include "Materials/MaterialExpression.h"
+
+#include "MaterialExpressionRotate2D.generated.h"
+
+/**
+ * A material expression that rotates a vector2 value about the origin in 2D.
+ */
+UCLASS(collapsecategories, hidecategories = Object, MinimalAPI)
+class UMaterialExpressionRotate2D : public UMaterialExpression
+{
+	GENERATED_UCLASS_BODY()
+
+	UPROPERTY()
+	FExpressionInput Input;
+
+	/** The amount to rotate, specified in degrees, with positive values rotating the
+	 *  incoming vector counterclockwise.*/
+	UPROPERTY(meta = (RequiredInput = "false", ToolTip = "RotationAngle in degrees. Defaults to 'ConstRotationAngle' if not specified"))
+	FExpressionInput RotationAngle;
+
+	/** only used if RotationAngle is not hooked up */
+	UPROPERTY(EditAnywhere, Category = MaterialExpressionAngle, meta = (OverridingInputProperty = "RotationAngle"))
+	float ConstRotationAngle = 0.f;
+
+	//~ Begin UMaterialExpression Interface
+#if WITH_EDITOR
+	virtual int32 Compile(class FMaterialCompiler* Compiler, int32 OutputIndex) override;
+	virtual void GetCaption(TArray<FString>& OutCaptions) const override;
+	virtual void GetExpressionToolTip(TArray<FString>& OutToolTip) override;
+#endif
+	//~ End UMaterialExpression Interface
+};
+
diff --git a/Engine/Plugins/Interchange/Runtime/Source/Import/Private/MaterialX/MaterialExpressions/MaterialExpressionSplitTopBottom.cpp b/Engine/Plugins/Interchange/Runtime/Source/Import/Private/MaterialX/MaterialExpressions/MaterialExpressionSplitTopBottom.cpp
--- a/Engine/Plugins/Interchange/Runtime/Source/Import/Private/MaterialX/MaterialExpressions/MaterialExpressionSplitTopBottom.cpp
+++ b/Engine/Plugins/Interchange/Runtime/Source/Import/Private/MaterialX/MaterialExpressions/MaterialExpressionSplitTopBottom.cpp
@@ -50,7 +50,7 @@ int32 UMaterialExpressionSplitTopBottom::Compile(FMaterialCompiler* Compiler, in
 	int32 AFWidth = Compiler->Mul(Compiler->Length(Compiler->AppendVector(Compiler->DDX(TexcoordYIndex), Compiler->DDY(TexcoordYIndex))), InvSqrt2);
 	int32 AAStep = Compiler->SmoothStep(Compiler->Sub(CenterIndex, AFWidth), Compiler->Add(CenterIndex, AFWidth), TexcoordYIndex);
 
-	return Compiler->Lerp(A.Compile(Compiler), B.Compile(Compiler), AAStep);
+	return Compiler->Lerp(B.Compile(Compiler), A.Compile(Compiler), AAStep);
 }
 
 void UMaterialExpressionSplitTopBottom::GetCaption(TArray<FString>& OutCaptions) const
diff --git a/Engine/Plugins/Interchange/Runtime/Source/Import/Private/MaterialX/MaterialXUtils/MaterialXManager.cpp b/Engine/Plugins/Interchange/Runtime/Source/Import/Private/MaterialX/MaterialXUtils/MaterialXManager.cpp
--- a/Engine/Plugins/Interchange/Runtime/Source/Import/Private/MaterialX/MaterialXUtils/MaterialXManager.cpp
+++ b/Engine/Plugins/Interchange/Runtime/Source/Import/Private/MaterialX/MaterialXUtils/MaterialXManager.cpp
@@ -34,6 +34,7 @@ FMaterialXManager::FMaterialXManager()
 		{{TEXT(""),                            TEXT("low")},       TEXT("Min")},
 		{{TEXT(""),                            TEXT("lumacoeffs")},TEXT("LuminanceFactors")}, // for the moment not yet handled by Interchange, because of the attribute being an advanced pin
 		{{TEXT(""),                            TEXT("mix")},       TEXT("Alpha")},
+		{{TEXT(""),                            TEXT("offset")},    TEXT("Offset")},
 		{{TEXT(""),                            TEXT("position")},  TEXT("Position")},
 		{{TEXT(""),                            TEXT("texcoord")},  TEXT("Coordinates")},
 		{{TEXT(""),                            TEXT("octaves")},   TEXT("Octaves")},
@@ -41,19 +42,21 @@ FMaterialXManager::FMaterialXManager()
 		{{TEXT(""),                            TEXT("outhigh")},   TEXT("TargetHigh")},
 		{{TEXT(""),                            TEXT("valuel")},    TEXT("A")},
 		{{TEXT(""),                            TEXT("valuer")},    TEXT("B")},
+		{{TEXT(""),                            TEXT("valuet")},    TEXT("A")},
+		{{TEXT(""),                            TEXT("valueb")},    TEXT("B")},
+		{{TEXT(""),                            TEXT("valuetl")},   TEXT("A")},
+		{{TEXT(""),                            TEXT("valuetr")},   TEXT("B")},
+		{{TEXT(""),                            TEXT("valuebl")},   TEXT("C")},
+		{{TEXT(""),                            TEXT("valuebr")},   TEXT("D")},
+		{{TEXT(""),                            TEXT("value1")},    TEXT("A")},
+		{{TEXT(""),                            TEXT("value2")},    TEXT("B")},
 		{{MaterialX::Category::Atan2,          TEXT("in1")},       TEXT("Y")},
 		{{MaterialX::Category::Atan2,          TEXT("in2")},       TEXT("X")},
 		{{MaterialX::Category::HeightToNormal, TEXT("scale")},     UE::Interchange::Materials::Standard::Nodes::NormalFromHeightMap::Inputs::Intensity.ToString()},
-		{{MaterialX::Category::IfGreater,      TEXT("value1")},    TEXT("A")},
-		{{MaterialX::Category::IfGreater,      TEXT("value2")},    TEXT("B")},
 		{{MaterialX::Category::IfGreater,      TEXT("in1")},       TEXT("AGreaterThanB")},
 		{{MaterialX::Category::IfGreater,      TEXT("in2")},       TEXT("ALessThanB")}, //another input is added for the case equal see ConnectIfGreater
-		{{MaterialX::Category::IfGreaterEq,    TEXT("value1")},    TEXT("A")},
-		{{MaterialX::Category::IfGreaterEq,    TEXT("value2")},    TEXT("B")},
 		{{MaterialX::Category::IfGreaterEq,    TEXT("in1")},       TEXT("AGreaterThanB")}, //another is added for the case equal see ConnectIfGreaterEq
 		{{MaterialX::Category::IfGreaterEq,    TEXT("in2")},       TEXT("ALessThanB")},
-		{{MaterialX::Category::IfEqual,        TEXT("value1")},    TEXT("A")},
-		{{MaterialX::Category::IfEqual,        TEXT("value2")},    TEXT("B")},
 		{{MaterialX::Category::IfEqual,        TEXT("in1")},       TEXT("AEqualsB")},
 		{{MaterialX::Category::IfEqual,        TEXT("in2")},       TEXT("ALessThanB")},  // another input is added for the case greater see ConnectIfEqual
 		{{MaterialX::Category::Inside,         TEXT("in")},        TEXT("A")},			  // Inside is treated as a Multiply node
@@ -69,6 +72,8 @@ FMaterialXManager::FMaterialXManager()
 		{{MaterialX::Category::Normalize,      TEXT("in")},        TEXT("VectorInput")},
 		{{MaterialX::Category::Outside,        TEXT("in")},        TEXT("A")},				// Outside is treated as Multiply node
 		{{MaterialX::Category::Outside,        TEXT("mask")},      TEXT("B")},				// Outside is treated as Multiply node
+		{{MaterialX::Category::Place2D,        TEXT("pivot")},     TEXT("Pivot")},
+		{{MaterialX::Category::Place2D,        TEXT("rotate")},    TEXT("RotationAngle")},
 		{{MaterialX::Category::Power,          TEXT("in1")},       TEXT("Base")},
 		{{MaterialX::Category::Power,	       TEXT("in2")},       TEXT("Exponent")},
 		{{MaterialX::Category::Rotate2D,       TEXT("amount")},    TEXT("RotationAngle")},
@@ -78,7 +83,7 @@ FMaterialXManager::FMaterialXManager()
 		{{MaterialX::Category::Rotate3D,       TEXT("in")},        TEXT("Position")},
 		{{MaterialX::Category::Saturate,       TEXT("amount")},    TEXT("Fraction")},
 		{{MaterialX::Category::Smoothstep,     TEXT("in")},        TEXT("Value")},
-}
+	}
 	, MatchingMaterialExpressions {
 		// Math nodes
 		{MaterialX::Category::Absval,       TEXT("Abs")},
@@ -102,7 +107,9 @@ FMaterialXManager::FMaterialXManager()
 		{MaterialX::Category::Modulo,       TEXT("Fmod")},
 		{MaterialX::Category::Multiply,     TEXT("Multiply")},
 		{MaterialX::Category::Normalize,    TEXT("Normalize")},
+		{MaterialX::Category::Place2D,		TEXT("Place2D")},
 		{MaterialX::Category::Power,        TEXT("Power")},
+		{MaterialX::Category::Rotate2D,		TEXT("Rotate2D")},
 		{MaterialX::Category::RampLR,       TEXT("RampLeftRight")},
 		{MaterialX::Category::RampTB,       TEXT("RampTopBottom")},
 		{MaterialX::Category::Sign,         TEXT("Sign")},
@@ -134,6 +141,12 @@ FMaterialXManager::FMaterialXManager()
 		{MaterialX::Category::Combine2,     TEXT("AppendVector")},
 		{MaterialX::Category::Combine3,     TEXT("Append3Vector")},
 		{MaterialX::Category::Combine4,     TEXT("Append4Vector")},
+		// Procedural2D nodes
+		{MaterialX::Category::Ramp4,        TEXT("Ramp4")},
+		{MaterialX::Category::RampLR,       TEXT("RampLeftRight")},
+		{MaterialX::Category::RampTB,       TEXT("RampTopBottom")},
+		{MaterialX::Category::SplitLR,      TEXT("SplitLeftRight")},
+		{MaterialX::Category::SplitTB,      TEXT("SplitTopBottom")},
 		// Procedural3D nodes
 		{MaterialX::Category::Fractal3D,    TEXT("Fractal3D")},
 		// Geometric nodes 
@@ -155,6 +168,7 @@ FMaterialXManager::FMaterialXManager()
 		TEXT("Base"),
 		TEXT("C"),
 		TEXT("Center"),
+		TEXT("Coordinates"),
 		TEXT("D"),
 		TEXT("Exponent"),
 		TEXT("Factor"),
@@ -167,14 +181,18 @@ FMaterialXManager::FMaterialXManager()
 		TEXT("Max"),
 		TEXT("Min"),
 		TEXT("Octaves"),
+		TEXT("Offset"),
+		TEXT("Pivot"),
 		TEXT("Position"),
+		TEXT("RotationAngle"),
+		TEXT("Scale"),
 		TEXT("TargetLow"),
 		TEXT("TargetHigh"),
 		TEXT("Value"),
 		TEXT("VectorInput"),
 		TEXT("X"),
 		TEXT("Y")
-}
+	}
 	, MaterialXContainerDelegates{
 		{mx::Category::SurfaceUnlit, FMaterialXManager::FOnGetMaterialXInstance::CreateStatic(&FMaterialXSurfaceUnlitShader::MakeInstance)},
 		{mx::Category::StandardSurface, FMaterialXManager::FOnGetMaterialXInstance::CreateStatic(&FMaterialXStandardSurfaceShader::MakeInstance)},
diff --git a/Engine/Plugins/Interchange/Runtime/Source/Import/Private/MaterialX/MaterialXUtils/MaterialXStandardSurfaceShader.cpp b/Engine/Plugins/Interchange/Runtime/Source/Import/Private/MaterialX/MaterialXUtils/MaterialXStandardSurfaceShader.cpp
--- a/Engine/Plugins/Interchange/Runtime/Source/Import/Private/MaterialX/MaterialXUtils/MaterialXStandardSurfaceShader.cpp
+++ b/Engine/Plugins/Interchange/Runtime/Source/Import/Private/MaterialX/MaterialXUtils/MaterialXStandardSurfaceShader.cpp
@@ -24,7 +24,7 @@ void FMaterialXStandardSurfaceShader::Translate(mx::NodePtr StandardSurfaceNode)
 
 	using namespace UE::Interchange::Materials;
 
-	UInterchangeShaderGraphNode* ShaderGraphNode = CreateShaderNode<UInterchangeShaderGraphNode>(SurfaceShaderNode->getName().c_str(), StandardSurface::Name.ToString());
+	ShaderGraphNode->SetCustomShaderType(StandardSurface::Name.ToString());
 
 	//Base
 	ConnectNodeOutputToInput(mx::StandardSurface::Input::Base, ShaderGraphNode, StandardSurface::Parameters::Base.ToString(), mx::StandardSurface::DefaultValue::Float::Base);
diff --git a/Engine/Plugins/Interchange/Runtime/Source/Import/Private/MaterialX/MaterialXUtils/MaterialXSurfaceMaterial.cpp b/Engine/Plugins/Interchange/Runtime/Source/Import/Private/MaterialX/MaterialXUtils/MaterialXSurfaceMaterial.cpp
--- a/Engine/Plugins/Interchange/Runtime/Source/Import/Private/MaterialX/MaterialXUtils/MaterialXSurfaceMaterial.cpp
+++ b/Engine/Plugins/Interchange/Runtime/Source/Import/Private/MaterialX/MaterialXUtils/MaterialXSurfaceMaterial.cpp
@@ -21,7 +21,22 @@ TSharedRef<FMaterialXBase> FMaterialXSurfaceMaterial::MakeInstance(UInterchangeB
 
 void FMaterialXSurfaceMaterial::Translate(MaterialX::NodePtr SurfaceMaterialNode)
 {
-	bool bHasSupportedSurfaceShader = false;
+	bool bHasSurfaceShader = false;
+	bool bIsShaderGraphInContainer = true;
+
+	 //We initialize the ShaderGraph outside of the loop, because a surfacematerial has only up to 2 inputs:
+	 //- a surfaceshader that we handle
+	 //- a displacementshader not yet supported
+	const FString ShaderGraphNodeUID = UInterchangeShaderNode::MakeNodeUid(ANSI_TO_TCHAR(SurfaceMaterialNode->getName().c_str()), FStringView{});
+	UInterchangeShaderGraphNode* ShaderGraphNode = const_cast<UInterchangeShaderGraphNode*>(Cast<UInterchangeShaderGraphNode>(NodeContainer.GetNode(ShaderGraphNodeUID)));
+
+	if(!ShaderGraphNode)
+	{
+		ShaderGraphNode = NewObject<UInterchangeShaderGraphNode>(&NodeContainer);
+		ShaderGraphNode->InitializeNode(ShaderGraphNodeUID, SurfaceMaterialNode->getName().c_str(), EInterchangeNodeContainerType::TranslatedAsset);
+		NodeContainer.SetNodeParentUid(ShaderGraphNodeUID, FString{});
+		bIsShaderGraphInContainer = false;
+	}
 
 	for(mx::InputPtr Input : SurfaceMaterialNode->getInputs())
 	{
@@ -29,16 +44,25 @@ void FMaterialXSurfaceMaterial::Translate(MaterialX::NodePtr SurfaceMaterialNode
 
 		if(ConnectedNode)
 		{
-			TSharedPtr<FMaterialXBase> ShaderTranslator = FMaterialXManager::GetInstance().GetShaderTranslator(ConnectedNode->getCategory().c_str(), NodeContainer);
-			bHasSupportedSurfaceShader = ShaderTranslator != nullptr;
-			if(bHasSupportedSurfaceShader)
+			// we use here a static_cast instead of a dynamic cast, because we know for sure that the underlying type is a surfaceshader
+			// otherwise the document is ill-formed and the Translator already handles it
+			
+			TSharedPtr<FMaterialXSurfaceShaderAbstract> ShaderTranslator = StaticCastSharedPtr<FMaterialXSurfaceShaderAbstract>(FMaterialXManager::GetInstance().GetShaderTranslator(ConnectedNode->getCategory().c_str(), NodeContainer));
+			bHasSurfaceShader = ShaderTranslator.IsValid();
+			if(bHasSurfaceShader)
 			{
+				ShaderTranslator->ShaderGraphNode = ShaderGraphNode;
 				ShaderTranslator->Translate(ConnectedNode);
 			}
 		}
 	}
 
-	if(!bHasSupportedSurfaceShader)
+	// We only add the ShaderGraph to the container if we found a supported surfaceshader
+	if(bHasSurfaceShader && !bIsShaderGraphInContainer)
+	{
+		NodeContainer.AddNode(ShaderGraphNode);
+	}
+	else
 	{
 		UE_LOG(LogInterchangeImport, Warning, TEXT("the surfaceshader of <%s> is not supported"), ANSI_TO_TCHAR(SurfaceMaterialNode->getName().c_str()));
 	}
diff --git a/Engine/Plugins/Interchange/Runtime/Source/Import/Private/MaterialX/MaterialXUtils/MaterialXSurfaceShaderAbstract.cpp b/Engine/Plugins/Interchange/Runtime/Source/Import/Private/MaterialX/MaterialXUtils/MaterialXSurfaceShaderAbstract.cpp
--- a/Engine/Plugins/Interchange/Runtime/Source/Import/Private/MaterialX/MaterialXUtils/MaterialXSurfaceShaderAbstract.cpp
+++ b/Engine/Plugins/Interchange/Runtime/Source/Import/Private/MaterialX/MaterialXUtils/MaterialXSurfaceShaderAbstract.cpp
@@ -148,7 +148,7 @@ bool FMaterialXSurfaceShaderAbstract::ConnectMatchingNodeOutputToInput(MaterialX
 	{
 		UInterchangeShaderNode* OperatorNode = nullptr;
 
-		OperatorNode = CreateShaderNode<UInterchangeShaderNode>(Node->getName().c_str(), *ShaderType);
+		OperatorNode = CreateShaderNode(Node->getName().c_str(), *ShaderType);
 
 		for(mx::InputPtr Input : Node->getInputs())
 		{
@@ -244,7 +244,7 @@ void FMaterialXSurfaceShaderAbstract::ConnectExtractInputToOutput(MaterialX::Nod
 {
 	using namespace UE::Interchange::Materials::Standard::Nodes;
 
-	UInterchangeShaderNode* MaskShaderNode = CreateShaderNode<UInterchangeShaderNode>(UpstreamNode->getName().c_str(), Mask::Name.ToString());
+	UInterchangeShaderNode* MaskShaderNode = CreateShaderNode(UpstreamNode->getName().c_str(), Mask::Name.ToString());
 
 	if(mx::InputPtr InputIndex = UpstreamNode->getInput("index"))
 	{
@@ -275,71 +275,21 @@ void FMaterialXSurfaceShaderAbstract::ConnectDotInputToOutput(MaterialX::NodePtr
 
 void FMaterialXSurfaceShaderAbstract::ConnectTransformPositionInputToOutput(MaterialX::NodePtr UpstreamNode, UInterchangeShaderNode* ParentShaderNode, const FString& InputChannelName)
 {
-	UInterchangeShaderNode* TransformNode = CreateShaderNode<UInterchangeShaderNode>(UpstreamNode->getName().c_str(), TEXT("TransformPosition"));
+	UInterchangeShaderNode* TransformNode = CreateShaderNode(UpstreamNode->getName().c_str(), TEXT("TransformPosition"));
 	AddAttributeFromValueOrInterface(UpstreamNode->getInput("in"), TEXT("Input"), TransformNode);
 	UInterchangeShaderPortsAPI::ConnectDefaultOuputToInput(ParentShaderNode, InputChannelName, TransformNode->GetUniqueID());
 }
 
 void FMaterialXSurfaceShaderAbstract::ConnectTransformVectorInputToOutput(MaterialX::NodePtr UpstreamNode, UInterchangeShaderNode* ParentShaderNode, const FString& InputChannelName)
 {
-	UInterchangeShaderNode* TransformNode = CreateShaderNode<UInterchangeShaderNode>(UpstreamNode->getName().c_str(), TEXT("Transform"));
+	UInterchangeShaderNode* TransformNode = CreateShaderNode(UpstreamNode->getName().c_str(), TEXT("Transform"));
 	AddAttributeFromValueOrInterface(UpstreamNode->getInput("in"), TEXT("Input"), TransformNode);
 	UInterchangeShaderPortsAPI::ConnectDefaultOuputToInput(ParentShaderNode, InputChannelName, TransformNode->GetUniqueID());
 }
 
-void FMaterialXSurfaceShaderAbstract::ConnectRotate2DInputToOutput(mx::NodePtr UpstreamNode, UInterchangeShaderNode* ParentShaderNode, const FString& InputChannelName)
-{
-	UInterchangeShaderNode* Rotate2DNode = CreateShaderNode<UInterchangeShaderNode>(UpstreamNode->getName().c_str(), TEXT("RotateAboutAxis"));
-	Rotate2DNode->AddLinearColorAttribute(UInterchangeShaderPortsAPI::MakeInputValueKey(TEXT("NormalizedRotationAxis")), FLinearColor(0, 0, 1));
-	Rotate2DNode->AddLinearColorAttribute(UInterchangeShaderPortsAPI::MakeInputValueKey(TEXT("PivotPoint")), FLinearColor(0.5, 0.5, 0));
-
-	{
-		mx::InputPtr Input = UpstreamNode->getInput("in");
-		AddAttributeFromValueOrInterface(UpstreamNode->getInput("in"), TEXT("Position"), Rotate2DNode);
-
-		// Rotate2D "in" input takes a vector2 input, but RotateAboutAxis takes a vector3 so we convert
-		if(Input)
-		{
-			mx::NodePtr ConvertNode = CreateNode(UpstreamNode->getParent()->asA<mx::NodeGraph>(),
-												 UpstreamNode->getName().c_str(),
-												 mx::Category::Convert,
-												 { {"in", Input} },
-												 {});
-			ConvertNode->setType(mx::Type::Vector3);
-			Input->setNodeName(ConvertNode->getName());
-		}
-	}
-
-	{
-
-		mx::InputPtr Input = UpstreamNode->getInput("amount");
-		AddAttributeFromValueOrInterface(Input, TEXT("RotationAngle"), Rotate2DNode);
-
-		// we create a Divide node to convert MaterialX angle in degrees to UE's angle which is a value between [0,1]
-		if(Input)
-		{
-			//copy "in" input into "in1" and create "amount" input as "in2" with the value 360 because UE's angle is a value between [0,1]
-			mx::NodePtr DivideNode = CreateNode(UpstreamNode->getParent()->asA<mx::NodeGraph>(),
-												UpstreamNode->getName().c_str(),
-												mx::Category::Divide,
-												{ {"in1", Input} }, // rename input to match <divide> input "in1"
-												{ {"in2", FAttributeValueArray{{"type", "float"}, {"value", "360"}}} });
-
-			// Input now points to the new node
-			Input->setNodeName(DivideNode->getName());
-		}
-	}
-
-	//We insert a mask node because the output has to be a vector2
-	UInterchangeShaderNode* MaskNode = CreateMaskShaderNode(0b1100, UpstreamNode->getName().c_str() + FString(TEXT("_Mask")));
-	UInterchangeShaderPortsAPI::ConnectDefaultOuputToInput(MaskNode, TEXT("Input"), Rotate2DNode->GetUniqueID());
-
-	UInterchangeShaderPortsAPI::ConnectDefaultOuputToInput(ParentShaderNode, InputChannelName, MaskNode->GetUniqueID());
-}
-
 void FMaterialXSurfaceShaderAbstract::ConnectRotate3DInputToOutput(mx::NodePtr UpstreamNode, UInterchangeShaderNode* ParentShaderNode, const FString& InputChannelName)
 {
-	UInterchangeShaderNode* Rotate3DNode = CreateShaderNode<UInterchangeShaderNode>(UpstreamNode->getName().c_str(), TEXT("RotateAboutAxis"));
+	UInterchangeShaderNode* Rotate3DNode = CreateShaderNode(UpstreamNode->getName().c_str(), TEXT("RotateAboutAxis"));
 	Rotate3DNode->AddLinearColorAttribute(UInterchangeShaderPortsAPI::MakeInputValueKey(TEXT("PivotPoint")), FLinearColor(0.5, 0.5, 0));
 
 	AddAttributeFromValueOrInterface(UpstreamNode->getInput("in"), TEXT("Position"), Rotate3DNode);
@@ -383,7 +333,7 @@ void FMaterialXSurfaceShaderAbstract::ConnectImageInputToOutput(MaterialX::NodeP
 			OutputChannel = TEXT("R");
 		}
 
-		UInterchangeShaderNode* TextureShaderNode = CreateShaderNode<UInterchangeShaderNode>(UpstreamNode->getName().c_str(), TextureSample::Name.ToString());
+		UInterchangeShaderNode* TextureShaderNode = CreateShaderNode(UpstreamNode->getName().c_str(), TextureSample::Name.ToString());
 		TextureShaderNode->AddStringAttribute(UInterchangeShaderPortsAPI::MakeInputValueKey(TextureSample::Inputs::Texture.ToString()), TextureNode->GetUniqueID());
 		UInterchangeShaderPortsAPI::ConnectOuputToInput(ParentShaderNode, InputChannelName, TextureShaderNode->GetUniqueID(), OutputChannel);
 	}
@@ -529,7 +479,7 @@ void FMaterialXSurfaceShaderAbstract::ConnectConvertInputToOutput(MaterialX::Nod
 
 void FMaterialXSurfaceShaderAbstract::ConnectIfGreaterInputToOutput(MaterialX::NodePtr UpstreamNode, UInterchangeShaderNode* ParentShaderNode, const FString& InputChannelName)
 {
-	UInterchangeShaderNode* NodeIf = CreateShaderNode<UInterchangeShaderNode>(UpstreamNode->getName().c_str(), TEXT("If"));
+	UInterchangeShaderNode* NodeIf = CreateShaderNode(UpstreamNode->getName().c_str(), TEXT("If"));
 	UInterchangeShaderPortsAPI::ConnectDefaultOuputToInput(ParentShaderNode, InputChannelName, NodeIf->GetUniqueID());
 
 	if(mx::InputPtr Input = UpstreamNode->getInput("value1"))
@@ -566,7 +516,7 @@ void FMaterialXSurfaceShaderAbstract::ConnectIfGreaterInputToOutput(MaterialX::N
 
 void FMaterialXSurfaceShaderAbstract::ConnectIfGreaterEqInputToOutput(MaterialX::NodePtr UpstreamNode, UInterchangeShaderNode* ParentShaderNode, const FString& InputChannelName)
 {
-	UInterchangeShaderNode* NodeIf = CreateShaderNode<UInterchangeShaderNode>(UpstreamNode->getName().c_str(), TEXT("If"));
+	UInterchangeShaderNode* NodeIf = CreateShaderNode(UpstreamNode->getName().c_str(), TEXT("If"));
 	UInterchangeShaderPortsAPI::ConnectDefaultOuputToInput(ParentShaderNode, InputChannelName, NodeIf->GetUniqueID());
 
 	if(mx::InputPtr Input = UpstreamNode->getInput("value1"))
@@ -603,7 +553,7 @@ void FMaterialXSurfaceShaderAbstract::ConnectIfGreaterEqInputToOutput(MaterialX:
 
 void FMaterialXSurfaceShaderAbstract::ConnectIfEqualInputToOutput(MaterialX::NodePtr UpstreamNode, UInterchangeShaderNode* ParentShaderNode, const FString& InputChannelName)
 {
-	UInterchangeShaderNode* NodeIf = CreateShaderNode<UInterchangeShaderNode>(UpstreamNode->getName().c_str(), TEXT("If"));
+	UInterchangeShaderNode* NodeIf = CreateShaderNode(UpstreamNode->getName().c_str(), TEXT("If"));
 	UInterchangeShaderPortsAPI::ConnectDefaultOuputToInput(ParentShaderNode, InputChannelName, NodeIf->GetUniqueID());
 
 	if(mx::InputPtr Input = UpstreamNode->getInput("value1"); Input && Input->hasValue())
@@ -641,11 +591,11 @@ void FMaterialXSurfaceShaderAbstract::ConnectIfEqualInputToOutput(MaterialX::Nod
 void FMaterialXSurfaceShaderAbstract::ConnectOutsideInputToOutput(MaterialX::NodePtr UpstreamNode, UInterchangeShaderNode* ParentShaderNode, const FString& InputChannelName)
 {
 	//in * (1 - mask)
-	UInterchangeShaderNode* NodeMultiply = CreateShaderNode<UInterchangeShaderNode>(UpstreamNode->getName().c_str(), TEXT("Multiply"));
+	UInterchangeShaderNode* NodeMultiply = CreateShaderNode(UpstreamNode->getName().c_str(), TEXT("Multiply"));
 	AddAttributeFromValueOrInterface(UpstreamNode->getInput("in"), TEXT("A"), NodeMultiply);
 	UInterchangeShaderPortsAPI::ConnectDefaultOuputToInput(ParentShaderNode, InputChannelName, NodeMultiply->GetUniqueID());
 
-	UInterchangeShaderNode* NodeOneMinus = CreateShaderNode<UInterchangeShaderNode>(UpstreamNode->getName().c_str() + FString(TEXT("_OneMinus")), TEXT("OneMinus"));
+	UInterchangeShaderNode* NodeOneMinus = CreateShaderNode(UpstreamNode->getName().c_str() + FString(TEXT("_OneMinus")), TEXT("OneMinus"));
 	AddAttributeFromValueOrInterface(UpstreamNode->getInput("mask"), TEXT("Input"), NodeOneMinus);
 	UInterchangeShaderPortsAPI::ConnectDefaultOuputToInput(NodeMultiply, TEXT("B"), NodeOneMinus->GetUniqueID());
 }
@@ -659,11 +609,11 @@ void FMaterialXSurfaceShaderAbstract::ConnectPositionInputToOutput(MaterialX::No
 
 	// For the moment we don't have the distinction between model/object, so let's just create an UMaterialExpressionWorldPosition
 	// In case of model/object we need to add a TransformPoint from world to local space
-	UInterchangeShaderNode* PositionNode = CreateShaderNode<UInterchangeShaderNode>(UpstreamNode->getName().c_str() + FString(TEXT("_Position")), TEXT("WorldPosition"));
+	UInterchangeShaderNode* PositionNode = CreateShaderNode(UpstreamNode->getName().c_str() + FString(TEXT("_Position")), TEXT("WorldPosition"));
 
 	// In case of the position node, it seems that the unit is different, we assume for now a conversion from mm -> m, even if UE by default is cm
 	// See standard_surface_marble_solid file, especially on the fractal3d node
-	UInterchangeShaderNode* UnitNode = CreateShaderNode<UInterchangeShaderNode>(UpstreamNode->getName().c_str(), TEXT("Multiply"));
+	UInterchangeShaderNode* UnitNode = CreateShaderNode(UpstreamNode->getName().c_str(), TEXT("Multiply"));
 	UnitNode->AddFloatAttribute(UInterchangeShaderPortsAPI::MakeInputValueKey(TEXT("B")), 0.001f);
 	UInterchangeShaderPortsAPI::ConnectDefaultOuputToInput(UnitNode, TEXT("A"), PositionNode->GetUniqueID());
 
@@ -683,7 +633,7 @@ void FMaterialXSurfaceShaderAbstract::ConnectPositionInputToOutput(MaterialX::No
 	{
 		using namespace UE::Interchange::Materials::Standard::Nodes;
 
-		UInterchangeShaderNode* TransformNode = CreateShaderNode<UInterchangeShaderNode>(UpstreamNode->getName().c_str() + FString(TEXT("_Transform")), TransformPosition::Name.ToString());
+		UInterchangeShaderNode* TransformNode = CreateShaderNode(UpstreamNode->getName().c_str() + FString(TEXT("_Transform")), TransformPosition::Name.ToString());
 		NodeToConnectTo = TransformNode;
 		InputToConnectTo = TransformPosition::Inputs::Input.ToString();
 		TransformNode->AddInt32Attribute(TransformPosition::Attributes::TransformSourceType, EMaterialPositionTransformSource::TRANSFORMPOSSOURCE_World);
@@ -703,7 +653,7 @@ void FMaterialXSurfaceShaderAbstract::ConnectNormalInputToOutput(MaterialX::Node
 
 	// For the moment we don't have the distinction between model/object, so let's just create an UMaterialExpressionVertexNormalWS
 	// In case of model/object we need to add a TransformVector from world to local space
-	UInterchangeShaderNode* NormalNode = CreateShaderNode<UInterchangeShaderNode>(UpstreamNode->getName().c_str(), TEXT("VertexNormalWS"));
+	UInterchangeShaderNode* NormalNode = CreateShaderNode(UpstreamNode->getName().c_str(), TEXT("VertexNormalWS"));
 	UInterchangeShaderNode* NodeToConnectTo = ParentShaderNode;
 	FString InputToConnectTo = InputChannelName;
 
@@ -720,7 +670,7 @@ void FMaterialXSurfaceShaderAbstract::ConnectNormalInputToOutput(MaterialX::Node
 	{
 		using namespace UE::Interchange::Materials::Standard::Nodes;
 
-		UInterchangeShaderNode* TransformNode = CreateShaderNode<UInterchangeShaderNode>(UpstreamNode->getName().c_str() + FString(TEXT("_Transform")), TransformVector::Name.ToString());
+		UInterchangeShaderNode* TransformNode = CreateShaderNode(UpstreamNode->getName().c_str() + FString(TEXT("_Transform")), TransformVector::Name.ToString());
 		NodeToConnectTo = TransformNode;
 		InputToConnectTo = TransformVector::Inputs::Input.ToString();
 		TransformNode->AddInt32Attribute(TransformVector::Attributes::TransformSourceType, EMaterialVectorCoordTransformSource::TRANSFORMSOURCE_World);
@@ -740,7 +690,7 @@ void FMaterialXSurfaceShaderAbstract::ConnectTangentInputToOutput(MaterialX::Nod
 
 	// For the moment we don't have the distinction between model/object, so let's just create an UMaterialExpressionVertexTangentWS
 	// In case of model/object we need to add a TransformVector from world to local space
-	UInterchangeShaderNode* TangentNode = CreateShaderNode<UInterchangeShaderNode>(UpstreamNode->getName().c_str(), TEXT("VertexTangentWS"));
+	UInterchangeShaderNode* TangentNode = CreateShaderNode(UpstreamNode->getName().c_str(), TEXT("VertexTangentWS"));
 	UInterchangeShaderNode* NodeToConnectTo = ParentShaderNode;
 	FString InputToConnectTo = InputChannelName;
 
@@ -757,7 +707,7 @@ void FMaterialXSurfaceShaderAbstract::ConnectTangentInputToOutput(MaterialX::Nod
 	{
 		using namespace UE::Interchange::Materials::Standard::Nodes;
 
-		UInterchangeShaderNode* TransformNode = CreateShaderNode<UInterchangeShaderNode>(UpstreamNode->getName().c_str() + FString(TEXT("_Transform")), TransformVector::Name.ToString());
+		UInterchangeShaderNode* TransformNode = CreateShaderNode(UpstreamNode->getName().c_str() + FString(TEXT("_Transform")), TransformVector::Name.ToString());
 		NodeToConnectTo = TransformNode;
 		InputToConnectTo = TransformVector::Inputs::Input.ToString();
 		TransformNode->AddInt32Attribute(TransformVector::Attributes::TransformSourceType, EMaterialVectorCoordTransformSource::TRANSFORMSOURCE_World);
@@ -777,9 +727,9 @@ void FMaterialXSurfaceShaderAbstract::ConnectBitangentInputToOutput(MaterialX::N
 
 	// For the moment we don't have the distinction between model/object, so let's just do the cross product between the normal and the tangent
 	// In case of model/object we need to add a TransformVector from world to local space
-	UInterchangeShaderNode* NormalNode = CreateShaderNode<UInterchangeShaderNode>(UpstreamNode->getName().c_str() + FString(TEXT("_Normal")), TEXT("VertexNormalWS"));
-	UInterchangeShaderNode* TangentNode = CreateShaderNode<UInterchangeShaderNode>(UpstreamNode->getName().c_str() + FString(TEXT("_Tangent")), TEXT("VertexTangentWS"));
-	UInterchangeShaderNode* BitangentNode = CreateShaderNode<UInterchangeShaderNode>(UpstreamNode->getName().c_str(), TEXT("CrossProduct"));
+	UInterchangeShaderNode* NormalNode = CreateShaderNode(UpstreamNode->getName().c_str() + FString(TEXT("_Normal")), TEXT("VertexNormalWS"));
+	UInterchangeShaderNode* TangentNode = CreateShaderNode(UpstreamNode->getName().c_str() + FString(TEXT("_Tangent")), TEXT("VertexTangentWS"));
+	UInterchangeShaderNode* BitangentNode = CreateShaderNode(UpstreamNode->getName().c_str(), TEXT("CrossProduct"));
 
 	UInterchangeShaderPortsAPI::ConnectDefaultOuputToInput(BitangentNode, TEXT("A"), NormalNode->GetUniqueID());
 	UInterchangeShaderPortsAPI::ConnectDefaultOuputToInput(BitangentNode, TEXT("B"), TangentNode->GetUniqueID());
@@ -801,7 +751,7 @@ void FMaterialXSurfaceShaderAbstract::ConnectBitangentInputToOutput(MaterialX::N
 	{
 		using namespace UE::Interchange::Materials::Standard::Nodes;
 
-		UInterchangeShaderNode* TransformNode = CreateShaderNode<UInterchangeShaderNode>(UpstreamNode->getName().c_str() + FString(TEXT("_Transform")), TransformVector::Name.ToString());
+		UInterchangeShaderNode* TransformNode = CreateShaderNode(UpstreamNode->getName().c_str() + FString(TEXT("_Transform")), TransformVector::Name.ToString());
 		NodeToConnectTo = TransformNode;
 		InputToConnectTo = TransformVector::Inputs::Input.ToString();
 		TransformNode->AddInt32Attribute(TransformVector::Attributes::TransformSourceType, EMaterialVectorCoordTransformSource::TRANSFORMSOURCE_World);
@@ -815,7 +765,7 @@ void FMaterialXSurfaceShaderAbstract::ConnectBitangentInputToOutput(MaterialX::N
 void FMaterialXSurfaceShaderAbstract::ConnectTimeInputToOutput(MaterialX::NodePtr UpstreamNode, UInterchangeShaderNode* ParentShaderNode, const FString& InputChannelName)
 {
 	using namespace UE::Interchange::Materials::Standard::Nodes;
-	UInterchangeShaderNode* TimeNode = CreateShaderNode<UInterchangeShaderNode>(UpstreamNode->getName().c_str(), TEXT("Time"));
+	UInterchangeShaderNode* TimeNode = CreateShaderNode(UpstreamNode->getName().c_str(), TEXT("Time"));
 	TimeNode->AddBooleanAttribute(Time::Attributes::OverridePeriod, true);
 
 	float FPS;
@@ -840,7 +790,7 @@ void FMaterialXSurfaceShaderAbstract::ConnectNoise3DInputToOutput(MaterialX::Nod
 	using namespace UE::Interchange::Materials::Standard::Nodes;
 
 	// MaterialX defines the Noise3d as Perlin Noise which is multiplied by the Amplitude then Added to Pivot
-	UInterchangeShaderNode* NoiseNode = CreateShaderNode<UInterchangeShaderNode>(UpstreamNode->getName().c_str(), Noise::Name.ToString());
+	UInterchangeShaderNode* NoiseNode = CreateShaderNode(UpstreamNode->getName().c_str(), Noise::Name.ToString());
 	NoiseNode->AddInt32Attribute(Noise::Attributes::Function, ENoiseFunction::NOISEFUNCTION_GradientTex);
 	NoiseNode->AddBooleanAttribute(Noise::Attributes::Turbulence, false);
 	NoiseNode->AddFloatAttribute(Noise::Attributes::OutputMin, 0);
@@ -856,7 +806,7 @@ void FMaterialXSurfaceShaderAbstract::ConnectNoise3DInputToOutput(MaterialX::Nod
 		}
 
 		const FString ShaderNodeName = UpstreamNode->getName().c_str() + FString{ TEXT("_") } + ShaderType;
-		UInterchangeShaderNode* ShaderNode = CreateShaderNode<UInterchangeShaderNode>(ShaderNodeName, ShaderType);
+		UInterchangeShaderNode* ShaderNode = CreateShaderNode(ShaderNodeName, ShaderType);
 
 		UInterchangeShaderPortsAPI::ConnectDefaultOuputToInput(ShaderNode, TEXT("A"), NodeToConnectTo->GetUniqueID());
 
@@ -888,7 +838,7 @@ void FMaterialXSurfaceShaderAbstract::ConnectCellNoise3DInputToOutput(MaterialX:
 	using namespace UE::Interchange::Materials::Standard::Nodes;
 
 	// Let's use a vector noise for this one, the only one that is close to MaterialX implementation
-	UInterchangeShaderNode* NoiseNode = CreateShaderNode<UInterchangeShaderNode>(UpstreamNode->getName().c_str(), VectorNoise::Name.ToString());
+	UInterchangeShaderNode* NoiseNode = CreateShaderNode(UpstreamNode->getName().c_str(), VectorNoise::Name.ToString());
 	NoiseNode->AddInt32Attribute(VectorNoise::Attributes::Function, EVectorNoiseFunction::VNF_CellnoiseALU);
 
 	UInterchangeShaderPortsAPI::ConnectDefaultOuputToInput(ParentShaderNode, InputChannelName, NoiseNode->GetUniqueID());
@@ -899,7 +849,7 @@ void FMaterialXSurfaceShaderAbstract::ConnectWorleyNoise3DInputToOutput(Material
 	using namespace UE::Interchange::Materials::Standard::Nodes;
 
 	//Also called Voronoi, the implementation is a bit different in UE, especially we don't have access to the jitter
-	UInterchangeShaderNode* NoiseNode = CreateShaderNode<UInterchangeShaderNode>(UpstreamNode->getName().c_str(), Noise::Name.ToString());
+	UInterchangeShaderNode* NoiseNode = CreateShaderNode(UpstreamNode->getName().c_str(), Noise::Name.ToString());
 	NoiseNode->AddInt32Attribute(Noise::Attributes::Function, ENoiseFunction::NOISEFUNCTION_VoronoiALU);
 
 	UInterchangeShaderPortsAPI::ConnectDefaultOuputToInput(ParentShaderNode, InputChannelName, NoiseNode->GetUniqueID());
@@ -929,11 +879,11 @@ void FMaterialXSurfaceShaderAbstract::ConnectHeightToNormalInputToOutput(Materia
 
 			if(UInterchangeTextureNode* TextureNode = CreateTextureNode<UInterchangeTexture2DNode>(UpstreamNode))
 			{
-				UInterchangeShaderNode* HeightMapNode = CreateShaderNode<UInterchangeShaderNode>(UpstreamNode->getName().c_str(), NormalFromHeightMap::Name.ToString());
+				UInterchangeShaderNode* HeightMapNode = CreateShaderNode(UpstreamNode->getName().c_str(), NormalFromHeightMap::Name.ToString());
 				UInterchangeShaderPortsAPI::ConnectDefaultOuputToInput(ParentShaderNode, InputChannelName, HeightMapNode->GetUniqueID());
 
 				const FString TextureNodeName = UpstreamNode->getName().c_str() + FString{ "_texture" };
-				UInterchangeShaderNode* TextureShaderNode = CreateShaderNode<UInterchangeShaderNode>(TextureNodeName, TextureObject::Name.ToString());
+				UInterchangeShaderNode* TextureShaderNode = CreateShaderNode(TextureNodeName, TextureObject::Name.ToString());
 				TextureShaderNode->AddStringAttribute(UInterchangeShaderPortsAPI::MakeInputValueKey(TextureObject::Inputs::Texture.ToString()), TextureNode->GetUniqueID());
 				UInterchangeShaderPortsAPI::ConnectDefaultOuputToInput(HeightMapNode, NormalFromHeightMap::Inputs::HeightMap.ToString(), TextureShaderNode->GetUniqueID());
 
@@ -989,7 +939,7 @@ void FMaterialXSurfaceShaderAbstract::ConnectBlurInputToOutput(MaterialX::NodePt
 					OutputChannel = TEXT("R");
 				}
 
-				UInterchangeShaderNode* TextureShaderNode = CreateShaderNode<UInterchangeShaderNode>(UpstreamNode->getName().c_str(), TextureSampleBlur::Name.ToString());
+				UInterchangeShaderNode* TextureShaderNode = CreateShaderNode(UpstreamNode->getName().c_str(), TextureSampleBlur::Name.ToString());
 				TextureShaderNode->AddStringAttribute(UInterchangeShaderPortsAPI::MakeInputValueKey(TextureSampleBlur::Inputs::Texture.ToString()), TextureNode->GetUniqueID());
 				UInterchangeShaderPortsAPI::ConnectOuputToInput(ParentShaderNode, InputChannelName, TextureShaderNode->GetUniqueID(), OutputChannel);
 
@@ -1029,7 +979,7 @@ UInterchangeShaderNode* FMaterialXSurfaceShaderAbstract::CreateMaskShaderNode(ui
 	bool bB = (0b0010 & RGBA) >> 1;
 	bool bA = (0b0001 & RGBA) >> 0;
 	using namespace UE::Interchange::Materials::Standard::Nodes;
-	UInterchangeShaderNode* MaskShaderNode = CreateShaderNode<UInterchangeShaderNode>(NodeName, Mask::Name.ToString());
+	UInterchangeShaderNode* MaskShaderNode = CreateShaderNode(NodeName, Mask::Name.ToString());
 	MaskShaderNode->AddBooleanAttribute(Mask::Attributes::R, bR);
 	MaskShaderNode->AddBooleanAttribute(Mask::Attributes::G, bG);
 	MaskShaderNode->AddBooleanAttribute(Mask::Attributes::B, bB);
@@ -1038,6 +988,27 @@ UInterchangeShaderNode* FMaterialXSurfaceShaderAbstract::CreateMaskShaderNode(ui
 	return MaskShaderNode;
 }
 
+UInterchangeShaderNode* FMaterialXSurfaceShaderAbstract::CreateShaderNode(const FString& NodeName, const FString& ShaderType)
+{
+	UInterchangeShaderNode* Node;
+
+	const FString NodeUID = UInterchangeShaderNode::MakeNodeUid(NodeName, FStringView{});
+
+	//Test directly in the NodeContainer, because the ShaderNodes can be altered during the node graph either by the parent (dot/normalmap),
+	//or by putting an intermediary node between the child and the parent (tiledimage)
+	if(Node = const_cast<UInterchangeShaderNode*>(Cast<UInterchangeShaderNode>(NodeContainer.GetNode(NodeUID))); !Node)
+	{
+		Node = NewObject<UInterchangeShaderNode>(&NodeContainer);
+		Node->InitializeNode(NodeUID, NodeName, EInterchangeNodeContainerType::TranslatedAsset);
+		NodeContainer.AddNode(Node);
+		Node->SetCustomShaderType(ShaderType);
+
+		ShaderNodes.Add(NodeName, Node);
+	}
+
+	return Node;
+}
+
 const FString& FMaterialXSurfaceShaderAbstract::GetMatchedInputName(MaterialX::NodePtr Node, MaterialX::InputPtr Input) const
 {
 	static FString EmptyString;
@@ -1149,7 +1120,6 @@ void FMaterialXSurfaceShaderAbstract::RegisterConnectNodeOutputToInputDelegates(
 	MatchingConnectNodeDelegates.Add(mx::Category::TransformPoint, FOnConnectNodeOutputToInput::CreateSP(this, &FMaterialXSurfaceShaderAbstract::ConnectTransformPositionInputToOutput));
 	MatchingConnectNodeDelegates.Add(mx::Category::TransformVector, FOnConnectNodeOutputToInput::CreateSP(this, &FMaterialXSurfaceShaderAbstract::ConnectTransformVectorInputToOutput));
 	MatchingConnectNodeDelegates.Add(mx::Category::TransformNormal, FOnConnectNodeOutputToInput::CreateSP(this, &FMaterialXSurfaceShaderAbstract::ConnectTransformVectorInputToOutput));
-	MatchingConnectNodeDelegates.Add(mx::Category::Rotate2D, FOnConnectNodeOutputToInput::CreateSP(this, &FMaterialXSurfaceShaderAbstract::ConnectRotate2DInputToOutput));
 	MatchingConnectNodeDelegates.Add(mx::Category::Rotate3D, FOnConnectNodeOutputToInput::CreateSP(this, &FMaterialXSurfaceShaderAbstract::ConnectRotate3DInputToOutput));
 	MatchingConnectNodeDelegates.Add(mx::Category::Image, FOnConnectNodeOutputToInput::CreateSP(this, &FMaterialXSurfaceShaderAbstract::ConnectImageInputToOutput));
 	MatchingConnectNodeDelegates.Add(mx::Category::Convert, FOnConnectNodeOutputToInput::CreateSP(this, &FMaterialXSurfaceShaderAbstract::ConnectConvertInputToOutput));
diff --git a/Engine/Plugins/Interchange/Runtime/Source/Import/Private/MaterialX/MaterialXUtils/MaterialXSurfaceUnlitShader.cpp b/Engine/Plugins/Interchange/Runtime/Source/Import/Private/MaterialX/MaterialXUtils/MaterialXSurfaceUnlitShader.cpp
--- a/Engine/Plugins/Interchange/Runtime/Source/Import/Private/MaterialX/MaterialXUtils/MaterialXSurfaceUnlitShader.cpp
+++ b/Engine/Plugins/Interchange/Runtime/Source/Import/Private/MaterialX/MaterialXUtils/MaterialXSurfaceUnlitShader.cpp
@@ -24,7 +24,7 @@ void FMaterialXSurfaceUnlitShader::Translate(mx::NodePtr SurfaceUnlitNode)
 
 	using namespace UE::Interchange::Materials;
 
-	UInterchangeShaderGraphNode* ShaderGraphNode = CreateShaderNode<UInterchangeShaderGraphNode>(SurfaceShaderNode->getName().c_str(), SurfaceUnlit::Name.ToString());
+	ShaderGraphNode->SetCustomShaderType(SurfaceUnlit::Name.ToString());
 
 	//Emission
 	ConnectNodeOutputToInput(mx::SurfaceUnlit::Input::Emission, ShaderGraphNode, SurfaceUnlit::Parameters::Emission.ToString(), mx::SurfaceUnlit::DefaultValue::Float::Emission);
diff --git a/Engine/Plugins/Interchange/Runtime/Source/Import/Public/MaterialX/MaterialXUtils/MaterialXSurfaceShaderAbstract.h b/Engine/Plugins/Interchange/Runtime/Source/Import/Public/MaterialX/MaterialXUtils/MaterialXSurfaceShaderAbstract.h
--- a/Engine/Plugins/Interchange/Runtime/Source/Import/Public/MaterialX/MaterialXUtils/MaterialXSurfaceShaderAbstract.h
+++ b/Engine/Plugins/Interchange/Runtime/Source/Import/Public/MaterialX/MaterialXUtils/MaterialXSurfaceShaderAbstract.h
@@ -5,13 +5,15 @@
 #if WITH_EDITOR
 #include "InterchangeShaderGraphNode.h"
 #include "InterchangeTranslatorBase.h"
-#include "InterchangeTextureNode.h"
+#include "InterchangeTexture2DNode.h"
 #include "MaterialXBase.h"
 
 class FMaterialXSurfaceShaderAbstract : public FMaterialXBase
 {
 protected:
 
+	friend class FMaterialXSurfaceMaterial;
+
 	DECLARE_DELEGATE_ThreeParams(FOnConnectNodeOutputToInput, MaterialX::NodePtr, UInterchangeShaderNode*, const FString&);
 
 	FMaterialXSurfaceShaderAbstract(UInterchangeBaseNodeContainer & BaseNodeContainer);
@@ -71,26 +73,31 @@ protected:
 	 * @param DefaultValue - The default value of the MaterialX input
 	 */
 	template<typename T>
-	void ConnectNodeOutputToInput(const char* InputName, UInterchangeShaderNode* ShaderNode, const FString& InputShaderName, T DefaultValue)
+	bool ConnectNodeOutputToInput(const char* InputName, UInterchangeShaderNode* ShaderNode, const FString& InputShaderName, T DefaultValue)
 	{
 		MaterialX::DocumentPtr Document = SurfaceShaderNode->getDocument();
 		MaterialX::InputPtr Input = GetInput(SurfaceShaderNode, InputName);
 
-		if(!ConnectNodeGraphOutputToInput(Input, ShaderNode, InputShaderName))
+		bool bIsConnected = ConnectNodeGraphOutputToInput(Input, ShaderNode, InputShaderName);
+
+		if(!bIsConnected)
 		{
-			if(!ConnectNodeNameOutputToInput(Input, ShaderNode, InputShaderName))
+			bIsConnected = ConnectNodeNameOutputToInput(Input, ShaderNode, InputShaderName);
+			if(!bIsConnected)
 			{
-				//the different inputs from a standard surface are either float, linear color, or no attribute at all
+				// only handle float and linear color here, for other types, the child should handle them as it is most likely not an input but a parameter to set in Interchange
 				if constexpr(std::is_same_v<decltype(DefaultValue), float>)
 				{
-					AddFloatAttribute(Input, InputShaderName, ShaderNode, DefaultValue);
+					bIsConnected = AddFloatAttribute(Input, InputShaderName, ShaderNode, DefaultValue);
 				}
 				else if constexpr(std::is_same_v<decltype(DefaultValue), FLinearColor>)
 				{
-					AddLinearColorAttribute(Input, InputShaderName, ShaderNode, DefaultValue);
+					bIsConnected = AddLinearColorAttribute(Input, InputShaderName, ShaderNode, DefaultValue);
 				}
 			}
 		}
+
+		return bIsConnected;
 	};
 
 	/**
@@ -155,9 +162,6 @@ protected:
 	/** <transformvector> */
 	void ConnectTransformVectorInputToOutput(MaterialX::NodePtr UpstreamNode, UInterchangeShaderNode* ParentShaderNode, const FString& InputChannelName);
 
-	/** <rotate2d> */
-	void ConnectRotate2DInputToOutput(MaterialX::NodePtr UpstreamNode, UInterchangeShaderNode* ParentShaderNode, const FString& InputChannelName);
-
 	/** <rotate3d> */
 	void ConnectRotate3DInputToOutput(MaterialX::NodePtr UpstreamNode, UInterchangeShaderNode* ParentShaderNode, const FString& InputChannelName);
 
@@ -222,40 +226,14 @@ protected:
 	UInterchangeShaderNode* CreateMaskShaderNode(uint8 RGBA, const FString& NodeName);
 
 	/**
-	 * Helper template function to create an InterchangeShaderNode
+	 * Helper function to create an InterchangeShaderNode
 	 *
 	 * @param NodeName - The name of the shader node
 	 * @param ShaderType - The shader node's type we want to create
 	 *
 	 * @return The shader node that was created
 	 */
-	template<typename ShaderNodeType>
-	ShaderNodeType* CreateShaderNode(const FString& NodeName, const FString& ShaderType)
-	{
-		static_assert(std::is_convertible_v<ShaderNodeType*, UInterchangeShaderNode*>, "CreateShaderNode only accepts type that derived from UInterchangeShaderNode");
-
-		ShaderNodeType* Node;
-
-		const FString NodeUID = UInterchangeShaderNode::MakeNodeUid(NodeName, FStringView{});
-
-		//Test directly in the NodeContainer, because the ShaderNodes can be altered during the node graph either by the parent (dot/normalmap),
-		//or by putting an intermediary node between the child and the parent (tiledimage)
-		if(Node = const_cast<ShaderNodeType*>(Cast<ShaderNodeType>(NodeContainer.GetNode(NodeUID))); !Node)
-		{
-			Node = NewObject<ShaderNodeType>(&NodeContainer);
-			Node->InitializeNode(NodeUID, NodeName, EInterchangeNodeContainerType::TranslatedAsset);
-			NodeContainer.AddNode(Node);
-			if constexpr(std::is_same_v<ShaderNodeType, UInterchangeShaderGraphNode>)
-			{
-				NodeContainer.SetNodeParentUid(NodeUID, FString{});
-			}
-			Node->SetCustomShaderType(ShaderType);
-
-			ShaderNodes.Add(NodeName, Node);
-		}
-
-		return Node;
-	}
+	UInterchangeShaderNode* CreateShaderNode(const FString& NodeName, const FString& ShaderType);
 
 	/**
 	 * Helper function to create an InterchangeTextureNode
@@ -268,7 +246,7 @@ protected:
 	UInterchangeTextureNode* CreateTextureNode(MaterialX::NodePtr Node) const
 	{
 		static_assert(std::is_convertible_v<TextureTypeNode*, UInterchangeTextureNode*>, "CreateTextureNode only accepts type that derived from UInterchangeTextureNode");
-		UInterchangeTextureNode* TextureNode = nullptr;
+		UInterchangeTexture2DNode* TextureNode = nullptr;
 
 		//A node image should have an input file otherwise the user should check its default value
 		if(Node)
@@ -282,7 +260,7 @@ protected:
 				const FString TextureNodeUID = TEXT("\\Texture\\") + Filename;
 
 				//Only add the TextureNode once
-				TextureNode = const_cast<UInterchangeTextureNode*>(Cast<UInterchangeTextureNode>(NodeContainer.GetNode(TextureNodeUID)));
+				TextureNode = const_cast<UInterchangeTexture2DNode*>(Cast<UInterchangeTexture2DNode>(NodeContainer.GetNode(TextureNodeUID)));
 				if(TextureNode == nullptr)
 				{
 					TextureNode = NewObject<TextureTypeNode>(&NodeContainer);
@@ -299,6 +277,35 @@ protected:
 					const FString ColorSpace = GetColorSpace(InputFile);
 					const bool bIsSRGB = ColorSpace == TEXT("srgb_texture");
 					TextureNode->SetCustomSRGB(bIsSRGB);
+
+					auto GetAddressMode = [Node](const char* InputName)
+					{
+						EInterchangeTextureWrapMode WrapMode = EInterchangeTextureWrapMode::Wrap;
+						if(MaterialX::InputPtr InputAddressMode = Node->getInput(InputName))
+						{
+							const std::string& AddressMode = InputAddressMode->getValueString();
+							if(AddressMode == "clamp")
+							{
+								WrapMode = EInterchangeTextureWrapMode::Clamp;
+							}
+							else if(AddressMode == "mirror")
+							{
+								WrapMode = EInterchangeTextureWrapMode::Mirror;
+							}
+							else
+							{
+								WrapMode = EInterchangeTextureWrapMode::Wrap;
+							}
+						}
+
+						return WrapMode;
+					};
+
+					EInterchangeTextureWrapMode WrapModeU = GetAddressMode("uaddressmode");
+					EInterchangeTextureWrapMode WrapModeV = GetAddressMode("vaddressmode");
+
+					TextureNode->SetCustomWrapU(WrapModeU);
+					TextureNode->SetCustomWrapV(WrapModeV);
 				}
 			}
 		}
@@ -374,10 +381,14 @@ protected:
 	/** Store the shader nodes only when we create the shader graph node*/
 	TMap<FString, UInterchangeShaderNode*> ShaderNodes;
 
+	/** Matching MaterialX category and Connect function*/
+	TMap<FString, FOnConnectNodeOutputToInput> MatchingConnectNodeDelegates;
+
 	/** The surface shader node processed during the Translate, up to the derived class to initialize it*/
 	MaterialX::NodePtr SurfaceShaderNode;
 
-	/** Matching MaterialX category and Connect function*/
-	TMap<FString, FOnConnectNodeOutputToInput> MatchingConnectNodeDelegates;
+	/** Initialized by the material shader (e.g: surfacematerial), the derived class should only set the ShaderType */
+	UInterchangeShaderGraphNode* ShaderGraphNode;
+
 };
 #endif
\ No newline at end of file
