Finalized support for MaterialFunctionCall in the generic material pipeline

diff --git a/Engine/Plugins/Experimental/Enterprise/DatasmithInterchange/Source/DatasmithInterchange/Private/InterchangeDatasmithMaterialNode.cpp b/Engine/Plugins/Experimental/Enterprise/DatasmithInterchange/Source/DatasmithInterchange/Private/InterchangeDatasmithMaterialNode.cpp
--- a/Engine/Plugins/Experimental/Enterprise/DatasmithInterchange/Source/DatasmithInterchange/Private/InterchangeDatasmithMaterialNode.cpp
+++ b/Engine/Plugins/Experimental/Enterprise/DatasmithInterchange/Source/DatasmithInterchange/Private/InterchangeDatasmithMaterialNode.cpp
@@ -38,7 +38,6 @@ const FString UInterchangeDatasmithPbrMaterialNode::MaterialFunctionsDependencie
 
 namespace UE::DatasmithInterchange::MaterialUtils
 {
-	const FName MaterialFunctionPathAttrName(TEXT("Datasmith:Material:FunctionCall:MaterialFunctionPath"));
 	const FName DefaultOutputIndexAttrName(TEXT("Datasmith:MaterialExpression:DefaultOutputIndex"));
 
 	const FName MaterialTypeAttrName(TEXT("Datasmith:ReferenceMaterial:MaterialType"));
@@ -68,7 +67,12 @@ namespace UE::DatasmithInterchange::MaterialUtils
 
 			void ConnectOuputToInput(const FString& OutputNodeUid, const FName& OutputName) const
 			{
-				UInterchangeShaderPortsAPI::ConnectOuputToInput(InputNode, InputName, OutputNodeUid, OutputName.ToString());
+				UInterchangeShaderPortsAPI::ConnectOuputToInputByName(InputNode, InputName, OutputNodeUid, OutputName.ToString());
+			}
+
+			void ConnectOuputToInput(const FString& OutputNodeUid, int32 OutputIndex) const
+			{
+				UInterchangeShaderPortsAPI::ConnectOuputToInputByIndex(InputNode, InputName, OutputNodeUid, OutputIndex);
 			}
 		};
 
@@ -446,7 +450,6 @@ namespace UE::DatasmithInterchange::MaterialUtils
 		FString FunctionPath = FPackageName::ExportTextPathToObjectPath(DatasmithExpression.GetFunctionPathName());
 
 		TArray<FString> InputNames;
-		UInterchangeShaderNode* ExpressionNode = nullptr;
 
 		if (FPackageName::DoesPackageExist(FunctionPath))
 		{
@@ -471,7 +474,7 @@ namespace UE::DatasmithInterchange::MaterialUtils
 			const bool bCanProceed = FunctionInputs.Num() < DatasmithExpression.GetInputCount() || FunctionOutputs.IsValidIndex(OutputIndex);
 			if (!ensure(bCanProceed))
 			{
-				// TODO: Log error
+				// TODO: Log warning invalid predefined asset reference used by function call expression
 				return;
 			}
 
@@ -480,10 +483,6 @@ namespace UE::DatasmithInterchange::MaterialUtils
 			{
 				InputNames.Add(FunctionInput.Input.InputName.ToString());
 			}
-
-			ExpressionNode = CreateExpressionNode(DatasmithExpression, ConnectionData.InputNode->GetUniqueID());
-
-			ExpressionNode->AddStringAttribute(MaterialFunctionPathAttrName, FunctionPath);
 		}
 		else
 		{
@@ -491,36 +490,45 @@ namespace UE::DatasmithInterchange::MaterialUtils
 			const UInterchangeDatasmithPbrMaterialNode* PbrMaterialNode = Cast<const UInterchangeDatasmithPbrMaterialNode>(NodeContainer.GetNode(MaterialUid));
 			if (!ensure(PbrMaterialNode))
 			{
-				// TODO: Log error
+				// TODO: Log warning invalid internal reference used by function call expression
 				return;
 			}
 
+			FunctionPath.Empty();
 			UInterchangeShaderPortsAPI::GatherInputs(PbrMaterialNode, InputNames);
 
-			UInterchangeFunctionCallShaderNode* FunctionCallExpressionNode = CreateExpressionNode<UInterchangeFunctionCallShaderNode>(DatasmithExpression, ConnectionData.InputNode->GetUniqueID());
-			FunctionCallExpressionNode->SetCustomMaterialFunction(MaterialUid);
-
 			MaterialFunctionUids.Add(MaterialUid);
+		}
 
-			ExpressionNode = FunctionCallExpressionNode;
+		if (FunctionPath.IsEmpty() && MaterialUid.IsEmpty())
+		{
+			// TODO: Log warning invalid reference used by function call expression
+			return;
 		}
 
-		if (!ensure(ExpressionNode))
+		UInterchangeFunctionCallShaderNode* FunctionCallExpressionNode = CreateExpressionNode<UInterchangeFunctionCallShaderNode>(DatasmithExpression, ConnectionData.InputNode->GetUniqueID());
+		if (!ensure(FunctionCallExpressionNode))
 		{
 			// TODO: Log error
 			return;
 		}
 
-		// TODO: Handle case where path is an existing asset vs a material function defined in file
-		ExpressionNode->SetCustomShaderType(TEXT("MaterialFunctionCall"));
+		if (!MaterialUid.IsEmpty())
+		{
+			FunctionCallExpressionNode->SetCustomMaterialFunction(MaterialUid);
+		}
+		else
+		{
+			FunctionCallExpressionNode->SetCustomMaterialFunction(FunctionPath);
+		}
 
-		ConnectionData.ConnectDefaultOuputToInput(ExpressionNode->GetUniqueID());
+		ConnectionData.ConnectOuputToInput(FunctionCallExpressionNode->GetUniqueID(), OutputIndex);
 
 		for (int32 InputIndex = 0; InputIndex < DatasmithExpression.GetInputCount(); ++InputIndex)
 		{
 			if (const IDatasmithExpressionInput* ExpressionInput = DatasmithExpression.GetInput(InputIndex))
 			{
-				FConnectionData InputConnectionData{ ExpressionNode, InputNames[InputIndex]};
+				FConnectionData InputConnectionData{ FunctionCallExpressionNode, InputNames[InputIndex] };
 
 				ConnectExpression(ExpressionInput->GetExpression(), InputConnectionData, ExpressionInput->GetOutputIndex());
 			}
@@ -554,8 +562,7 @@ namespace UE::DatasmithInterchange::MaterialUtils
 			}
 			else
 			{
-				// TODO: Need to find node's output list
-				ConnectionData.ConnectOuputToInput(ExpressionNode->GetUniqueID(), NAME_None);
+				ConnectionData.ConnectOuputToInput(ExpressionNode->GetUniqueID(), OutputIndex);
 			}
 		}
 
diff --git a/Engine/Plugins/Experimental/Enterprise/DatasmithInterchange/Source/DatasmithInterchange/Private/InterchangeDatasmithMaterialNode.h b/Engine/Plugins/Experimental/Enterprise/DatasmithInterchange/Source/DatasmithInterchange/Private/InterchangeDatasmithMaterialNode.h
--- a/Engine/Plugins/Experimental/Enterprise/DatasmithInterchange/Source/DatasmithInterchange/Private/InterchangeDatasmithMaterialNode.h
+++ b/Engine/Plugins/Experimental/Enterprise/DatasmithInterchange/Source/DatasmithInterchange/Private/InterchangeDatasmithMaterialNode.h
@@ -127,7 +127,6 @@ private:
 namespace UE::DatasmithInterchange::MaterialUtils
 {
 	extern const FName DefaultOutputIndexAttrName;
-	extern const FName MaterialFunctionPathAttrName;
 
 	extern const FName MaterialTypeAttrName;
 	extern const FName MaterialQualityAttrName;
diff --git a/Engine/Plugins/Experimental/Enterprise/DatasmithInterchange/Source/DatasmithInterchange/Private/InterchangeDatasmithMaterialPipeline.cpp b/Engine/Plugins/Experimental/Enterprise/DatasmithInterchange/Source/DatasmithInterchange/Private/InterchangeDatasmithMaterialPipeline.cpp
--- a/Engine/Plugins/Experimental/Enterprise/DatasmithInterchange/Source/DatasmithInterchange/Private/InterchangeDatasmithMaterialPipeline.cpp
+++ b/Engine/Plugins/Experimental/Enterprise/DatasmithInterchange/Source/DatasmithInterchange/Private/InterchangeDatasmithMaterialPipeline.cpp
@@ -295,18 +295,7 @@ void UInterchangeDatasmithMaterialPipeline::UpdateMaterialFactoryNodes(const TAr
 			continue;
 		}
 
-		FString MaterialFunctionPath;
-		if (ShaderNode->GetStringAttribute(MaterialUtils::MaterialFunctionPathAttrName, MaterialFunctionPath))
-		{
-			static const FName MaterialFunctionMemberName = GET_MEMBER_NAME_CHECKED(UMaterialExpressionMaterialFunctionCall, MaterialFunction);
-
-			if (UInterchangeMaterialExpressionFactoryNode* ExpressionFactoryNode = Cast<UInterchangeMaterialExpressionFactoryNode>(FactoryNode))
-			{
-				ExpressionFactoryNode->AddStringAttribute(MaterialFunctionMemberName, MaterialFunctionPath);
-				ExpressionFactoryNode->AddApplyAndFillDelegates<FString>(MaterialFunctionMemberName, UMaterialExpressionMaterialFunctionCall::StaticClass(), MaterialFunctionMemberName);
-			}
-		}
-		else if (UInterchangeDatasmithPbrMaterialNode* PbrMaterialNode = Cast<UInterchangeDatasmithPbrMaterialNode>(ShaderNode))
+		if (UInterchangeDatasmithPbrMaterialNode* PbrMaterialNode = Cast<UInterchangeDatasmithPbrMaterialNode>(ShaderNode))
 		{
 			UInterchangeBaseMaterialFactoryNode* BaseMaterialFactoryNode = Cast<UInterchangeBaseMaterialFactoryNode>(FactoryNode);
 			ensure(BaseMaterialFactoryNode);
diff --git a/Engine/Plugins/Interchange/Runtime/Source/FactoryNodes/Private/InterchangeMaterialFactoryNode.cpp b/Engine/Plugins/Interchange/Runtime/Source/FactoryNodes/Private/InterchangeMaterialFactoryNode.cpp
--- a/Engine/Plugins/Interchange/Runtime/Source/FactoryNodes/Private/InterchangeMaterialFactoryNode.cpp
+++ b/Engine/Plugins/Interchange/Runtime/Source/FactoryNodes/Private/InterchangeMaterialFactoryNode.cpp
@@ -47,7 +47,7 @@ bool UInterchangeMaterialFactoryNode::ConnectToBaseColor(const FString& Attribut
 	
 bool UInterchangeMaterialFactoryNode::ConnectOutputToBaseColor(const FString& ExpressionNodeUid, const FString& OutputName)
 {
-	return UInterchangeShaderPortsAPI::ConnectOuputToInput(this, UE::Interchange::Materials::PBR::Parameters::BaseColor.ToString(), ExpressionNodeUid, OutputName);
+	return UInterchangeShaderPortsAPI::ConnectOuputToInputByName(this, UE::Interchange::Materials::PBR::Parameters::BaseColor.ToString(), ExpressionNodeUid, OutputName);
 }
 
 bool UInterchangeMaterialFactoryNode::GetMetallicConnection(FString& ExpressionNodeUid, FString& OutputName) const
@@ -62,7 +62,7 @@ bool UInterchangeMaterialFactoryNode::ConnectToMetallic(const FString& Attribute
 
 bool UInterchangeMaterialFactoryNode::ConnectOutputToMetallic(const FString& ExpressionNodeUid, const FString& OutputName)
 {
-	return UInterchangeShaderPortsAPI::ConnectOuputToInput(this, UE::Interchange::Materials::PBR::Parameters::Metallic.ToString(), ExpressionNodeUid, OutputName);
+	return UInterchangeShaderPortsAPI::ConnectOuputToInputByName(this, UE::Interchange::Materials::PBR::Parameters::Metallic.ToString(), ExpressionNodeUid, OutputName);
 }
 
 bool UInterchangeMaterialFactoryNode::GetSpecularConnection(FString& ExpressionNodeUid, FString& OutputName) const
@@ -77,7 +77,7 @@ bool UInterchangeMaterialFactoryNode::ConnectToSpecular(const FString& Expressio
 
 bool UInterchangeMaterialFactoryNode::ConnectOutputToSpecular(const FString& ExpressionNodeUid, const FString& OutputName)
 {
-	return UInterchangeShaderPortsAPI::ConnectOuputToInput(this, UE::Interchange::Materials::PBR::Parameters::Specular.ToString(), ExpressionNodeUid, OutputName);
+	return UInterchangeShaderPortsAPI::ConnectOuputToInputByName(this, UE::Interchange::Materials::PBR::Parameters::Specular.ToString(), ExpressionNodeUid, OutputName);
 }
 
 bool UInterchangeMaterialFactoryNode::GetRoughnessConnection(FString& ExpressionNodeUid, FString& OutputName) const
@@ -92,7 +92,7 @@ bool UInterchangeMaterialFactoryNode::ConnectToRoughness(const FString& Expressi
 
 bool UInterchangeMaterialFactoryNode::ConnectOutputToRoughness(const FString& ExpressionNodeUid, const FString& OutputName)
 {
-	return UInterchangeShaderPortsAPI::ConnectOuputToInput(this, UE::Interchange::Materials::PBR::Parameters::Roughness.ToString(), ExpressionNodeUid, OutputName);
+	return UInterchangeShaderPortsAPI::ConnectOuputToInputByName(this, UE::Interchange::Materials::PBR::Parameters::Roughness.ToString(), ExpressionNodeUid, OutputName);
 }
 
 bool UInterchangeMaterialFactoryNode::GetAnisotropyConnection(FString& ExpressionNodeUid, FString& OutputName) const
@@ -107,7 +107,7 @@ bool UInterchangeMaterialFactoryNode::ConnectToAnisotropy(const FString& Express
 
 bool UInterchangeMaterialFactoryNode::ConnectOutputToAnisotropy(const FString& ExpressionNodeUid, const FString& OutputName)
 {
-	return UInterchangeShaderPortsAPI::ConnectOuputToInput(this, UE::Interchange::Materials::PBR::Parameters::Anisotropy.ToString(), ExpressionNodeUid, OutputName);
+	return UInterchangeShaderPortsAPI::ConnectOuputToInputByName(this, UE::Interchange::Materials::PBR::Parameters::Anisotropy.ToString(), ExpressionNodeUid, OutputName);
 }
 
 bool UInterchangeMaterialFactoryNode::GetEmissiveColorConnection(FString& ExpressionNodeUid, FString& OutputName) const
@@ -122,7 +122,7 @@ bool UInterchangeMaterialFactoryNode::ConnectToEmissiveColor(const FString& Expr
 
 bool UInterchangeMaterialFactoryNode::ConnectOutputToEmissiveColor(const FString& ExpressionNodeUid, const FString& OutputName)
 {
-	return UInterchangeShaderPortsAPI::ConnectOuputToInput(this, UE::Interchange::Materials::PBR::Parameters::EmissiveColor.ToString(), ExpressionNodeUid, OutputName);
+	return UInterchangeShaderPortsAPI::ConnectOuputToInputByName(this, UE::Interchange::Materials::PBR::Parameters::EmissiveColor.ToString(), ExpressionNodeUid, OutputName);
 }
 
 bool UInterchangeMaterialFactoryNode::GetNormalConnection(FString& ExpressionNodeUid, FString& OutputName) const
@@ -137,7 +137,7 @@ bool UInterchangeMaterialFactoryNode::ConnectToNormal(const FString& ExpressionN
 
 bool UInterchangeMaterialFactoryNode::ConnectOutputToNormal(const FString& ExpressionNodeUid, const FString& OutputName)
 {
-	return UInterchangeShaderPortsAPI::ConnectOuputToInput(this, UE::Interchange::Materials::PBR::Parameters::Normal.ToString(), ExpressionNodeUid, OutputName);
+	return UInterchangeShaderPortsAPI::ConnectOuputToInputByName(this, UE::Interchange::Materials::PBR::Parameters::Normal.ToString(), ExpressionNodeUid, OutputName);
 }
 
 bool UInterchangeMaterialFactoryNode::GetTangentConnection(FString& ExpressionNodeUid, FString& OutputName) const
@@ -152,7 +152,7 @@ bool UInterchangeMaterialFactoryNode::ConnectToTangent(const FString& Expression
 
 bool UInterchangeMaterialFactoryNode::ConnectOutputToTangent(const FString& ExpressionNodeUid, const FString& OutputName)
 {
-	return UInterchangeShaderPortsAPI::ConnectOuputToInput(this, UE::Interchange::Materials::PBR::Parameters::Tangent.ToString(), ExpressionNodeUid, OutputName);
+	return UInterchangeShaderPortsAPI::ConnectOuputToInputByName(this, UE::Interchange::Materials::PBR::Parameters::Tangent.ToString(), ExpressionNodeUid, OutputName);
 }
 
 bool UInterchangeMaterialFactoryNode::GetSubsurfaceConnection(FString& ExpressionNodeUid, FString& OutputName) const
@@ -167,7 +167,7 @@ bool UInterchangeMaterialFactoryNode::ConnectToSubsurface(const FString& Express
 
 bool UInterchangeMaterialFactoryNode::ConnectOutputToSubsurface(const FString& ExpressionNodeUid, const FString& OutputName)
 {
-	return UInterchangeShaderPortsAPI::ConnectOuputToInput(this, UE::Interchange::Materials::Subsurface::Parameters::SubsurfaceColor.ToString(), ExpressionNodeUid, OutputName);
+	return UInterchangeShaderPortsAPI::ConnectOuputToInputByName(this, UE::Interchange::Materials::Subsurface::Parameters::SubsurfaceColor.ToString(), ExpressionNodeUid, OutputName);
 }
 
 bool UInterchangeMaterialFactoryNode::GetOpacityConnection(FString& ExpressionNodeUid, FString& OutputName) const
@@ -182,7 +182,7 @@ bool UInterchangeMaterialFactoryNode::ConnectToOpacity(const FString& AttributeV
 
 bool UInterchangeMaterialFactoryNode::ConnectOutputToOpacity(const FString& ExpressionNodeUid, const FString& OutputName)
 {
-	return UInterchangeShaderPortsAPI::ConnectOuputToInput(this, UE::Interchange::Materials::PBR::Parameters::Opacity.ToString(), ExpressionNodeUid, OutputName);
+	return UInterchangeShaderPortsAPI::ConnectOuputToInputByName(this, UE::Interchange::Materials::PBR::Parameters::Opacity.ToString(), ExpressionNodeUid, OutputName);
 }
 
 bool UInterchangeMaterialFactoryNode::GetOcclusionConnection(FString& ExpressionNodeUid, FString& OutputName) const
@@ -197,7 +197,7 @@ bool UInterchangeMaterialFactoryNode::ConnectToOcclusion(const FString& Attribut
 
 bool UInterchangeMaterialFactoryNode::ConnectOutputToOcclusion(const FString& ExpressionNodeUid, const FString& OutputName)
 {
-	return UInterchangeShaderPortsAPI::ConnectOuputToInput(this, UE::Interchange::Materials::PBR::Parameters::Occlusion.ToString(), ExpressionNodeUid, OutputName);
+	return UInterchangeShaderPortsAPI::ConnectOuputToInputByName(this, UE::Interchange::Materials::PBR::Parameters::Occlusion.ToString(), ExpressionNodeUid, OutputName);
 }
 
 bool UInterchangeMaterialFactoryNode::GetRefractionConnection(FString& ExpressionNodeUid, FString& OutputName) const
@@ -212,7 +212,7 @@ bool UInterchangeMaterialFactoryNode::ConnectToRefraction(const FString& Attribu
 
 bool UInterchangeMaterialFactoryNode::ConnectOutputToRefraction(const FString& ExpressionNodeUid, const FString& OutputName)
 {
-	return UInterchangeShaderPortsAPI::ConnectOuputToInput(this, UE::Interchange::Materials::Common::Parameters::IndexOfRefraction.ToString(), ExpressionNodeUid, OutputName);
+	return UInterchangeShaderPortsAPI::ConnectOuputToInputByName(this, UE::Interchange::Materials::Common::Parameters::IndexOfRefraction.ToString(), ExpressionNodeUid, OutputName);
 }
 
 bool UInterchangeMaterialFactoryNode::GetClearCoatConnection(FString& ExpressionNodeUid, FString& OutputName) const
@@ -227,7 +227,7 @@ bool UInterchangeMaterialFactoryNode::ConnectToClearCoat(const FString& Attribut
 
 bool UInterchangeMaterialFactoryNode::ConnectOutputToClearCoat(const FString& ExpressionNodeUid, const FString& OutputName)
 {
-	return UInterchangeShaderPortsAPI::ConnectOuputToInput(this, UE::Interchange::Materials::ClearCoat::Parameters::ClearCoat.ToString(), ExpressionNodeUid, OutputName);
+	return UInterchangeShaderPortsAPI::ConnectOuputToInputByName(this, UE::Interchange::Materials::ClearCoat::Parameters::ClearCoat.ToString(), ExpressionNodeUid, OutputName);
 }
 
 bool UInterchangeMaterialFactoryNode::GetClearCoatRoughnessConnection(FString& ExpressionNodeUid, FString& OutputName) const
@@ -242,7 +242,7 @@ bool UInterchangeMaterialFactoryNode::ConnectToClearCoatRoughness(const FString&
 
 bool UInterchangeMaterialFactoryNode::ConnectOutputToClearCoatRoughness(const FString& ExpressionNodeUid, const FString& OutputName)
 {
-	return UInterchangeShaderPortsAPI::ConnectOuputToInput(this, UE::Interchange::Materials::ClearCoat::Parameters::ClearCoatRoughness.ToString(), ExpressionNodeUid, OutputName);
+	return UInterchangeShaderPortsAPI::ConnectOuputToInputByName(this, UE::Interchange::Materials::ClearCoat::Parameters::ClearCoatRoughness.ToString(), ExpressionNodeUid, OutputName);
 }
 
 bool UInterchangeMaterialFactoryNode::GetClearCoatNormalConnection(FString& ExpressionNodeUid, FString& OutputName) const
@@ -257,7 +257,7 @@ bool UInterchangeMaterialFactoryNode::ConnectToClearCoatNormal(const FString& At
 
 bool UInterchangeMaterialFactoryNode::ConnectOutputToClearCoatNormal(const FString& ExpressionNodeUid, const FString& OutputName)
 {
-	return UInterchangeShaderPortsAPI::ConnectOuputToInput(this, UE::Interchange::Materials::ClearCoat::Parameters::ClearCoatNormal.ToString(), ExpressionNodeUid, OutputName);
+	return UInterchangeShaderPortsAPI::ConnectOuputToInputByName(this, UE::Interchange::Materials::ClearCoat::Parameters::ClearCoatNormal.ToString(), ExpressionNodeUid, OutputName);
 }
 
 bool UInterchangeMaterialFactoryNode::GetTransmissionColorConnection(FString& ExpressionNodeUid, FString& OutputName) const
@@ -272,7 +272,7 @@ bool UInterchangeMaterialFactoryNode::ConnectToTransmissionColor(const FString&
 
 bool UInterchangeMaterialFactoryNode::ConnectOutputToTransmissionColor(const FString& ExpressionNodeUid, const FString& OutputName)
 {
-	return UInterchangeShaderPortsAPI::ConnectOuputToInput(this, UE::Interchange::Materials::ThinTranslucent::Parameters::TransmissionColor.ToString(), ExpressionNodeUid, OutputName);
+	return UInterchangeShaderPortsAPI::ConnectOuputToInputByName(this, UE::Interchange::Materials::ThinTranslucent::Parameters::TransmissionColor.ToString(), ExpressionNodeUid, OutputName);
 }
 
 bool UInterchangeMaterialFactoryNode::GetFuzzColorConnection(FString& ExpressionNodeUid, FString& OutputName) const
@@ -287,7 +287,7 @@ bool UInterchangeMaterialFactoryNode::ConnectToFuzzColor(const FString& Attribut
 
 bool UInterchangeMaterialFactoryNode::ConnectOutputToFuzzColor(const FString& ExpressionNodeUid, const FString& OutputName)
 {
-	return UInterchangeShaderPortsAPI::ConnectOuputToInput(this, UE::Interchange::Materials::Sheen::Parameters::SheenColor.ToString(), ExpressionNodeUid, OutputName);
+	return UInterchangeShaderPortsAPI::ConnectOuputToInputByName(this, UE::Interchange::Materials::Sheen::Parameters::SheenColor.ToString(), ExpressionNodeUid, OutputName);
 }
 
 bool UInterchangeMaterialFactoryNode::GetClothConnection(FString& ExpressionNodeUid, FString& OutputName) const
@@ -302,7 +302,7 @@ bool UInterchangeMaterialFactoryNode::ConnectToCloth(const FString& AttributeVal
 
 bool UInterchangeMaterialFactoryNode::ConnectOutputToCloth(const FString& ExpressionNodeUid, const FString& OutputName)
 {
-	return UInterchangeShaderPortsAPI::ConnectOuputToInput(this, UE::Interchange::Materials::Sheen::Parameters::SheenRoughness.ToString(), ExpressionNodeUid, OutputName);
+	return UInterchangeShaderPortsAPI::ConnectOuputToInputByName(this, UE::Interchange::Materials::Sheen::Parameters::SheenRoughness.ToString(), ExpressionNodeUid, OutputName);
 }
 
 bool UInterchangeMaterialFactoryNode::GetCustomShadingModel(TEnumAsByte<EMaterialShadingModel>& AttributeValue) const
diff --git a/Engine/Plugins/Interchange/Runtime/Source/Import/Private/Gltf/InterchangeGltfTranslator.cpp b/Engine/Plugins/Interchange/Runtime/Source/Import/Private/Gltf/InterchangeGltfTranslator.cpp
--- a/Engine/Plugins/Interchange/Runtime/Source/Import/Private/Gltf/InterchangeGltfTranslator.cpp
+++ b/Engine/Plugins/Interchange/Runtime/Source/Import/Private/Gltf/InterchangeGltfTranslator.cpp
@@ -461,7 +461,7 @@ void UInterchangeGltfTranslator::HandleGltfMaterialParameter( UInterchangeBaseNo
 					FactorOneMinusNode->AddFloatAttribute( UInterchangeShaderPortsAPI::MakeInputValueKey( Standard::Nodes::OneMinus::Inputs::Input.ToString() ), MapFactor.Get< float >() );
 				}
 
-				UInterchangeShaderPortsAPI::ConnectOuputToInput( FactorNode, Standard::Nodes::FlattenNormal::Inputs::Normal.ToString(), ColorNode->GetUniqueID(), OutputChannel );
+				UInterchangeShaderPortsAPI::ConnectOuputToInputByName( FactorNode, Standard::Nodes::FlattenNormal::Inputs::Normal.ToString(), ColorNode->GetUniqueID(), OutputChannel );
 				UInterchangeShaderPortsAPI::ConnectDefaultOuputToInput( FactorNode, Standard::Nodes::FlattenNormal::Inputs::Flatness.ToString(), FactorOneMinusNode->GetUniqueID() );
 			}
 			else
@@ -477,14 +477,14 @@ void UInterchangeGltfTranslator::HandleGltfMaterialParameter( UInterchangeBaseNo
 					FactorNode->AddLinearColorAttribute( UInterchangeShaderPortsAPI::MakeInputValueKey( Standard::Nodes::Multiply::Inputs::B.ToString() ), MapFactor.Get< FLinearColor >() );
 				}
 
-				UInterchangeShaderPortsAPI::ConnectOuputToInput( FactorNode, Standard::Nodes::Multiply::Inputs::A.ToString(), ColorNode->GetUniqueID(), OutputChannel );
+				UInterchangeShaderPortsAPI::ConnectOuputToInputByName( FactorNode, Standard::Nodes::Multiply::Inputs::A.ToString(), ColorNode->GetUniqueID(), OutputChannel );
 			}
 
 			UInterchangeShaderPortsAPI::ConnectDefaultOuputToInput( NodeToConnectTo, InputToConnectTo, FactorNode->GetUniqueID() );
 		}
 		else
 		{
-			UInterchangeShaderPortsAPI::ConnectOuputToInput( NodeToConnectTo, InputToConnectTo, ColorNode->GetUniqueID(), OutputChannel );
+			UInterchangeShaderPortsAPI::ConnectOuputToInputByName( NodeToConnectTo, InputToConnectTo, ColorNode->GetUniqueID(), OutputChannel );
 		}
 	}
 	else
@@ -819,7 +819,7 @@ void UInterchangeGltfTranslator::HandleGltfTransmission( UInterchangeBaseNodeCon
 
 				LerpMetallicNode->AddFloatAttribute( UInterchangeShaderPortsAPI::MakeInputValueKey( Standard::Nodes::Lerp::Inputs::B.ToString() ), 1.f );
 				UInterchangeShaderPortsAPI::ConnectDefaultOuputToInput( LerpMetallicNode, Standard::Nodes::Lerp::Inputs::A.ToString(), CurrentNode->GetUniqueID() );
-				UInterchangeShaderPortsAPI::ConnectOuputToInput( LerpMetallicNode, Standard::Nodes::Lerp::Inputs::Factor.ToString(), MetallicNodeUid, MetallicNodeOutput );
+				UInterchangeShaderPortsAPI::ConnectOuputToInputByName( LerpMetallicNode, Standard::Nodes::Lerp::Inputs::Factor.ToString(), MetallicNodeUid, MetallicNodeOutput );
 
 				CurrentNode = LerpMetallicNode;
 				CurrentInput = TEXT("");
@@ -836,7 +836,7 @@ void UInterchangeGltfTranslator::HandleGltfTransmission( UInterchangeBaseNodeCon
 
 				OpacityLerpNode->AddFloatAttribute( UInterchangeShaderPortsAPI::MakeInputValueKey( Standard::Nodes::Lerp::Inputs::A.ToString() ), 0.f );
 				UInterchangeShaderPortsAPI::ConnectDefaultOuputToInput( OpacityLerpNode, Standard::Nodes::Lerp::Inputs::B.ToString(), CurrentNode->GetUniqueID() );
-				UInterchangeShaderPortsAPI::ConnectOuputToInput( OpacityLerpNode, Standard::Nodes::Lerp::Inputs::Factor.ToString(), OpacityNodeUid, OpacityNodeOutput );
+				UInterchangeShaderPortsAPI::ConnectOuputToInputByName( OpacityLerpNode, Standard::Nodes::Lerp::Inputs::Factor.ToString(), OpacityNodeUid, OpacityNodeOutput );
 
 				CurrentNode = OpacityLerpNode;
 				CurrentInput = TEXT("");
@@ -881,11 +881,11 @@ void UInterchangeGltfTranslator::HandleGltfTransmission( UInterchangeBaseNodeCon
 				OpacityLerpNode->SetCustomShaderType( Standard::Nodes::Lerp::Name.ToString() );
 
 				OpacityLerpNode->AddLinearColorAttribute( UInterchangeShaderPortsAPI::MakeInputValueKey( Standard::Nodes::Lerp::Inputs::A.ToString() ), FLinearColor::White );
-				UInterchangeShaderPortsAPI::ConnectOuputToInput( OpacityLerpNode, Standard::Nodes::Lerp::Inputs::Factor.ToString(), OpacityNodeUid, OpacityNodeOutput );
+				UInterchangeShaderPortsAPI::ConnectOuputToInputByName( OpacityLerpNode, Standard::Nodes::Lerp::Inputs::Factor.ToString(), OpacityNodeUid, OpacityNodeOutput );
 
 				if ( CurrentNode )
 				{
-					UInterchangeShaderPortsAPI::ConnectOuputToInput( OpacityLerpNode, Standard::Nodes::Lerp::Inputs::B.ToString(), CurrentNode->GetUniqueID(), CurrentOuput );
+					UInterchangeShaderPortsAPI::ConnectOuputToInputByName( OpacityLerpNode, Standard::Nodes::Lerp::Inputs::B.ToString(), CurrentNode->GetUniqueID(), CurrentOuput );
 				}
 				else
 				{
@@ -899,7 +899,7 @@ void UInterchangeGltfTranslator::HandleGltfTransmission( UInterchangeBaseNodeCon
 
 		if ( CurrentNode )
 		{
-			UInterchangeShaderPortsAPI::ConnectOuputToInput( &ShaderGraphNode, ThinTranslucent::Parameters::TransmissionColor.ToString(), CurrentNode->GetUniqueID(), CurrentOuput );
+			UInterchangeShaderPortsAPI::ConnectOuputToInputByName( &ShaderGraphNode, ThinTranslucent::Parameters::TransmissionColor.ToString(), CurrentNode->GetUniqueID(), CurrentOuput );
 		}
 		else
 		{
diff --git a/Engine/Plugins/Interchange/Runtime/Source/Import/Private/Material/InterchangeMaterialFactory.cpp b/Engine/Plugins/Interchange/Runtime/Source/Import/Private/Material/InterchangeMaterialFactory.cpp
--- a/Engine/Plugins/Interchange/Runtime/Source/Import/Private/Material/InterchangeMaterialFactory.cpp
+++ b/Engine/Plugins/Interchange/Runtime/Source/Import/Private/Material/InterchangeMaterialFactory.cpp
@@ -101,7 +101,17 @@ namespace UE
 
 				int32 GetOutputIndex(UMaterialExpression& MaterialExpression, const FString& OutputName)
 				{
-					int32 ExpressionOutputIndex = 0;
+					// Check whether OutputName stores an index
+					int32 ExpressionOutputIndex = UInterchangeShaderPortsAPI::GetOutputIndexFromName(OutputName);
+					if (ExpressionOutputIndex != INDEX_NONE)
+					{
+						if (MaterialExpression.GetOutputs().IsValidIndex(ExpressionOutputIndex))
+						{
+							return ExpressionOutputIndex;
+						}
+					}
+
+					ExpressionOutputIndex = 0;
 
 					for (const FExpressionOutput& ExpressionOutput : MaterialExpression.GetOutputs())
 					{
diff --git a/Engine/Plugins/Interchange/Runtime/Source/Import/Private/MaterialX/MaterialXUtils/MaterialXSurfaceShaderAbstract.cpp b/Engine/Plugins/Interchange/Runtime/Source/Import/Private/MaterialX/MaterialXUtils/MaterialXSurfaceShaderAbstract.cpp
--- a/Engine/Plugins/Interchange/Runtime/Source/Import/Private/MaterialX/MaterialXUtils/MaterialXSurfaceShaderAbstract.cpp
+++ b/Engine/Plugins/Interchange/Runtime/Source/Import/Private/MaterialX/MaterialXUtils/MaterialXSurfaceShaderAbstract.cpp
@@ -335,7 +335,7 @@ void FMaterialXSurfaceShaderAbstract::ConnectImageInputToOutput(MaterialX::NodeP
 
 		UInterchangeShaderNode* TextureShaderNode = CreateShaderNode(UpstreamNode->getName().c_str(), TextureSample::Name.ToString());
 		TextureShaderNode->AddStringAttribute(UInterchangeShaderPortsAPI::MakeInputValueKey(TextureSample::Inputs::Texture.ToString()), TextureNode->GetUniqueID());
-		UInterchangeShaderPortsAPI::ConnectOuputToInput(ParentShaderNode, InputChannelName, TextureShaderNode->GetUniqueID(), OutputChannel);
+		UInterchangeShaderPortsAPI::ConnectOuputToInputByName(ParentShaderNode, InputChannelName, TextureShaderNode->GetUniqueID(), OutputChannel);
 	}
 	else
 	{
@@ -941,7 +941,7 @@ void FMaterialXSurfaceShaderAbstract::ConnectBlurInputToOutput(MaterialX::NodePt
 
 				UInterchangeShaderNode* TextureShaderNode = CreateShaderNode(UpstreamNode->getName().c_str(), TextureSampleBlur::Name.ToString());
 				TextureShaderNode->AddStringAttribute(UInterchangeShaderPortsAPI::MakeInputValueKey(TextureSampleBlur::Inputs::Texture.ToString()), TextureNode->GetUniqueID());
-				UInterchangeShaderPortsAPI::ConnectOuputToInput(ParentShaderNode, InputChannelName, TextureShaderNode->GetUniqueID(), OutputChannel);
+				UInterchangeShaderPortsAPI::ConnectOuputToInputByName(ParentShaderNode, InputChannelName, TextureShaderNode->GetUniqueID(), OutputChannel);
 
 				if(mx::InputPtr InputKernel = UpstreamNode->getInput("filtertype"))
 				{
diff --git a/Engine/Plugins/Interchange/Runtime/Source/Nodes/Private/InterchangeShaderGraphNode.cpp b/Engine/Plugins/Interchange/Runtime/Source/Nodes/Private/InterchangeShaderGraphNode.cpp
--- a/Engine/Plugins/Interchange/Runtime/Source/Nodes/Private/InterchangeShaderGraphNode.cpp
+++ b/Engine/Plugins/Interchange/Runtime/Source/Nodes/Private/InterchangeShaderGraphNode.cpp
@@ -7,6 +7,7 @@
 
 const TCHAR* UInterchangeShaderPortsAPI::InputPrefix = TEXT("Inputs");
 const TCHAR* UInterchangeShaderPortsAPI::InputSeparator = TEXT(":");
+const TCHAR* UInterchangeShaderPortsAPI::OutputByIndex = TEXT("__ByIndex__");
 
 FName UInterchangeShaderPortsAPI::MakeInputConnectionKey(const FString& InputName)
 {
@@ -87,7 +88,7 @@ bool UInterchangeShaderPortsAPI::ConnectDefaultOuputToInput(UInterchangeBaseNode
 	return InterchangeNode->AddStringAttribute(MakeInputConnectionKey(InputName), ExpressionUid);
 }
 
-bool UInterchangeShaderPortsAPI::ConnectOuputToInput(UInterchangeBaseNode* InterchangeNode, const FString& InputName, const FString& ExpressionUid, const FString& OutputName)
+bool UInterchangeShaderPortsAPI::ConnectOuputToInputByName(UInterchangeBaseNode* InterchangeNode, const FString& InputName, const FString& ExpressionUid, const FString& OutputName)
 {
 	if (OutputName.IsEmpty())
 	{
@@ -99,6 +100,18 @@ bool UInterchangeShaderPortsAPI::ConnectOuputToInput(UInterchangeBaseNode* Inter
 	}
 }
 
+bool UInterchangeShaderPortsAPI::ConnectOuputToInputByIndex(UInterchangeBaseNode* InterchangeNode, const FString& InputName, const FString& ExpressionUid, int32 OutputIndex)
+{
+	TStringBuilder<128> StringBuilder;
+	StringBuilder.Append(ExpressionUid);
+	StringBuilder.Append(InputSeparator);
+	StringBuilder.Append(OutputByIndex);
+	StringBuilder.Append(InputSeparator);
+	StringBuilder.Append(FString::FromInt(OutputIndex));
+
+	return InterchangeNode->AddStringAttribute(MakeInputConnectionKey(InputName), StringBuilder.ToString());
+}
+
 UE::Interchange::EAttributeTypes UInterchangeShaderPortsAPI::GetInputType(const UInterchangeBaseNode* InterchangeNode, const FString& InputName)
 {
 	return InterchangeNode->GetAttributeType(UE::Interchange::FAttributeKey(MakeInputValueKey(InputName)));
@@ -115,6 +128,21 @@ bool UInterchangeShaderPortsAPI::GetInputConnection(const UInterchangeBaseNode*
 	return false;
 }
 
+int32 UInterchangeShaderPortsAPI::GetOutputIndexFromName(const FString& OutputName)
+{
+	if (OutputName.StartsWith(OutputByIndex))
+	{
+		FString OutputIndex;
+		FString Discard;
+
+		OutputName.Split(InputSeparator, &Discard, &OutputIndex, ESearchCase::IgnoreCase, ESearchDir::FromStart);
+
+		return FCString::Atoi(*OutputIndex);
+	}
+
+	return INDEX_NONE;
+}
+
 FString UInterchangeShaderNode::MakeNodeUid(const FStringView NodeName, const FStringView ParentNodeUid)
 {
 	FString ShaderNodeUid = FString(UInterchangeBaseNode::HierarchySeparator) + TEXT("Shaders") + FString(UInterchangeBaseNode::HierarchySeparator);
diff --git a/Engine/Plugins/Interchange/Runtime/Source/Nodes/Public/InterchangeShaderGraphNode.h b/Engine/Plugins/Interchange/Runtime/Source/Nodes/Public/InterchangeShaderGraphNode.h
--- a/Engine/Plugins/Interchange/Runtime/Source/Nodes/Public/InterchangeShaderGraphNode.h
+++ b/Engine/Plugins/Interchange/Runtime/Source/Nodes/Public/InterchangeShaderGraphNode.h
@@ -70,26 +70,50 @@ public:
 	 * @param InterchangeNode	The Node to create the input on.
 	 * @param InputName			The name to give to the input.
 	 * @param ExpressionUid		The unique id of the node to connect to the input.
-	 * @param OutputName		The name of the ouput from ExpressionUid to connect to the input.
-	 * @return					true if the input connection was succesfully added to the node.
+	 * @param OutputName		The name of the output from ExpressionUid to connect to the input.
+	 * @return					true if the input connection was successfully added to the node.
 	 */
 	UFUNCTION(BlueprintCallable, Category = "Interchange | Node | Material")
-	static bool ConnectOuputToInput(UInterchangeBaseNode* InterchangeNode, const FString& InputName, const FString& ExpressionUid, const FString& OutputName);
+	static bool ConnectOuputToInputByName(UInterchangeBaseNode* InterchangeNode, const FString& InputName, const FString& ExpressionUid, const FString& OutputName);
+	UE_DEPRECATED(5.3, "This function is replace by ConnectOuputToInputByName.")
+	static bool ConnectOuputToInput(UInterchangeBaseNode* InterchangeNode, const FString& InputName, const FString& ExpressionUid, const FString& OutputName)
+	{
+		return ConnectOuputToInputByName(InterchangeNode, InputName, ExpressionUid, OutputName);
+	}
+
+	/**
+	 * Adds an input connection attribute.
+	 * @param InterchangeNode	The Node to create the input on.
+	 * @param InputName			The name to give to the input.
+	 * @param ExpressionUid		The unique id of the node to connect to the input.
+	 * @param OutputIndex		The index of the output from ExpressionUid to connect to the input.
+	 * @return					true if the input connection was successfully added to the node.
+	 * OutputIndex is encoded in a string in the following pattern  ExpressionUid:OutputByIndex:FString::FromInt(OutputIndex)
+	 * The index should be retrieved using UInterchangeShaderPortsAPI::GetOutputIndexFromName
+	 */
+	UFUNCTION(BlueprintCallable, Category = "Interchange | Node | Material")
+	static bool ConnectOuputToInputByIndex(UInterchangeBaseNode* InterchangeNode, const FString& InputName, const FString& ExpressionUid, int32 OutputIndex);
  
 	/**
-	 * Retrieves the node unique id and the ouputname connected to a given input, if any.
+	 * Retrieves the node unique id and the output name connected to a given input, if any.
 	 */
 	UFUNCTION(BlueprintCallable, Category = "Interchange | Node | Material")
 	static bool GetInputConnection(const UInterchangeBaseNode* InterchangeNode, const FString& InputName, FString& OutExpressionUid, FString& OutputName);
-	
+
 	/**
 	 * For an input with a value, returns the type of the stored value.
 	 */
 	static UE::Interchange::EAttributeTypes GetInputType(const UInterchangeBaseNode* InterchangeNode, const FString& InputName);
- 
+
+	/**
+	 * Returns INDEX_NONE if OutputName is not an index.
+	 */
+	static int32 GetOutputIndexFromName(const FString& OutputName);
+
 private:
 	static const TCHAR* InputPrefix;
 	static const TCHAR* InputSeparator;
+	static const TCHAR* OutputByIndex;
 };
 
 /**
diff --git a/Engine/Plugins/Interchange/Runtime/Source/Pipelines/Private/InterchangeGenericMaterialPipeline.cpp b/Engine/Plugins/Interchange/Runtime/Source/Pipelines/Private/InterchangeGenericMaterialPipeline.cpp
--- a/Engine/Plugins/Interchange/Runtime/Source/Pipelines/Private/InterchangeGenericMaterialPipeline.cpp
+++ b/Engine/Plugins/Interchange/Runtime/Source/Pipelines/Private/InterchangeGenericMaterialPipeline.cpp
@@ -373,6 +373,41 @@ namespace UE::Interchange::InterchangeGenericMaterialPipeline::Private
 			}
 		}
 	}
+
+	void UpdateFunctionCallExpression(UInterchangeMaterialExpressionFactoryNode& FunctionCallExpression, const FString& MaterialFunctionPath)
+	{
+		const FName MaterialFunctionMemberName = GET_MEMBER_NAME_CHECKED(UMaterialExpressionMaterialFunctionCall, MaterialFunction);
+		UClass* CustomExpressionClass = UMaterialExpressionMaterialFunctionCall::StaticClass();
+
+		FunctionCallExpression.SetCustomExpressionClassName(CustomExpressionClass->GetName());
+		
+		FunctionCallExpression.AddStringAttribute(MaterialFunctionMemberName, MaterialFunctionPath);
+		FunctionCallExpression.AddApplyAndFillDelegates<FString>(MaterialFunctionMemberName, CustomExpressionClass, MaterialFunctionMemberName);
+	}
+
+	UInterchangeMaterialExpressionFactoryNode* CreateExpressionWithMaterialFunction(
+		UInterchangeBaseNodeContainer* BaseNodeContainer,
+		UInterchangeMaterialFactoryNode* MaterialFactoryNode,
+		const FString& Label,
+		const FString& MaterialFunctionPath)
+	{
+		const FString UniqueId = MaterialFactoryNode->GetUniqueID() + Label;
+		
+		UInterchangeMaterialExpressionFactoryNode* Expression = NewObject<UInterchangeMaterialExpressionFactoryNode>(BaseNodeContainer, NAME_None);
+		if (!Expression)
+		{
+			return nullptr;
+		}
+
+		Expression->InitializeNode(UniqueId, Label, EInterchangeNodeContainerType::FactoryData);
+
+		BaseNodeContainer->AddNode(Expression);
+		BaseNodeContainer->SetNodeParentUid(UniqueId, MaterialFactoryNode->GetUniqueID());
+
+		UpdateFunctionCallExpression(*Expression, MaterialFunctionPath);
+		
+		return Expression;
+	}
 }
 
 UInterchangeGenericMaterialPipeline::UInterchangeGenericMaterialPipeline()
@@ -696,25 +731,15 @@ bool UInterchangeGenericMaterialPipeline::IsSurfaceUnlitModel(const UInterchange
 bool UInterchangeGenericMaterialPipeline::HandlePhongModel(const UInterchangeShaderGraphNode* ShaderGraphNode, UInterchangeMaterialFactoryNode* MaterialFactoryNode)
 {
 	using namespace UE::Interchange::Materials::Phong;
+	using namespace UE::Interchange::InterchangeGenericMaterialPipeline::Private;
 
 	if (IsPhongModel(ShaderGraphNode))
 	{
 		// ConvertFromDiffSpec function call
-		UInterchangeMaterialExpressionFactoryNode* FunctionCallExpression = NewObject<UInterchangeMaterialExpressionFactoryNode>(BaseNodeContainer, NAME_None);
-		FunctionCallExpression->SetCustomExpressionClassName(UMaterialExpressionMaterialFunctionCall::StaticClass()->GetName());
-		FString FunctionCallExpressionUid = MaterialFactoryNode->GetUniqueID() + TEXT("\\Inputs\\BaseColor\\DiffSpecFunc");
-		FunctionCallExpression->InitializeNode(FunctionCallExpressionUid, TEXT("DiffSpecFunc"), EInterchangeNodeContainerType::FactoryData);
-
-		BaseNodeContainer->AddNode(FunctionCallExpression);
-		BaseNodeContainer->SetNodeParentUid(FunctionCallExpressionUid, MaterialFactoryNode->GetUniqueID());
-
-		const FName MaterialFunctionMemberName = GET_MEMBER_NAME_CHECKED(UMaterialExpressionMaterialFunctionCall, MaterialFunction);
-
-		FunctionCallExpression->AddStringAttribute(
-			MaterialFunctionMemberName,
-			TEXT("MaterialFunction'/Engine/Functions/Engine_MaterialFunctions01/Shading/ConvertFromDiffSpec.ConvertFromDiffSpec'"));
-		FunctionCallExpression->AddApplyAndFillDelegates<FString>(MaterialFunctionMemberName, UMaterialExpressionMaterialFunctionCall::StaticClass(), MaterialFunctionMemberName);
+		const FString MaterialFunctionPath = TEXT("MaterialFunction'/Engine/Functions/Engine_MaterialFunctions01/Shading/ConvertFromDiffSpec.ConvertFromDiffSpec'");
+		UInterchangeMaterialExpressionFactoryNode* FunctionCallExpression = CreateExpressionWithMaterialFunction(BaseNodeContainer, MaterialFactoryNode, TEXT("DiffSpecFunc"), MaterialFunctionPath);
 
+		const FString FunctionCallExpressionUid = FunctionCallExpression->GetUniqueID();
 		MaterialFactoryNode->ConnectOutputToBaseColor(FunctionCallExpressionUid, UE::Interchange::Materials::PBR::Parameters::BaseColor.ToString());
 		MaterialFactoryNode->ConnectOutputToMetallic(FunctionCallExpressionUid, UE::Interchange::Materials::PBR::Parameters::Metallic.ToString());
 		MaterialFactoryNode->ConnectOutputToSpecular(FunctionCallExpressionUid, UE::Interchange::Materials::PBR::Parameters::Specular.ToString());
@@ -726,7 +751,7 @@ bool UInterchangeGenericMaterialPipeline::HandlePhongModel(const UInterchangeSha
 
 			if (DiffuseExpressionFactoryNode.Get<0>())
 			{
-				UInterchangeShaderPortsAPI::ConnectOuputToInput(FunctionCallExpression, Parameters::DiffuseColor.ToString(),
+				UInterchangeShaderPortsAPI::ConnectOuputToInputByName(FunctionCallExpression, Parameters::DiffuseColor.ToString(),
 					DiffuseExpressionFactoryNode.Get<0>()->GetUniqueID(), DiffuseExpressionFactoryNode.Get<1>());
 			}
 		}
@@ -738,7 +763,7 @@ bool UInterchangeGenericMaterialPipeline::HandlePhongModel(const UInterchangeSha
 
 			if (SpecularExpressionFactoryNode.Get<0>())
 			{
-				UInterchangeShaderPortsAPI::ConnectOuputToInput(FunctionCallExpression, Parameters::SpecularColor.ToString(),
+				UInterchangeShaderPortsAPI::ConnectOuputToInputByName(FunctionCallExpression, Parameters::SpecularColor.ToString(),
 					SpecularExpressionFactoryNode.Get<0>()->GetUniqueID(), SpecularExpressionFactoryNode.Get<1>());
 			}
 		}
@@ -764,7 +789,7 @@ bool UInterchangeGenericMaterialPipeline::HandlePhongModel(const UInterchangeSha
 					DivideShininessNode->AddApplyAndFillDelegates<float>(ShininessScaleParameterName, UMaterialExpressionDivide::StaticClass(),  GET_MEMBER_NAME_CHECKED(UMaterialExpressionDivide, ConstB));
 
 					// Connect Shininess to Divide
-					UInterchangeShaderPortsAPI::ConnectOuputToInput(DivideShininessNode, GET_MEMBER_NAME_CHECKED(UMaterialExpressionDivide, A).ToString(),
+					UInterchangeShaderPortsAPI::ConnectOuputToInputByName(DivideShininessNode, GET_MEMBER_NAME_CHECKED(UMaterialExpressionDivide, A).ToString(),
 						ShininessExpressionFactoryNode.Get<0>()->GetUniqueID(), ShininessExpressionFactoryNode.Get<1>());
 
 					UInterchangeMaterialExpressionFactoryNode* InverseShininessNode =
@@ -894,32 +919,22 @@ bool UInterchangeGenericMaterialPipeline::HandlePBRModel(const UInterchangeShade
 
 FString UInterchangeGenericMaterialPipeline::HandleSurfaceShaderModel(const UInterchangeShaderGraphNode* ShaderGraphNode, UInterchangeMaterialFactoryNode* MaterialFactoryNode, const FSurfaceShader& SurfaceShader)
 {
-	UInterchangeMaterialExpressionFactoryNode* FunctionCallExpression = NewObject<UInterchangeMaterialExpressionFactoryNode>(BaseNodeContainer, NAME_None);
-	FunctionCallExpression->SetCustomExpressionClassName(UMaterialExpressionMaterialFunctionCall::StaticClass()->GetName());
-	FString FunctionCallExpressionUid = MaterialFactoryNode->GetUniqueID() + SurfaceShader.UniqueID;
-	FunctionCallExpression->InitializeNode(FunctionCallExpressionUid, SurfaceShader.UniqueID, EInterchangeNodeContainerType::FactoryData);
-
-	BaseNodeContainer->AddNode(FunctionCallExpression);
-	BaseNodeContainer->SetNodeParentUid(FunctionCallExpressionUid, MaterialFactoryNode->GetUniqueID());
+	using namespace UE::Interchange::InterchangeGenericMaterialPipeline::Private;
 
-	const FName MaterialFunctionMemberName = GET_MEMBER_NAME_CHECKED(UMaterialExpressionMaterialFunctionCall, MaterialFunction);
+	UInterchangeMaterialExpressionFactoryNode* FunctionCallExpression = CreateExpressionWithMaterialFunction(BaseNodeContainer, MaterialFactoryNode, SurfaceShader.UniqueID, SurfaceShader.Path);
 
-
-	FunctionCallExpression->AddStringAttribute(
-		MaterialFunctionMemberName,
-		SurfaceShader.Path);
-	FunctionCallExpression->AddApplyAndFillDelegates<FString>(MaterialFunctionMemberName, UMaterialExpressionMaterialFunctionCall::StaticClass(), MaterialFunctionMemberName);
+	const FString FunctionCallExpressionUid = FunctionCallExpression->GetUniqueID();
 
 	auto ConnectMaterialExpressionOutputToInput = [&](const FString& InputName, EMaterialInputType InputType)
 	{
 		TGuardValue<EMaterialInputType> InputTypeBeingProcessedGuard(MaterialCreationContext.InputTypeBeingProcessed, InputType);
 
 		TTuple<UInterchangeMaterialExpressionFactoryNode*, FString> ExpressionFactoryNode =
-			CreateMaterialExpressionForInput(MaterialFactoryNode, ShaderGraphNode, InputName, FunctionCallExpression->GetUniqueID());
+			CreateMaterialExpressionForInput(MaterialFactoryNode, ShaderGraphNode, InputName, FunctionCallExpressionUid);
 
 		if(ExpressionFactoryNode.Get<0>())
 		{
-			UInterchangeShaderPortsAPI::ConnectOuputToInput(FunctionCallExpression, InputName,
+			UInterchangeShaderPortsAPI::ConnectOuputToInputByName(FunctionCallExpression, InputName,
 															ExpressionFactoryNode.Get<0>()->GetUniqueID(), ExpressionFactoryNode.Get<1>());
 		}
 	};
@@ -1241,7 +1256,7 @@ bool UInterchangeGenericMaterialPipeline::HandleSheen(const UInterchangeShaderGr
 				UInterchangeMaterialExpressionFactoryNode* InverseSheenRoughnessNode =
 						CreateExpressionNode(TEXT("InverseSheenRoughness"), ExpressionFactoryNode.Get<0>()->GetUniqueID(), UMaterialExpressionOneMinus::StaticClass());
 
-				UInterchangeShaderPortsAPI::ConnectOuputToInput(InverseSheenRoughnessNode, GET_MEMBER_NAME_CHECKED(UMaterialExpressionOneMinus, Input).ToString(),
+				UInterchangeShaderPortsAPI::ConnectOuputToInputByName(InverseSheenRoughnessNode, GET_MEMBER_NAME_CHECKED(UMaterialExpressionOneMinus, Input).ToString(),
 					ExpressionFactoryNode.Get<0>()->GetUniqueID(), ExpressionFactoryNode.Get<1>());
 
 				MaterialFactoryNode->ConnectToCloth(InverseSheenRoughnessNode->GetUniqueID());
@@ -1411,7 +1426,7 @@ void UInterchangeGenericMaterialPipeline::HandleCommonParameters(const UIntercha
 				IORLerp->AddFloatAttribute(ConstAMemberName, AirIOR);
 				IORLerp->AddApplyAndFillDelegates<float>(ConstAMemberName, UMaterialExpressionLinearInterpolate::StaticClass(), ConstAMemberName);
 
-				UInterchangeShaderPortsAPI::ConnectOuputToInput(IORLerp, GET_MEMBER_NAME_CHECKED(UMaterialExpressionLinearInterpolate, B).ToString(),
+				UInterchangeShaderPortsAPI::ConnectOuputToInputByName(IORLerp, GET_MEMBER_NAME_CHECKED(UMaterialExpressionLinearInterpolate, B).ToString(),
 					ExpressionFactoryNode.Get<0>()->GetUniqueID(), ExpressionFactoryNode.Get<1>());
 
 				UInterchangeMaterialExpressionFactoryNode* IORFresnel = CreateExpressionNode(TEXT("IORFresnel"), ShaderGraphNode->GetUniqueID(), UMaterialExpressionFresnel::StaticClass());
@@ -1428,12 +1443,15 @@ void UInterchangeGenericMaterialPipeline::HandleFlattenNormalNode(const UInterch
 	UInterchangeMaterialExpressionFactoryNode* FlattenNormalFactoryNode)
 {
 	using namespace UE::Interchange::Materials::Standard::Nodes::FlattenNormal;
+	using namespace UE::Interchange::InterchangeGenericMaterialPipeline::Private;
 
-	FlattenNormalFactoryNode->SetCustomExpressionClassName(UMaterialExpressionMaterialFunctionCall::StaticClass()->GetName());
+	if (!FlattenNormalFactoryNode)
+	{
+		return;
+	}
 
-	const FName MaterialFunctionMemberName = GET_MEMBER_NAME_CHECKED(UMaterialExpressionMaterialFunctionCall, MaterialFunction);
-	FlattenNormalFactoryNode->AddStringAttribute(MaterialFunctionMemberName, TEXT("/Engine/Functions/Engine_MaterialFunctions01/Texturing/FlattenNormal.FlattenNormal"));
-	FlattenNormalFactoryNode->AddApplyAndFillDelegates<FString>(MaterialFunctionMemberName, UMaterialExpressionMaterialFunctionCall::StaticClass(), MaterialFunctionMemberName);
+	const FString MaterialFunctionPath = TEXT("/Engine/Functions/Engine_MaterialFunctions01/Texturing/FlattenNormal.FlattenNormal");
+	UpdateFunctionCallExpression(*FlattenNormalFactoryNode, MaterialFunctionPath);
 
 	// Normal
 	{
@@ -1442,7 +1460,7 @@ void UInterchangeGenericMaterialPipeline::HandleFlattenNormalNode(const UInterch
 
 		if (NormalExpression.Get<0>())
 		{
-			UInterchangeShaderPortsAPI::ConnectOuputToInput(FlattenNormalFactoryNode, TEXT("Normal"),
+			UInterchangeShaderPortsAPI::ConnectOuputToInputByName(FlattenNormalFactoryNode, TEXT("Normal"),
 				NormalExpression.Get<0>()->GetUniqueID(), NormalExpression.Get<1>());
 		}
 	}
@@ -1454,7 +1472,7 @@ void UInterchangeGenericMaterialPipeline::HandleFlattenNormalNode(const UInterch
 
 		if (FlatnessExpression.Get<0>())
 		{
-			UInterchangeShaderPortsAPI::ConnectOuputToInput(FlattenNormalFactoryNode, TEXT("Flatness"),
+			UInterchangeShaderPortsAPI::ConnectOuputToInputByName(FlattenNormalFactoryNode, TEXT("Flatness"),
 				FlatnessExpression.Get<0>()->GetUniqueID(), FlatnessExpression.Get<1>());
 		}
 	}
@@ -1463,12 +1481,15 @@ void UInterchangeGenericMaterialPipeline::HandleFlattenNormalNode(const UInterch
 void UInterchangeGenericMaterialPipeline::HandleNormalFromHeightMapNode(const UInterchangeShaderNode* ShaderNode, UInterchangeBaseMaterialFactoryNode* MaterialFactoryNode, UInterchangeMaterialExpressionFactoryNode* NormalFromHeightMapFactoryNode)
 {
 	using namespace UE::Interchange::Materials::Standard::Nodes::NormalFromHeightMap;
+	using namespace UE::Interchange::InterchangeGenericMaterialPipeline::Private;
 
-	NormalFromHeightMapFactoryNode->SetCustomExpressionClassName(UMaterialExpressionMaterialFunctionCall::StaticClass()->GetName());
+	if (!NormalFromHeightMapFactoryNode)
+	{
+		return;
+	}
 
-	const FName MaterialFunctionMemberName = GET_MEMBER_NAME_CHECKED(UMaterialExpressionMaterialFunctionCall, MaterialFunction);
-	NormalFromHeightMapFactoryNode->AddStringAttribute(MaterialFunctionMemberName, TEXT("/Engine/Functions/Engine_MaterialFunctions03/Procedurals/NormalFromHeightmap.NormalFromHeightmap"));
-	NormalFromHeightMapFactoryNode->AddApplyAndFillDelegates<FString>(MaterialFunctionMemberName, UMaterialExpressionMaterialFunctionCall::StaticClass(), MaterialFunctionMemberName);
+	const FString MaterialFunctionPath = TEXT("/Engine/Functions/Engine_MaterialFunctions03/Procedurals/NormalFromHeightmap.NormalFromHeightmap");
+	UpdateFunctionCallExpression(*NormalFromHeightMapFactoryNode, MaterialFunctionPath);
 
 	// Heightmap
 	{
@@ -1478,7 +1499,7 @@ void UInterchangeGenericMaterialPipeline::HandleNormalFromHeightMapNode(const UI
 
 		if(HeightMapExpression.Get<0>())
 		{
-			UInterchangeShaderPortsAPI::ConnectOuputToInput(NormalFromHeightMapFactoryNode, HeightMapInput, HeightMapExpression.Get<0>()->GetUniqueID(), HeightMapExpression.Get<1>());
+			UInterchangeShaderPortsAPI::ConnectOuputToInputByName(NormalFromHeightMapFactoryNode, HeightMapInput, HeightMapExpression.Get<0>()->GetUniqueID(), HeightMapExpression.Get<1>());
 		}
 	}
 
@@ -1490,7 +1511,7 @@ void UInterchangeGenericMaterialPipeline::HandleNormalFromHeightMapNode(const UI
 
 		if(IntensityExpression.Get<0>())
 		{
-			UInterchangeShaderPortsAPI::ConnectOuputToInput(NormalFromHeightMapFactoryNode, Intensity, IntensityExpression.Get<0>()->GetUniqueID(), IntensityExpression.Get<1>());
+			UInterchangeShaderPortsAPI::ConnectOuputToInputByName(NormalFromHeightMapFactoryNode, Intensity, IntensityExpression.Get<0>()->GetUniqueID(), IntensityExpression.Get<1>());
 		}
 	}
 
@@ -1502,7 +1523,7 @@ void UInterchangeGenericMaterialPipeline::HandleNormalFromHeightMapNode(const UI
 
 		if(OffsetExpression.Get<0>())
 		{
-			UInterchangeShaderPortsAPI::ConnectOuputToInput(NormalFromHeightMapFactoryNode, Offset, OffsetExpression.Get<0>()->GetUniqueID(), OffsetExpression.Get<1>());
+			UInterchangeShaderPortsAPI::ConnectOuputToInputByName(NormalFromHeightMapFactoryNode, Offset, OffsetExpression.Get<0>()->GetUniqueID(), OffsetExpression.Get<1>());
 		}
 	}
 
@@ -1514,7 +1535,7 @@ void UInterchangeGenericMaterialPipeline::HandleNormalFromHeightMapNode(const UI
 
 		if(CoordinatesExpression.Get<0>())
 		{
-			UInterchangeShaderPortsAPI::ConnectOuputToInput(NormalFromHeightMapFactoryNode, Coordinates, CoordinatesExpression.Get<0>()->GetUniqueID(), CoordinatesExpression.Get<1>());
+			UInterchangeShaderPortsAPI::ConnectOuputToInputByName(NormalFromHeightMapFactoryNode, Coordinates, CoordinatesExpression.Get<0>()->GetUniqueID(), CoordinatesExpression.Get<1>());
 		}
 	}
 
@@ -1526,7 +1547,7 @@ void UInterchangeGenericMaterialPipeline::HandleNormalFromHeightMapNode(const UI
 
 		if(ChannelExpression.Get<0>())
 		{
-			UInterchangeShaderPortsAPI::ConnectOuputToInput(NormalFromHeightMapFactoryNode, Channel, ChannelExpression.Get<0>()->GetUniqueID(), ChannelExpression.Get<1>());
+			UInterchangeShaderPortsAPI::ConnectOuputToInputByName(NormalFromHeightMapFactoryNode, Channel, ChannelExpression.Get<0>()->GetUniqueID(), ChannelExpression.Get<1>());
 		}
 	}
 }
@@ -1535,18 +1556,21 @@ void UInterchangeGenericMaterialPipeline::HandleMakeFloat3Node(const UInterchang
 	UInterchangeMaterialExpressionFactoryNode* MakeFloat3FactoryNode)
 {
 	using namespace UE::Interchange::Materials::Standard::Nodes::MakeFloat3;
+	using namespace UE::Interchange::InterchangeGenericMaterialPipeline::Private;
 
-	MakeFloat3FactoryNode->SetCustomExpressionClassName(UMaterialExpressionMaterialFunctionCall::StaticClass()->GetName());
+	if (!MakeFloat3FactoryNode)
+	{
+		return;
+	}
 
-	const FName MaterialFunctionMemberName = GET_MEMBER_NAME_CHECKED(UMaterialExpressionMaterialFunctionCall, MaterialFunction);
-	MakeFloat3FactoryNode->AddStringAttribute(MaterialFunctionMemberName, TEXT("/Engine/Functions/Engine_MaterialFunctions02/Utility/MakeFloat3.MakeFloat3"));
-	MakeFloat3FactoryNode->AddApplyAndFillDelegates<FString>(MaterialFunctionMemberName, UMaterialExpressionMaterialFunctionCall::StaticClass(), MaterialFunctionMemberName);
+	const FString MaterialFunctionPath = TEXT("/Engine/Functions/Engine_MaterialFunctions02/Utility/MakeFloat3.MakeFloat3");
+	UpdateFunctionCallExpression(*MakeFloat3FactoryNode, MaterialFunctionPath);
 
 	TTuple<UInterchangeMaterialExpressionFactoryNode*, FString> RedChannelExpression =
 		CreateMaterialExpressionForInput(MaterialFactoryNode, ShaderNode, Inputs::X.ToString(), MakeFloat3FactoryNode->GetUniqueID());
 	if (RedChannelExpression.Get<0>())
 	{
-		UInterchangeShaderPortsAPI::ConnectOuputToInput(MakeFloat3FactoryNode, TEXT("X"),
+		UInterchangeShaderPortsAPI::ConnectOuputToInputByName(MakeFloat3FactoryNode, TEXT("X"),
 			RedChannelExpression.Get<0>()->GetUniqueID(), RedChannelExpression.Get<1>());
 	}
 
@@ -1554,7 +1578,7 @@ void UInterchangeGenericMaterialPipeline::HandleMakeFloat3Node(const UInterchang
 		CreateMaterialExpressionForInput(MaterialFactoryNode, ShaderNode, Inputs::Y.ToString(), MakeFloat3FactoryNode->GetUniqueID());
 	if (GreenChannelExpression.Get<0>())
 	{
-		UInterchangeShaderPortsAPI::ConnectOuputToInput(MakeFloat3FactoryNode, TEXT("Y"),
+		UInterchangeShaderPortsAPI::ConnectOuputToInputByName(MakeFloat3FactoryNode, TEXT("Y"),
 			GreenChannelExpression.Get<0>()->GetUniqueID(), GreenChannelExpression.Get<1>());
 	}
 
@@ -1562,7 +1586,7 @@ void UInterchangeGenericMaterialPipeline::HandleMakeFloat3Node(const UInterchang
 		CreateMaterialExpressionForInput(MaterialFactoryNode, ShaderNode, Inputs::Z.ToString(), MakeFloat3FactoryNode->GetUniqueID());
 	if (BlueChannelExpression.Get<0>())
 	{
-		UInterchangeShaderPortsAPI::ConnectOuputToInput(MakeFloat3FactoryNode, TEXT("Z"),
+		UInterchangeShaderPortsAPI::ConnectOuputToInputByName(MakeFloat3FactoryNode, TEXT("Z"),
 			BlueChannelExpression.Get<0>()->GetUniqueID(), BlueChannelExpression.Get<1>());
 	}
 }
@@ -1699,7 +1723,7 @@ void UInterchangeGenericMaterialPipeline::HandleTextureSampleNode(const UInterch
 
 		if (CoordinatesExpression.Get<0>())
 		{
-			UInterchangeShaderPortsAPI::ConnectOuputToInput(TextureSampleFactoryNode, GET_MEMBER_NAME_CHECKED(UMaterialExpressionTextureSample, Coordinates).ToString(),
+			UInterchangeShaderPortsAPI::ConnectOuputToInputByName(TextureSampleFactoryNode, GET_MEMBER_NAME_CHECKED(UMaterialExpressionTextureSample, Coordinates).ToString(),
 				CoordinatesExpression.Get<0>()->GetUniqueID(), CoordinatesExpression.Get<1>());
 		}
 	}
@@ -1751,6 +1775,7 @@ void UInterchangeGenericMaterialPipeline::HandleTextureCoordinateNode(const UInt
 	UInterchangeMaterialExpressionFactoryNode*& TexCoordFactoryNode)
 {
 	using namespace UE::Interchange::Materials::Standard;
+	using namespace UE::Interchange::InterchangeGenericMaterialPipeline::Private;
 
 	TexCoordFactoryNode->SetCustomExpressionClassName(UMaterialExpressionTextureCoordinate::StaticClass()->GetName());
 
@@ -1801,7 +1826,7 @@ void UInterchangeGenericMaterialPipeline::HandleTextureCoordinateNode(const UInt
 
 			UInterchangeShaderPortsAPI::ConnectDefaultOuputToInput(MultiplyExpression, GET_MEMBER_NAME_CHECKED(UMaterialExpressionMultiply, A).ToString(),
 				TexCoordFactoryNode->GetUniqueID());
-			UInterchangeShaderPortsAPI::ConnectOuputToInput(MultiplyExpression, GET_MEMBER_NAME_CHECKED(UMaterialExpressionMultiply, B).ToString(),
+			UInterchangeShaderPortsAPI::ConnectOuputToInputByName(MultiplyExpression, GET_MEMBER_NAME_CHECKED(UMaterialExpressionMultiply, B).ToString(),
 				ScaleExpression.Get<0>()->GetUniqueID(), ScaleExpression.Get<1>());
 
 			TexCoordFactoryNode = MultiplyExpression;
@@ -1818,19 +1843,18 @@ void UInterchangeGenericMaterialPipeline::HandleTextureCoordinateNode(const UInt
 			UInterchangeMaterialExpressionFactoryNode* CallRotatorExpression =
 				CreateExpressionNode(RotateExpression.Get<0>()->GetDisplayLabel() + TEXT("_Rotator"), TexCoordFactoryNode->GetUniqueID(), UMaterialExpressionMaterialFunctionCall::StaticClass());
 
-			const FName MaterialFunctionMemberName = GET_MEMBER_NAME_CHECKED(UMaterialExpressionMaterialFunctionCall, MaterialFunction);
-			CallRotatorExpression->AddStringAttribute(MaterialFunctionMemberName, TEXT("/Engine/Functions/Engine_MaterialFunctions02/Texturing/CustomRotator.CustomRotator"));
-			CallRotatorExpression->AddApplyAndFillDelegates<FString>(MaterialFunctionMemberName, UMaterialExpressionMaterialFunctionCall::StaticClass(), MaterialFunctionMemberName);
+			const FString MaterialFunctionPath = TEXT("/Engine/Functions/Engine_MaterialFunctions02/Texturing/CustomRotator.CustomRotator");
+			UpdateFunctionCallExpression(*CallRotatorExpression, MaterialFunctionPath);
 
 			UInterchangeShaderPortsAPI::ConnectDefaultOuputToInput(CallRotatorExpression, TEXT("UVs"), TexCoordFactoryNode->GetUniqueID());
-			UInterchangeShaderPortsAPI::ConnectOuputToInput(CallRotatorExpression, TEXT("Rotation Angle (0-1)"), RotateExpression.Get<0>()->GetUniqueID(), RotateExpression.Get<1>());
+			UInterchangeShaderPortsAPI::ConnectOuputToInputByName(CallRotatorExpression, TEXT("Rotation Angle (0-1)"), RotateExpression.Get<0>()->GetUniqueID(), RotateExpression.Get<1>());
 
 			TTuple<UInterchangeMaterialExpressionFactoryNode*, FString> RotationCenterExpression =
 				CreateMaterialExpressionForInput(MaterialFactoryNode, ShaderNode, Nodes::TextureCoordinate::Inputs::RotationCenter.ToString(), TexCoordFactoryNode->GetUniqueID());
 
 			if (RotationCenterExpression.Get<0>())
 			{
-				UInterchangeShaderPortsAPI::ConnectOuputToInput(CallRotatorExpression, TEXT("Rotation Center"), RotationCenterExpression.Get<0>()->GetUniqueID(), RotationCenterExpression.Get<1>());
+				UInterchangeShaderPortsAPI::ConnectOuputToInputByName(CallRotatorExpression, TEXT("Rotation Center"), RotationCenterExpression.Get<0>()->GetUniqueID(), RotationCenterExpression.Get<1>());
 			}
 
 			TexCoordFactoryNode = CallRotatorExpression;
@@ -1849,7 +1873,7 @@ void UInterchangeGenericMaterialPipeline::HandleTextureCoordinateNode(const UInt
 
 			UInterchangeShaderPortsAPI::ConnectDefaultOuputToInput(AddExpression, GET_MEMBER_NAME_CHECKED(UMaterialExpressionAdd, A).ToString(),
 				TexCoordFactoryNode->GetUniqueID());
-			UInterchangeShaderPortsAPI::ConnectOuputToInput(AddExpression, GET_MEMBER_NAME_CHECKED(UMaterialExpressionAdd, B).ToString(),
+			UInterchangeShaderPortsAPI::ConnectOuputToInputByName(AddExpression, GET_MEMBER_NAME_CHECKED(UMaterialExpressionAdd, B).ToString(),
 				OffsetExpression.Get<0>()->GetUniqueID(), OffsetExpression.Get<1>());
 
 			TexCoordFactoryNode = AddExpression;
@@ -1870,7 +1894,7 @@ void UInterchangeGenericMaterialPipeline::HandleLerpNode(const UInterchangeShade
 
 		if (ColorAExpression.Get<0>())
 		{
-			UInterchangeShaderPortsAPI::ConnectOuputToInput(LerpFactoryNode, GET_MEMBER_NAME_CHECKED(UMaterialExpressionLinearInterpolate, A).ToString(),
+			UInterchangeShaderPortsAPI::ConnectOuputToInputByName(LerpFactoryNode, GET_MEMBER_NAME_CHECKED(UMaterialExpressionLinearInterpolate, A).ToString(),
 				ColorAExpression.Get<0>()->GetUniqueID(), ColorAExpression.Get<1>());
 		}
 	}
@@ -1882,7 +1906,7 @@ void UInterchangeGenericMaterialPipeline::HandleLerpNode(const UInterchangeShade
 
 		if (ColorBExpression.Get<0>())
 		{
-			UInterchangeShaderPortsAPI::ConnectOuputToInput(LerpFactoryNode, GET_MEMBER_NAME_CHECKED(UMaterialExpressionLinearInterpolate, B).ToString(),
+			UInterchangeShaderPortsAPI::ConnectOuputToInputByName(LerpFactoryNode, GET_MEMBER_NAME_CHECKED(UMaterialExpressionLinearInterpolate, B).ToString(),
 				ColorBExpression.Get<0>()->GetUniqueID(), ColorBExpression.Get<1>());
 		}
 	}
@@ -1894,7 +1918,7 @@ void UInterchangeGenericMaterialPipeline::HandleLerpNode(const UInterchangeShade
 
 		if (FactorExpression.Get<0>())
 		{
-			UInterchangeShaderPortsAPI::ConnectOuputToInput(LerpFactoryNode, GET_MEMBER_NAME_CHECKED(UMaterialExpressionLinearInterpolate, Alpha).ToString(),
+			UInterchangeShaderPortsAPI::ConnectOuputToInputByName(LerpFactoryNode, GET_MEMBER_NAME_CHECKED(UMaterialExpressionLinearInterpolate, Alpha).ToString(),
 				FactorExpression.Get<0>()->GetUniqueID(), FactorExpression.Get<1>());
 		}
 	}
@@ -1954,7 +1978,7 @@ void UInterchangeGenericMaterialPipeline::HandleMaskNode(const UInterchangeShade
 
 		if(InputExpression.Get<0>())
 		{
-			UInterchangeShaderPortsAPI::ConnectOuputToInput(MaskFactoryNode, GET_MEMBER_NAME_CHECKED(UMaterialExpressionComponentMask, Input).ToString(),
+			UInterchangeShaderPortsAPI::ConnectOuputToInputByName(MaskFactoryNode, GET_MEMBER_NAME_CHECKED(UMaterialExpressionComponentMask, Input).ToString(),
 				InputExpression.Get<0>()->GetUniqueID(), InputExpression.Get<1>());
 		}
 	}
@@ -2004,7 +2028,7 @@ void UInterchangeGenericMaterialPipeline::HandleTransformPositionNode(const UInt
 
 		if(InputExpression.Get<0>())
 		{
-			UInterchangeShaderPortsAPI::ConnectOuputToInput(TransformPositionFactoryNode, GET_MEMBER_NAME_CHECKED(UMaterialExpressionTransformPosition, Input).ToString(),
+			UInterchangeShaderPortsAPI::ConnectOuputToInputByName(TransformPositionFactoryNode, GET_MEMBER_NAME_CHECKED(UMaterialExpressionTransformPosition, Input).ToString(),
 															InputExpression.Get<0>()->GetUniqueID(), InputExpression.Get<1>());
 		}
 	}
@@ -2039,7 +2063,7 @@ void UInterchangeGenericMaterialPipeline::HandleTransformVectorNode(const UInter
 
 		if(InputExpression.Get<0>())
 		{
-			UInterchangeShaderPortsAPI::ConnectOuputToInput(TransformVectorFactoryNode, GET_MEMBER_NAME_CHECKED(UMaterialExpressionTransform, Input).ToString(),
+			UInterchangeShaderPortsAPI::ConnectOuputToInputByName(TransformVectorFactoryNode, GET_MEMBER_NAME_CHECKED(UMaterialExpressionTransform, Input).ToString(),
 															InputExpression.Get<0>()->GetUniqueID(), InputExpression.Get<1>());
 		}
 	}
@@ -2074,7 +2098,7 @@ void UInterchangeGenericMaterialPipeline::HandleNoiseNode(const UInterchangeShad
 
 		if(PositionExpression.Get<0>())
 		{
-			UInterchangeShaderPortsAPI::ConnectOuputToInput(NoiseFactoryNode, GET_MEMBER_NAME_CHECKED(UMaterialExpressionNoise, Position).ToString(),
+			UInterchangeShaderPortsAPI::ConnectOuputToInputByName(NoiseFactoryNode, GET_MEMBER_NAME_CHECKED(UMaterialExpressionNoise, Position).ToString(),
 															PositionExpression.Get<0>()->GetUniqueID(), PositionExpression.Get<1>());
 		}
 	}
@@ -2086,7 +2110,7 @@ void UInterchangeGenericMaterialPipeline::HandleNoiseNode(const UInterchangeShad
 
 		if(FilterWidthExpression.Get<0>())
 		{
-			UInterchangeShaderPortsAPI::ConnectOuputToInput(NoiseFactoryNode, GET_MEMBER_NAME_CHECKED(UMaterialExpressionNoise, FilterWidth).ToString(),
+			UInterchangeShaderPortsAPI::ConnectOuputToInputByName(NoiseFactoryNode, GET_MEMBER_NAME_CHECKED(UMaterialExpressionNoise, FilterWidth).ToString(),
 															FilterWidthExpression.Get<0>()->GetUniqueID(), FilterWidthExpression.Get<1>());
 		}
 	}
@@ -2185,7 +2209,7 @@ void UInterchangeGenericMaterialPipeline::HandleVectorNoiseNode(const UInterchan
 
 		if(PositionExpression.Get<0>())
 		{
-			UInterchangeShaderPortsAPI::ConnectOuputToInput(NoiseFactoryNode, GET_MEMBER_NAME_CHECKED(UMaterialExpressionVectorNoise, Position).ToString(),
+			UInterchangeShaderPortsAPI::ConnectOuputToInputByName(NoiseFactoryNode, GET_MEMBER_NAME_CHECKED(UMaterialExpressionVectorNoise, Position).ToString(),
 															PositionExpression.Get<0>()->GetUniqueID(), PositionExpression.Get<1>());
 		}
 	}
@@ -2242,23 +2266,17 @@ UInterchangeMaterialExpressionFactoryNode* UInterchangeGenericMaterialPipeline::
 		return MaterialExpression;
 	}
 
-	if (const UInterchangeFunctionCallShaderNode* FunctionCallShaderNode = Cast<UInterchangeFunctionCallShaderNode>(ShaderNode))
+	// Create function call expression if applicable
+	MaterialExpression = CreateFunctionCallExpression(ShaderNode, MaterialExpressionUid, MaterialFactoryNode);
+	if (MaterialExpression)
 	{
-		UInterchangeMaterialFunctionCallExpressionFactoryNode* MaterialFunctionCallExpression = NewObject<UInterchangeMaterialFunctionCallExpressionFactoryNode>(BaseNodeContainer);
-		
-		FString MaterialFunctionUid;
-		if (FunctionCallShaderNode->GetCustomMaterialFunction(MaterialFunctionUid))
-		{
-			const FString MaterialFunctionFactoryNodeUid = UInterchangeFactoryBaseNode::BuildFactoryNodeUid(MaterialFunctionUid);
-			MaterialFunctionCallExpression->SetCustomMaterialFunctionDependency(MaterialFunctionFactoryNodeUid);
-			MaterialFunctionCallExpression->AddFactoryDependencyUid(MaterialFunctionFactoryNodeUid);
-		}
-
-		MaterialExpression = MaterialFunctionCallExpression;
+		return MaterialExpression;
 	}
-	else
+
+	MaterialExpression = NewObject<UInterchangeMaterialExpressionFactoryNode>(BaseNodeContainer);
+	if (!MaterialExpression)
 	{
-		MaterialExpression = NewObject<UInterchangeMaterialExpressionFactoryNode>(BaseNodeContainer);
+		return nullptr;
 	}
 
 	FString ShaderType;
@@ -2319,7 +2337,7 @@ UInterchangeMaterialExpressionFactoryNode* UInterchangeGenericMaterialPipeline::
 	{
 		HandleVectorNoiseNode(ShaderNode, MaterialFactoryNode, MaterialExpression);
 	}
-	else
+	else if (ensure(!ShaderType.IsEmpty()))
 	{
 		const FString ExpressionClassName = TEXT("MaterialExpression") + ShaderType;
 		MaterialExpression->SetCustomExpressionClassName(ExpressionClassName);
@@ -2334,7 +2352,7 @@ UInterchangeMaterialExpressionFactoryNode* UInterchangeGenericMaterialPipeline::
 
 			if (InputExpression.Get<0>())
 			{
-				UInterchangeShaderPortsAPI::ConnectOuputToInput(MaterialExpression, InputName, InputExpression.Get<0>()->GetUniqueID(), InputExpression.Get<1>());
+				UInterchangeShaderPortsAPI::ConnectOuputToInputByName(MaterialExpression, InputName, InputExpression.Get<0>()->GetUniqueID(), InputExpression.Get<1>());
 			}
 		}
 	}
@@ -2905,7 +2923,7 @@ bool UInterchangeGenericMaterialPipeline::HandleBxDFInput(const UInterchangeShad
 
 	if (ExpressionFactoryNode.Get<0>())
 	{
-		UInterchangeShaderPortsAPI::ConnectOuputToInput(MaterialFactoryNode, Common::Parameters::BxDF.ToString(), ExpressionFactoryNode.Get<0>()->GetUniqueID(), ExpressionFactoryNode.Get<1>());
+		UInterchangeShaderPortsAPI::ConnectOuputToInputByName(MaterialFactoryNode, Common::Parameters::BxDF.ToString(), ExpressionFactoryNode.Get<0>()->GetUniqueID(), ExpressionFactoryNode.Get<1>());
 	}
 
 	// Make sure the bUseMaterialAttributes property of the material is set to true
@@ -2931,7 +2949,7 @@ UInterchangeMaterialFunctionFactoryNode* UInterchangeGenericMaterialPipeline::Cr
 
 		if (ExpressionFactoryNode.Get<0>())
 		{
-			UInterchangeShaderPortsAPI::ConnectOuputToInput(FactoryNode, InputName, ExpressionFactoryNode.Get<0>()->GetUniqueID(), ExpressionFactoryNode.Get<1>());
+			UInterchangeShaderPortsAPI::ConnectOuputToInputByName(FactoryNode, InputName, ExpressionFactoryNode.Get<0>()->GetUniqueID(), ExpressionFactoryNode.Get<1>());
 		}
 	}
 
@@ -2999,4 +3017,72 @@ bool UInterchangeGenericMaterialPipeline::HandleUnlitModel(const UInterchangeSha
 	return bShadingModelHandled;
 }
 
+UInterchangeMaterialExpressionFactoryNode* UInterchangeGenericMaterialPipeline::CreateFunctionCallExpression(const UInterchangeShaderNode* ShaderNode, const FString& MaterialExpressionUid, UInterchangeBaseMaterialFactoryNode* MaterialFactoryNode)
+{
+	using namespace UE::Interchange::InterchangeGenericMaterialPipeline::Private;
+
+	const UInterchangeFunctionCallShaderNode* FunctionCallShaderNode = Cast<UInterchangeFunctionCallShaderNode>(ShaderNode);
+	if (!FunctionCallShaderNode)
+	{
+		return nullptr;
+	}
+
+	UInterchangeMaterialFunctionCallExpressionFactoryNode* FunctionCallFactoryNode = NewObject<UInterchangeMaterialFunctionCallExpressionFactoryNode>(BaseNodeContainer);
+	if (!FunctionCallFactoryNode)
+	{
+		return nullptr;
+	}
+
+	// Check whether the MaterialFunction attribute is valid
+	FString MaterialFunctionAttribute;
+	if (FunctionCallShaderNode->GetCustomMaterialFunction(MaterialFunctionAttribute))
+	{
+		if (!BaseNodeContainer->GetNode(MaterialFunctionAttribute))
+		{
+			if (!FPackageName::IsValidObjectPath(MaterialFunctionAttribute))
+			{
+				MaterialFunctionAttribute.Empty();
+			}
+		}
+	}
+
+	// Nothing to do if the MaterialFunction attribute is not valid
+	if (MaterialFunctionAttribute.IsEmpty())
+	{
+		// TODO: Log a warning
+		return nullptr;
+	}
+
+	FunctionCallFactoryNode->InitializeNode(MaterialExpressionUid, ShaderNode->GetDisplayLabel(), EInterchangeNodeContainerType::FactoryData);
+	BaseNodeContainer->AddNode(FunctionCallFactoryNode);
+
+	if (BaseNodeContainer->GetNode(MaterialFunctionAttribute))
+	{
+		const FString MaterialFunctionFactoryNodeUid = UInterchangeFactoryBaseNode::BuildFactoryNodeUid(MaterialFunctionAttribute);
+		FunctionCallFactoryNode->SetCustomMaterialFunctionDependency(MaterialFunctionFactoryNodeUid);
+		FunctionCallFactoryNode->AddFactoryDependencyUid(MaterialFunctionFactoryNodeUid);
+	}
+	else if (FPackageName::IsValidObjectPath(MaterialFunctionAttribute))
+	{
+		FunctionCallFactoryNode->SetCustomMaterialFunctionDependency(MaterialFunctionAttribute);
+		UpdateFunctionCallExpression(*FunctionCallFactoryNode, MaterialFunctionAttribute);
+	}
+
+	TArray<FString> Inputs;
+	UInterchangeShaderPortsAPI::GatherInputs(ShaderNode, Inputs);
+
+	for (const FString& InputName : Inputs)
+	{
+		TTuple<UInterchangeMaterialExpressionFactoryNode*, FString> InputExpression =
+			CreateMaterialExpressionForInput(MaterialFactoryNode, ShaderNode, InputName, MaterialExpressionUid);
+
+		if (InputExpression.Get<0>())
+		{
+			UInterchangeShaderPortsAPI::ConnectOuputToInputByName(FunctionCallFactoryNode, InputName, InputExpression.Get<0>()->GetUniqueID(), InputExpression.Get<1>());
+		}
+	}
+
+	return FunctionCallFactoryNode;
+}
+
 #undef LOCTEXT_NAMESPACE
diff --git a/Engine/Plugins/Interchange/Runtime/Source/Pipelines/Public/InterchangeGenericMaterialPipeline.h b/Engine/Plugins/Interchange/Runtime/Source/Pipelines/Public/InterchangeGenericMaterialPipeline.h
--- a/Engine/Plugins/Interchange/Runtime/Source/Pipelines/Public/InterchangeGenericMaterialPipeline.h
+++ b/Engine/Plugins/Interchange/Runtime/Source/Pipelines/Public/InterchangeGenericMaterialPipeline.h
@@ -196,6 +196,7 @@ private:
 	UInterchangeMaterialExpressionFactoryNode* CreateScalarParameterExpression(const UInterchangeShaderNode* ShaderNode, const FString& InputName, const FString& ParentUid);
 	UInterchangeMaterialExpressionFactoryNode* CreateVectorParameterExpression(const UInterchangeShaderNode* ShaderNode, const FString& InputName, const FString& ParentUid);
 	UInterchangeMaterialExpressionFactoryNode* CreateVector2ParameterExpression(const UInterchangeShaderNode* ShaderNode, const FString& InputName, const FString& ParentUid);
+	UInterchangeMaterialExpressionFactoryNode* CreateFunctionCallExpression(const UInterchangeShaderNode* ShaderNode, const FString& MaterialExpressionUid, UInterchangeBaseMaterialFactoryNode* MaterialFactoryNode);
 
 	/**
 	 * Visits a given shader node and its connections to find its strongest value.
