Fix C++20 compile error. Implicit capture of 'this' via '[=]' is deprecated in '/std:c++20', so fixed an instance of this error in FAnimationPayloadData::CalculateDataFor.

diff --git a/Engine/Plugins/Interchange/Runtime/Source/Parsers/CommonParser/Private/InterchangeCommonAnimationPayload.cpp b/Engine/Plugins/Interchange/Runtime/Source/Parsers/CommonParser/Private/InterchangeCommonAnimationPayload.cpp
--- a/Engine/Plugins/Interchange/Runtime/Source/Parsers/CommonParser/Private/InterchangeCommonAnimationPayload.cpp
+++ b/Engine/Plugins/Interchange/Runtime/Source/Parsers/CommonParser/Private/InterchangeCommonAnimationPayload.cpp
@@ -72,7 +72,7 @@ namespace UE::Interchange
 			int32 BakeKeyCount = FMath::RoundToInt32(SequenceLength * BakeFrequency) + 1;
 			const FFrameRate ResampleFrameRate(BakeFrequency, 1);
 
-			auto EvaluateCurve = [=](const int32& CurveIndex, double CurrentTime, float DefaultValue)
+			auto EvaluateCurve = [this](const int32& CurveIndex, double CurrentTime, float DefaultValue)
 			{
 				if (Curves[CurveIndex].IsEmpty())
 				{
