[interchange] [gltf] Two root node created in the world outliner

diff --git a/Engine/Plugins/Interchange/Runtime/Source/Pipelines/Private/InterchangeGenericScenesPipeline.cpp b/Engine/Plugins/Interchange/Runtime/Source/Pipelines/Private/InterchangeGenericScenesPipeline.cpp
--- a/Engine/Plugins/Interchange/Runtime/Source/Pipelines/Private/InterchangeGenericScenesPipeline.cpp
+++ b/Engine/Plugins/Interchange/Runtime/Source/Pipelines/Private/InterchangeGenericScenesPipeline.cpp
@@ -28,6 +28,43 @@
 
 #include UE_INLINE_GENERATED_CPP_BY_NAME(InterchangeGenericScenesPipeline)
 
+namespace UE::Interchange::Private
+{
+	//Either a non-joint, or a rootjoint can be a parent (only those get FactoryNodes) :
+	FString FindFactoryParentSceneNodeUid(UInterchangeBaseNodeContainer* BaseNodeContainer, const UInterchangeSceneNode* SceneNode)
+	{
+		FString ParentUid = SceneNode->GetParentUid();
+		if (const UInterchangeSceneNode* ParentSceneNode = Cast<UInterchangeSceneNode>(BaseNodeContainer->GetNode(ParentUid)))
+		{
+			bool bParentIsJoint = ParentSceneNode->IsSpecializedTypeContains(UE::Interchange::FSceneNodeStaticData::GetJointSpecializeTypeString());
+			if (bParentIsJoint)
+			{
+				//Check if its a rootjoint:
+				//	aka check parent's parent if its !joint:
+				FString ParentsParentUid = ParentSceneNode->GetParentUid();
+				if (const UInterchangeSceneNode* ParentsParentSceneNode = Cast<UInterchangeSceneNode>(BaseNodeContainer->GetNode(ParentsParentUid)))
+				{
+					bool bParentsParentIsJoint = ParentSceneNode->IsSpecializedTypeContains(UE::Interchange::FSceneNodeStaticData::GetJointSpecializeTypeString());
+					if (bParentsParentIsJoint)
+					{
+						return FindFactoryParentSceneNodeUid(BaseNodeContainer, ParentSceneNode);
+					}
+					else
+					{
+						return ParentUid;
+					}
+				}
+			}
+			else
+			{
+				return ParentUid;
+			}
+		}
+
+		return UInterchangeBaseNode::InvalidNodeUid();
+	}
+}
+
 void UInterchangeGenericLevelPipeline::ExecutePipeline(UInterchangeBaseNodeContainer* InBaseNodeContainer, const TArray<UInterchangeSourceData*>& InSourceDatas)
 {
 	if (!InBaseNodeContainer)
@@ -178,11 +215,18 @@ void UInterchangeGenericLevelPipeline::ExecuteSceneNodePreImport(const FTransfor
 
 	ActorFactoryNode->InitializeNode(FactoryNodeUid, SceneNode->GetDisplayLabel(), EInterchangeNodeContainerType::FactoryData);
 	const FString ActorFactoryNodeUid = BaseNodeContainer->AddNode(ActorFactoryNode);
+
 	if (!SceneNode->GetParentUid().IsEmpty())
 	{
-		const FString ParentFactoryNodeUid = TEXT("Factory_") + SceneNode->GetParentUid();
-		BaseNodeContainer->SetNodeParentUid(ActorFactoryNodeUid, ParentFactoryNodeUid);
-		ActorFactoryNode->AddFactoryDependencyUid(ParentFactoryNodeUid);
+		FString ParentNodeUid = UE::Interchange::Private::FindFactoryParentSceneNodeUid(BaseNodeContainer, SceneNode);
+		if (ParentNodeUid != UInterchangeBaseNode::InvalidNodeUid())
+		{
+			FString ParentFactoryNodeUid = UInterchangeFactoryBaseNode::BuildFactoryNodeUid(ParentNodeUid);
+			if (BaseNodeContainer->SetNodeParentUid(ActorFactoryNodeUid, ParentFactoryNodeUid))
+			{
+				ActorFactoryNode->AddFactoryDependencyUid(ParentFactoryNodeUid);
+			}
+		}
 	}
 
 	if (bRootJointNode)
