[Interchange] UInterchangeFactoryBaseNode::BuildFactoryNodeUid usage review

diff --git a/Engine/Plugins/Interchange/Runtime/Source/FactoryNodes/Private/InterchangeMaterialFactoryNode.cpp b/Engine/Plugins/Interchange/Runtime/Source/FactoryNodes/Private/InterchangeMaterialFactoryNode.cpp
--- a/Engine/Plugins/Interchange/Runtime/Source/FactoryNodes/Private/InterchangeMaterialFactoryNode.cpp
+++ b/Engine/Plugins/Interchange/Runtime/Source/FactoryNodes/Private/InterchangeMaterialFactoryNode.cpp
@@ -16,7 +16,7 @@
 
 FString UInterchangeBaseMaterialFactoryNode::GetMaterialFactoryNodeUidFromMaterialNodeUid(const FString& TranslatedNodeUid)
 {
-	FString NewUid = TEXT("Factory_") + TranslatedNodeUid;
+	FString NewUid = UInterchangeFactoryBaseNode::BuildFactoryNodeUid(TranslatedNodeUid);
 	return NewUid;
 }
 
diff --git a/Engine/Plugins/Interchange/Runtime/Source/FactoryNodes/Public/InterchangeTextureFactoryNode.h b/Engine/Plugins/Interchange/Runtime/Source/FactoryNodes/Public/InterchangeTextureFactoryNode.h
--- a/Engine/Plugins/Interchange/Runtime/Source/FactoryNodes/Public/InterchangeTextureFactoryNode.h
+++ b/Engine/Plugins/Interchange/Runtime/Source/FactoryNodes/Public/InterchangeTextureFactoryNode.h
@@ -58,7 +58,7 @@ public:
 
 	static FString GetTextureFactoryNodeUidFromTextureNodeUid(const FString& TranslatedNodeUid)
 	{
-		FString NewUid = TEXT("Factory_") + TranslatedNodeUid;
+		FString NewUid = UInterchangeFactoryBaseNode::BuildFactoryNodeUid(TranslatedNodeUid);
 		return NewUid;
 	}
 
diff --git a/Engine/Plugins/Interchange/Runtime/Source/Import/Private/Animation/InterchangeAnimationTrackSetFactory.cpp b/Engine/Plugins/Interchange/Runtime/Source/Import/Private/Animation/InterchangeAnimationTrackSetFactory.cpp
--- a/Engine/Plugins/Interchange/Runtime/Source/Import/Private/Animation/InterchangeAnimationTrackSetFactory.cpp
+++ b/Engine/Plugins/Interchange/Runtime/Source/Import/Private/Animation/InterchangeAnimationTrackSetFactory.cpp
@@ -46,7 +46,7 @@ namespace UE::Interchange::Private
 		FString ActorNodeUid;
 		if (TrackNode->GetCustomActorDependencyUid(ActorNodeUid))
 		{
-			const FString ActorFactoryNodeUid = TEXT("Factory_") + ActorNodeUid;
+			const FString ActorFactoryNodeUid = UInterchangeFactoryBaseNode::BuildFactoryNodeUid(ActorNodeUid);
 			const UInterchangeFactoryBaseNode* ActorFactoryNode = Cast<UInterchangeFactoryBaseNode>(NodeContainer->GetNode(ActorFactoryNodeUid));
 
 			if (ActorFactoryNode)
@@ -89,7 +89,7 @@ namespace UE::Interchange::Private
 						continue;
 					}
 
-					const FString TrackSetFactoryNodeUid = TEXT("Factory_") + TrackSetNodeUid;
+					const FString TrackSetFactoryNodeUid = UInterchangeFactoryBaseNode::BuildFactoryNodeUid(TrackSetNodeUid);
 					const UInterchangeAnimationTrackSetFactoryNode* InstanceFactoryNode = Cast<UInterchangeAnimationTrackSetFactoryNode>(NodeContainer->GetNode(TrackSetFactoryNodeUid));
 
 					if (!InstanceFactoryNode)
@@ -361,7 +361,7 @@ namespace UE::Interchange::Private
 			return;
 		}
 
-		const FString TrackSetFactoryNodeUid = TEXT("Factory_") + TrackSetNodeUid;
+		const FString TrackSetFactoryNodeUid = UInterchangeFactoryBaseNode::BuildFactoryNodeUid(TrackSetNodeUid);
 		const UInterchangeAnimationTrackSetFactoryNode* InstanceFactoryNode = Cast<UInterchangeAnimationTrackSetFactoryNode>(NodeContainer.GetNode(TrackSetFactoryNodeUid));
 		
 		FString InstanceNodeDisplayLabel = InstanceNode.GetDisplayLabel();
@@ -544,7 +544,7 @@ namespace UE::Interchange::Private
 		FString ActorNodeUid;
 		if (TrackNode.GetCustomActorDependencyUid(ActorNodeUid))
 		{
-			const FString ActorFactoryNodeUid = TEXT("Factory_") + ActorNodeUid;
+			const FString ActorFactoryNodeUid = UInterchangeFactoryBaseNode::BuildFactoryNodeUid(ActorNodeUid);
 			const UInterchangeFactoryBaseNode* ActorFactoryNode = Cast<UInterchangeFactoryBaseNode>(NodeContainer.GetNode(ActorFactoryNodeUid));
 
 			if (ActorFactoryNode)
diff --git a/Engine/Plugins/Interchange/Runtime/Source/Pipelines/Private/InterchangeGenericAnimationPipeline.cpp b/Engine/Plugins/Interchange/Runtime/Source/Pipelines/Private/InterchangeGenericAnimationPipeline.cpp
--- a/Engine/Plugins/Interchange/Runtime/Source/Pipelines/Private/InterchangeGenericAnimationPipeline.cpp
+++ b/Engine/Plugins/Interchange/Runtime/Source/Pipelines/Private/InterchangeGenericAnimationPipeline.cpp
@@ -215,7 +215,7 @@ void UInterchangeGenericAnimationPipeline::CreateAnimationTrackSetFactoryNode(UI
 				FString ActorNodeUid;
 				if (TransformTrackNode->GetCustomActorDependencyUid(ActorNodeUid))
 				{
-					const FString ActorFactoryNodeUid = TEXT("Factory_") + ActorNodeUid;
+					const FString ActorFactoryNodeUid = UInterchangeFactoryBaseNode::BuildFactoryNodeUid(ActorNodeUid);
 					FactoryNode->AddFactoryDependencyUid(ActorFactoryNodeUid);
 				}
 			}
@@ -224,7 +224,7 @@ void UInterchangeGenericAnimationPipeline::CreateAnimationTrackSetFactoryNode(UI
 				FString TrackSetNodeUid;
 				if (InstanceTrackNode->GetCustomTrackSetDependencyUid(TrackSetNodeUid))
 				{
-					const FString TrackSetFactoryNodeUid = TEXT("Factory_") + TrackSetNodeUid;
+					const FString TrackSetFactoryNodeUid = UInterchangeFactoryBaseNode::BuildFactoryNodeUid(TrackSetNodeUid);
 					FactoryNode->AddFactoryDependencyUid(TrackSetFactoryNodeUid);
 				}
 
