MaterialX:
- Use FunctionCallShaderNode instead of ShaderGraphNode, and apply directly the outputs to the shadergraph in the Translator instead of the Material Pipeline (StandardSurfaceShader, SurfaceUnlitShader)

Translator:
- Adding a key to the TexturePayloadKey to retrieve CompressionSettings from a TextureNode

Material Pipeline:
- Remove the OneMinus on the SheenRoughness (was only useful for the GltfTranslator)
- Connect and enable the Refraction, which was disabled by default in UE
- Connect the Anisotropy in the PBR model
- Connect the Tangent in the Common Parameters
- Set the Blend mode to Opaque for Subsurface Materials

diff --git a/Engine/Plugins/Interchange/Runtime/Source/FactoryNodes/Private/InterchangeMaterialFactoryNode.cpp b/Engine/Plugins/Interchange/Runtime/Source/FactoryNodes/Private/InterchangeMaterialFactoryNode.cpp
--- a/Engine/Plugins/Interchange/Runtime/Source/FactoryNodes/Private/InterchangeMaterialFactoryNode.cpp
+++ b/Engine/Plugins/Interchange/Runtime/Source/FactoryNodes/Private/InterchangeMaterialFactoryNode.cpp
@@ -355,6 +355,16 @@ bool UInterchangeMaterialFactoryNode::SetCustomOpacityMaskClipValue(const float&
 	IMPLEMENT_NODE_ATTRIBUTE_SETTER(UInterchangeMaterialFactoryNode, OpacityMaskClipValue, float, UMaterial);
 }
 
+bool UInterchangeMaterialFactoryNode::GetCustomRefractionMethod(TEnumAsByte<ERefractionMode>& AttributeValue) const
+{
+	IMPLEMENT_NODE_ATTRIBUTE_GETTER(RefractionMethod, TEnumAsByte<ERefractionMode>);
+}
+
+bool UInterchangeMaterialFactoryNode::SetCustomRefractionMethod(const TEnumAsByte<ERefractionMode>& AttributeValue, bool bAddApplyDelegate)
+{
+	IMPLEMENT_NODE_ATTRIBUTE_SETTER(UInterchangeMaterialFactoryNode, RefractionMethod, TEnumAsByte<ERefractionMode>, UMaterial);
+}
+
 FString UInterchangeMaterialExpressionFactoryNode::GetTypeName() const
 {
 	const FString TypeName = TEXT("MaterialExpressionFactoryNode");
diff --git a/Engine/Plugins/Interchange/Runtime/Source/FactoryNodes/Public/InterchangeMaterialFactoryNode.h b/Engine/Plugins/Interchange/Runtime/Source/FactoryNodes/Public/InterchangeMaterialFactoryNode.h
--- a/Engine/Plugins/Interchange/Runtime/Source/FactoryNodes/Public/InterchangeMaterialFactoryNode.h
+++ b/Engine/Plugins/Interchange/Runtime/Source/FactoryNodes/Public/InterchangeMaterialFactoryNode.h
@@ -231,12 +231,19 @@ public:
 	UFUNCTION(BlueprintCallable, Category = "Interchange | Node | Material")
 	bool SetCustomOpacityMaskClipValue(const float& AttributeValue, bool bAddApplyDelegate = true);
 
+	UFUNCTION(BlueprintCallable, Category = "Interchange | Node | Material")
+	bool GetCustomRefractionMethod(TEnumAsByte<ERefractionMode>& AttributeValue) const;
+
+	UFUNCTION(BlueprintCallable, Category = "Interchange | Node | Material")
+	bool SetCustomRefractionMethod(const TEnumAsByte<ERefractionMode>& AttributeValue, bool bAddApplyDelegate = true);
+
 private:
 	const UE::Interchange::FAttributeKey Macro_CustomShadingModelKey = UE::Interchange::FAttributeKey(TEXT("ShadingModel"));
 	const UE::Interchange::FAttributeKey Macro_CustomTranslucencyLightingModeKey = UE::Interchange::FAttributeKey(TEXT("TranslucencyLightingMode"));
 	const UE::Interchange::FAttributeKey Macro_CustomBlendModeKey = UE::Interchange::FAttributeKey(TEXT("BlendMode"));
 	const UE::Interchange::FAttributeKey Macro_CustomTwoSidedKey = UE::Interchange::FAttributeKey(TEXT("TwoSided"));
 	const UE::Interchange::FAttributeKey Macro_CustomOpacityMaskClipValueKey = UE::Interchange::FAttributeKey(TEXT("OpacityMaskClipValue"));
+	const UE::Interchange::FAttributeKey Macro_CustomRefractionMethodKey = UE::Interchange::FAttributeKey(TEXT("RefractionMethod"));
 };
 
 UCLASS(BlueprintType)
diff --git a/Engine/Plugins/Interchange/Runtime/Source/Import/Private/Gltf/InterchangeGltfTranslator.cpp b/Engine/Plugins/Interchange/Runtime/Source/Import/Private/Gltf/InterchangeGltfTranslator.cpp
--- a/Engine/Plugins/Interchange/Runtime/Source/Import/Private/Gltf/InterchangeGltfTranslator.cpp
+++ b/Engine/Plugins/Interchange/Runtime/Source/Import/Private/Gltf/InterchangeGltfTranslator.cpp
@@ -759,8 +759,9 @@ void UInterchangeGltfTranslator::HandleGltfSheen( UInterchangeBaseNodeContainer&
 		TVariant< FLinearColor, float > SheenRoughnessFactor;
 		SheenRoughnessFactor.Set< float >( GltfMaterial.Sheen.SheenRoughnessFactor );
 
+		const bool bInverse = true;
 		HandleGltfMaterialParameter( NodeContainer, GltfMaterial.Sheen.SheenRoughnessMap, ShaderGraphNode, Sheen::Parameters::SheenRoughness.ToString(),
-			SheenRoughnessFactor, Standard::Nodes::TextureSample::Outputs::A.ToString() );
+			SheenRoughnessFactor, Standard::Nodes::TextureSample::Outputs::A.ToString(), bInverse);
 	}
 }
 
diff --git a/Engine/Plugins/Interchange/Runtime/Source/Import/Private/MaterialX/InterchangeMaterialXTranslator.cpp b/Engine/Plugins/Interchange/Runtime/Source/Import/Private/MaterialX/InterchangeMaterialXTranslator.cpp
--- a/Engine/Plugins/Interchange/Runtime/Source/Import/Private/MaterialX/InterchangeMaterialXTranslator.cpp
+++ b/Engine/Plugins/Interchange/Runtime/Source/Import/Private/MaterialX/InterchangeMaterialXTranslator.cpp
@@ -193,7 +193,17 @@ bool UInterchangeMaterialXTranslator::Translate(UInterchangeBaseNodeContainer& B
 
 TOptional<UE::Interchange::FImportImage> UInterchangeMaterialXTranslator::GetTexturePayloadData(const UInterchangeSourceData* InSourceData, const FString& PayloadKey) const
 {
-	UInterchangeSourceData* PayloadSourceData = UInterchangeManager::GetInterchangeManager().CreateSourceData(PayloadKey);
+	FString Filename = PayloadKey;
+	TextureCompressionSettings CompressionSettings = TextureCompressionSettings::TC_Default;
+
+#if WITH_EDITOR
+	if(int32 IndexTextureCompression; PayloadKey.FindChar(FMaterialXManager::TexturePayloadSeparator, IndexTextureCompression))
+	{
+		Filename = PayloadKey.Mid(0, IndexTextureCompression);
+		CompressionSettings = TextureCompressionSettings(FCString::Atoi(*PayloadKey.Mid(IndexTextureCompression + 1)));
+	}
+#endif
+	UInterchangeSourceData* PayloadSourceData = UInterchangeManager::GetInterchangeManager().CreateSourceData(Filename);
 	FGCObjectScopeGuard ScopedSourceData(PayloadSourceData);
 
 	if(!PayloadSourceData)
@@ -210,7 +220,14 @@ TOptional<UE::Interchange::FImportImage> UInterchangeMaterialXTranslator::GetTex
 		return TOptional<UE::Interchange::FImportImage>();
 	}
 
-	return TextureTranslator->GetTexturePayloadData(PayloadSourceData, PayloadKey);
+	TOptional<UE::Interchange::FImportImage> TexturePayloadData = TextureTranslator->GetTexturePayloadData(PayloadSourceData, Filename);
+
+	if(TexturePayloadData.IsSet())
+	{
+		TexturePayloadData.GetValue().CompressionSettings = CompressionSettings;
+	}
+
+	return TexturePayloadData;
 }
 
 #undef LOCTEXT_NAMESPACE
diff --git a/Engine/Plugins/Interchange/Runtime/Source/Import/Private/MaterialX/MaterialXUtils/MaterialXManager.cpp b/Engine/Plugins/Interchange/Runtime/Source/Import/Private/MaterialX/MaterialXUtils/MaterialXManager.cpp
--- a/Engine/Plugins/Interchange/Runtime/Source/Import/Private/MaterialX/MaterialXUtils/MaterialXManager.cpp
+++ b/Engine/Plugins/Interchange/Runtime/Source/Import/Private/MaterialX/MaterialXUtils/MaterialXManager.cpp
@@ -15,6 +15,9 @@
 
 namespace mx = MaterialX;
 
+//not a good solution to use semicolon because of drive disk on Windows
+const TCHAR FMaterialXManager::TexturePayloadSeparator = TEXT('{');
+
 FMaterialXManager::FMaterialXManager()
 	: MatchingInputNames {
 		{{TEXT(""),                            TEXT("amplitude")}, TEXT("Amplitude")},
diff --git a/Engine/Plugins/Interchange/Runtime/Source/Import/Private/MaterialX/MaterialXUtils/MaterialXStandardSurfaceShader.cpp b/Engine/Plugins/Interchange/Runtime/Source/Import/Private/MaterialX/MaterialXUtils/MaterialXStandardSurfaceShader.cpp
--- a/Engine/Plugins/Interchange/Runtime/Source/Import/Private/MaterialX/MaterialXUtils/MaterialXStandardSurfaceShader.cpp
+++ b/Engine/Plugins/Interchange/Runtime/Source/Import/Private/MaterialX/MaterialXUtils/MaterialXStandardSurfaceShader.cpp
@@ -21,105 +21,155 @@ TSharedRef<FMaterialXBase> FMaterialXStandardSurfaceShader::MakeInstance(UInterc
 void FMaterialXStandardSurfaceShader::Translate(mx::NodePtr StandardSurfaceNode)
 {
 	this->SurfaceShaderNode = StandardSurfaceNode;
-
+	
 	using namespace UE::Interchange::Materials;
 
-	ShaderGraphNode->SetCustomShaderType(StandardSurface::Name.ToString());
+	UInterchangeFunctionCallShaderNode* StandardSurfaceShaderNode;
+	
+	const FString NodeUID = UInterchangeShaderNode::MakeNodeUid(ANSI_TO_TCHAR(StandardSurfaceNode->getName().c_str()), FStringView{});
+
+	if(StandardSurfaceShaderNode = const_cast<UInterchangeFunctionCallShaderNode*>(Cast<UInterchangeFunctionCallShaderNode>(NodeContainer.GetNode(NodeUID))); !StandardSurfaceShaderNode)
+	{
+		const FString NodeName = StandardSurfaceNode->getName().c_str();
+		StandardSurfaceShaderNode = NewObject<UInterchangeFunctionCallShaderNode>(&NodeContainer);
+		StandardSurfaceShaderNode->InitializeNode(NodeUID, NodeName, EInterchangeNodeContainerType::TranslatedAsset);
+
+		StandardSurfaceShaderNode->SetCustomMaterialFunction(TEXT("/Interchange/Functions/MX_StandardSurface.MX_StandardSurface"));
+		NodeContainer.AddNode(StandardSurfaceShaderNode);
+
+		ShaderNodes.Add(NodeName, StandardSurfaceShaderNode);
+	}
 
+	// Inputs
 	//Base
-	ConnectNodeOutputToInput(mx::StandardSurface::Input::Base, ShaderGraphNode, StandardSurface::Parameters::Base.ToString(), mx::StandardSurface::DefaultValue::Float::Base);
+	ConnectNodeOutputToInput(mx::StandardSurface::Input::Base, StandardSurfaceShaderNode, StandardSurface::Parameters::Base.ToString(), mx::StandardSurface::DefaultValue::Float::Base);
 
 	//Base Color
-	ConnectNodeOutputToInput(mx::StandardSurface::Input::BaseColor, ShaderGraphNode, StandardSurface::Parameters::BaseColor.ToString(), mx::StandardSurface::DefaultValue::Color3::BaseColor);
+	ConnectNodeOutputToInput(mx::StandardSurface::Input::BaseColor, StandardSurfaceShaderNode, StandardSurface::Parameters::BaseColor.ToString(), mx::StandardSurface::DefaultValue::Color3::BaseColor);
 
 	//Diffuse Roughness
-	ConnectNodeOutputToInput(mx::StandardSurface::Input::DiffuseRoughness, ShaderGraphNode, StandardSurface::Parameters::DiffuseRoughness.ToString(), mx::StandardSurface::DefaultValue::Float::DiffuseRoughness);
+	ConnectNodeOutputToInput(mx::StandardSurface::Input::DiffuseRoughness, StandardSurfaceShaderNode, StandardSurface::Parameters::DiffuseRoughness.ToString(), mx::StandardSurface::DefaultValue::Float::DiffuseRoughness);
 
 	//Specular
-	ConnectNodeOutputToInput(mx::StandardSurface::Input::Specular, ShaderGraphNode, StandardSurface::Parameters::Specular.ToString(), mx::StandardSurface::DefaultValue::Float::Specular);
+	ConnectNodeOutputToInput(mx::StandardSurface::Input::Specular, StandardSurfaceShaderNode, StandardSurface::Parameters::Specular.ToString(), mx::StandardSurface::DefaultValue::Float::Specular);
 
 	//Specular Roughness
-	ConnectNodeOutputToInput(mx::StandardSurface::Input::SpecularRoughness, ShaderGraphNode, StandardSurface::Parameters::SpecularRoughness.ToString(), mx::StandardSurface::DefaultValue::Float::SpecularRoughness);
+	ConnectNodeOutputToInput(mx::StandardSurface::Input::SpecularRoughness, StandardSurfaceShaderNode, StandardSurface::Parameters::SpecularRoughness.ToString(), mx::StandardSurface::DefaultValue::Float::SpecularRoughness);
 
 	//Specular IOR
-	ConnectNodeOutputToInput(mx::StandardSurface::Input::SpecularIOR, ShaderGraphNode, StandardSurface::Parameters::SpecularIOR.ToString(), mx::StandardSurface::DefaultValue::Float::SpecularIOR);
+	ConnectNodeOutputToInput(mx::StandardSurface::Input::SpecularIOR, StandardSurfaceShaderNode, StandardSurface::Parameters::SpecularIOR.ToString(), mx::StandardSurface::DefaultValue::Float::SpecularIOR);
 
 	//Specular Anisotropy
-	ConnectNodeOutputToInput(mx::StandardSurface::Input::SpecularAnisotropy, ShaderGraphNode, StandardSurface::Parameters::SpecularAnisotropy.ToString(), mx::StandardSurface::DefaultValue::Float::SpecularAnisotropy);
+	ConnectNodeOutputToInput(mx::StandardSurface::Input::SpecularAnisotropy, StandardSurfaceShaderNode, StandardSurface::Parameters::SpecularAnisotropy.ToString(), mx::StandardSurface::DefaultValue::Float::SpecularAnisotropy);
 
 	//Specular Rotation
-	ConnectNodeOutputToInput(mx::StandardSurface::Input::SpecularRotation, ShaderGraphNode, StandardSurface::Parameters::SpecularRotation.ToString(), mx::StandardSurface::DefaultValue::Float::SpecularRotation);
+	ConnectNodeOutputToInput(mx::StandardSurface::Input::SpecularRotation, StandardSurfaceShaderNode, StandardSurface::Parameters::SpecularRotation.ToString(), mx::StandardSurface::DefaultValue::Float::SpecularRotation);
 
 	//Metallic
-	ConnectNodeOutputToInput(mx::StandardSurface::Input::Metalness, ShaderGraphNode, StandardSurface::Parameters::Metalness.ToString(), mx::StandardSurface::DefaultValue::Float::Metalness);
+	ConnectNodeOutputToInput(mx::StandardSurface::Input::Metalness, StandardSurfaceShaderNode, StandardSurface::Parameters::Metalness.ToString(), mx::StandardSurface::DefaultValue::Float::Metalness);
 
 	//Subsurface
-	ConnectNodeOutputToInput(mx::StandardSurface::Input::Subsurface, ShaderGraphNode, StandardSurface::Parameters::Subsurface.ToString(), mx::StandardSurface::DefaultValue::Float::Subsurface);
+	ConnectNodeOutputToInput(mx::StandardSurface::Input::Subsurface, StandardSurfaceShaderNode, StandardSurface::Parameters::Subsurface.ToString(), mx::StandardSurface::DefaultValue::Float::Subsurface);
 
 	//Subsurface Color
-	ConnectNodeOutputToInput(mx::StandardSurface::Input::SubsurfaceColor, ShaderGraphNode, StandardSurface::Parameters::SubsurfaceColor.ToString(), mx::StandardSurface::DefaultValue::Color3::SubsurfaceColor);
+	ConnectNodeOutputToInput(mx::StandardSurface::Input::SubsurfaceColor, StandardSurfaceShaderNode, StandardSurface::Parameters::SubsurfaceColor.ToString(), mx::StandardSurface::DefaultValue::Color3::SubsurfaceColor);
 
 	//Subsurface Radius
-	ConnectNodeOutputToInput(mx::StandardSurface::Input::SubsurfaceRadius, ShaderGraphNode, StandardSurface::Parameters::SubsurfaceRadius.ToString(), mx::StandardSurface::DefaultValue::Color3::SubsurfaceRadius);
+	ConnectNodeOutputToInput(mx::StandardSurface::Input::SubsurfaceRadius, StandardSurfaceShaderNode, StandardSurface::Parameters::SubsurfaceRadius.ToString(), mx::StandardSurface::DefaultValue::Color3::SubsurfaceRadius);
 
 	//Subsurface Scale
-	ConnectNodeOutputToInput(mx::StandardSurface::Input::SubsurfaceScale, ShaderGraphNode, StandardSurface::Parameters::SubsurfaceScale.ToString(), mx::StandardSurface::DefaultValue::Float::SubsurfaceScale);
+	ConnectNodeOutputToInput(mx::StandardSurface::Input::SubsurfaceScale, StandardSurfaceShaderNode, StandardSurface::Parameters::SubsurfaceScale.ToString(), mx::StandardSurface::DefaultValue::Float::SubsurfaceScale);
 
 	//Sheen
-	ConnectNodeOutputToInput(mx::StandardSurface::Input::Sheen, ShaderGraphNode, StandardSurface::Parameters::Sheen.ToString(), mx::StandardSurface::DefaultValue::Float::Sheen);
+	ConnectNodeOutputToInput(mx::StandardSurface::Input::Sheen, StandardSurfaceShaderNode, StandardSurface::Parameters::Sheen.ToString(), mx::StandardSurface::DefaultValue::Float::Sheen);
 
 	//Sheen Color
-	ConnectNodeOutputToInput(mx::StandardSurface::Input::SheenColor, ShaderGraphNode, StandardSurface::Parameters::SheenColor.ToString(), mx::StandardSurface::DefaultValue::Color3::SheenColor);
+	ConnectNodeOutputToInput(mx::StandardSurface::Input::SheenColor, StandardSurfaceShaderNode, StandardSurface::Parameters::SheenColor.ToString(), mx::StandardSurface::DefaultValue::Color3::SheenColor);
 
 	//Sheen Roughness
-	ConnectNodeOutputToInput(mx::StandardSurface::Input::SheenRoughness, ShaderGraphNode, StandardSurface::Parameters::SheenRoughness.ToString(), mx::StandardSurface::DefaultValue::Float::SheenRoughness);
+	ConnectNodeOutputToInput(mx::StandardSurface::Input::SheenRoughness, StandardSurfaceShaderNode, StandardSurface::Parameters::SheenRoughness.ToString(), mx::StandardSurface::DefaultValue::Float::SheenRoughness);
 
 	//Coat
-	ConnectNodeOutputToInput(mx::StandardSurface::Input::Coat, ShaderGraphNode, StandardSurface::Parameters::Coat.ToString(), mx::StandardSurface::DefaultValue::Float::Coat);
+	ConnectNodeOutputToInput(mx::StandardSurface::Input::Coat, StandardSurfaceShaderNode, StandardSurface::Parameters::Coat.ToString(), mx::StandardSurface::DefaultValue::Float::Coat);
 
 	//Coat Color
-	ConnectNodeOutputToInput(mx::StandardSurface::Input::CoatColor, ShaderGraphNode, StandardSurface::Parameters::CoatColor.ToString(), mx::StandardSurface::DefaultValue::Color3::CoatColor);
+	ConnectNodeOutputToInput(mx::StandardSurface::Input::CoatColor, StandardSurfaceShaderNode, StandardSurface::Parameters::CoatColor.ToString(), mx::StandardSurface::DefaultValue::Color3::CoatColor);
 
 	//Coat Roughness
-	ConnectNodeOutputToInput(mx::StandardSurface::Input::CoatRoughness, ShaderGraphNode, StandardSurface::Parameters::CoatRoughness.ToString(), mx::StandardSurface::DefaultValue::Float::CoatRoughness);
+	ConnectNodeOutputToInput(mx::StandardSurface::Input::CoatRoughness, StandardSurfaceShaderNode, StandardSurface::Parameters::CoatRoughness.ToString(), mx::StandardSurface::DefaultValue::Float::CoatRoughness);
 
 	//Coat Normal: No need to take the default input if there is no CoatNormal input
-	ConnectNodeOutputToInput(mx::StandardSurface::Input::CoatNormal, ShaderGraphNode, StandardSurface::Parameters::CoatNormal.ToString(), nullptr);
+	ConnectNodeOutputToInput(mx::StandardSurface::Input::CoatNormal, StandardSurfaceShaderNode, StandardSurface::Parameters::CoatNormal.ToString(), nullptr, TextureCompressionSettings::TC_Normalmap);
 
 	//Thin Film Thickness
-	ConnectNodeOutputToInput(mx::StandardSurface::Input::ThinFilmThickness, ShaderGraphNode, StandardSurface::Parameters::ThinFilmThickness.ToString(), mx::StandardSurface::DefaultValue::Float::ThinFilmThickness);
+	ConnectNodeOutputToInput(mx::StandardSurface::Input::ThinFilmThickness, StandardSurfaceShaderNode, StandardSurface::Parameters::ThinFilmThickness.ToString(), mx::StandardSurface::DefaultValue::Float::ThinFilmThickness);
 
 	//Emission
-	ConnectNodeOutputToInput(mx::StandardSurface::Input::Emission, ShaderGraphNode, StandardSurface::Parameters::Emission.ToString(), mx::StandardSurface::DefaultValue::Float::Emission);
+	ConnectNodeOutputToInput(mx::StandardSurface::Input::Emission, StandardSurfaceShaderNode, StandardSurface::Parameters::Emission.ToString(), mx::StandardSurface::DefaultValue::Float::Emission);
 
 	//Emission Color
-	ConnectNodeOutputToInput(mx::StandardSurface::Input::EmissionColor, ShaderGraphNode, StandardSurface::Parameters::EmissionColor.ToString(), mx::StandardSurface::DefaultValue::Color3::EmissionColor);
+	ConnectNodeOutputToInput(mx::StandardSurface::Input::EmissionColor, StandardSurfaceShaderNode, StandardSurface::Parameters::EmissionColor.ToString(), mx::StandardSurface::DefaultValue::Color3::EmissionColor);
 
 	//Normal: No need to take the default input if there is no Normal input
-	ConnectNodeOutputToInput(mx::StandardSurface::Input::Normal, ShaderGraphNode, StandardSurface::Parameters::Normal.ToString(), nullptr);
+	ConnectNodeOutputToInput(mx::StandardSurface::Input::Normal, StandardSurfaceShaderNode, StandardSurface::Parameters::Normal.ToString(), nullptr, TextureCompressionSettings::TC_Normalmap);
 
 	//Tangent: No need to take the default input if there is no Tangent input
-	ConnectNodeOutputToInput(mx::StandardSurface::Input::Tangent, ShaderGraphNode, StandardSurface::Parameters::Tangent.ToString(), nullptr);
+	ConnectNodeOutputToInput(mx::StandardSurface::Input::Tangent, StandardSurfaceShaderNode, StandardSurface::Parameters::Tangent.ToString(), nullptr, TextureCompressionSettings::TC_Normalmap);
 
 	//Transmission
-	ConnectNodeOutputToInput(mx::StandardSurface::Input::Transmission, ShaderGraphNode, StandardSurface::Parameters::Transmission.ToString(), mx::StandardSurface::DefaultValue::Float::Transmission);
+	ConnectNodeOutputToInput(mx::StandardSurface::Input::Transmission, StandardSurfaceShaderNode, StandardSurface::Parameters::Transmission.ToString(), mx::StandardSurface::DefaultValue::Float::Transmission);
 
 	//Transmission Color
-	ConnectNodeOutputToInput(mx::StandardSurface::Input::TransmissionColor, ShaderGraphNode, StandardSurface::Parameters::TransmissionColor.ToString(), mx::StandardSurface::DefaultValue::Color3::TransmissionColor);
+	ConnectNodeOutputToInput(mx::StandardSurface::Input::TransmissionColor, StandardSurfaceShaderNode, StandardSurface::Parameters::TransmissionColor.ToString(), mx::StandardSurface::DefaultValue::Color3::TransmissionColor);
 
 	//Transmission Depth
-	ConnectNodeOutputToInput(mx::StandardSurface::Input::TransmissionDepth, ShaderGraphNode, StandardSurface::Parameters::TransmissionDepth.ToString(), mx::StandardSurface::DefaultValue::Float::TransmissionDepth);
+	ConnectNodeOutputToInput(mx::StandardSurface::Input::TransmissionDepth, StandardSurfaceShaderNode, StandardSurface::Parameters::TransmissionDepth.ToString(), mx::StandardSurface::DefaultValue::Float::TransmissionDepth);
 
 	//Transmission Scatter
-	ConnectNodeOutputToInput(mx::StandardSurface::Input::TransmissionScatter, ShaderGraphNode, StandardSurface::Parameters::TransmissionScatter.ToString(), mx::StandardSurface::DefaultValue::Color3::TransmissionScatter);
+	ConnectNodeOutputToInput(mx::StandardSurface::Input::TransmissionScatter, StandardSurfaceShaderNode, StandardSurface::Parameters::TransmissionScatter.ToString(), mx::StandardSurface::DefaultValue::Color3::TransmissionScatter);
 
 	//Transmission Scatter Anisotropy
-	ConnectNodeOutputToInput(mx::StandardSurface::Input::TransmissionScatterAnisotropy, ShaderGraphNode, StandardSurface::Parameters::TransmissionScatterAnisotropy.ToString(), mx::StandardSurface::DefaultValue::Float::TransmissionScatterAnisotropy);
+	ConnectNodeOutputToInput(mx::StandardSurface::Input::TransmissionScatterAnisotropy, StandardSurfaceShaderNode, StandardSurface::Parameters::TransmissionScatterAnisotropy.ToString(), mx::StandardSurface::DefaultValue::Float::TransmissionScatterAnisotropy);
 
 	//Transmission Dispersion
-	ConnectNodeOutputToInput(mx::StandardSurface::Input::TransmissionDispersion, ShaderGraphNode, StandardSurface::Parameters::TransmissionDispersion.ToString(), mx::StandardSurface::DefaultValue::Float::TransmissionDispersion);
+	ConnectNodeOutputToInput(mx::StandardSurface::Input::TransmissionDispersion, StandardSurfaceShaderNode, StandardSurface::Parameters::TransmissionDispersion.ToString(), mx::StandardSurface::DefaultValue::Float::TransmissionDispersion);
 
 	//Transmission Extra Roughness
-	ConnectNodeOutputToInput(mx::StandardSurface::Input::TransmissionExtraRoughness, ShaderGraphNode, StandardSurface::Parameters::TransmissionExtraRoughness.ToString(), mx::StandardSurface::DefaultValue::Float::TransmissionExtraRoughness);
+	ConnectNodeOutputToInput(mx::StandardSurface::Input::TransmissionExtraRoughness, StandardSurfaceShaderNode, StandardSurface::Parameters::TransmissionExtraRoughness.ToString(), mx::StandardSurface::DefaultValue::Float::TransmissionExtraRoughness);
+
+	// Outputs
+	UInterchangeShaderPortsAPI::ConnectOuputToInputByName(ShaderGraphNode, PBR::Parameters::BaseColor.ToString(), StandardSurfaceShaderNode->GetUniqueID(), TEXT("Base Color"));
+	UInterchangeShaderPortsAPI::ConnectOuputToInputByName(ShaderGraphNode, PBR::Parameters::Metallic.ToString(), StandardSurfaceShaderNode->GetUniqueID(), PBR::Parameters::Metallic.ToString());
+	UInterchangeShaderPortsAPI::ConnectOuputToInputByName(ShaderGraphNode, PBR::Parameters::Specular.ToString(), StandardSurfaceShaderNode->GetUniqueID(), PBR::Parameters::Specular.ToString());
+	UInterchangeShaderPortsAPI::ConnectOuputToInputByName(ShaderGraphNode, PBR::Parameters::Roughness.ToString(), StandardSurfaceShaderNode->GetUniqueID(), PBR::Parameters::Roughness.ToString());
+	UInterchangeShaderPortsAPI::ConnectOuputToInputByName(ShaderGraphNode, PBR::Parameters::EmissiveColor.ToString(), StandardSurfaceShaderNode->GetUniqueID(), PBR::Parameters::EmissiveColor.ToString());
+	UInterchangeShaderPortsAPI::ConnectOuputToInputByName(ShaderGraphNode, PBR::Parameters::Anisotropy.ToString(), StandardSurfaceShaderNode->GetUniqueID(), PBR::Parameters::Anisotropy.ToString());
+	UInterchangeShaderPortsAPI::ConnectOuputToInputByName(ShaderGraphNode, PBR::Parameters::Normal.ToString(), StandardSurfaceShaderNode->GetUniqueID(), PBR::Parameters::Normal.ToString());
+	UInterchangeShaderPortsAPI::ConnectOuputToInputByName(ShaderGraphNode, PBR::Parameters::Tangent.ToString(), StandardSurfaceShaderNode->GetUniqueID(), PBR::Parameters::Tangent.ToString());
+
+	// We can't have all shading models at once, so we have to make a choice here
+	if(UInterchangeShaderPortsAPI::HasInput(StandardSurfaceShaderNode, StandardSurface::Parameters::Transmission))
+	{
+		StandardSurfaceShaderNode->SetCustomMaterialFunction(TEXT("/Interchange/Functions/MX_TransmissionSurface.MX_TransmissionSurface"));
+		UInterchangeShaderPortsAPI::ConnectOuputToInputByName(ShaderGraphNode, PBR::Parameters::Opacity.ToString(), StandardSurfaceShaderNode->GetUniqueID(), PBR::Parameters::Opacity.ToString());
+		UInterchangeShaderPortsAPI::ConnectOuputToInputByName(ShaderGraphNode, ThinTranslucent::Parameters::TransmissionColor.ToString(), StandardSurfaceShaderNode->GetUniqueID(), ThinTranslucent::Parameters::TransmissionColor.ToString());
+		UInterchangeShaderPortsAPI::ConnectOuputToInputByName(ShaderGraphNode, Common::Parameters::Refraction.ToString(), StandardSurfaceShaderNode->GetUniqueID(), Common::Parameters::Refraction.ToString());
+	}
+	else if(UInterchangeShaderPortsAPI::HasInput(StandardSurfaceShaderNode, StandardSurface::Parameters::Sheen))
+	{
+		UInterchangeShaderPortsAPI::ConnectOuputToInputByName(ShaderGraphNode, Sheen::Parameters::SheenColor.ToString(), StandardSurfaceShaderNode->GetUniqueID(), Sheen::Parameters::SheenColor.ToString());
+		UInterchangeShaderPortsAPI::ConnectOuputToInputByName(ShaderGraphNode, Sheen::Parameters::SheenRoughness.ToString(), StandardSurfaceShaderNode->GetUniqueID(), Sheen::Parameters::SheenRoughness.ToString());
+	}
+	else if(UInterchangeShaderPortsAPI::HasInput(StandardSurfaceShaderNode, StandardSurface::Parameters::Coat))
+	{
+		UInterchangeShaderPortsAPI::ConnectOuputToInputByName(ShaderGraphNode, ClearCoat::Parameters::ClearCoat.ToString(), StandardSurfaceShaderNode->GetUniqueID(), ClearCoat::Parameters::ClearCoat.ToString());
+		UInterchangeShaderPortsAPI::ConnectOuputToInputByName(ShaderGraphNode, ClearCoat::Parameters::ClearCoatRoughness.ToString(), StandardSurfaceShaderNode->GetUniqueID(), ClearCoat::Parameters::ClearCoatRoughness.ToString());
+		UInterchangeShaderPortsAPI::ConnectOuputToInputByName(ShaderGraphNode, ClearCoat::Parameters::ClearCoatNormal.ToString(), StandardSurfaceShaderNode->GetUniqueID(), ClearCoat::Parameters::ClearCoatNormal.ToString());
+	}
+	else if(UInterchangeShaderPortsAPI::HasInput(StandardSurfaceShaderNode, StandardSurface::Parameters::Subsurface))
+	{
+		UInterchangeShaderPortsAPI::ConnectOuputToInputByName(ShaderGraphNode, PBR::Parameters::Opacity.ToString(), StandardSurfaceShaderNode->GetUniqueID(), PBR::Parameters::Opacity.ToString());
+		UInterchangeShaderPortsAPI::ConnectOuputToInputByName(ShaderGraphNode, Subsurface::Parameters::SubsurfaceColor.ToString(), StandardSurfaceShaderNode->GetUniqueID(), Subsurface::Parameters::SubsurfaceColor.ToString());
+	}
 }
 #endif
\ No newline at end of file
diff --git a/Engine/Plugins/Interchange/Runtime/Source/Import/Private/MaterialX/MaterialXUtils/MaterialXSurfaceUnlitShader.cpp b/Engine/Plugins/Interchange/Runtime/Source/Import/Private/MaterialX/MaterialXUtils/MaterialXSurfaceUnlitShader.cpp
--- a/Engine/Plugins/Interchange/Runtime/Source/Import/Private/MaterialX/MaterialXUtils/MaterialXSurfaceUnlitShader.cpp
+++ b/Engine/Plugins/Interchange/Runtime/Source/Import/Private/MaterialX/MaterialXUtils/MaterialXSurfaceUnlitShader.cpp
@@ -22,23 +22,51 @@ void FMaterialXSurfaceUnlitShader::Translate(mx::NodePtr SurfaceUnlitNode)
 {
 	this->SurfaceShaderNode = SurfaceUnlitNode;
 
-	using namespace UE::Interchange::Materials;
+	UInterchangeFunctionCallShaderNode* SurfaceUnlitShaderNode;
+
+	const FString NodeUID = UInterchangeShaderNode::MakeNodeUid(ANSI_TO_TCHAR(SurfaceUnlitNode->getName().c_str()), FStringView{});
+
+	if(SurfaceUnlitShaderNode = const_cast<UInterchangeFunctionCallShaderNode*>(Cast<UInterchangeFunctionCallShaderNode>(NodeContainer.GetNode(NodeUID))); !SurfaceUnlitShaderNode)
+	{
+		const FString NodeName = SurfaceUnlitNode->getName().c_str();
+		SurfaceUnlitShaderNode = NewObject<UInterchangeFunctionCallShaderNode>(&NodeContainer);
+		SurfaceUnlitShaderNode->InitializeNode(NodeUID, NodeName, EInterchangeNodeContainerType::TranslatedAsset);
+
+		SurfaceUnlitShaderNode->SetCustomMaterialFunction(TEXT("/Interchange/Functions/MX_SurfaceUnlit.MX_SurfaceUnlit"));
+		NodeContainer.AddNode(SurfaceUnlitShaderNode);
 
-	ShaderGraphNode->SetCustomShaderType(SurfaceUnlit::Name.ToString());
+		ShaderNodes.Add(NodeName, SurfaceUnlitShaderNode);
+	}
+
+	using namespace UE::Interchange::Materials;
 
 	//Emission
-	ConnectNodeOutputToInput(mx::SurfaceUnlit::Input::Emission, ShaderGraphNode, SurfaceUnlit::Parameters::Emission.ToString(), mx::SurfaceUnlit::DefaultValue::Float::Emission);
+	ConnectNodeOutputToInput(mx::SurfaceUnlit::Input::Emission, SurfaceUnlitShaderNode, SurfaceUnlit::Parameters::Emission.ToString(), mx::SurfaceUnlit::DefaultValue::Float::Emission);
 
 	//Emission Color
-	ConnectNodeOutputToInput(mx::SurfaceUnlit::Input::EmissionColor, ShaderGraphNode, SurfaceUnlit::Parameters::EmissionColor.ToString(), mx::SurfaceUnlit::DefaultValue::Color3::EmissionColor);
+	ConnectNodeOutputToInput(mx::SurfaceUnlit::Input::EmissionColor, SurfaceUnlitShaderNode, SurfaceUnlit::Parameters::EmissionColor.ToString(), mx::SurfaceUnlit::DefaultValue::Color3::EmissionColor);
 
 	//Opacity
-	ConnectNodeOutputToInput(mx::SurfaceUnlit::Input::Opacity, ShaderGraphNode, SurfaceUnlit::Parameters::Opacity.ToString(), mx::SurfaceUnlit::DefaultValue::Float::Opacity);
+	ConnectNodeOutputToInput(mx::SurfaceUnlit::Input::Opacity, SurfaceUnlitShaderNode, SurfaceUnlit::Parameters::Opacity.ToString(), mx::SurfaceUnlit::DefaultValue::Float::Opacity);
 
 	//Transmission
-	ConnectNodeOutputToInput(mx::SurfaceUnlit::Input::Transmission, ShaderGraphNode, SurfaceUnlit::Parameters::Transmission.ToString(), mx::SurfaceUnlit::DefaultValue::Float::Transmission);
+	ConnectNodeOutputToInput(mx::SurfaceUnlit::Input::Transmission, SurfaceUnlitShaderNode, SurfaceUnlit::Parameters::Transmission.ToString(), mx::SurfaceUnlit::DefaultValue::Float::Transmission);
 
 	//Transmission Color
-	ConnectNodeOutputToInput(mx::SurfaceUnlit::Input::TransmissionColor, ShaderGraphNode, SurfaceUnlit::Parameters::TransmissionColor.ToString(), mx::SurfaceUnlit::DefaultValue::Color3::TransmissionColor);
+	ConnectNodeOutputToInput(mx::SurfaceUnlit::Input::TransmissionColor, SurfaceUnlitShaderNode, SurfaceUnlit::Parameters::TransmissionColor.ToString(), mx::SurfaceUnlit::DefaultValue::Color3::TransmissionColor);
+
+	// Outputs
+	UInterchangeShaderPortsAPI::ConnectOuputToInputByName(ShaderGraphNode, PBR::Parameters::EmissiveColor.ToString(), SurfaceUnlitShaderNode->GetUniqueID(), PBR::Parameters::EmissiveColor.ToString());
+
+	//We can't have both Opacity and Opacity Mask we need to make a choice	
+	if(UInterchangeShaderPortsAPI::HasInput(SurfaceUnlitShaderNode, SurfaceUnlit::Parameters::Transmission))
+	{
+		UInterchangeShaderPortsAPI::ConnectOuputToInputByName(ShaderGraphNode, PBR::Parameters::Opacity.ToString(), SurfaceUnlitShaderNode->GetUniqueID(), PBR::Parameters::Opacity.ToString());
+	}
+	else if(UInterchangeShaderPortsAPI::HasInput(SurfaceUnlitShaderNode, SurfaceUnlit::Parameters::Opacity))
+	{
+		UInterchangeShaderPortsAPI::ConnectOuputToInputByName(ShaderGraphNode, PBR::Parameters::Opacity.ToString(), SurfaceUnlitShaderNode->GetUniqueID(), TEXT("OpacityMask"));
+		ShaderGraphNode->SetCustomOpacityMaskClipValue(1.f); // just to connect to the opacity mask
+	}
 }
 #endif
\ No newline at end of file
diff --git a/Engine/Plugins/Interchange/Runtime/Source/Import/Public/MaterialX/MaterialXUtils/MaterialXManager.h b/Engine/Plugins/Interchange/Runtime/Source/Import/Public/MaterialX/MaterialXUtils/MaterialXManager.h
--- a/Engine/Plugins/Interchange/Runtime/Source/Import/Public/MaterialX/MaterialXUtils/MaterialXManager.h
+++ b/Engine/Plugins/Interchange/Runtime/Source/Import/Public/MaterialX/MaterialXUtils/MaterialXManager.h
@@ -36,6 +36,8 @@ public:
 	FMaterialXManager(const FMaterialXManager&) = delete;
 	FMaterialXManager& operator=(const FMaterialXManager&) = delete;
 
+	static const TCHAR TexturePayloadSeparator;
+
 private:
 
 	FMaterialXManager();
diff --git a/Engine/Plugins/Interchange/Runtime/Source/Import/Public/MaterialX/MaterialXUtils/MaterialXSurfaceShaderAbstract.h b/Engine/Plugins/Interchange/Runtime/Source/Import/Public/MaterialX/MaterialXUtils/MaterialXSurfaceShaderAbstract.h
--- a/Engine/Plugins/Interchange/Runtime/Source/Import/Public/MaterialX/MaterialXUtils/MaterialXSurfaceShaderAbstract.h
+++ b/Engine/Plugins/Interchange/Runtime/Source/Import/Public/MaterialX/MaterialXUtils/MaterialXSurfaceShaderAbstract.h
@@ -3,10 +3,12 @@
 #pragma once
 
 #if WITH_EDITOR
+#include "Engine/TextureDefines.h"
 #include "InterchangeShaderGraphNode.h"
 #include "InterchangeTranslatorBase.h"
 #include "InterchangeTexture2DNode.h"
 #include "MaterialXBase.h"
+#include "MaterialXManager.h"
 
 class FMaterialXSurfaceShaderAbstract : public FMaterialXBase
 {
@@ -71,13 +73,16 @@ protected:
 	 * @param ShaderNode - The Interchange shader node to connect the MaterialX's node or node graph to
 	 * @param InputShaderName - The name of the input of the shader node to connect to
 	 * @param DefaultValue - The default value of the MaterialX input
+	 * @param OptionalTextureCompression - Set the texture compression for all textures along the path of an input
 	 */
 	template<typename T>
-	bool ConnectNodeOutputToInput(const char* InputName, UInterchangeShaderNode* ShaderNode, const FString& InputShaderName, T DefaultValue)
+	bool ConnectNodeOutputToInput(const char* InputName, UInterchangeShaderNode* ShaderNode, const FString& InputShaderName, T DefaultValue, const TOptional<TextureCompressionSettings>& OptionalTextureCompression = TOptional<TextureCompressionSettings>{})
 	{
 		MaterialX::DocumentPtr Document = SurfaceShaderNode->getDocument();
 		MaterialX::InputPtr Input = GetInput(SurfaceShaderNode, InputName);
 
+		TGuardValue<TOptional<TextureCompressionSettings>> InputTypeBeingProcessedGuard(TextureCompression, OptionalTextureCompression);
+
 		bool bIsConnected = ConnectNodeGraphOutputToInput(Input, ShaderNode, InputShaderName);
 
 		if(!bIsConnected)
@@ -306,6 +311,12 @@ protected:
 
 					TextureNode->SetCustomWrapU(WrapModeU);
 					TextureNode->SetCustomWrapV(WrapModeV);
+
+					// Encode the compression in the payloadKey
+					if(TextureCompression.IsSet())
+					{
+						TextureNode->SetPayLoadKey(TextureNode->GetPayLoadKey().GetValue() + FMaterialXManager::TexturePayloadSeparator + FString::FromInt(TextureCompression.GetValue()));
+					}
 				}
 			}
 		}
@@ -390,5 +401,6 @@ protected:
 	/** Initialized by the material shader (e.g: surfacematerial), the derived class should only set the ShaderType */
 	UInterchangeShaderGraphNode* ShaderGraphNode;
 
+	TOptional<TextureCompressionSettings> TextureCompression;
 };
 #endif
\ No newline at end of file
diff --git a/Engine/Plugins/Interchange/Runtime/Source/Pipelines/Private/InterchangeGenericMaterialPipeline.cpp b/Engine/Plugins/Interchange/Runtime/Source/Pipelines/Private/InterchangeGenericMaterialPipeline.cpp
--- a/Engine/Plugins/Interchange/Runtime/Source/Pipelines/Private/InterchangeGenericMaterialPipeline.cpp
+++ b/Engine/Plugins/Interchange/Runtime/Source/Pipelines/Private/InterchangeGenericMaterialPipeline.cpp
@@ -914,202 +914,27 @@ bool UInterchangeGenericMaterialPipeline::HandlePBRModel(const UInterchangeShade
 		}
 	}
 
-	return bShadingModelHandled;
-}
-
-FString UInterchangeGenericMaterialPipeline::HandleSurfaceShaderModel(const UInterchangeShaderGraphNode* ShaderGraphNode, UInterchangeMaterialFactoryNode* MaterialFactoryNode, const FSurfaceShader& SurfaceShader)
-{
-	using namespace UE::Interchange::InterchangeGenericMaterialPipeline::Private;
-
-	UInterchangeMaterialExpressionFactoryNode* FunctionCallExpression = CreateExpressionWithMaterialFunction(BaseNodeContainer, MaterialFactoryNode, SurfaceShader.UniqueID, SurfaceShader.Path);
-
-	const FString FunctionCallExpressionUid = FunctionCallExpression->GetUniqueID();
-
-	auto ConnectMaterialExpressionOutputToInput = [&](const FString& InputName, EMaterialInputType InputType)
+	// Anisotropy
 	{
-		TGuardValue<EMaterialInputType> InputTypeBeingProcessedGuard(MaterialCreationContext.InputTypeBeingProcessed, InputType);
+		const bool bHasInput = UInterchangeShaderPortsAPI::HasInput(ShaderGraphNode, Parameters::Anisotropy);
 
-		TTuple<UInterchangeMaterialExpressionFactoryNode*, FString> ExpressionFactoryNode =
-			CreateMaterialExpressionForInput(MaterialFactoryNode, ShaderGraphNode, InputName, FunctionCallExpressionUid);
-
-		if(ExpressionFactoryNode.Get<0>())
-		{
-			UInterchangeShaderPortsAPI::ConnectOuputToInputByName(FunctionCallExpression, InputName,
-															ExpressionFactoryNode.Get<0>()->GetUniqueID(), ExpressionFactoryNode.Get<1>());
-		}
-	};
-
-	for(auto& Arguments : SurfaceShader.Parameters)
-	{
-		const FString& InputName = Arguments.Get<0>();
-		EMaterialInputType InputType = Arguments.Get<1>();
-		ConnectMaterialExpressionOutputToInput(InputName, InputType);
-	}
-
-	return FunctionCallExpressionUid;
-}
-
-bool UInterchangeGenericMaterialPipeline::HandleStandardSurfaceModel(const UInterchangeShaderGraphNode* ShaderGraphNode, UInterchangeMaterialFactoryNode* MaterialFactoryNode)
-{
-	using namespace UE::Interchange::Materials;
-
-	if(IsStandardSurfaceModel(ShaderGraphNode))
-	{
-		static FSurfaceShader StandardSurfaceShader
-		{
-			// Parameters
-			{
-			{StandardSurface::Parameters::Base.ToString(), EMaterialInputType::Scalar},                          // Base
-			{StandardSurface::Parameters::BaseColor.ToString(), EMaterialInputType::Color},                      // BaseColor
-			{StandardSurface::Parameters::DiffuseRoughness.ToString(),  EMaterialInputType::Scalar},             // DiffuseRoughness
-			{StandardSurface::Parameters::Metalness.ToString(), EMaterialInputType::Scalar},                     // Metalness
-			{StandardSurface::Parameters::Specular.ToString(), EMaterialInputType::Scalar},                      // Specular
-			{StandardSurface::Parameters::SpecularRoughness.ToString(), EMaterialInputType::Scalar},             // SpecularRoughness
-			{StandardSurface::Parameters::SpecularIOR.ToString(), EMaterialInputType::Scalar},                   // SpecularIOR
-			{StandardSurface::Parameters::SpecularAnisotropy.ToString(), EMaterialInputType::Scalar},            // SpecularAnisotropy
-			{StandardSurface::Parameters::SpecularRotation.ToString(), EMaterialInputType::Scalar},              // SpecularRotation
-			{StandardSurface::Parameters::Subsurface.ToString(), EMaterialInputType::Scalar},                    // Subsurface
-			{StandardSurface::Parameters::SubsurfaceColor.ToString(), EMaterialInputType::Scalar},               // SubsurfaceColor
-			{StandardSurface::Parameters::SubsurfaceRadius.ToString(), EMaterialInputType::Scalar},              // SubsurfaceRadius
-			{StandardSurface::Parameters::SubsurfaceScale.ToString(), EMaterialInputType::Scalar},               // SubsurfaceScale
-			{StandardSurface::Parameters::Sheen.ToString(), EMaterialInputType::Scalar},                         // Sheen
-			{StandardSurface::Parameters::SheenColor.ToString(), EMaterialInputType::Color},                     // SheenColor
-			{StandardSurface::Parameters::SheenRoughness.ToString(), EMaterialInputType::Scalar},                // SheenRoughness
-			{StandardSurface::Parameters::Coat.ToString(), EMaterialInputType::Scalar},                          // Coat
-			{StandardSurface::Parameters::CoatColor.ToString(),  EMaterialInputType::Color},                     // CoatColor
-			{StandardSurface::Parameters::CoatNormal.ToString(),  EMaterialInputType::Vector},                   // CoatNormal
-			{StandardSurface::Parameters::CoatRoughness.ToString(), EMaterialInputType::Scalar},                 // CoatRoughness
-			{StandardSurface::Parameters::ThinFilmThickness.ToString(), EMaterialInputType::Scalar},             // ThinFilmThickness
-			{StandardSurface::Parameters::Emission.ToString(), EMaterialInputType::Scalar},                      // Emission
-			{StandardSurface::Parameters::EmissionColor.ToString(), EMaterialInputType::Color},                  // EmissionColor
-			{StandardSurface::Parameters::Normal.ToString(), EMaterialInputType::Vector},                        // Normal
-			{StandardSurface::Parameters::Tangent.ToString(), EMaterialInputType::Vector},                       // Tangent
-			{StandardSurface::Parameters::Transmission.ToString(),  EMaterialInputType::Scalar},                 // Transmission
-			{StandardSurface::Parameters::TransmissionColor.ToString(), EMaterialInputType::Color},              // TransmissionColor
-			{StandardSurface::Parameters::TransmissionDepth.ToString(), EMaterialInputType::Scalar},             // TransmissionDepth
-			{StandardSurface::Parameters::TransmissionScatter.ToString(), EMaterialInputType::Color},            // TransmissionScatter
-			{StandardSurface::Parameters::TransmissionScatterAnisotropy.ToString(), EMaterialInputType::Scalar}, // TransmissionScatterAnisotropy
-			{StandardSurface::Parameters::TransmissionDispersion.ToString(), EMaterialInputType::Scalar},        // TransmissionDispersion
-			{StandardSurface::Parameters::TransmissionExtraRoughness.ToString(), EMaterialInputType::Scalar},    // TransmissionExtraRoughness
-			},
-
-			// UniqueID
-			TEXT("StandardSurface")
-		};
-
-		StandardSurfaceShader.Path = MaterialXSettings.GetAssetPathString(EInterchangeMaterialXShaders::StandardSurface);
-
-		if(UInterchangeShaderPortsAPI::HasInput(ShaderGraphNode, StandardSurface::Parameters::Transmission))
-		{
-			StandardSurfaceShader.Path = MaterialXSettings.GetAssetPathString(EInterchangeMaterialXShaders::StandardSurfaceTransmission);
-		}
-
-		const FString FunctionCallExpressionUid = HandleSurfaceShaderModel(ShaderGraphNode, MaterialFactoryNode, StandardSurfaceShader);
-		
-		MaterialFactoryNode->ConnectOutputToBaseColor(FunctionCallExpressionUid,  PBR::Parameters::BaseColor.ToString());
-		MaterialFactoryNode->ConnectOutputToMetallic(FunctionCallExpressionUid, PBR::Parameters::Metallic.ToString());
-		MaterialFactoryNode->ConnectOutputToSpecular(FunctionCallExpressionUid, PBR::Parameters::Specular.ToString());
-		MaterialFactoryNode->ConnectOutputToRoughness(FunctionCallExpressionUid, PBR::Parameters::Roughness.ToString());
-		MaterialFactoryNode->ConnectOutputToEmissiveColor(FunctionCallExpressionUid, PBR::Parameters::EmissiveColor.ToString());
-		MaterialFactoryNode->ConnectOutputToAnisotropy(FunctionCallExpressionUid, PBR::Parameters::Anisotropy.ToString());
-		MaterialFactoryNode->ConnectOutputToNormal(FunctionCallExpressionUid, PBR::Parameters::Normal.ToString());
-		MaterialFactoryNode->ConnectOutputToTangent(FunctionCallExpressionUid, PBR::Parameters::Tangent.ToString());
-		MaterialFactoryNode->ConnectOutputToOpacity(FunctionCallExpressionUid, PBR::Parameters::Opacity.ToString());
-		
-
-		// We can't have all shading models at once, so we have to make a choice here
-		if(UInterchangeShaderPortsAPI::HasInput(ShaderGraphNode, StandardSurface::Parameters::Sheen))
-		{
-			MaterialFactoryNode->ConnectOutputToFuzzColor(FunctionCallExpressionUid, Sheen::Parameters::SheenColor.ToString());
-			MaterialFactoryNode->ConnectOutputToCloth(FunctionCallExpressionUid, Sheen::Parameters::SheenRoughness.ToString());
-			MaterialFactoryNode->SetCustomShadingModel(MSM_Cloth);
-		}
-		else if(UInterchangeShaderPortsAPI::HasInput(ShaderGraphNode, StandardSurface::Parameters::Coat))
-		{
-			MaterialFactoryNode->ConnectOutputToClearCoat(FunctionCallExpressionUid, ClearCoat::Parameters::ClearCoat.ToString());
-			MaterialFactoryNode->ConnectOutputToClearCoatRoughness(FunctionCallExpressionUid, ClearCoat::Parameters::ClearCoatRoughness.ToString());
-			MaterialFactoryNode->ConnectOutputToClearCoatNormal(FunctionCallExpressionUid, ClearCoat::Parameters::ClearCoatNormal.ToString());
-			MaterialFactoryNode->SetCustomShadingModel(MSM_ClearCoat);
-		}
-		else if(UInterchangeShaderPortsAPI::HasInput(ShaderGraphNode, StandardSurface::Parameters::Subsurface))
-		{
-			MaterialFactoryNode->ConnectOutputToSubsurface(FunctionCallExpressionUid, Subsurface::Parameters::SubsurfaceColor.ToString());
-			MaterialFactoryNode->SetCustomShadingModel(MSM_Subsurface);
-		}
-		else if (UInterchangeShaderPortsAPI::HasInput(ShaderGraphNode, StandardSurface::Parameters::Transmission))
-		{
-			MaterialFactoryNode->ConnectOutputToTransmissionColor(FunctionCallExpressionUid, ThinTranslucent::Parameters::TransmissionColor.ToString());
-			MaterialFactoryNode->ConnectOutputToRefraction(FunctionCallExpressionUid, Common::Parameters::Refraction.ToString());
-
-			MaterialFactoryNode->SetCustomBlendMode(EBlendMode::BLEND_Translucent);
-			MaterialFactoryNode->SetCustomShadingModel(EMaterialShadingModel::MSM_ThinTranslucent);
-			MaterialFactoryNode->SetCustomTranslucencyLightingMode(ETranslucencyLightingMode::TLM_SurfacePerPixelLighting);
-		}
-		else
+		if(bHasInput)
 		{
-			MaterialFactoryNode->SetCustomShadingModel(MSM_DefaultLit);
-		}
-
-		return true;
-	}
-
-	return false;
-}
+			TGuardValue<EMaterialInputType> InputTypeBeingProcessedGuard(MaterialCreationContext.InputTypeBeingProcessed, EMaterialInputType::Scalar);
 
-bool UInterchangeGenericMaterialPipeline::HandleSurfaceUnlitModel(const UInterchangeShaderGraphNode* ShaderGraphNode, UInterchangeMaterialFactoryNode* MaterialFactoryNode)
-{
-	using namespace UE::Interchange::Materials;
-	using namespace UE::Interchange::InterchangeGenericMaterialPipeline;
+			TTuple<UInterchangeMaterialExpressionFactoryNode*, FString> ExpressionFactoryNode =
+				CreateMaterialExpressionForInput(MaterialFactoryNode, ShaderGraphNode, Parameters::Anisotropy.ToString(), MaterialFactoryNode->GetUniqueID());
 
-	if(IsSurfaceUnlitModel(ShaderGraphNode))
-	{
-		static FSurfaceShader StandardSurfaceShader
-		{
-			// Parameters
+			if(ExpressionFactoryNode.Get<0>())
 			{
-			{SurfaceUnlit::Parameters::Emission.ToString(), EMaterialInputType::Scalar},                      // Emission
-			{SurfaceUnlit::Parameters::EmissionColor.ToString(), EMaterialInputType::Color},                  // EmissionColor
-			{SurfaceUnlit::Parameters::Opacity.ToString(), EMaterialInputType::Scalar},                       // Opacity
-			{SurfaceUnlit::Parameters::Transmission.ToString(),  EMaterialInputType::Scalar},                 // Transmission
-			{SurfaceUnlit::Parameters::TransmissionColor.ToString(), EMaterialInputType::Color},              // TransmissionColor
-			},
-
-			// UniqueID
-			TEXT("SurfaceUnlit"),
-
-			// Path
-			MaterialXSettings.GetAssetPathString(EInterchangeMaterialXShaders::SurfaceUnlit)
-		};
-
-		const FString FunctionCallExpressionUid = HandleSurfaceShaderModel(ShaderGraphNode, MaterialFactoryNode, StandardSurfaceShader);
-
-		MaterialFactoryNode->ConnectOutputToEmissiveColor(FunctionCallExpressionUid, PBR::Parameters::EmissiveColor.ToString());
-		
-		//We can't have both Opacity and Opacity Mask we need to make a choice
-		bool bHasOpacityOrTransmission = false;
-		if(UInterchangeShaderPortsAPI::HasInput(ShaderGraphNode, SurfaceUnlit::Parameters::Transmission))
-		{
-			MaterialFactoryNode->ConnectOutputToOpacity(FunctionCallExpressionUid, PBR::Parameters::Opacity.ToString());
-			bHasOpacityOrTransmission = true;
-		}
-		else if(UInterchangeShaderPortsAPI::HasInput(ShaderGraphNode, SurfaceUnlit::Parameters::Opacity))
-		{
-			MaterialFactoryNode->ConnectOutputToOpacity(FunctionCallExpressionUid, "OpacityMask");
-			MaterialFactoryNode->SetCustomOpacityMaskClipValue(0.f); // just to connect to the opacity mask
-			bHasOpacityOrTransmission = true;
-		}
+				MaterialFactoryNode->ConnectOutputToAnisotropy(ExpressionFactoryNode.Get<0>()->GetUniqueID(), ExpressionFactoryNode.Get<1>());
+			}
 
-		if(bHasOpacityOrTransmission)
-		{
-			Private::UpdateBlendModeBasedOnOpacityAttributes(ShaderGraphNode, MaterialFactoryNode);
+			bShadingModelHandled = true;
 		}
-
-		MaterialFactoryNode->SetCustomShadingModel(EMaterialShadingModel::MSM_Unlit);
-
-		return true;
 	}
 
-	return false;
+	return bShadingModelHandled;
 }
 
 bool UInterchangeGenericMaterialPipeline::HandleClearCoat(const UInterchangeShaderGraphNode* ShaderGraphNode, UInterchangeMaterialFactoryNode* MaterialFactoryNode)
@@ -1213,6 +1038,7 @@ bool UInterchangeGenericMaterialPipeline::HandleSubsurface(const UInterchangeSha
 	if(bShadingModelHandled)
 	{
 		MaterialFactoryNode->SetCustomShadingModel(EMaterialShadingModel::MSM_Subsurface);
+		MaterialFactoryNode->SetCustomBlendMode(EBlendMode::BLEND_Opaque); // Opacity in Subsurface doesn't mean Translucency, according to UE doc
 	}
 
 	return bShadingModelHandled;
@@ -1253,13 +1079,7 @@ bool UInterchangeGenericMaterialPipeline::HandleSheen(const UInterchangeShaderGr
 
 			if (ExpressionFactoryNode.Get<0>())
 			{
-				UInterchangeMaterialExpressionFactoryNode* InverseSheenRoughnessNode =
-						CreateExpressionNode(TEXT("InverseSheenRoughness"), ExpressionFactoryNode.Get<0>()->GetUniqueID(), UMaterialExpressionOneMinus::StaticClass());
-
-				UInterchangeShaderPortsAPI::ConnectOuputToInputByName(InverseSheenRoughnessNode, GET_MEMBER_NAME_CHECKED(UMaterialExpressionOneMinus, Input).ToString(),
-					ExpressionFactoryNode.Get<0>()->GetUniqueID(), ExpressionFactoryNode.Get<1>());
-
-				MaterialFactoryNode->ConnectToCloth(InverseSheenRoughnessNode->GetUniqueID());
+				MaterialFactoryNode->ConnectOutputToCloth(ExpressionFactoryNode.Get<0>()->GetUniqueID(), ExpressionFactoryNode.Get<1>());
 			}
 
 			bShadingModelHandled = true;
@@ -1354,6 +1174,24 @@ void UInterchangeGenericMaterialPipeline::HandleCommonParameters(const UIntercha
 		}
 	}
 
+	// Tangent
+	{
+		const bool bHasNormalInput = UInterchangeShaderPortsAPI::HasInput(ShaderGraphNode, Parameters::Tangent);
+
+		if(bHasNormalInput)
+		{
+			TGuardValue<EMaterialInputType> InputTypeBeingProcessedGuard(MaterialCreationContext.InputTypeBeingProcessed, EMaterialInputType::Vector);
+
+			TTuple<UInterchangeMaterialExpressionFactoryNode*, FString> ExpressionFactoryNode =
+				CreateMaterialExpressionForInput(MaterialFactoryNode, ShaderGraphNode, Parameters::Tangent.ToString(), MaterialFactoryNode->GetUniqueID());
+
+			if(ExpressionFactoryNode.Get<0>())
+			{
+				MaterialFactoryNode->ConnectOutputToTangent(ExpressionFactoryNode.Get<0>()->GetUniqueID(), ExpressionFactoryNode.Get<1>());
+			}
+		}
+	}
+
 	// Opacity
 	{
 		TGuardValue<EMaterialInputType> InputTypeBeingProcessedGuard(MaterialCreationContext.InputTypeBeingProcessed, EMaterialInputType::Scalar);
@@ -1405,6 +1243,26 @@ void UInterchangeGenericMaterialPipeline::HandleCommonParameters(const UIntercha
 		}
 	}
 
+	// Refraction
+	// probably unlikely that someone will use both at same time but to keep backwards compability IndexOfRefraction will override this one
+	{
+		TGuardValue<EMaterialInputType> InputTypeBeingProcessedGuard(MaterialCreationContext.InputTypeBeingProcessed, EMaterialInputType::Scalar);
+
+		const bool bHasIorInput = UInterchangeShaderPortsAPI::HasInput(ShaderGraphNode, Parameters::Refraction);
+
+		if(bHasIorInput)
+		{
+			TTuple<UInterchangeMaterialExpressionFactoryNode*, FString> ExpressionFactoryNode =
+				CreateMaterialExpressionForInput(MaterialFactoryNode, ShaderGraphNode, Parameters::Refraction.ToString(), MaterialFactoryNode->GetUniqueID());
+
+			if(ExpressionFactoryNode.Get<0>())
+			{
+				MaterialFactoryNode->SetCustomRefractionMethod(ERefractionMode::RM_IndexOfRefraction);
+				MaterialFactoryNode->ConnectOutputToRefraction(ExpressionFactoryNode.Get<0>()->GetUniqueID(), ExpressionFactoryNode.Get<1>());
+			}
+		}
+	}
+
 	// Index of Refraction (IOR)
 	// We'll lerp between Air IOR (1) and the IOR from the shader graph based on a fresnel, as per UE doc on refraction.
 	{
@@ -1419,6 +1277,7 @@ void UInterchangeGenericMaterialPipeline::HandleCommonParameters(const UIntercha
 
 			if (ExpressionFactoryNode.Get<0>())
 			{
+				MaterialFactoryNode->SetCustomRefractionMethod(ERefractionMode::RM_IndexOfRefraction);
 				UInterchangeMaterialExpressionFactoryNode* IORLerp = CreateExpressionNode(TEXT("IORLerp"), ShaderGraphNode->GetUniqueID(), UMaterialExpressionLinearInterpolate::StaticClass());
 
 				const float AirIOR = 1.f;
@@ -2513,16 +2372,6 @@ UInterchangeMaterialFactoryNode* UInterchangeGenericMaterialPipeline::CreateMate
 {
 	UInterchangeMaterialFactoryNode* MaterialFactoryNode = Cast<UInterchangeMaterialFactoryNode>( CreateBaseMaterialFactoryNode(ShaderGraphNode, UInterchangeMaterialFactoryNode::StaticClass()) );
 
-	if(HandleStandardSurfaceModel(ShaderGraphNode, MaterialFactoryNode))
-	{
-		return MaterialFactoryNode;
-	}
-
-	if(HandleSurfaceUnlitModel(ShaderGraphNode, MaterialFactoryNode))
-	{
-		return MaterialFactoryNode;
-	}
-
 	// Handle the case where the material will be connected through the material attributes input
 	if (HandleBxDFInput(ShaderGraphNode, MaterialFactoryNode))
 	{
diff --git a/Engine/Plugins/Interchange/Runtime/Source/Pipelines/Public/InterchangeGenericMaterialPipeline.h b/Engine/Plugins/Interchange/Runtime/Source/Pipelines/Public/InterchangeGenericMaterialPipeline.h
--- a/Engine/Plugins/Interchange/Runtime/Source/Pipelines/Public/InterchangeGenericMaterialPipeline.h
+++ b/Engine/Plugins/Interchange/Runtime/Source/Pipelines/Public/InterchangeGenericMaterialPipeline.h
@@ -163,8 +163,6 @@ private:
 	bool HandlePhongModel(const UInterchangeShaderGraphNode* ShaderGraphNode, UInterchangeMaterialFactoryNode* MaterialFactoryNode);
 	bool HandleLambertModel(const UInterchangeShaderGraphNode* ShaderGraphNode, UInterchangeMaterialFactoryNode* MaterialFactoryNode);
 	bool HandlePBRModel(const UInterchangeShaderGraphNode* ShaderGraphNode, UInterchangeMaterialFactoryNode* MaterialFactoryNode);
-	bool HandleStandardSurfaceModel(const UInterchangeShaderGraphNode* ShaderGraphNode, UInterchangeMaterialFactoryNode* MaterialFactoryNode);
-	bool HandleSurfaceUnlitModel(const UInterchangeShaderGraphNode* ShaderGraphNode, UInterchangeMaterialFactoryNode* MaterialFactoryNode);
 	bool HandleClearCoat(const UInterchangeShaderGraphNode* ShaderGraphNode, UInterchangeMaterialFactoryNode* MaterialFactoryNode);
 	bool HandleSubsurface(const UInterchangeShaderGraphNode* ShaderGraphNode, UInterchangeMaterialFactoryNode* MaterialFactoryNode);
 	bool HandleSheen(const UInterchangeShaderGraphNode* ShaderGraphNode, UInterchangeMaterialFactoryNode* MaterialFactoryNode);
@@ -239,13 +237,4 @@ private:
 	TArray<FMaterialExpressionCreationContext> MaterialExpressionCreationContextStack;
 
 	using FParameterMaterialInputType = TTuple<FString, EMaterialInputType>;
-
-	struct FSurfaceShader
-	{
-		TArray<FParameterMaterialInputType> Parameters;
-		FString UniqueID;
-		FString Path;
-	};
-
-	FString HandleSurfaceShaderModel(const UInterchangeShaderGraphNode* ShaderGraphNode, UInterchangeMaterialFactoryNode* MaterialFactoryNode, const FSurfaceShader& SurfaceShader);
 };
