Fixed issue with file path stored in texture assets created from  non-texture related translators
  - Deprecated all API with signature API(const UInterchangeSourceData* PayloadSourceData, const FString& PayloadKey)
  - Replaced depreacted API with API(const FString& PayloadKey, TOptional<FString>& AlternateTexturePath)
  - Replaced use of UE_LOG with a new helper to log error in the result's message list
  - Refactor some similar code used in several cpp

diff --git a/Engine/Plugins/Experimental/Enterprise/DatasmithInterchange/Source/DatasmithInterchange/Private/InterchangeDatasmithTranslator.cpp b/Engine/Plugins/Experimental/Enterprise/DatasmithInterchange/Source/DatasmithInterchange/Private/InterchangeDatasmithTranslator.cpp
--- a/Engine/Plugins/Experimental/Enterprise/DatasmithInterchange/Source/DatasmithInterchange/Private/InterchangeDatasmithTranslator.cpp
+++ b/Engine/Plugins/Experimental/Enterprise/DatasmithInterchange/Source/DatasmithInterchange/Private/InterchangeDatasmithTranslator.cpp
@@ -508,7 +508,7 @@ UInterchangeBaseLightNode* UInterchangeDatasmithTranslator::AddLightNode(UInterc
 	return LightNode;
 }
 
-TOptional<UE::Interchange::FImportImage> UInterchangeDatasmithTranslator::GetTexturePayloadData(const UInterchangeSourceData* InPayloadSourceData, const FString& PayloadKey) const
+TOptional<UE::Interchange::FImportImage> UInterchangeDatasmithTranslator::GetTexturePayloadData(const FString& PayloadKey, TOptional<FString>& AlternateTexturePath) const
 {
 	if (!LoadedExternalSource || !LoadedExternalSource->GetDatasmithScene())
 	{
@@ -545,7 +545,9 @@ TOptional<UE::Interchange::FImportImage> UInterchangeDatasmithTranslator::GetTex
 		return TOptional<UE::Interchange::FImportImage>();
 	}
 
-	return TextureTranslator->GetTexturePayloadData(PayloadSourceData, TextureElement->GetFile());
+	AlternateTexturePath = TextureElement->GetFile();
+
+	return TextureTranslator->GetTexturePayloadData(PayloadKey, AlternateTexturePath);
 }
 
 TFuture<TOptional<UE::Interchange::FMeshPayloadData>> UInterchangeDatasmithTranslator::GetMeshPayloadData(const FInterchangeMeshPayLoadKey& PayLoadKey) const
diff --git a/Engine/Plugins/Experimental/Enterprise/DatasmithInterchange/Source/DatasmithInterchange/Public/InterchangeDatasmithTranslator.h b/Engine/Plugins/Experimental/Enterprise/DatasmithInterchange/Source/DatasmithInterchange/Public/InterchangeDatasmithTranslator.h
--- a/Engine/Plugins/Experimental/Enterprise/DatasmithInterchange/Source/DatasmithInterchange/Public/InterchangeDatasmithTranslator.h
+++ b/Engine/Plugins/Experimental/Enterprise/DatasmithInterchange/Source/DatasmithInterchange/Public/InterchangeDatasmithTranslator.h
@@ -75,7 +75,7 @@ public:
 	virtual void ImportFinish() override;
 
 	/* IInterchangeTexturePayloadInterface Begin */
-	virtual TOptional<UE::Interchange::FImportImage> GetTexturePayloadData(const UInterchangeSourceData* InPayloadSourceData, const FString& PayloadKey) const override;
+	virtual TOptional<UE::Interchange::FImportImage> GetTexturePayloadData(const FString& PayloadKey, TOptional<FString>& AlternateTexturePath) const override;
 	/* IInterchangeTexturePayloadInterface End */
 
 	/* IInterchangeStaticMeshPayloadInterface Begin */
diff --git a/Engine/Plugins/Interchange/Runtime/Source/Import/Private/Fbx/InterchangeFbxTranslator.cpp b/Engine/Plugins/Interchange/Runtime/Source/Import/Private/Fbx/InterchangeFbxTranslator.cpp
--- a/Engine/Plugins/Interchange/Runtime/Source/Import/Private/Fbx/InterchangeFbxTranslator.cpp
+++ b/Engine/Plugins/Interchange/Runtime/Source/Import/Private/Fbx/InterchangeFbxTranslator.cpp
@@ -163,7 +163,7 @@ void UInterchangeFbxTranslator::ImportFinish()
 }
 
 
-TOptional<UE::Interchange::FImportImage> UInterchangeFbxTranslator::GetTexturePayloadData(const UInterchangeSourceData* InSourceData, const FString& PayLoadKey) const
+TOptional<UE::Interchange::FImportImage> UInterchangeFbxTranslator::GetTexturePayloadData(const FString& PayLoadKey, TOptional<FString>& AlternateTexturePath) const
 {
 	UInterchangeSourceData* PayloadSourceData = UInterchangeManager::GetInterchangeManager().CreateSourceData(PayLoadKey);
 	FGCObjectScopeGuard ScopedSourceData(PayloadSourceData);
@@ -179,7 +179,10 @@ TOptional<UE::Interchange::FImportImage> UInterchangeFbxTranslator::GetTexturePa
 	{
 		return TOptional<UE::Interchange::FImportImage>();
 	}
-	return TextureTranslator->GetTexturePayloadData(PayloadSourceData, PayLoadKey);
+
+	AlternateTexturePath = PayLoadKey;
+
+	return TextureTranslator->GetTexturePayloadData(PayLoadKey, AlternateTexturePath);
 }
 
 TFuture<TOptional<UE::Interchange::FMeshPayloadData>> UInterchangeFbxTranslator::GetMeshPayloadData(const FInterchangeMeshPayLoadKey& PayLoadKey) const
diff --git a/Engine/Plugins/Interchange/Runtime/Source/Import/Private/Gltf/InterchangeGltfTranslator.cpp b/Engine/Plugins/Interchange/Runtime/Source/Import/Private/Gltf/InterchangeGltfTranslator.cpp
--- a/Engine/Plugins/Interchange/Runtime/Source/Import/Private/Gltf/InterchangeGltfTranslator.cpp
+++ b/Engine/Plugins/Interchange/Runtime/Source/Import/Private/Gltf/InterchangeGltfTranslator.cpp
@@ -1316,10 +1316,10 @@ bool UInterchangeGltfTranslator::Translate( UInterchangeBaseNodeContainer& NodeC
 	return true;
 }
 
-TOptional< UE::Interchange::FImportImage > UInterchangeGltfTranslator::GetTexturePayloadData( const UInterchangeSourceData* InSourceData, const FString& PayLoadKey ) const
+TOptional< UE::Interchange::FImportImage > UInterchangeGltfTranslator::GetTexturePayloadData(const FString& PayloadKey, TOptional<FString>& AlternateTexturePath) const
 {
 	int32 TextureIndex = 0;
-	LexFromString( TextureIndex, *PayLoadKey );
+	LexFromString( TextureIndex, *PayloadKey);
 
 	if ( !GltfAsset.Textures.IsValidIndex( TextureIndex ) )
 	{
@@ -1355,7 +1355,9 @@ TOptional< UE::Interchange::FImportImage > UInterchangeGltfTranslator::GetTextur
 			return TOptional<UE::Interchange::FImportImage>();
 		}
 
-		return TextureTranslator->GetTexturePayloadData(PayloadSourceData, TextureFilePath);
+		AlternateTexturePath = TextureFilePath;
+
+		return TextureTranslator->GetTexturePayloadData(PayloadKey, AlternateTexturePath);
 	}
 }
 
diff --git a/Engine/Plugins/Interchange/Runtime/Source/Import/Private/MaterialX/InterchangeMaterialXTranslator.cpp b/Engine/Plugins/Interchange/Runtime/Source/Import/Private/MaterialX/InterchangeMaterialXTranslator.cpp
--- a/Engine/Plugins/Interchange/Runtime/Source/Import/Private/MaterialX/InterchangeMaterialXTranslator.cpp
+++ b/Engine/Plugins/Interchange/Runtime/Source/Import/Private/MaterialX/InterchangeMaterialXTranslator.cpp
@@ -191,7 +191,7 @@ bool UInterchangeMaterialXTranslator::Translate(UInterchangeBaseNodeContainer& B
 	return bIsDocumentValid && bIsReferencesValid;
 }
 
-TOptional<UE::Interchange::FImportImage> UInterchangeMaterialXTranslator::GetTexturePayloadData(const UInterchangeSourceData* InSourceData, const FString& PayloadKey) const
+TOptional<UE::Interchange::FImportImage> UInterchangeMaterialXTranslator::GetTexturePayloadData(const FString& PayloadKey, TOptional<FString>& AlternateTexturePath) const
 {
 	FString Filename = PayloadKey;
 	TextureCompressionSettings CompressionSettings = TextureCompressionSettings::TC_Default;
@@ -220,7 +220,9 @@ TOptional<UE::Interchange::FImportImage> UInterchangeMaterialXTranslator::GetTex
 		return TOptional<UE::Interchange::FImportImage>();
 	}
 
-	TOptional<UE::Interchange::FImportImage> TexturePayloadData = TextureTranslator->GetTexturePayloadData(PayloadSourceData, Filename);
+	AlternateTexturePath = Filename;
+
+	TOptional<UE::Interchange::FImportImage> TexturePayloadData = TextureTranslator->GetTexturePayloadData(PayloadKey, AlternateTexturePath);
 
 	if(TexturePayloadData.IsSet())
 	{
diff --git a/Engine/Plugins/Interchange/Runtime/Source/Import/Private/Mesh/InterchangeOBJTranslator.cpp b/Engine/Plugins/Interchange/Runtime/Source/Import/Private/Mesh/InterchangeOBJTranslator.cpp
--- a/Engine/Plugins/Interchange/Runtime/Source/Import/Private/Mesh/InterchangeOBJTranslator.cpp
+++ b/Engine/Plugins/Interchange/Runtime/Source/Import/Private/Mesh/InterchangeOBJTranslator.cpp
@@ -1043,7 +1043,7 @@ TFuture<TOptional<UE::Interchange::FMeshPayloadData>> UInterchangeOBJTranslator:
 }
 
 
-TOptional<UE::Interchange::FImportImage> UInterchangeOBJTranslator::GetTexturePayloadData(const UInterchangeSourceData* InSourceData, const FString& PayLoadKey) const
+TOptional<UE::Interchange::FImportImage> UInterchangeOBJTranslator::GetTexturePayloadData(const FString& PayLoadKey, TOptional<FString>& AlternateTexturePath) const
 {
 	// @TODO: API design
 	// This method is copied verbatim from InterchangeFBXTranslator.
@@ -1066,7 +1066,9 @@ TOptional<UE::Interchange::FImportImage> UInterchangeOBJTranslator::GetTexturePa
 		return TOptional<UE::Interchange::FImportImage>();
 	}
 
-	return TextureTranslator->GetTexturePayloadData(PayloadSourceData, PayLoadKey);
+	AlternateTexturePath = PayLoadKey;
+
+	return TextureTranslator->GetTexturePayloadData(PayLoadKey, AlternateTexturePath);
 }
 
 
diff --git a/Engine/Plugins/Interchange/Runtime/Source/Import/Private/Texture/InterchangeDDSTranslator.cpp b/Engine/Plugins/Interchange/Runtime/Source/Import/Private/Texture/InterchangeDDSTranslator.cpp
--- a/Engine/Plugins/Interchange/Runtime/Source/Import/Private/Texture/InterchangeDDSTranslator.cpp
+++ b/Engine/Plugins/Interchange/Runtime/Source/Import/Private/Texture/InterchangeDDSTranslator.cpp
@@ -128,28 +128,6 @@ namespace UE::Interchange::Private::InterchangeDDSTranslator
 		return MipCount;
 	}
 
-	bool PerformPrePayloadValidations(const UInterchangeDDSTranslator& DDSTranslator, const UInterchangeSourceData* PayloadSourceData, const FString& PayLoadKey)
-	{
-		check(PayloadSourceData == DDSTranslator.GetSourceData());
-
-		if (!PayloadSourceData)
-		{
-			UE_LOG(LogInterchangeImport, Error, TEXT("Failed to import DDS, bad source data."));
-			return false;
-		}
-
-		FString Filename = PayloadSourceData->GetFilename();
-
-		//Make sure the key fit the filename, The key should always be valid
-		if (!Filename.Equals(PayLoadKey))
-		{
-			UE_LOG(LogInterchangeImport, Error, TEXT("Failed to import DDS, wrong payload key. [%s]"), *Filename);
-			return false;
-		}
-
-		return true;
-	}
-
 	bool DDSPayloadIsSRGB(const TUniquePtr<UE::DDS::FDDSFile>& DDS, ERawImageFormat::Type RawFormat)
 	{
 		bool bSRGB = false;
@@ -292,11 +270,12 @@ bool UInterchangeDDSTranslator::Translate(UInterchangeBaseNodeContainer& BaseNod
 	return false;
 }
 
-TOptional<UE::Interchange::FImportImage> UInterchangeDDSTranslator::GetTexturePayloadData(const UInterchangeSourceData* PayloadSourceData, const FString& PayLoadKey) const
+TOptional<UE::Interchange::FImportImage> UInterchangeDDSTranslator::GetTexturePayloadData(const FString& /*PayloadKey*/, TOptional<FString>& /*AlternateTexturePath*/) const
 {
 	using namespace UE::Interchange::Private::InterchangeDDSTranslator;
+	using namespace UE::Interchange;
 
-	if (!PerformPrePayloadValidations(*this, PayloadSourceData, PayLoadKey))
+	if (!FTextureTranslatorUtilities::IsTranslatorValid(*this, TEXT("DDS")))
 	{
 		return {};
 	}
@@ -307,7 +286,8 @@ TOptional<UE::Interchange::FImportImage> UInterchangeDDSTranslator::GetTexturePa
 
 	if (!DDS || !DDS->IsValidTexture2D())
 	{
-		UE_LOG(LogInterchangeImport, Error, TEXT("Failed to import DDS, unsupported format. [%s]"), *Filename);
+		FText ErrorMessage = FText::Format(NSLOCTEXT("InterchangeDDSTranslator", "UnsupportedFormat", "Failed to import DDS, unsupported format. [{0}]"), FText::FromString(Filename));
+		FTextureTranslatorUtilities::LogError(*this, MoveTemp(ErrorMessage));
 		return TOptional<UE::Interchange::FImportImage>();
 	}
 
@@ -315,9 +295,8 @@ TOptional<UE::Interchange::FImportImage> UInterchangeDDSTranslator::GetTexturePa
 	ERawImageFormat::Type RawFormat = UE::DDS::DXGIFormatGetClosestRawFormat(DDS->DXGIFormat);
 	if (RawFormat == ERawImageFormat::Invalid)
 	{
-		UInterchangeResultError_Generic* Error = AddMessage<UInterchangeResultError_Generic>();
-		Error->AssetType = UTexture2D::StaticClass();
-		Error->Text = FText::Format(NSLOCTEXT("InterchangeDDSTranslator", "InvalidDXGIFormat", "DDS DXGIFormat not supported : {0} : {1}"), (int)DDS->DXGIFormat, FText::FromString(UE::DDS::DXGIFormatGetName(DDS->DXGIFormat)));
+		FText ErrorMessage = FText::Format(NSLOCTEXT("InterchangeDDSTranslator", "InvalidDXGIFormat", "DDS DXGIFormat not supported : {0} : {1}"), (int)DDS->DXGIFormat, FText::FromString(UE::DDS::DXGIFormatGetName(DDS->DXGIFormat)));
+		FTextureTranslatorUtilities::LogError(*this, MoveTemp(ErrorMessage));
 		return {};
 	}
 
@@ -346,7 +325,7 @@ TOptional<UE::Interchange::FImportImage> UInterchangeDDSTranslator::GetTexturePa
 
 		if (!DDS->GetMipImage(DestMipData, MipIndex))
 		{
-			UE_LOG(LogInterchangeImport, Error, TEXT("DDS could not convert pixel format"));
+			FTextureTranslatorUtilities::LogError(*this, NSLOCTEXT("InterchangeDDSTranslator", "FailedPixelFormat", "DDS could not convert pixel format"));
 			return {};
 		}
 
@@ -358,11 +337,12 @@ TOptional<UE::Interchange::FImportImage> UInterchangeDDSTranslator::GetTexturePa
 	return PayloadData;
 }
 
-TOptional<UE::Interchange::FImportSlicedImage> UInterchangeDDSTranslator::GetSlicedTexturePayloadData(const UInterchangeSourceData* PayloadSourceData, const FString& PayLoadKey) const
+TOptional<UE::Interchange::FImportSlicedImage> UInterchangeDDSTranslator::GetSlicedTexturePayloadData(const FString& /*PayloadKey*/, TOptional<FString>& /*AlternateTexturePath*/) const
 {
+	using namespace UE::Interchange;
 	using namespace UE::Interchange::Private::InterchangeDDSTranslator;
 
-	if (!PerformPrePayloadValidations(*this, PayloadSourceData, PayLoadKey))
+	if (!FTextureTranslatorUtilities::IsTranslatorValid(*this, TEXT("DDS")))
 	{
 		return {};
 	}
@@ -373,7 +353,8 @@ TOptional<UE::Interchange::FImportSlicedImage> UInterchangeDDSTranslator::GetSli
 
 	if (!DDS)
 	{
-		UE_LOG(LogInterchangeImport, Error, TEXT("Failed to import DDS, unsupported format. [%s]"), *Filename);
+		FText ErrorMessage = FText::Format(NSLOCTEXT("InterchangeDDSTranslator", "UnsupportedFormat", "Failed to import DDS, unsupported format. [{0}]"), FText::FromString(Filename));
+		FTextureTranslatorUtilities::LogError(*this, MoveTemp(ErrorMessage));
 		return {};
 	}
 
@@ -397,9 +378,8 @@ TOptional<UE::Interchange::FImportSlicedImage> UInterchangeDDSTranslator::GetSli
 	ERawImageFormat::Type RawFormat = UE::DDS::DXGIFormatGetClosestRawFormat(DDS->DXGIFormat);
 	if (RawFormat == ERawImageFormat::Invalid)
 	{
-		UInterchangeResultError_Generic* Error = AddMessage<UInterchangeResultError_Generic>();
-		Error->AssetType = UTexture2D::StaticClass();
-		Error->Text = FText::Format(NSLOCTEXT("InterchangeDDSTranslator", "InvalidDXGIFormat", "DDS DXGIFormat not supported : {0} : {1}"), (int)DDS->DXGIFormat, FText::FromString(UE::DDS::DXGIFormatGetName(DDS->DXGIFormat)));
+		FText ErrorMessage = FText::Format(NSLOCTEXT("InterchangeDDSTranslator", "InvalidDXGIFormat", "DDS DXGIFormat not supported : {0} : {1}"), (int)DDS->DXGIFormat, FText::FromString(UE::DDS::DXGIFormatGetName(DDS->DXGIFormat)));
+		FTextureTranslatorUtilities::LogError(*this, MoveTemp(ErrorMessage));
 		return {};
 	}
 
@@ -454,7 +434,7 @@ TOptional<UE::Interchange::FImportSlicedImage> UInterchangeDDSTranslator::GetSli
 
 			if (!DDS->GetMipImage(DestMipData, MipIndex))
 			{
-				UE_LOG(LogInterchangeImport, Error, TEXT("DDS could not convert pixel format"));
+				FTextureTranslatorUtilities::LogError(*this, NSLOCTEXT("InterchangeDDSTranslator", "InvalidPixelFormat", "DDS could not convert pixel format"));
 				return {};
 			}
 		}
@@ -473,7 +453,7 @@ TOptional<UE::Interchange::FImportSlicedImage> UInterchangeDDSTranslator::GetSli
 
 				if (!DDS->GetMipImage(DestSliceData, DDSMipIndex))
 				{
-					UE_LOG(LogInterchangeImport, Error, TEXT("DDS could not convert pixel format"));
+					FTextureTranslatorUtilities::LogError(*this, NSLOCTEXT("InterchangeDDSTranslator", "FailedPixelFormat", "DDS could not convert pixel format"));
 					return {};
 				}
 			}
diff --git a/Engine/Plugins/Interchange/Runtime/Source/Import/Private/Texture/InterchangeIESTranslator.cpp b/Engine/Plugins/Interchange/Runtime/Source/Import/Private/Texture/InterchangeIESTranslator.cpp
--- a/Engine/Plugins/Interchange/Runtime/Source/Import/Private/Texture/InterchangeIESTranslator.cpp
+++ b/Engine/Plugins/Interchange/Runtime/Source/Import/Private/Texture/InterchangeIESTranslator.cpp
@@ -22,20 +22,6 @@ static FAutoConsoleVariableRef CCvarInterchangeEnableIESImport(
 	TEXT("Whether IES support is enabled."),
 	ECVF_Default);
 
-namespace UE::Interchange::IESTranslator::Private
-{
-	void LogError(const UInterchangeIESTranslator& IESTranslator, const FString& UniqueID, FText&& ErrorText)
-	{
-		check(IESTranslator.GetSourceData() != nullptr);
-
-		UInterchangeResultError_Generic* ErrorMessage = IESTranslator.AddMessage<UInterchangeResultError_Generic>();
-		ErrorMessage->SourceAssetName = IESTranslator.GetSourceData()->GetFilename();
-		ErrorMessage->AssetType = UTextureLightProfile::StaticClass();
-		ErrorMessage->InterchangeKey = UniqueID;
-		ErrorMessage->Text = MoveTemp(ErrorText);
-	}
-}
-
 TArray<FString> UInterchangeIESTranslator::GetSupportedFormats() const
 {
 	if (GInterchangeEnableIESImport || GIsAutomationTesting)
@@ -54,45 +40,13 @@ bool UInterchangeIESTranslator::Translate(UInterchangeBaseNodeContainer& BaseNod
 	return UE::Interchange::FTextureTranslatorUtilities::GenericTextureLightProfileTranslate(GetSourceData(), BaseNodeContainer);
 }
 
-TOptional<UE::Interchange::FImportLightProfile> UInterchangeIESTranslator::GetLightProfilePayloadData(const UInterchangeSourceData* PayloadSourceData, const FString& PayLoadKey) const
+TOptional<UE::Interchange::FImportLightProfile> UInterchangeIESTranslator::GetLightProfilePayloadData(const FString& /*PayloadKey*/, TOptional<FString>& /*AlternateTexturePath*/) const
 {
 	using namespace UE::Interchange;
 
-	check(GetSourceData() == PayloadSourceData);
-
-	if (!GetSourceData())
-	{
-		IESTranslator::Private::LogError(*this, PayLoadKey,
-			LOCTEXT("IESImportFailure_BadData", "Failed to import IES, bad source data."));
-
-		return {};
-	}
-
 	TArray64<uint8> SourceDataBuffer;
-	FString Filename = GetSourceData()->GetFilename();
-
-	//Make sure the key fit the filename, The key should always be valid
-	if (!Filename.Equals(PayLoadKey))
-	{
-		IESTranslator::Private::LogError(*this, PayLoadKey,
-			LOCTEXT("IESImportFailure_BadPayloadKey", "Failed to import IES, wrong payload key."));
-
-		return {};
-	}
-
-	if (!FPaths::FileExists(Filename))
+	if (!FTextureTranslatorUtilities::LoadSourceBuffer(*this, TEXT("IES"), SourceDataBuffer))
 	{
-		IESTranslator::Private::LogError(*this, PayLoadKey,
-			LOCTEXT("IESImportFailure_OpenFile", "Failed to import IES, cannot open file."));
-
-		return {};
-	}
-
-	if (!FFileHelper::LoadFileToArray(SourceDataBuffer, *Filename))
-	{
-		IESTranslator::Private::LogError(*this, PayLoadKey,
-			LOCTEXT("IESImportFailure_LoadFile", "Failed to import IES, cannot load file content into an array."));
-
 		return {};
 	}
 
@@ -124,8 +78,7 @@ TOptional<UE::Interchange::FImportLightProfile> UInterchangeIESTranslator::GetLi
 	}
 	else
 	{
-		IESTranslator::Private::LogError(*this, PayLoadKey,
-			FText::Format(LOCTEXT("IESImportFailure_Generic", "IES import failed: \"{0}\""), FText::FromString(IESConverter.GetError())));
+		FTextureTranslatorUtilities::LogError(*this, FText::Format(LOCTEXT("IESImportFailure_Generic", "IES import failed: \"{0}\""), FText::FromString(IESConverter.GetError())));
 	}
 
 	return {};
diff --git a/Engine/Plugins/Interchange/Runtime/Source/Import/Private/Texture/InterchangeImageWrapperTranslator.cpp b/Engine/Plugins/Interchange/Runtime/Source/Import/Private/Texture/InterchangeImageWrapperTranslator.cpp
--- a/Engine/Plugins/Interchange/Runtime/Source/Import/Private/Texture/InterchangeImageWrapperTranslator.cpp
+++ b/Engine/Plugins/Interchange/Runtime/Source/Import/Private/Texture/InterchangeImageWrapperTranslator.cpp
@@ -122,32 +122,18 @@ bool UInterchangeImageWrapperTranslator::Translate(UInterchangeBaseNodeContainer
 	return UE::Interchange::FTextureTranslatorUtilities::Generic2DTextureTranslate(GetSourceData(), BaseNodeContainer);
 }
 
-TOptional<UE::Interchange::FImportImage> UInterchangeImageWrapperTranslator::GetTexturePayloadData(const UInterchangeSourceData* PayloadSourceData, const FString& PayLoadKey) const
+TOptional<UE::Interchange::FImportImage> UInterchangeImageWrapperTranslator::GetTexturePayloadData(const FString& /*PayloadKey*/, TOptional<FString>& /*AlternateTexturePath*/) const
 {
-	check(PayloadSourceData == GetSourceData());
+	using namespace UE::Interchange;
 
-	if (!GetSourceData())
+	if (!FTextureTranslatorUtilities::IsTranslatorValid(*this, TEXT("ImageWrapper")))
 	{
-		UE_LOG(LogInterchangeImport, Error, TEXT("Failed to import Texture, bad source data."));
-		return TOptional<UE::Interchange::FImportImage>();
+		return {};
 	}
 
 	TArray64<uint8> SourceDataBuffer;
 	FString Filename = GetSourceData()->GetFilename();
 
-	//Make sure the key fit the filename, The key should always be valid
-	if (!Filename.Equals(PayLoadKey))
-	{
-		UE_LOG(LogInterchangeImport, Error, TEXT("Failed to import Texture, wrong payload key. [%s]"), *Filename);
-		return TOptional<UE::Interchange::FImportImage>();
-	}
-
-	if (!FPaths::FileExists(Filename))
-	{
-		UE_LOG(LogInterchangeImport, Error, TEXT("Failed to import Texture, cannot open file. [%s]"), *Filename);
-		return TOptional<UE::Interchange::FImportImage>();
-	}
-
 	if (!FFileHelper::LoadFileToArray(SourceDataBuffer, *Filename))
 	{
 		UE_LOG(LogInterchangeImport, Error, TEXT("Failed to import Texture, cannot load file content into an array. [%s]"), *Filename);
diff --git a/Engine/Plugins/Interchange/Runtime/Source/Import/Private/Texture/InterchangeJPGTranslator.cpp b/Engine/Plugins/Interchange/Runtime/Source/Import/Private/Texture/InterchangeJPGTranslator.cpp
--- a/Engine/Plugins/Interchange/Runtime/Source/Import/Private/Texture/InterchangeJPGTranslator.cpp
+++ b/Engine/Plugins/Interchange/Runtime/Source/Import/Private/Texture/InterchangeJPGTranslator.cpp
@@ -47,17 +47,18 @@ bool UInterchangeJPGTranslator::Translate(UInterchangeBaseNodeContainer& BaseNod
 	return UE::Interchange::FTextureTranslatorUtilities::Generic2DTextureTranslate(GetSourceData(), BaseNodeContainer);
 }
 
-TOptional<UE::Interchange::FImportImage> UInterchangeJPGTranslator::GetTexturePayloadData(const UInterchangeSourceData* PayloadSourceData, const FString& PayLoadKey) const
+TOptional<UE::Interchange::FImportImage> UInterchangeJPGTranslator::GetTexturePayloadData(const FString& /*PayloadKey*/, TOptional<FString>& AlternateTexturePath) const
 {
-	TArray64<uint8> SourceDataBuffer = LoadSourceFile(PayloadSourceData, PayLoadKey);
+	using namespace UE::Interchange;
 
-	if (SourceDataBuffer.IsEmpty())
+	TArray64<uint8> SourceDataBuffer;
+	if (!FTextureTranslatorUtilities::LoadSourceBuffer(*this, TEXT("JPEG"), SourceDataBuffer))
 	{
 		return {};
 	}
 
 	const bool bImportRaw = false;
-	return GetTexturePayloadImplementation(MoveTemp(SourceDataBuffer), PayLoadKey, bImportRaw);
+	return GetTexturePayloadImplementation(MoveTemp(SourceDataBuffer), bImportRaw);
 
 }
 
@@ -66,56 +67,23 @@ bool UInterchangeJPGTranslator::SupportCompressedTexturePayloadData() const
 	return true;
 }
 
-TOptional<UE::Interchange::FImportImage> UInterchangeJPGTranslator::GetCompressedTexturePayloadData(const UInterchangeSourceData* PayloadSourceData, const FString& PayLoadKey) const
+TOptional<UE::Interchange::FImportImage> UInterchangeJPGTranslator::GetCompressedTexturePayloadData(const FString& /*PayloadKey*/, TOptional<FString>& /*AlternateTexturePath*/) const
 {
-	TArray64<uint8> SourceDataBuffer = LoadSourceFile(PayloadSourceData, PayLoadKey);
+	using namespace UE::Interchange;
 
-	if (SourceDataBuffer.IsEmpty())
+	TArray64<uint8> SourceDataBuffer;
+	if (!FTextureTranslatorUtilities::LoadSourceBuffer(*this, TEXT("JPEG"), SourceDataBuffer))
 	{
 		return {};
 	}
 
 	const bool bImportRaw = true;
-	return GetTexturePayloadImplementation(MoveTemp(SourceDataBuffer), PayLoadKey, bImportRaw);
-}
-
-TArray64<uint8> UInterchangeJPGTranslator::LoadSourceFile(const UInterchangeSourceData* PayloadSourceData, const FString& PayLoadKey) const
-{
-	check(PayloadSourceData == GetSourceData());
-
-	TArray64<uint8> SourceDataBuffer;
-
-	if (!GetSourceData())
-	{
-		UE_LOG(LogInterchangeImport, Error, TEXT("Failed to import JPEG, bad source data."));
-		return SourceDataBuffer;
-	}
-
-	FString Filename = GetSourceData()->GetFilename();
-
-	//Make sure the key fit the filename, The key should always be valid
-	if (!Filename.Equals(PayLoadKey))
-	{
-		UE_LOG(LogInterchangeImport, Error, TEXT("Failed to import JPEG, wrong payload key. [%s]"), *Filename);
-		return SourceDataBuffer;
-	}
-
-	if (!FPaths::FileExists(Filename))
-	{
-		UE_LOG(LogInterchangeImport, Error, TEXT("Failed to import JPEG, cannot open file. [%s]"), *Filename);
-		return SourceDataBuffer;
-	}
-
-	if (!FFileHelper::LoadFileToArray(SourceDataBuffer, *Filename))
-	{
-		UE_LOG(LogInterchangeImport, Error, TEXT("Failed to import JPEG, cannot load file content into an array. [%s]"), *Filename);
-	}
-
-	return SourceDataBuffer;
+	return GetTexturePayloadImplementation(MoveTemp(SourceDataBuffer), bImportRaw);
 }
 
-TOptional<UE::Interchange::FImportImage> UInterchangeJPGTranslator::GetTexturePayloadImplementation(TArray64<uint8>&& SourceDataBuffer, const FString& Filename, bool bShouldImportRaw) const
+TOptional<UE::Interchange::FImportImage> UInterchangeJPGTranslator::GetTexturePayloadImplementation(TArray64<uint8>&& SourceDataBuffer, bool bShouldImportRaw) const
 {
+	using namespace UE::Interchange;
 
 	const uint8* Buffer = SourceDataBuffer.GetData();
 	const uint8* BufferEnd = Buffer + SourceDataBuffer.Num();
@@ -130,7 +98,7 @@ TOptional<UE::Interchange::FImportImage> UInterchangeJPGTranslator::GetTexturePa
 	TSharedPtr<IImageWrapper> JpegImageWrapper = ImageWrapperModule.CreateImageWrapper(EImageFormat::JPEG);
 	if (!JpegImageWrapper.IsValid() || !JpegImageWrapper->SetCompressed(Buffer, Length))
 	{
-		UE_LOG(LogInterchangeImport, Error, TEXT("Failed to decode JPEG. [%s]"), *Filename);
+		FTextureTranslatorUtilities::LogError(*this, NSLOCTEXT("InterchangeJPEGTranslator", "DecodingFailed", "Failed to decode JPEG."));
 		return TOptional<UE::Interchange::FImportImage>();
 	}
 
@@ -160,7 +128,7 @@ TOptional<UE::Interchange::FImportImage> UInterchangeJPGTranslator::GetTexturePa
 
 	if (TextureFormat == TSF_Invalid)
 	{
-		UE_LOG(LogInterchangeImport, Error, TEXT("JPEG file [%s] contains data in an unsupported format"), *Filename);
+		FTextureTranslatorUtilities::LogError(*this, NSLOCTEXT("InterchangeJPEGTranslator", "UnsupportedFormat", "JPEG file contains data in an unsupported format."));
 		return TOptional<UE::Interchange::FImportImage>();
 	}
 
@@ -189,7 +157,7 @@ TOptional<UE::Interchange::FImportImage> UInterchangeJPGTranslator::GetTexturePa
 	}
 	else
 	{
-		UE_LOG(LogInterchangeImport, Error, TEXT("Failed to decode JPEG. [%s]"), *Filename);
+		FTextureTranslatorUtilities::LogError(*this, NSLOCTEXT("InterchangeJPEGTranslator", "DecodingFailed", "Failed to decode JPEG."));
 		return TOptional<UE::Interchange::FImportImage>();
 	}
 
diff --git a/Engine/Plugins/Interchange/Runtime/Source/Import/Private/Texture/InterchangePCXTranslator.cpp b/Engine/Plugins/Interchange/Runtime/Source/Import/Private/Texture/InterchangePCXTranslator.cpp
--- a/Engine/Plugins/Interchange/Runtime/Source/Import/Private/Texture/InterchangePCXTranslator.cpp
+++ b/Engine/Plugins/Interchange/Runtime/Source/Import/Private/Texture/InterchangePCXTranslator.cpp
@@ -93,36 +93,14 @@ bool UInterchangePCXTranslator::Translate(UInterchangeBaseNodeContainer& BaseNod
 	return UE::Interchange::FTextureTranslatorUtilities::Generic2DTextureTranslate(GetSourceData(), BaseNodeContainer);
 }
 
-TOptional<UE::Interchange::FImportImage> UInterchangePCXTranslator::GetTexturePayloadData(const UInterchangeSourceData* PayloadSourceData, const FString& PayLoadKey) const
+TOptional<UE::Interchange::FImportImage> UInterchangePCXTranslator::GetTexturePayloadData(const FString& /*PayloadKey*/, TOptional<FString>& AlternateTexturePath) const
 {
-	check(PayloadSourceData == GetSourceData());
-
-	if (!GetSourceData())
-	{
-		UE_LOG(LogInterchangeImport, Error, TEXT("Failed to import PCX, bad source data."));
-		return TOptional<UE::Interchange::FImportImage>();
-	}
+	using namespace UE::Interchange;
 
 	TArray64<uint8> SourceDataBuffer;
-	FString Filename = GetSourceData()->GetFilename();
-	
-	//Make sure the key fit the filename, The key should always be valid
-	if (!Filename.Equals(PayLoadKey))
+	if (!FTextureTranslatorUtilities::LoadSourceBuffer(*this, TEXT("PCX"), SourceDataBuffer))
 	{
-		UE_LOG(LogInterchangeImport, Error, TEXT("Failed to import PCX, wrong payload key. [%s]"), *Filename);
-		return TOptional<UE::Interchange::FImportImage>();
-	}
-
-	if (!FPaths::FileExists(Filename))
-	{
-		UE_LOG(LogInterchangeImport, Error, TEXT("Failed to import PCX, cannot open file. [%s]"), *Filename);
-		return TOptional<UE::Interchange::FImportImage>();
-	}
-
-	if (!FFileHelper::LoadFileToArray(SourceDataBuffer, *Filename))
-	{
-		UE_LOG(LogInterchangeImport, Error, TEXT("Failed to import PCX, cannot load file content into an array. [%s]"), *Filename);
-		return TOptional<UE::Interchange::FImportImage>();
+		return {};
 	}
 
 	const uint8* Buffer = SourceDataBuffer.GetData();
@@ -238,13 +216,14 @@ TOptional<UE::Interchange::FImportImage> UInterchangePCXTranslator::GetTexturePa
 		}
 		else
 		{
-			UE_LOG(LogInterchangeImport, Error, TEXT("Failed to import PCX, unsupported format. (%i/%i) [%s]"), PCX->NumPlanes, PCX->BitsPerPixel, *Filename);
+			FText ErrorMessage = FText::Format(NSLOCTEXT("InterchangePCXTranslator", "UnsupportedFormat", "Failed to import PCX, unsupported format. ({0}{1})"), (int)PCX->NumPlanes, (int)PCX->BitsPerPixel);
+			FTextureTranslatorUtilities::LogError(*this, MoveTemp(ErrorMessage));
 			return TOptional<UE::Interchange::FImportImage>();
 		}
 	}
 	else
 	{
-		UE_LOG(LogInterchangeImport, Error, TEXT("Failed to import PCX, unsupported file version. [%s]"), *Filename);
+		FTextureTranslatorUtilities::LogError(*this, NSLOCTEXT("InterchangePCXTranslator", "UnsupportedVersion", "Failed to import PCX, unsupported file version."));
 		return TOptional<UE::Interchange::FImportImage>();
 	}
 
diff --git a/Engine/Plugins/Interchange/Runtime/Source/Import/Private/Texture/InterchangePSDTranslator.cpp b/Engine/Plugins/Interchange/Runtime/Source/Import/Private/Texture/InterchangePSDTranslator.cpp
--- a/Engine/Plugins/Interchange/Runtime/Source/Import/Private/Texture/InterchangePSDTranslator.cpp
+++ b/Engine/Plugins/Interchange/Runtime/Source/Import/Private/Texture/InterchangePSDTranslator.cpp
@@ -395,36 +395,14 @@ bool UInterchangePSDTranslator::Translate(UInterchangeBaseNodeContainer& BaseNod
 	return UE::Interchange::FTextureTranslatorUtilities::Generic2DTextureTranslate(GetSourceData(), BaseNodeContainer);
 }
 
-TOptional<UE::Interchange::FImportImage> UInterchangePSDTranslator::GetTexturePayloadData(const UInterchangeSourceData* PayloadSourceData, const FString& PayLoadKey) const
+TOptional<UE::Interchange::FImportImage> UInterchangePSDTranslator::GetTexturePayloadData(const FString& /*PayloadKey*/, TOptional<FString>& /*AlternateTexturePath*/) const
 {
-	check(PayloadSourceData == GetSourceData());
-
-	if (!GetSourceData())
-	{
-		UE_LOG(LogInterchangeImport, Error, TEXT("Failed to import PSD, bad source data."));
-		return TOptional<UE::Interchange::FImportImage>();
-	}
+	using namespace UE::Interchange;
 
 	TArray64<uint8> SourceDataBuffer;
-	FString Filename = GetSourceData()->GetFilename();
-
-	//Make sure the key fit the filename, The key should always be valid
-	if (!Filename.Equals(PayLoadKey))
-	{
-		UE_LOG(LogInterchangeImport, Error, TEXT("Failed to import PSD, wrong payload key. [%s]"), *Filename);
-		return TOptional<UE::Interchange::FImportImage>();
-	}
-
-	if (!FPaths::FileExists(Filename))
+	if (!FTextureTranslatorUtilities::LoadSourceBuffer(*this, TEXT("PSD"), SourceDataBuffer))
 	{
-		UE_LOG(LogInterchangeImport, Error, TEXT("Failed to import PSD, cannot open file. [%s]"), *Filename);
-		return TOptional<UE::Interchange::FImportImage>();
-	}
-
-	if (!FFileHelper::LoadFileToArray(SourceDataBuffer, *Filename))
-	{
-		UE_LOG(LogInterchangeImport, Error, TEXT("Failed to import PSD, cannot load file content into an array. [%s]"), *Filename);
-		return TOptional<UE::Interchange::FImportImage>();
+		return {};
 	}
 
 	const uint8* Buffer = SourceDataBuffer.GetData();
@@ -450,7 +428,7 @@ TOptional<UE::Interchange::FImportImage> UInterchangePSDTranslator::GetTexturePa
 
 	if (!PSDHeader.IsSupported())
 	{
-		UE_LOG(LogInterchangeImport, Error, TEXT("Format of this PSD is not supported. Only Grayscale and RGBColor PSD images are currently supported, in 8-bit or 16-bit."));
+		FTextureTranslatorUtilities::LogError(*this, NSLOCTEXT("InterchangeDDSTranslator", "UnsupportedFormat", "Format of this PSD is not supported. Only Grayscale and RGBColor PSD images are currently supported, in 8-bit or 16-bit."));
 		return TOptional<UE::Interchange::FImportImage>();
 	}
 
@@ -467,7 +445,7 @@ TOptional<UE::Interchange::FImportImage> UInterchangePSDTranslator::GetTexturePa
 
 	if (TextureFormat == TSF_Invalid)
 	{
-		UE_LOG(LogInterchangeImport, Error, TEXT("PSD file contains data in an unsupported format."));
+		FTextureTranslatorUtilities::LogError(*this, NSLOCTEXT("InterchangeDDSTranslator", "UnsupportedPixelFormat", "This PSD file contains data in an unsupported pixel format."));
 		return TOptional<UE::Interchange::FImportImage>();
 	}
 
@@ -482,7 +460,7 @@ TOptional<UE::Interchange::FImportImage> UInterchangePSDTranslator::GetTexturePa
 
 	if (!UE::Interchange::Private::ReadData(static_cast<uint8*>(PayloadData.RawData.GetData()), Buffer, PSDHeader))
 	{
-		UE_LOG(LogInterchangeImport, Error, TEXT("Failed to read this PSD"));
+		FTextureTranslatorUtilities::LogError(*this, NSLOCTEXT("InterchangeDDSTranslator", "FailedRead", "Failed to read this PSD."));
 		return TOptional<UE::Interchange::FImportImage>();
 	}
 
diff --git a/Engine/Plugins/Interchange/Runtime/Source/Import/Private/Texture/InterchangeTextureFactory.cpp b/Engine/Plugins/Interchange/Runtime/Source/Import/Private/Texture/InterchangeTextureFactory.cpp
--- a/Engine/Plugins/Interchange/Runtime/Source/Import/Private/Texture/InterchangeTextureFactory.cpp
+++ b/Engine/Plugins/Interchange/Runtime/Source/Import/Private/Texture/InterchangeTextureFactory.cpp
@@ -498,8 +498,9 @@ namespace UE::Interchange::Private::InterchangeTextureFactory
 				TOptional<UE::Interchange::FImportImage> Payload;
 				if (Translator->CanImportSourceData(SourceDataForBlock))
 				{
+					TOptional<FString> DummyAlternateTexturePath;
 					Translator->SourceData = SourceDataForBlock;
-					Payload = TextureTranslator->GetTexturePayloadData(SourceDataForBlock, SourceDataForBlock->GetFilename());
+					Payload = TextureTranslator->GetTexturePayloadData(FString(), DummyAlternateTexturePath);
 				}
 
 				if (Payload.IsSet())
@@ -579,7 +580,7 @@ namespace UE::Interchange::Private::InterchangeTextureFactory
 		return BlockedImage;
 	}
 
-	FTexturePayloadVariant GetTexturePayload(const UInterchangeSourceData* SourceData, const FString& PayloadKey, const FTextureNodeVariant& TextureNodeVariant, const FTextureFactoryNodeVariant& FactoryNodeVariant, const UInterchangeTranslatorBase* Translator)
+	FTexturePayloadVariant GetTexturePayload(const UInterchangeSourceData* SourceData, const FString& PayloadKey, const FTextureNodeVariant& TextureNodeVariant, const FTextureFactoryNodeVariant& FactoryNodeVariant, const UInterchangeTranslatorBase* Translator, TOptional<FString>& AlternateTexturePath)
 	{
 		static_assert(TVariantSize<FTextureFactoryNodeVariant>::Value == 7, "Please update the code below and this assert to reflect the change to the variant type.");
 
@@ -619,11 +620,11 @@ namespace UE::Interchange::Private::InterchangeTextureFactory
 
 					if (bShoudImportCompressedImage && TextureTranslator->SupportCompressedTexturePayloadData())
 					{
-						return FTexturePayloadVariant(TInPlaceType<TOptional<FImportImage>>(), TextureTranslator->GetCompressedTexturePayloadData(SourceData, PayloadKey));
+						return FTexturePayloadVariant(TInPlaceType<TOptional<FImportImage>>(), TextureTranslator->GetCompressedTexturePayloadData(PayloadKey, AlternateTexturePath));
 					}
 					else
 					{
-						return FTexturePayloadVariant(TInPlaceType<TOptional<FImportImage>>(), TextureTranslator->GetTexturePayloadData(SourceData, PayloadKey));
+						return FTexturePayloadVariant(TInPlaceType<TOptional<FImportImage>>(), TextureTranslator->GetTexturePayloadData(PayloadKey, AlternateTexturePath));
 					}
 				}
 				else
@@ -633,7 +634,7 @@ namespace UE::Interchange::Private::InterchangeTextureFactory
 			}
 			else if (const IInterchangeBlockedTexturePayloadInterface* BlockedTextureTranslator = Cast<IInterchangeBlockedTexturePayloadInterface>(Translator))
 			{
-				return FTexturePayloadVariant(TInPlaceType<TOptional<FImportBlockedImage>>(), BlockedTextureTranslator->GetBlockedTexturePayloadData(SourceData, PayloadKey));
+				return FTexturePayloadVariant(TInPlaceType<TOptional<FImportBlockedImage>>(), BlockedTextureTranslator->GetBlockedTexturePayloadData(PayloadKey, AlternateTexturePath));
 			}
 		}
 
@@ -643,7 +644,7 @@ namespace UE::Interchange::Private::InterchangeTextureFactory
 		{
 			if (const IInterchangeSlicedTexturePayloadInterface* SlicedTextureTranslator = Cast<IInterchangeSlicedTexturePayloadInterface>(Translator))
 			{
-				return FTexturePayloadVariant(TInPlaceType<TOptional<FImportSlicedImage>>(), SlicedTextureTranslator->GetSlicedTexturePayloadData(SourceData, PayloadKey));
+				return FTexturePayloadVariant(TInPlaceType<TOptional<FImportSlicedImage>>(), SlicedTextureTranslator->GetSlicedTexturePayloadData(PayloadKey, AlternateTexturePath));
 			}
 		}
 
@@ -652,7 +653,7 @@ namespace UE::Interchange::Private::InterchangeTextureFactory
 		{
 			if (const IInterchangeTextureLightProfilePayloadInterface* LightProfileTranslator = Cast<IInterchangeTextureLightProfilePayloadInterface>(Translator))
 			{
-				return FTexturePayloadVariant(TInPlaceType<TOptional<FImportLightProfile>>(), LightProfileTranslator->GetLightProfilePayloadData(SourceData, PayloadKey));
+				return FTexturePayloadVariant(TInPlaceType<TOptional<FImportLightProfile>>(), LightProfileTranslator->GetLightProfilePayloadData(PayloadKey, AlternateTexturePath));
 			}
 		}
 
@@ -1611,7 +1612,8 @@ UObject* UInterchangeTextureFactory::ImportAssetObject_Async(const FImportAssetO
 		return nullptr;
 	}
 
-	FTexturePayloadVariant TexturePayload = GetTexturePayload(Arguments.SourceData, PayLoadKey.GetValue(), TextureNodeVariant, TextureFactoryNodeVariant, Arguments.Translator);
+	AlternateTexturePath.Reset();
+	FTexturePayloadVariant TexturePayload = GetTexturePayload(Arguments.SourceData, PayLoadKey.GetValue(), TextureNodeVariant, TextureFactoryNodeVariant, Arguments.Translator, AlternateTexturePath);
 
 	if(TexturePayload.IsType<FEmptyVariantState>())
 	{
@@ -1787,6 +1789,13 @@ void UInterchangeTextureFactory::SetupObject_GameThread(const FSetupObjectParams
 		SetImportAssetDataParameters.SourceFiles = MoveTemp(SourceFiles);
 
 		Texture->AssetImportData = UE::Interchange::FFactoryCommon::SetImportAssetData(SetImportAssetDataParameters);
+		
+		if (AlternateTexturePath.IsSet())
+		{
+			// A file different from the source data has been used to create the texture.
+			// Set this file's path as the first file name
+			Texture->AssetImportData->AddFileName(AlternateTexturePath.GetValue(), 0);
+		}
 	}
 #endif
 }
diff --git a/Engine/Plugins/Interchange/Runtime/Source/Import/Private/Texture/TextureTranslatorUtilities.cpp b/Engine/Plugins/Interchange/Runtime/Source/Import/Private/Texture/TextureTranslatorUtilities.cpp
--- a/Engine/Plugins/Interchange/Runtime/Source/Import/Private/Texture/TextureTranslatorUtilities.cpp
+++ b/Engine/Plugins/Interchange/Runtime/Source/Import/Private/Texture/TextureTranslatorUtilities.cpp
@@ -2,6 +2,7 @@
 #include "Texture/TextureTranslatorUtilities.h"
 
 #include "CoreMinimal.h"
+#include "InterchangeImportLog.h"
 #include "InterchangeSourceData.h"
 #include "InterchangeTexture2DArrayNode.h"
 #include "InterchangeTexture2DNode.h"
@@ -9,6 +10,7 @@
 #include "InterchangeTextureCubeNode.h"
 #include "InterchangeTextureLightProfileNode.h"
 #include "InterchangeTextureNode.h"
+#include "InterchangeTranslatorBase.h"
 #include "InterchangeVolumeTextureNode.h"
 #include "Misc/Paths.h"
 #include "Nodes/InterchangeBaseNodeContainer.h"
@@ -16,6 +18,8 @@
 #include "UObject/Object.h"
 #include "UObject/ObjectMacros.h"
 
+#define LOCTEXT_NAMESPACE "InterchangeTextureTranslator"
+
 namespace UE::Interchange::Private
 {
 	bool GenericTextureTranslate(const UInterchangeSourceData* SourceData, UInterchangeBaseNodeContainer& BaseNodeContainer, const TSubclassOf<UInterchangeTextureNode>& TextureNodeClass)
@@ -79,3 +83,61 @@ bool UE::Interchange::FTextureTranslatorUtilities::GenericTextureLightProfileTra
 {
 	return Private::GenericTextureTranslate(SourceData, BaseNodeContainer, UInterchangeTextureLightProfileNode::StaticClass());
 }
+
+void UE::Interchange::FTextureTranslatorUtilities::LogError(const UInterchangeTranslatorBase& TextureTranslator, FText&& ErrorText)
+{
+	UInterchangeResultError_Generic* ErrorMessage = TextureTranslator.AddMessage<UInterchangeResultError_Generic>();
+	ErrorMessage->AssetType = TextureTranslator.GetClass();
+
+	if (TextureTranslator.GetSourceData())
+	{
+		const FString Filename = TextureTranslator.GetSourceData()->GetFilename();
+		ErrorMessage->SourceAssetName = Filename;
+		ErrorMessage->InterchangeKey = FPaths::GetBaseFilename(Filename);
+	}
+	else
+	{
+		ErrorMessage->InterchangeKey = TEXT("Undefined");
+	}
+
+	ErrorMessage->Text = MoveTemp(ErrorText);
+}
+
+bool UE::Interchange::FTextureTranslatorUtilities::IsTranslatorValid(const UInterchangeTranslatorBase& TextureTranslator, const TCHAR* Format)
+{
+	const UInterchangeSourceData* SourceData = TextureTranslator.GetSourceData();
+
+	if (!SourceData)
+	{
+		LogError(TextureTranslator, FText::Format(LOCTEXT("TextureImportFailure_BadData", "Failed to import {0}, bad source data."), FText::FromString(Format)));
+		return false;
+	}
+
+	FString Filename = SourceData->GetFilename();
+
+	if (!FPaths::FileExists(Filename))
+	{
+		LogError(TextureTranslator, FText::Format(LOCTEXT("TextureImportFailure_OpenFile", "Failed to import {0}, cannot open file."), FText::FromString(Format)));
+		return false;
+	}
+
+	return true;
+}
+
+bool UE::Interchange::FTextureTranslatorUtilities::LoadSourceBuffer(const UInterchangeTranslatorBase& TextureTranslator, const TCHAR* Format, TArray64<uint8>& SourceDataBuffer)
+{
+	if (!IsTranslatorValid(TextureTranslator, Format))
+	{
+		return false;
+	}
+
+	if (!FFileHelper::LoadFileToArray(SourceDataBuffer, *TextureTranslator.GetSourceData()->GetFilename()))
+	{
+		LogError(TextureTranslator, FText::Format(LOCTEXT("TextureImportFailure_LoadBuffer", "Failed to import {0}, cannot load file content into an array."), FText::FromString(Format)));
+		return false;
+	}
+
+	return !SourceDataBuffer.IsEmpty();
+}
+
+#undef LOCTEXT_NAMESPACE
diff --git a/Engine/Plugins/Interchange/Runtime/Source/Import/Private/Texture/TextureTranslatorUtilities.h b/Engine/Plugins/Interchange/Runtime/Source/Import/Private/Texture/TextureTranslatorUtilities.h
--- a/Engine/Plugins/Interchange/Runtime/Source/Import/Private/Texture/TextureTranslatorUtilities.h
+++ b/Engine/Plugins/Interchange/Runtime/Source/Import/Private/Texture/TextureTranslatorUtilities.h
@@ -8,29 +8,34 @@
 #include "UObject/ObjectMacros.h"
 
 class UInterchangeTextureNode;
+class UInterchangeTranslatorBase;
 
 template <class T>
 class TSubclassOf;
 
-namespace UE
+namespace UE::Interchange
 {
-	namespace Interchange
+	class FTextureTranslatorUtilities
 	{
-		class FTextureTranslatorUtilities
-		{
-		public:
-			static bool Generic2DTextureTranslate(const UInterchangeSourceData* SourceData, UInterchangeBaseNodeContainer& BaseNodeContainer);
+	public:
+		static bool Generic2DTextureTranslate(const UInterchangeSourceData* SourceData, UInterchangeBaseNodeContainer& BaseNodeContainer);
 
-			static bool GenericTextureCubeTranslate(const UInterchangeSourceData* SourceData, UInterchangeBaseNodeContainer& BaseNodeContainer);
+		static bool GenericTextureCubeTranslate(const UInterchangeSourceData* SourceData, UInterchangeBaseNodeContainer& BaseNodeContainer);
 
-			static bool GenericTextureCubeArrayTranslate(const UInterchangeSourceData* SourceData, UInterchangeBaseNodeContainer& BaseNodeContainer);
+		static bool GenericTextureCubeArrayTranslate(const UInterchangeSourceData* SourceData, UInterchangeBaseNodeContainer& BaseNodeContainer);
 
-			static bool GenericTexture2DArrayTranslate(const UInterchangeSourceData* SourceData, UInterchangeBaseNodeContainer& BaseNodeContainer);
+		static bool GenericTexture2DArrayTranslate(const UInterchangeSourceData* SourceData, UInterchangeBaseNodeContainer& BaseNodeContainer);
 
-			static bool GenericVolumeTextureTranslate(const UInterchangeSourceData* SourceData, UInterchangeBaseNodeContainer& BaseNodeContainer);
+		static bool GenericVolumeTextureTranslate(const UInterchangeSourceData* SourceData, UInterchangeBaseNodeContainer& BaseNodeContainer);
 
-			static bool GenericTextureLightProfileTranslate(const UInterchangeSourceData* SourceData, UInterchangeBaseNodeContainer& BaseNodeContainer);
-		};
-	}//ns Interchange
+		static bool GenericTextureLightProfileTranslate(const UInterchangeSourceData* SourceData, UInterchangeBaseNodeContainer& BaseNodeContainer);
+
+		static bool IsTranslatorValid(const UInterchangeTranslatorBase& TextureTranslator, const TCHAR* Format);
+
+		static bool LoadSourceBuffer(const UInterchangeTranslatorBase& TextureTranslator, const TCHAR* Format, TArray64<uint8>& SourceDataBuffer);
+
+		static void LogError(const UInterchangeTranslatorBase& TextureTranslator, FText&& ErrorText);
+
+	};
 }//ns UE
 
diff --git a/Engine/Plugins/Interchange/Runtime/Source/Import/Public/Fbx/InterchangeFbxTranslator.h b/Engine/Plugins/Interchange/Runtime/Source/Import/Public/Fbx/InterchangeFbxTranslator.h
--- a/Engine/Plugins/Interchange/Runtime/Source/Import/Public/Fbx/InterchangeFbxTranslator.h
+++ b/Engine/Plugins/Interchange/Runtime/Source/Import/Public/Fbx/InterchangeFbxTranslator.h
@@ -47,7 +47,7 @@ public:
 	 * @param PayloadKey - The key to retrieve the a particular payload contain into the specified source data.
 	 * @return a PayloadData containing the imported data. The TOptional will not be set if there is an error.
 	 */
-	virtual TOptional<UE::Interchange::FImportImage> GetTexturePayloadData(const UInterchangeSourceData* InSourceData, const FString& PayLoadKey) const override;
+	virtual TOptional<UE::Interchange::FImportImage> GetTexturePayloadData(const FString& PayloadKey, TOptional<FString>& AlternateTexturePath) const override;
 
 	/* IInterchangeTexturePayloadInterface End */
 
diff --git a/Engine/Plugins/Interchange/Runtime/Source/Import/Public/Gltf/InterchangeGltfTranslator.h b/Engine/Plugins/Interchange/Runtime/Source/Import/Public/Gltf/InterchangeGltfTranslator.h
--- a/Engine/Plugins/Interchange/Runtime/Source/Import/Public/Gltf/InterchangeGltfTranslator.h
+++ b/Engine/Plugins/Interchange/Runtime/Source/Import/Public/Gltf/InterchangeGltfTranslator.h
@@ -47,7 +47,7 @@ public:
 
 	/* IInterchangeTexturePayloadInterface Begin */
 
-	virtual TOptional< UE::Interchange::FImportImage > GetTexturePayloadData( const UInterchangeSourceData* InSourceData, const FString& PayLoadKey ) const override;
+	virtual TOptional< UE::Interchange::FImportImage > GetTexturePayloadData(const FString& PayloadKey, TOptional<FString>& AlternateTexturePath) const override;
 
 	/* IInterchangeTexturePayloadInterface End */
 
diff --git a/Engine/Plugins/Interchange/Runtime/Source/Import/Public/MaterialX/InterchangeMaterialXTranslator.h b/Engine/Plugins/Interchange/Runtime/Source/Import/Public/MaterialX/InterchangeMaterialXTranslator.h
--- a/Engine/Plugins/Interchange/Runtime/Source/Import/Public/MaterialX/InterchangeMaterialXTranslator.h
+++ b/Engine/Plugins/Interchange/Runtime/Source/Import/Public/MaterialX/InterchangeMaterialXTranslator.h
@@ -36,5 +36,5 @@ public:
 	virtual bool Translate( UInterchangeBaseNodeContainer& BaseNodeContainer ) const override;
 	/** End UInterchangeTranslatorBase API*/
 
-	virtual TOptional<UE::Interchange::FImportImage> GetTexturePayloadData(const UInterchangeSourceData* PayloadSourceData, const FString& PayloadKey) const override;
+	virtual TOptional<UE::Interchange::FImportImage> GetTexturePayloadData(const FString& PayloadKey, TOptional<FString>& AlternateTexturePath) const override;
 };
\ No newline at end of file
diff --git a/Engine/Plugins/Interchange/Runtime/Source/Import/Public/Mesh/InterchangeOBJTranslator.h b/Engine/Plugins/Interchange/Runtime/Source/Import/Public/Mesh/InterchangeOBJTranslator.h
--- a/Engine/Plugins/Interchange/Runtime/Source/Import/Public/Mesh/InterchangeOBJTranslator.h
+++ b/Engine/Plugins/Interchange/Runtime/Source/Import/Public/Mesh/InterchangeOBJTranslator.h
@@ -65,7 +65,7 @@ public:
 	 * @param PayloadKey - The key to retrieve the a particular payload contain into the specified source data.
 	 * @return a PayloadData containing the imported data. The TOptional will not be set if there is an error.
 	 */
-	virtual TOptional<UE::Interchange::FImportImage> GetTexturePayloadData(const UInterchangeSourceData* InSourceData, const FString& PayLoadKey) const override;
+	virtual TOptional<UE::Interchange::FImportImage> GetTexturePayloadData(const FString& PayloadKey, TOptional<FString>& AlternateTexturePath) const override;
 
 	/* IInterchangeTexturePayloadInterface End */
 
diff --git a/Engine/Plugins/Interchange/Runtime/Source/Import/Public/Texture/InterchangeBlockedTexturePayloadInterface.h b/Engine/Plugins/Interchange/Runtime/Source/Import/Public/Texture/InterchangeBlockedTexturePayloadInterface.h
--- a/Engine/Plugins/Interchange/Runtime/Source/Import/Public/Texture/InterchangeBlockedTexturePayloadInterface.h
+++ b/Engine/Plugins/Interchange/Runtime/Source/Import/Public/Texture/InterchangeBlockedTexturePayloadInterface.h
@@ -26,14 +26,22 @@ class INTERCHANGEIMPORT_API IInterchangeBlockedTexturePayloadInterface
 {
 	GENERATED_BODY()
 public:
+
+	UE_DEPRECATED(5.3, "Deprecated. Use GetBlockedTexturePayloadData(const FString&, TOptional<FString>&) instead.")
+	virtual TOptional<UE::Interchange::FImportBlockedImage> GetBlockedTexturePayloadData(const UInterchangeSourceData* PayloadSourceData, const FString& PayloadKey)
+	{
+		TOptional<FString> AlternateTexturePath;
+		return GetBlockedTexturePayloadData(PayloadKey, AlternateTexturePath);
+	}
+
 	/**
 	 * Once the translation is done, the import process need a way to retrieve payload data.
 	 * This payload will be use by the factories to create the asset.
 	 *
-	 * @param SourceData - The source data containing the data to translate
-	 * @param PayloadKey - The key to retrieve the a particular payload contain into the specified source data
+	 * @param PayloadKey - The key to retrieve the a particular payload contain into the specified source data.
+	 * @param AlternateTexturePath - When applicable, set to the path of the file actually loaded to create the FImportImage.
 	 * @return a PayloadData containing the import blocked image data. The TOptional will not be set if there is an error.
 	 */
-	virtual TOptional<UE::Interchange::FImportBlockedImage> GetBlockedTexturePayloadData(const UInterchangeSourceData* PayloadSourceData, const FString& PayLoadKey) const {return {};}
+	virtual TOptional<UE::Interchange::FImportBlockedImage> GetBlockedTexturePayloadData(const FString& PayloadKey, TOptional<FString>& AlternateTexturePath) const {return {};}
 };
 
diff --git a/Engine/Plugins/Interchange/Runtime/Source/Import/Public/Texture/InterchangeDDSTranslator.h b/Engine/Plugins/Interchange/Runtime/Source/Import/Public/Texture/InterchangeDDSTranslator.h
--- a/Engine/Plugins/Interchange/Runtime/Source/Import/Public/Texture/InterchangeDDSTranslator.h
+++ b/Engine/Plugins/Interchange/Runtime/Source/Import/Public/Texture/InterchangeDDSTranslator.h
@@ -42,11 +42,11 @@ public:
 	 * Once the translation is done, the import process need a way to retrieve payload data.
 	 * This payload will be use by the factories to create the asset.
 	 *
-	 * @param PayloadSourceData - The source data containing the data to translate
-	 * @param PayloadKey - The key to retrieve the a particular payload contain into the specified source data.
+	 * @param PayloadKey - Unused. The translator uses its SourceData property to extract the payload.
+	 * @param AlternateTexturePath - Unused. The translator uses its SourceData property to extract the payload.
 	 * @return a PayloadData containing the import image data. The TOptional will not be set if there is an error.
 	 */
-	virtual TOptional<UE::Interchange::FImportImage> GetTexturePayloadData(const UInterchangeSourceData* PayloadSourceData, const FString& PayLoadKey) const override;
+	virtual TOptional<UE::Interchange::FImportImage> GetTexturePayloadData(const FString& PayloadKey, TOptional<FString>& AlternateTexturePath) const override;
 
 	/* IInterchangeTexturePayloadInterface End */
 
@@ -57,11 +57,11 @@ public:
 	 * Once the translation is done, the import process need a way to retrieve payload data.
 	 * This payload will be use by the factories to create the asset.
 	 *
-	 * @param SourceData - The source data containing the data to translate
-	 * @param PayloadKey - The key to retrieve the a particular payload contain into the specified source data.
+	 * @param PayloadKey - Unused. The translator uses its SourceData property to extract the payload.
+	 * @param AlternateTexturePath - Unused. The translator uses its SourceData property to extract the payload.
 	 * @return a PayloadData containing the import image data. The TOptional will not be set if there is an error.
 	 */
-	virtual TOptional<UE::Interchange::FImportSlicedImage> GetSlicedTexturePayloadData(const UInterchangeSourceData* SourceData, const FString& PayloadKey) const override;
+	virtual TOptional<UE::Interchange::FImportSlicedImage> GetSlicedTexturePayloadData(const FString& PayloadKey, TOptional<FString>& AlternateTexturePath) const override;
 
 	/* IInterchangeSlicedTexturePayloadInterface End */
 };
diff --git a/Engine/Plugins/Interchange/Runtime/Source/Import/Public/Texture/InterchangeIESTranslator.h b/Engine/Plugins/Interchange/Runtime/Source/Import/Public/Texture/InterchangeIESTranslator.h
--- a/Engine/Plugins/Interchange/Runtime/Source/Import/Public/Texture/InterchangeIESTranslator.h
+++ b/Engine/Plugins/Interchange/Runtime/Source/Import/Public/Texture/InterchangeIESTranslator.h
@@ -36,11 +36,11 @@ public:
 	 * Once the translation is done, the import process need a way to retrieve payload data.
 	 * This payload will be use by the factories to create the asset.
 	 *
-	 * @param SourceData - The source data containing the data to translate
-	 * @param PayloadKey - The key to retrieve the a particular payload contain into the specified source data.
+	 * @param PayloadKey - Unused. The translator uses its SourceData property to extract the payload.
+	 * @param AlternateTexturePath - Unused. The translator uses its SourceData property to extract the payload.
 	 * @return a PayloadData containing the import light profile data. The TOptional will not be set if there is an error.
 	 */
-	virtual TOptional<UE::Interchange::FImportLightProfile> GetLightProfilePayloadData(const UInterchangeSourceData* PayloadSourceData, const FString& PayLoadKey) const override;
+	virtual TOptional<UE::Interchange::FImportLightProfile> GetLightProfilePayloadData(const FString& PayloadKey, TOptional<FString>& AlternateTexturePath) const override;
 
 	/* IInterchangeTexturePayloadInterface End */
 };
\ No newline at end of file
diff --git a/Engine/Plugins/Interchange/Runtime/Source/Import/Public/Texture/InterchangeImageWrapperTranslator.h b/Engine/Plugins/Interchange/Runtime/Source/Import/Public/Texture/InterchangeImageWrapperTranslator.h
--- a/Engine/Plugins/Interchange/Runtime/Source/Import/Public/Texture/InterchangeImageWrapperTranslator.h
+++ b/Engine/Plugins/Interchange/Runtime/Source/Import/Public/Texture/InterchangeImageWrapperTranslator.h
@@ -40,11 +40,11 @@ public:
 	 * Once the translation is done, the import process need a way to retrieve payload data.
 	 * This payload will be use by the factories to create the asset.
 	 *
-	 * @param PayloadSourceData - The source data containing the data to translate
-	 * @param PayloadKey - The key to retrieve the a particular payload contain into the specified source data.
+	 * @param PayloadKey - Unused. The translator uses its SourceData property to extract the payload.
+	 * @param AlternateTexturePath - Unused. The translator uses its SourceData property to extract the payload.
 	 * @return a PayloadData containing the import image data. The TOptional will not be set if there is an error.
 	 */
-	virtual TOptional<UE::Interchange::FImportImage> GetTexturePayloadData(const UInterchangeSourceData* PayloadSourceData, const FString& PayLoadKey) const override;
+	virtual TOptional<UE::Interchange::FImportImage> GetTexturePayloadData(const FString& PayloadKey, TOptional<FString>& AlternateTexturePath) const override;
 
 	/* IInterchangeTexturePayloadInterface End */
 
diff --git a/Engine/Plugins/Interchange/Runtime/Source/Import/Public/Texture/InterchangeJPGTranslator.h b/Engine/Plugins/Interchange/Runtime/Source/Import/Public/Texture/InterchangeJPGTranslator.h
--- a/Engine/Plugins/Interchange/Runtime/Source/Import/Public/Texture/InterchangeJPGTranslator.h
+++ b/Engine/Plugins/Interchange/Runtime/Source/Import/Public/Texture/InterchangeJPGTranslator.h
@@ -35,23 +35,29 @@ public:
 	 * Once the translation is done, the import process need a way to retrieve payload data.
 	 * This payload will be use by the factories to create the asset.
 	 *
-	 * @param PayloadSourceData - The source data containing the data to translate
-	 * @param PayloadKey - The key to retrieve the a particular payload contain into the specified source data.
+	 * @param PayloadKey - Unused. The translator uses its SourceData property to extract the payload.
+	 * @param AlternateTexturePath - Unused. The translator uses its SourceData property to extract the payload.
 	 * @return a PayloadData containing the import image data. The TOptional will not be set if there is an error.
 	 */
-	virtual TOptional<UE::Interchange::FImportImage> GetTexturePayloadData(const UInterchangeSourceData* PayloadSourceData, const FString& PayLoadKey) const override;
+	virtual TOptional<UE::Interchange::FImportImage> GetTexturePayloadData(const FString& PayloadKey, TOptional<FString>& AlternateTexturePath) const override;
 
 	virtual bool SupportCompressedTexturePayloadData() const override;
 
-	virtual TOptional<UE::Interchange::FImportImage> GetCompressedTexturePayloadData(const UInterchangeSourceData* PayloadSourceData, const FString& PayLoadKey) const override;
+	/**
+	 * Once the translation is done, the import process need a way to retrieve payload data.
+	 * This payload will be use by the factories to create the asset.
+	 *
+	 * @param PayloadKey - Unused. The translator uses its SourceData property to extract the payload.
+	 * @param AlternateTexturePath - Unused. The translator uses its SourceData property to extract the payload.
+	 * @return a PayloadData containing the import image data. The TOptional will not be set if there is an error.
+	 */
+	virtual TOptional<UE::Interchange::FImportImage> GetCompressedTexturePayloadData(const FString& PayloadKey, TOptional<FString>& AlternateTexturePath) const override;
 	
 	/* IInterchangeTexturePayloadInterface End */
 
 private:
 
-	TArray64<uint8> LoadSourceFile(const UInterchangeSourceData* PayloadSourceData, const FString& PayLoadKey) const;
-
-	TOptional<UE::Interchange::FImportImage> GetTexturePayloadImplementation(TArray64<uint8>&& SourceDataBuffer, const FString& Filename, bool bShouldImportRaw) const;
+	TOptional<UE::Interchange::FImportImage> GetTexturePayloadImplementation(TArray64<uint8>&& SourceDataBuffer, bool bShouldImportRaw) const;
 };
 
 
diff --git a/Engine/Plugins/Interchange/Runtime/Source/Import/Public/Texture/InterchangePCXTranslator.h b/Engine/Plugins/Interchange/Runtime/Source/Import/Public/Texture/InterchangePCXTranslator.h
--- a/Engine/Plugins/Interchange/Runtime/Source/Import/Public/Texture/InterchangePCXTranslator.h
+++ b/Engine/Plugins/Interchange/Runtime/Source/Import/Public/Texture/InterchangePCXTranslator.h
@@ -35,11 +35,11 @@ public:
 	 * Once the translation is done, the import process need a way to retrieve payload data.
 	 * This payload will be use by the factories to create the asset.
 	 *
-	 * @param PayloadSourceData - The source data containing the data to translate
-	 * @param PayloadKey - The key to retrieve the a particular payload contain into the specified source data.
+	 * @param PayloadKey - Unused. The translator uses its SourceData property to extract the payload.
+	 * @param AlternateTexturePath - Unused. The translator uses its SourceData property to extract the payload.
 	 * @return a PayloadData containing the import image data. The TOptional will not be set if there is an error.
 	 */
-	virtual TOptional<UE::Interchange::FImportImage> GetTexturePayloadData(const UInterchangeSourceData* PayloadSourceData, const FString& PayLoadKey) const override;
+	virtual TOptional<UE::Interchange::FImportImage> GetTexturePayloadData(const FString& PayloadKey, TOptional<FString>& AlternateTexturePath) const override;
 
 	/* IInterchangeTexturePayloadInterface End */
 };
diff --git a/Engine/Plugins/Interchange/Runtime/Source/Import/Public/Texture/InterchangePSDTranslator.h b/Engine/Plugins/Interchange/Runtime/Source/Import/Public/Texture/InterchangePSDTranslator.h
--- a/Engine/Plugins/Interchange/Runtime/Source/Import/Public/Texture/InterchangePSDTranslator.h
+++ b/Engine/Plugins/Interchange/Runtime/Source/Import/Public/Texture/InterchangePSDTranslator.h
@@ -35,11 +35,11 @@ public:
 	 * Once the translation is done, the import process need a way to retrieve payload data.
 	 * This payload will be use by the factories to create the asset.
 	 *
-	 * @param PayloadSourceData - The source data containing the data to translate
-	 * @param PayloadKey - The key to retrieve the a particular payload contain into the specified source data.
+	 * @param PayloadKey - Unused. The translator uses its SourceData property to extract the payload.
+	 * @param AlternateTexturePath - Unused. The translator uses its SourceData property to extract the payload.
 	 * @return a PayloadData containing the import image data. The TOptional will not be set if there is an error.
 	 */
-	virtual TOptional<UE::Interchange::FImportImage> GetTexturePayloadData(const UInterchangeSourceData* PayloadSourceData, const FString& PayLoadKey) const override;
+	virtual TOptional<UE::Interchange::FImportImage> GetTexturePayloadData(const FString& PayloadKey, TOptional<FString>& AlternateTexturePath) const override;
 
 	/* IInterchangeTexturePayloadInterface End */
 };
diff --git a/Engine/Plugins/Interchange/Runtime/Source/Import/Public/Texture/InterchangeSlicedTexturePayloadInterface.h b/Engine/Plugins/Interchange/Runtime/Source/Import/Public/Texture/InterchangeSlicedTexturePayloadInterface.h
--- a/Engine/Plugins/Interchange/Runtime/Source/Import/Public/Texture/InterchangeSlicedTexturePayloadInterface.h
+++ b/Engine/Plugins/Interchange/Runtime/Source/Import/Public/Texture/InterchangeSlicedTexturePayloadInterface.h
@@ -23,13 +23,20 @@ class INTERCHANGEIMPORT_API IInterchangeSlicedTexturePayloadInterface
 	GENERATED_BODY()
 public:
 
+	UE_DEPRECATED(5.3, "Deprecated. Use GetSlicedTexturePayloadData(const FString&, TOptional<FString>&) instead.")
+	virtual TOptional<UE::Interchange::FImportSlicedImage> GetSlicedTexturePayloadData(const UInterchangeSourceData* SourceData, const FString& PayloadKey) const
+	{
+		TOptional<FString> AlternateTexturePath;
+		return GetSlicedTexturePayloadData(PayloadKey, AlternateTexturePath);
+	}
+	
 	/**
 	 * Once the translation is done, the import process need a way to retrieve payload data.
 	 * This payload will be use by the factories to create the asset.
 	 *
-	 * @param SourceData - The source data containing the data to translate
 	 * @param PayloadKey - The key to retrieve the a particular payload contain into the specified source data.
+	 * @param AlternateTexturePath - When applicable, set to the path of the file actually loaded to create the FImportImage.
 	 * @return a PayloadData containing the import image data. The TOptional will not be set if there is an error.
 	 */
-	virtual TOptional<UE::Interchange::FImportSlicedImage> GetSlicedTexturePayloadData(const UInterchangeSourceData* SourceData, const FString& PayloadKey) const = 0;
+	virtual TOptional<UE::Interchange::FImportSlicedImage> GetSlicedTexturePayloadData(const FString& PayloadKey, TOptional<FString>& AlternateTexturePath) const PURE_VIRTUAL(IInterchangeSlicedTexturePayloadInterface::GetSlicedTexturePayloadData, return{};);
 };
diff --git a/Engine/Plugins/Interchange/Runtime/Source/Import/Public/Texture/InterchangeTextureFactory.h b/Engine/Plugins/Interchange/Runtime/Source/Import/Public/Texture/InterchangeTextureFactory.h
--- a/Engine/Plugins/Interchange/Runtime/Source/Import/Public/Texture/InterchangeTextureFactory.h
+++ b/Engine/Plugins/Interchange/Runtime/Source/Import/Public/Texture/InterchangeTextureFactory.h
@@ -72,6 +72,7 @@ private:
 	void CheckForInvalidResolutions(UE::Interchange::Private::InterchangeTextureFactory::FTexturePayloadVariant& InPayloadVariant, const UInterchangeSourceData* SourceData, const UInterchangeTextureFactoryNode* TextureFactoryNode);
 	
 	UE::Interchange::Private::InterchangeTextureFactory::FProcessedPayload ProcessedPayload;
+	TOptional<FString> AlternateTexturePath;
 
 #if WITH_EDITORONLY_DATA
 	//  The data for the source files will be stored here during the import
diff --git a/Engine/Plugins/Interchange/Runtime/Source/Import/Public/Texture/InterchangeTextureLightProfilePayloadInterface.h b/Engine/Plugins/Interchange/Runtime/Source/Import/Public/Texture/InterchangeTextureLightProfilePayloadInterface.h
--- a/Engine/Plugins/Interchange/Runtime/Source/Import/Public/Texture/InterchangeTextureLightProfilePayloadInterface.h
+++ b/Engine/Plugins/Interchange/Runtime/Source/Import/Public/Texture/InterchangeTextureLightProfilePayloadInterface.h
@@ -23,13 +23,20 @@ class INTERCHANGEIMPORT_API IInterchangeTextureLightProfilePayloadInterface
 	GENERATED_BODY()
 public:
 
+	UE_DEPRECATED(5.3, "Deprecated. Use GetLightProfilePayloadData(const FString&, TOptional<FString>&) instead.")
+	virtual TOptional<UE::Interchange::FImportLightProfile> GetLightProfilePayloadData(const UInterchangeSourceData* SourceData, const FString& PayloadKey) const
+	{
+		TOptional<FString> AlternateTexturePath;
+		return GetLightProfilePayloadData(PayloadKey, AlternateTexturePath);
+	}
+
 	/**
 	 * Once the translation is done, the import process need a way to retrieve payload data.
 	 * This payload will be use by the factories to create the asset.
 	 *
-	 * @param SourceData - The source data containing the data to translate
 	 * @param PayloadKey - The key to retrieve the a particular payload contain into the specified source data.
-	 * @return a PayloadData containing the import image data. The TOptional will not be set if there is an error.
+	 * @param AlternateTexturePath - When applicable, set to the path of the file actually loaded to create the FImportImage.
+	 * @return a PayloadData containing the import light profile data. The TOptional will not be set if there is an error.
 	 */
-	virtual TOptional<UE::Interchange::FImportLightProfile> GetLightProfilePayloadData(const UInterchangeSourceData* SourceData, const FString& PayloadKey) const = 0;
+	virtual TOptional<UE::Interchange::FImportLightProfile> GetLightProfilePayloadData(const FString& PayloadKey, TOptional<FString>& AlternateTexturePath) const PURE_VIRTUAL(IInterchangeTextureLightProfilePayloadInterface::GetLightProfilePayloadData, return{};);
 };
diff --git a/Engine/Plugins/Interchange/Runtime/Source/Import/Public/Texture/InterchangeTexturePayloadInterface.h b/Engine/Plugins/Interchange/Runtime/Source/Import/Public/Texture/InterchangeTexturePayloadInterface.h
--- a/Engine/Plugins/Interchange/Runtime/Source/Import/Public/Texture/InterchangeTexturePayloadInterface.h
+++ b/Engine/Plugins/Interchange/Runtime/Source/Import/Public/Texture/InterchangeTexturePayloadInterface.h
@@ -23,30 +23,44 @@ class INTERCHANGEIMPORT_API IInterchangeTexturePayloadInterface
 	GENERATED_BODY()
 public:
 
+	UE_DEPRECATED(5.3, "Deprecated. Use GetTexturePayloadData(const FString&, TOptional<FString>&) instead.")
+	virtual TOptional<UE::Interchange::FImportImage> GetTexturePayloadData(const UInterchangeSourceData* PayloadSourceData, const FString& PayloadKey) const
+	{
+		TOptional<FString> AlternateTexturePath;
+		return GetTexturePayloadData(PayloadKey, AlternateTexturePath);
+	}
+
 	/**
 	 * Once the translation is done, the import process need a way to retrieve payload data.
 	 * This payload will be use by the factories to create the asset.
 	 *
-	 * @param PayloadSourceData - The source data containing the data to translate
 	 * @param PayloadKey - The key to retrieve the a particular payload contain into the specified source data.
+	 * @param AlternateTexturePath - When applicable, set to the path of the file actually loaded to create the FImportImage.
 	 * @return a PayloadData containing the import image data. The TOptional will not be set if there is an error.
 	 */
-	virtual TOptional<UE::Interchange::FImportImage> GetTexturePayloadData(const UInterchangeSourceData* PayloadSourceData, const FString& PayloadKey) const = 0;
+	virtual TOptional<UE::Interchange::FImportImage> GetTexturePayloadData(const FString& PayloadKey, TOptional<FString>& AlternateTexturePath) const PURE_VIRTUAL(IInterchangeTexturePayloadInterface::GetTexturePayloadData, return{};);
 
 	/**
 	 * Does this translator support the request for a compressed source data?
 	 */
 	virtual bool SupportCompressedTexturePayloadData() const { return false; }
 
+	UE_DEPRECATED(5.3, "Deprecated. Use GetCompressedTexturePayloadData(const FString&, TOptional<FString>&) instead.")
+	virtual TOptional<UE::Interchange::FImportImage> GetCompressedTexturePayloadData(const UInterchangeSourceData* PayloadSourceData, const FString& PayloadKey)
+	{
+		TOptional<FString> AlternateTexturePath;
+		return GetCompressedTexturePayloadData(FString(), AlternateTexturePath);
+	}
+
 	/**
 	 * Once the translation is done, the import process need a way to retrieve payload data.
 	 * This payload will be use by the factories to create the asset.
 	 *
-	 * @param PayloadSourceData - The source data containing the data to translate
 	 * @param PayloadKey - The key to retrieve the a particular payload contain into the specified source data.
+	 * @param AlternateTexturePath - When applicable, set to the path of the file actually loaded to create the FImportImage.
 	 * @return a PayloadData containing the import image data with image data being the original file format generally. The TOptional will not be set if there is an error.
 	 */
-	virtual TOptional<UE::Interchange::FImportImage> GetCompressedTexturePayloadData(const UInterchangeSourceData* PayloadSourceData, const FString& PayloadKey) const { return {}; }
+	virtual TOptional<UE::Interchange::FImportImage> GetCompressedTexturePayloadData(const FString& PayloadKey, TOptional<FString>& AlternateTexturePath) const { return {}; }
 };
 
 
