Added support to reimport policies for material instances

jira none

diff --git a/Engine/Plugins/Interchange/Runtime/Source/Import/Private/Material/InterchangeMaterialFactory.cpp b/Engine/Plugins/Interchange/Runtime/Source/Import/Private/Material/InterchangeMaterialFactory.cpp
--- a/Engine/Plugins/Interchange/Runtime/Source/Import/Private/Material/InterchangeMaterialFactory.cpp
+++ b/Engine/Plugins/Interchange/Runtime/Source/Import/Private/Material/InterchangeMaterialFactory.cpp
@@ -341,6 +341,119 @@ namespace UE
 					TMap<FString, UMaterialExpression*> Expressions;
 				};
 #endif // #if WITH_EDITOR
+
+				void UpdateParameterBool(UMaterialInstance& MaterialInstance, const FString& InputName, const UInterchangeMaterialInstanceFactoryNode& FactoryNode)
+				{
+#if WITH_EDITORONLY_DATA
+					const FName ParameterName = *InputName;
+					bool bInstanceValue;
+					FGuid Uid;
+					if (MaterialInstance.GetStaticSwitchParameterValue(ParameterName, bInstanceValue, Uid))
+					{
+						bool bInputValue = false;
+						FactoryNode.GetBooleanAttribute(UInterchangeShaderPortsAPI::MakeInputValueKey(InputName), bInputValue);
+
+						if (bInputValue != bInstanceValue)
+						{
+							if (UMaterialInstanceConstant* MaterialInstanceConstant = Cast<UMaterialInstanceConstant>(&MaterialInstance))
+							{
+								MaterialInstanceConstant->SetStaticSwitchParameterValueEditorOnly(ParameterName, bInputValue);
+							}
+						}
+					}
+#endif // #if WITH_EDITORONLY_DATA
+				}
+
+				void UpdateParameterFloat(UMaterialInstance& MaterialInstance, const FString& InputName, const UInterchangeMaterialInstanceFactoryNode& FactoryNode)
+				{
+					const FName ParameterName = *InputName;
+					float InstanceValue;
+
+					if (MaterialInstance.GetScalarParameterValue(ParameterName, InstanceValue))
+					{
+						float InputValue = 0.f;
+						FactoryNode.GetFloatAttribute(UInterchangeShaderPortsAPI::MakeInputValueKey(InputName), InputValue);
+
+						if (!FMath::IsNearlyEqual(InputValue, InstanceValue))
+						{
+#if WITH_EDITOR
+							if (UMaterialInstanceConstant* MaterialInstanceConstant = Cast<UMaterialInstanceConstant>(&MaterialInstance))
+							{
+								MaterialInstanceConstant->SetScalarParameterValueEditorOnly(ParameterName, InputValue);
+							}
+							else
+#endif // #if WITH_EDITOR
+							if (UMaterialInstanceDynamic* MaterialInstanceDynamic = Cast<UMaterialInstanceDynamic>(&MaterialInstance))
+							{
+								MaterialInstanceDynamic->SetScalarParameterValue(ParameterName, InputValue);
+							}
+						}
+					}
+				}
+
+				void UpdateParameterLinearColor(UMaterialInstance& MaterialInstance, const FString& InputName, const UInterchangeMaterialInstanceFactoryNode& FactoryNode)
+				{
+					const FName ParameterName = *InputName;
+					FLinearColor InstanceValue;
+
+					if (MaterialInstance.GetVectorParameterValue(ParameterName, InstanceValue))
+					{
+						FLinearColor InputValue;
+						if (FactoryNode.GetLinearColorAttribute(UInterchangeShaderPortsAPI::MakeInputValueKey(InputName), InputValue))
+						{
+							if (!InputValue.Equals(InstanceValue))
+							{
+#if WITH_EDITOR
+								if (UMaterialInstanceConstant* MaterialInstanceConstant = Cast<UMaterialInstanceConstant>(&MaterialInstance))
+								{
+									MaterialInstanceConstant->SetVectorParameterValueEditorOnly(ParameterName, InputValue);
+								}
+								else
+#endif // #if WITH_EDITOR
+								if (UMaterialInstanceDynamic* MaterialInstanceDynamic = Cast<UMaterialInstanceDynamic>(&MaterialInstance))
+								{
+									MaterialInstanceDynamic->SetVectorParameterValue(ParameterName, InputValue);
+								}
+							}
+						}
+					}
+				}
+
+				void UpdateParameterTexture(UMaterialInstance& MaterialInstance, const FString& InputName, const UInterchangeMaterialInstanceFactoryNode& FactoryNode, const UInterchangeBaseNodeContainer& NodeContainer)
+				{
+					const FName ParameterName = *InputName;
+					UTexture* InstanceValue;
+
+					if (MaterialInstance.GetTextureParameterValue(ParameterName, InstanceValue))
+					{
+						FString InputValue;
+						if (FactoryNode.GetStringAttribute(UInterchangeShaderPortsAPI::MakeInputValueKey(InputName), InputValue))
+						{
+							if (const UInterchangeTextureFactoryNode* TextureFactoryNode = Cast<UInterchangeTextureFactoryNode>(NodeContainer.GetNode(InputValue)))
+							{
+								FSoftObjectPath ReferenceObject;
+								TextureFactoryNode->GetCustomReferenceObject(ReferenceObject);
+								if (UTexture* InputTexture = Cast<UTexture>(ReferenceObject.TryLoad()))
+								{
+									if (InputTexture != InstanceValue)
+									{
+#if WITH_EDITOR
+										if (UMaterialInstanceConstant* MaterialInstanceConstant = Cast<UMaterialInstanceConstant>(&MaterialInstance))
+										{
+											MaterialInstanceConstant->SetTextureParameterValueEditorOnly(ParameterName, InputTexture);
+										}
+										else
+#endif // #if WITH_EDITOR
+										if (UMaterialInstanceDynamic* MaterialInstanceDynamic = Cast<UMaterialInstanceDynamic>(&MaterialInstance))
+										{
+											MaterialInstanceDynamic->SetTextureParameterValue(ParameterName, InputTexture);
+										}
+									}
+								}
+							}
+						}
+					}
+				}
 			}
 		}
 	}
@@ -423,20 +536,63 @@ UObject* UInterchangeMaterialFactory::BeginImportAssetObject_GameThread(const FI
 
 #if WITH_EDITOR
 	Material->PreEditChange(nullptr);
+#endif //WITH_EDITOR
 
-	if (UMaterialInstanceConstant* MaterialInstanceConstant = Cast<UMaterialInstanceConstant>(Material))
+	// Setup material instance based on reimport policy if applicable
+	if (const UInterchangeMaterialInstanceFactoryNode* MaterialInstanceFactoryNode = Cast<UInterchangeMaterialInstanceFactoryNode>(MaterialFactoryNode))
 	{
-		if (const UInterchangeMaterialInstanceFactoryNode* MaterialInstanceFactoryNode = Cast<UInterchangeMaterialInstanceFactoryNode>(MaterialFactoryNode))
+		if (UMaterialInstance* MaterialInstance = Cast<UMaterialInstance>(Material))
 		{
-			FString ParentPath;
-			if (MaterialInstanceFactoryNode->GetCustomParent(ParentPath))
+			const EReimportStrategyFlags ReimportStrategyFlags = MaterialFactoryNode->GetReimportStrategyFlags();
+			bool bApplyPipelineProperties = !Arguments.ReimportObject || ReimportStrategyFlags != EReimportStrategyFlags::ApplyNoProperties;
+
+#if WITH_EDITORONLY_DATA
+			// For the time being, reimport policies are only enforced on UMaterialInstanceConstant
+			if (Arguments.ReimportObject && MaterialInstance->IsA<UMaterialInstanceConstant>())
+			{
+
+				if (ReimportStrategyFlags == EReimportStrategyFlags::ApplyEditorChangedProperties)
+				{
+					if (UInterchangeAssetImportData* AssetImportData = Cast<UInterchangeAssetImportData>(MaterialInstance->AssetImportData))
+					{
+						UInterchangeMaterialInstanceFactoryNode* PreviousNode = Cast<UInterchangeMaterialInstanceFactoryNode>(AssetImportData->NodeContainer->GetFactoryNode(AssetImportData->NodeUniqueID));
+
+						if (PreviousNode)
+						{
+							FString PreviousParentPath;
+							FString ParentPath;
+							if (PreviousNode->GetCustomParent(PreviousParentPath) && MaterialInstanceFactoryNode->GetCustomParent(ParentPath))
+							{
+								if (ParentPath == PreviousParentPath)
+								{
+									SetupReimportedMaterialInstance(*MaterialInstance, *Arguments.NodeContainer, *MaterialInstanceFactoryNode, *PreviousNode);
+									bApplyPipelineProperties = false;
+								}
+							}
+						}
+					}
+				}
+			}
+#endif
+
+			if (bApplyPipelineProperties)
 			{
-				FSoftObjectPath ParentMaterial(ParentPath);
-				MaterialInstanceConstant->SetParentEditorOnly(Cast<UMaterialInterface>(ParentMaterial.TryLoad()));
+				SetupMaterialInstance(*MaterialInstance, *Arguments.NodeContainer, *MaterialInstanceFactoryNode, !Arguments.ReimportObject);
+
+#if WITH_EDITOR
+				if (UMaterialInstanceConstant* MaterialInstanceConstant = Cast<UMaterialInstanceConstant>(MaterialInstance))
+				{
+					FString ParentPath;
+					if (MaterialInstanceFactoryNode->GetCustomParent(ParentPath))
+					{
+						UMaterialInterface* ParentMaterial = Cast<UMaterialInterface>(FSoftObjectPath(ParentPath).TryLoad());
+						MaterialInstanceConstant->SetParentEditorOnly(ParentMaterial);
+					}
+				}
+#endif
 			}
 		}
 	}
-#endif //WITH_EDITOR
 
 	return Material;
 }
@@ -488,49 +644,43 @@ UObject* UInterchangeMaterialFactory::ImportAssetObject_Async(const FImportAsset
 		return nullptr;
 	}
 
-	if (MaterialObject)
+	// Currently re-import of UMaterial will not touch the material at all
+	//TODO: Design a re-import process for the material (expressions and input connections)
+	if(!Arguments.ReimportObject)
 	{
-		//Currently material re-import will not touch the material at all
-		//TODO design a re-import process for the material (expressions and input connections)
-		if(!Arguments.ReimportObject)
+		if (UMaterialInterface* MaterialInterface = Cast<UMaterialInterface>(MaterialObject))
 		{
-			if (UMaterialInterface* MaterialInterface = Cast<UMaterialInterface>(MaterialObject))
-			{
 #if WITH_EDITOR
-				if (UMaterial* Material = Cast<UMaterial>(MaterialObject))
-				{
-					SetupMaterial(Material, Arguments, MaterialFactoryNode);
-				}
+			if (UMaterial* Material = Cast<UMaterial>(MaterialObject))
+			{
+				SetupMaterial(Material, Arguments, MaterialFactoryNode);
+			}
 #endif // #if WITH_EDITOR
 
-				MaterialFactoryNode->ApplyAllCustomAttributeToObject(MaterialObject);
-			}
+			MaterialFactoryNode->ApplyAllCustomAttributeToObject(MaterialObject);
 		}
+	}
 #if WITH_EDITORONLY_DATA
-		else
+	else if(UMaterialInstance* MaterialInstance = Cast<UMaterialInstance>(MaterialObject))
+	{
+		//Apply the re import strategy 
+		UInterchangeAssetImportData* InterchangeAssetImportData = Cast<UInterchangeAssetImportData>(MaterialInstance->AssetImportData);
+		UInterchangeFactoryBaseNode* PreviousNode = nullptr;
+		if(InterchangeAssetImportData)
 		{
-			if(UMaterialInstance* MaterialInstance = Cast<UMaterialInstance>(MaterialObject))
-			{
-				//Apply the re import strategy 
-				UInterchangeAssetImportData* InterchangeAssetImportData = Cast<UInterchangeAssetImportData>(MaterialInstance->AssetImportData);
-				UInterchangeFactoryBaseNode* PreviousNode = nullptr;
-				if(InterchangeAssetImportData)
-				{
-					PreviousNode = InterchangeAssetImportData->NodeContainer->GetFactoryNode(InterchangeAssetImportData->NodeUniqueID);
-				}
-				UInterchangeFactoryBaseNode* CurrentNode = NewObject<UInterchangeMaterialFactoryNode>(GetTransientPackage());
-				UInterchangeBaseNode::CopyStorage(MaterialFactoryNode, CurrentNode);
-				CurrentNode->FillAllCustomAttributeFromObject(MaterialInstance);
-				UE::Interchange::FFactoryCommon::ApplyReimportStrategyToAsset(MaterialInstance, PreviousNode, CurrentNode, MaterialFactoryNode);
-			}
+			PreviousNode = InterchangeAssetImportData->NodeContainer->GetFactoryNode(InterchangeAssetImportData->NodeUniqueID);
 		}
+		UInterchangeFactoryBaseNode* CurrentNode = NewObject<UInterchangeMaterialFactoryNode>(GetTransientPackage());
+		UInterchangeBaseNode::CopyStorage(MaterialFactoryNode, CurrentNode);
+		CurrentNode->FillAllCustomAttributeFromObject(MaterialInstance);
+		UE::Interchange::FFactoryCommon::ApplyReimportStrategyToAsset(MaterialInstance, PreviousNode, CurrentNode, MaterialFactoryNode);
+	}
 #endif // #if WITH_EDITORONLY_DATA
 
-		//Getting the file Hash will cache it into the source data
-		Arguments.SourceData->GetFileContentHash();
+	//Getting the file Hash will cache it into the source data
+	Arguments.SourceData->GetFileContentHash();
 
-		//The interchange completion task (call in the GameThread after the factories pass), will call PostEditChange which will trig another asynchronous system that will build all material in parallel
-	}
+	//The interchange completion task (call in the GameThread after the factories pass), will call PostEditChange which will trig another asynchronous system that will build all material in parallel
 
 	return MaterialObject;
 }
@@ -558,13 +708,7 @@ void UInterchangeMaterialFactory::SetupObject_GameThread(const FSetupObjectParam
 				}
 			}
 		}
-		else
 #endif // WITH_EDITOR
-		if (UMaterialInstance* MaterialInstance = Cast<UMaterialInstance>(ImportedMaterialInterface))
-		{
-			// Material instances expect their parameters to only be updated from the game thread
-			SetupMaterialInstance(MaterialInstance, Arguments.NodeContainer, Cast<UInterchangeBaseMaterialFactoryNode>(Arguments.FactoryNode));
-		}
 
 #if WITH_EDITORONLY_DATA
 		UE::Interchange::FFactoryCommon::FUpdateImportAssetDataParameters UpdateImportAssetDataParameters(ImportedMaterialInterface
@@ -1048,126 +1192,151 @@ void UInterchangeMaterialFactory::SetupMaterial(UMaterial* Material, const FImpo
 }
 #endif // #if WITH_EDITOR
 
-void UInterchangeMaterialFactory::SetupMaterialInstance(UMaterialInstance* MaterialInstance, const UInterchangeBaseNodeContainer* NodeContainer, const UInterchangeBaseMaterialFactoryNode* MaterialFactoryNode)
+void UInterchangeMaterialFactory::SetupMaterialInstance(UMaterialInstance& MaterialInstance, const UInterchangeBaseNodeContainer& NodeContainer, const UInterchangeMaterialInstanceFactoryNode& FactoryNode, bool bResetInstance)
 {
-	if (!MaterialFactoryNode || !NodeContainer)
+	using namespace UE::Interchange::MaterialFactory::Internal;
+
+	if (bResetInstance)
 	{
-		return;
+		// Clear all material instance's parameters before applying new ones
+#if WITH_EDITOR
+		if (UMaterialInstanceConstant* MaterialInstanceConstant = Cast<UMaterialInstanceConstant>(&MaterialInstance))
+		{
+			MaterialInstanceConstant->ClearParameterValuesEditorOnly();
+		}
+		else
+#endif
+		if (UMaterialInstanceDynamic* MaterialInstanceDynamic = Cast<UMaterialInstanceDynamic>(&MaterialInstance))
+		{
+			MaterialInstanceDynamic->ClearParameterValues();
+		}
 	}
 
 	TArray<FString> Inputs;
-	UInterchangeShaderPortsAPI::GatherInputs(MaterialFactoryNode, Inputs);
+	UInterchangeShaderPortsAPI::GatherInputs(&FactoryNode, Inputs);
+
+	for (const FString& InputName : Inputs)
+	{
+		const FName ParameterName = *InputName;
+
+		switch(UInterchangeShaderPortsAPI::GetInputType(&FactoryNode, InputName))
+		{
+		case UE::Interchange::EAttributeTypes::Bool:
+			UpdateParameterBool(MaterialInstance, InputName, FactoryNode);
+		    break;
+		case UE::Interchange::EAttributeTypes::Float:
+			UpdateParameterFloat(MaterialInstance, InputName, FactoryNode);
+			break;
+		case UE::Interchange::EAttributeTypes::LinearColor:
+			UpdateParameterLinearColor(MaterialInstance, InputName, FactoryNode);
+			break;
+		case UE::Interchange::EAttributeTypes::String:
+			UpdateParameterTexture(MaterialInstance, InputName, FactoryNode, NodeContainer);
+			break;
+		}
+	}
+}
+
+void UInterchangeMaterialFactory::SetupReimportedMaterialInstance(UMaterialInstance& MaterialInstance, const UInterchangeBaseNodeContainer& NodeContainer, const UInterchangeMaterialInstanceFactoryNode& FactoryNode, const UInterchangeMaterialInstanceFactoryNode& PreviousFactoryNode)
+{
+	using namespace UE::Interchange::MaterialFactory::Internal;
+
+	auto ValidateBool = [&](const FString& InputName) -> bool
+	{
+
+		return true;
+	};
+
+	TArray<FString> Inputs;
+	UInterchangeShaderPortsAPI::GatherInputs(&FactoryNode, Inputs);
 
 	for (const FString& InputName : Inputs)
 	{
 		const FName ParameterName = *InputName;
+		const FName AttributKey = UInterchangeShaderPortsAPI::MakeInputValueKey(InputName);
 
-		switch(UInterchangeShaderPortsAPI::GetInputType(MaterialFactoryNode, InputName))
+		FGuid Uid;
+		switch (UInterchangeShaderPortsAPI::GetInputType(&FactoryNode, InputName))
 		{
 #if WITH_EDITORONLY_DATA
 		case UE::Interchange::EAttributeTypes::Bool:
-		    {
+			{
 				bool bInstanceValue;
-				FGuid Uid;
-				if(MaterialInstance->GetStaticSwitchParameterValue(ParameterName, bInstanceValue, Uid))
+				if (MaterialInstance.GetStaticSwitchParameterValue(*InputName, bInstanceValue, Uid))
 				{
+					bool bPreviousInputValue = false;
+					if (!PreviousFactoryNode.GetBooleanAttribute(AttributKey, bPreviousInputValue))
+					{
+						bPreviousInputValue = bInstanceValue;
+					}
+
 					bool bInputValue = false;
-					MaterialFactoryNode->GetBooleanAttribute(UInterchangeShaderPortsAPI::MakeInputValueKey(InputName), bInputValue);
+					FactoryNode.GetBooleanAttribute(AttributKey, bInputValue);
 
-					if(bInputValue != bInstanceValue)
+					if (bInstanceValue == bPreviousInputValue)
 					{
-						if(UMaterialInstanceConstant* MaterialInstanceConstant = Cast<UMaterialInstanceConstant>(MaterialInstance))
-						{
-							MaterialInstanceConstant->SetStaticSwitchParameterValueEditorOnly(ParameterName, bInputValue);
-						}
+						UpdateParameterBool(MaterialInstance, InputName, FactoryNode);
 					}
 				}
-		    }
-		    break;
-#endif // #if WITH_EDITORONLY_DATA
+			}
+			break;
+#endif
 		case UE::Interchange::EAttributeTypes::Float:
 			{
-				float InstanceValue;
-				if (MaterialInstance->GetScalarParameterValue(ParameterName, InstanceValue))
+				bool bUpdateParameter = true;
+				float PreviousValue;
+				if (PreviousFactoryNode.GetFloatAttribute(AttributKey, PreviousValue))
 				{
-					float InputValue = 0.f;
-					MaterialFactoryNode->GetFloatAttribute(UInterchangeShaderPortsAPI::MakeInputValueKey(InputName), InputValue);
-
-					if ( !FMath::IsNearlyEqual(InputValue, InstanceValue) )
+					float CurrentValue;
+					if (MaterialInstance.GetScalarParameterValue(*InputName, CurrentValue))
 					{
-#if WITH_EDITOR
-						if (UMaterialInstanceConstant* MaterialInstanceConstant = Cast<UMaterialInstanceConstant>(MaterialInstance))
-						{
-							MaterialInstanceConstant->SetScalarParameterValueEditorOnly(ParameterName, InputValue);
-						}
-						else
-#endif // #if WITH_EDITOR
-						if (UMaterialInstanceDynamic* MaterialInstanceDynamic = Cast<UMaterialInstanceDynamic>(MaterialInstance))
-						{
-							MaterialInstanceDynamic->SetScalarParameterValue(ParameterName, InputValue);
-						}
+						bUpdateParameter = FMath::IsNearlyEqual(PreviousValue, CurrentValue);
 					}
 				}
+
+				if(bUpdateParameter)
+				{
+					UpdateParameterFloat(MaterialInstance, InputName, FactoryNode);
+				}
 			}
 			break;
 		case UE::Interchange::EAttributeTypes::LinearColor:
 			{
-				FLinearColor InstanceValue;
-				if (MaterialInstance->GetVectorParameterValue(ParameterName, InstanceValue))
+				bool bUpdateParameter = true;
+				FLinearColor PreviousValue;
+				if (PreviousFactoryNode.GetLinearColorAttribute(AttributKey, PreviousValue))
 				{
-					FLinearColor InputValue;
-					if (MaterialFactoryNode->GetLinearColorAttribute(UInterchangeShaderPortsAPI::MakeInputValueKey(InputName), InputValue))
+					FLinearColor CurrentValue;
+					if (MaterialInstance.GetVectorParameterValue(*InputName, CurrentValue))
 					{
-						if ( !InputValue.Equals(InstanceValue) )
-						{
-#if WITH_EDITOR
-							if (UMaterialInstanceConstant* MaterialInstanceConstant = Cast<UMaterialInstanceConstant>(MaterialInstance))
-							{
-								MaterialInstanceConstant->SetVectorParameterValueEditorOnly(ParameterName, InputValue);
-							}
-							else
-#endif // #if WITH_EDITOR
-							if (UMaterialInstanceDynamic* MaterialInstanceDynamic = Cast<UMaterialInstanceDynamic>(MaterialInstance))
-							{
-								MaterialInstanceDynamic->SetVectorParameterValue(ParameterName, InputValue);
-							}
-						}
+						bUpdateParameter = PreviousValue.Equals(CurrentValue);
 					}
 				}
+
+				if (bUpdateParameter)
+				{
+					UpdateParameterLinearColor(MaterialInstance, InputName, FactoryNode);
+				}
 			}
 			break;
 		case UE::Interchange::EAttributeTypes::String:
 			{
-				UTexture* InstanceValue;
-				if (MaterialInstance->GetTextureParameterValue(ParameterName, InstanceValue))
+				bool bUpdateParameter = true;
+				FString PreviousValue;
+				if (PreviousFactoryNode.GetStringAttribute(AttributKey, PreviousValue))
 				{
-					FString InputValue;
-					if (MaterialFactoryNode->GetStringAttribute(UInterchangeShaderPortsAPI::MakeInputValueKey(InputName), InputValue))
+					UTexture* CurrentValue;
+					if (MaterialInstance.GetTextureParameterValue(*InputName, CurrentValue))
 					{
-						if (const UInterchangeTextureFactoryNode* TextureFactoryNode = Cast<UInterchangeTextureFactoryNode>(NodeContainer->GetNode(InputValue)))
-						{
-							FSoftObjectPath ReferenceObject;
-							TextureFactoryNode->GetCustomReferenceObject(ReferenceObject);
-							if (UTexture* InputTexture = Cast<UTexture>(ReferenceObject.TryLoad()))
-							{
-								if ( InputTexture != InstanceValue )
-								{
-#if WITH_EDITOR
-									if (UMaterialInstanceConstant* MaterialInstanceConstant = Cast<UMaterialInstanceConstant>(MaterialInstance))
-									{
-										MaterialInstanceConstant->SetTextureParameterValueEditorOnly(ParameterName, InputTexture);
-									}
-									else
-#endif // #if WITH_EDITOR
-									if (UMaterialInstanceDynamic* MaterialInstanceDynamic = Cast<UMaterialInstanceDynamic>(MaterialInstance))
-									{
-										MaterialInstanceDynamic->SetTextureParameterValue(ParameterName, InputTexture);
-									}
-								}
-							}
-						}
+						FSoftObjectPath PreviousObjectPath(PreviousValue);
+						bUpdateParameter = PreviousObjectPath == FSoftObjectPath(CurrentValue);
 					}
 				}
+
+				if (bUpdateParameter)
+				{
+					UpdateParameterTexture(MaterialInstance, InputName, FactoryNode, NodeContainer);
+				}
 			}
 			break;
 		}
diff --git a/Engine/Plugins/Interchange/Runtime/Source/Import/Public/Material/InterchangeMaterialFactory.h b/Engine/Plugins/Interchange/Runtime/Source/Import/Public/Material/InterchangeMaterialFactory.h
--- a/Engine/Plugins/Interchange/Runtime/Source/Import/Public/Material/InterchangeMaterialFactory.h
+++ b/Engine/Plugins/Interchange/Runtime/Source/Import/Public/Material/InterchangeMaterialFactory.h
@@ -12,6 +12,7 @@
 
 class UInterchangeBaseMaterialFactoryNode;
 class UInterchangeMaterialExpressionFactoryNode;
+class UInterchangeMaterialInstanceFactoryNode;
 class UMaterial;
 class UMaterialExpression;
 class UMaterialExpressionMaterialFunctionCall;
@@ -41,7 +42,8 @@ private:
 	void SetupMaterial(UMaterial* Material, const FImportAssetObjectParams& Arguments, const UInterchangeBaseMaterialFactoryNode* MaterialFactoryNode);
 #endif // #if WITH_EDITOR
 
-	void SetupMaterialInstance(UMaterialInstance* MaterialInstance, const UInterchangeBaseNodeContainer* NodeContainer, const UInterchangeBaseMaterialFactoryNode* MaterialFactoryNode);
+	void SetupMaterialInstance(UMaterialInstance& MaterialInstance, const UInterchangeBaseNodeContainer& NodeContainer, const UInterchangeMaterialInstanceFactoryNode& FactoryNode, bool bResetInstance);
+	void SetupReimportedMaterialInstance(UMaterialInstance& MaterialInstance, const UInterchangeBaseNodeContainer& NodeContainer, const UInterchangeMaterialInstanceFactoryNode& FactoryNode, const UInterchangeMaterialInstanceFactoryNode& PreviousFactoryNode);
 };
 
 UCLASS(BlueprintType)
