Interchange: make sure we do not create unused materials

#rnx

diff --git a/Engine/Plugins/Interchange/Runtime/Source/Pipelines/Private/InterchangeGenericMaterialPipeline.cpp b/Engine/Plugins/Interchange/Runtime/Source/Pipelines/Private/InterchangeGenericMaterialPipeline.cpp
--- a/Engine/Plugins/Interchange/Runtime/Source/Pipelines/Private/InterchangeGenericMaterialPipeline.cpp
+++ b/Engine/Plugins/Interchange/Runtime/Source/Pipelines/Private/InterchangeGenericMaterialPipeline.cpp
@@ -536,7 +536,6 @@ void UInterchangeGenericMaterialPipeline::ExecutePipeline(UInterchangeBaseNodeCo
 	{
 		SourceNode->GetCustomImportUnusedMaterial(bImportUnusedMaterial);
 	}
-	bImportUnusedMaterial |= bImportMaterials;
 
 #if !WITH_EDITOR
 	// Can't import materials at runtime, fallback to instances
