ocalization issue 281950

diff --git a/Engine/Plugins/Interchange/Runtime/Source/Import/Private/Texture/InterchangePSDTranslator.cpp b/Engine/Plugins/Interchange/Runtime/Source/Import/Private/Texture/InterchangePSDTranslator.cpp
--- a/Engine/Plugins/Interchange/Runtime/Source/Import/Private/Texture/InterchangePSDTranslator.cpp
+++ b/Engine/Plugins/Interchange/Runtime/Source/Import/Private/Texture/InterchangePSDTranslator.cpp
@@ -428,7 +428,7 @@ TOptional<UE::Interchange::FImportImage> UInterchangePSDTranslator::GetTexturePa
 
 	if (!PSDHeader.IsSupported())
 	{
-		FTextureTranslatorUtilities::LogError(*this, NSLOCTEXT("InterchangeDDSTranslator", "UnsupportedFormat", "Format of this PSD is not supported. Only Grayscale and RGBColor PSD images are currently supported, in 8-bit or 16-bit."));
+		FTextureTranslatorUtilities::LogError(*this, NSLOCTEXT("InterchangePSDTranslator", "UnsupportedFormat", "Format of this PSD is not supported. Only Grayscale and RGBColor PSD images are currently supported, in 8-bit or 16-bit."));
 		return TOptional<UE::Interchange::FImportImage>();
 	}
 
@@ -445,7 +445,7 @@ TOptional<UE::Interchange::FImportImage> UInterchangePSDTranslator::GetTexturePa
 
 	if (TextureFormat == TSF_Invalid)
 	{
-		FTextureTranslatorUtilities::LogError(*this, NSLOCTEXT("InterchangeDDSTranslator", "UnsupportedPixelFormat", "This PSD file contains data in an unsupported pixel format."));
+		FTextureTranslatorUtilities::LogError(*this, NSLOCTEXT("InterchangePSDTranslator", "UnsupportedPixelFormat", "This PSD file contains data in an unsupported pixel format."));
 		return TOptional<UE::Interchange::FImportImage>();
 	}
 
@@ -460,7 +460,7 @@ TOptional<UE::Interchange::FImportImage> UInterchangePSDTranslator::GetTexturePa
 
 	if (!UE::Interchange::Private::ReadData(static_cast<uint8*>(PayloadData.RawData.GetData()), Buffer, PSDHeader))
 	{
-		FTextureTranslatorUtilities::LogError(*this, NSLOCTEXT("InterchangeDDSTranslator", "FailedRead", "Failed to read this PSD."));
+		FTextureTranslatorUtilities::LogError(*this, NSLOCTEXT("InterchangePSDTranslator", "FailedRead", "Failed to read this PSD."));
 		return TOptional<UE::Interchange::FImportImage>();
 	}
 
