Fixing texture not correctly applied on standard_surface_greysphere_calibration, <place2d> and <ramp4> were not filtered when the subgraphs are flattened

diff --git a/Engine/Plugins/Interchange/Runtime/Source/Import/Private/MaterialX/MaterialXUtils/MaterialXBase.cpp b/Engine/Plugins/Interchange/Runtime/Source/Import/Private/MaterialX/MaterialXUtils/MaterialXBase.cpp
--- a/Engine/Plugins/Interchange/Runtime/Source/Import/Private/MaterialX/MaterialXUtils/MaterialXBase.cpp
+++ b/Engine/Plugins/Interchange/Runtime/Source/Import/Private/MaterialX/MaterialXUtils/MaterialXBase.cpp
@@ -17,12 +17,14 @@ void FMaterialXBase::UpdateDocumentRecursively(MaterialX::GraphElementPtr Graph)
 	AddTexCoordToTiledImageNodes(Graph);
 
 	static std::set<std::string> SkippedCategories{
-		mx::Category::StandardSurface,
 		mx::Category::Extract,
+		mx::Category::Place2D,
+		mx::Category::Ramp4,
 		mx::Category::Saturate,
 		mx::Category::Separate2,
 		mx::Category::Separate3,
-		mx::Category::Separate4
+		mx::Category::Separate4,
+		mx::Category::StandardSurface
 	};
 
 	auto FilterNode = [](mx::NodePtr Node)
