MaterialX:
- Adding support for UsdPreviewSurface
- Adding method AddVectorAttribute for MaterialX inputs that are vectors
- Putting the different MaterialX material functions under the MaterialX category for easier research

diff --git a/Engine/Plugins/Interchange/Editor/Source/Pipelines/Private/InterchangeMaterialXPipelineSettingsCustomization.cpp b/Engine/Plugins/Interchange/Editor/Source/Pipelines/Private/InterchangeMaterialXPipelineSettingsCustomization.cpp
--- a/Engine/Plugins/Interchange/Editor/Source/Pipelines/Private/InterchangeMaterialXPipelineSettingsCustomization.cpp
+++ b/Engine/Plugins/Interchange/Editor/Source/Pipelines/Private/InterchangeMaterialXPipelineSettingsCustomization.cpp
@@ -34,13 +34,17 @@ void FInterchangeMaterialXPipelineSettingsCustomization::CustomizeChildren(TShar
 			TSharedPtr<IPropertyHandle> ChildPropertyHandle = SurfaceShaderProperty->GetChildHandle(i);
 			TSharedPtr<IPropertyHandle> KeyPropertyHandle = ChildPropertyHandle->GetKeyHandle();
 
-			OnShouldFilterAssetFunc OnShouldFilterAsset = &FInterchangeMaterialXPipelineSettingsCustomization::OnShouldFilterAssetStandardSurface;
+			OnShouldFilterAssetFunc OnShouldFilterAsset = &FInterchangeMaterialXPipelineSettingsCustomization::OnShouldFilterAsset;
 
 			if(KeyPropertyHandle.IsValid())
 			{
 				FString Str;
 				KeyPropertyHandle->GetValueAsDisplayString(Str);
-				if(Str == TEXT("Standard Surface Transmission"))
+				if(Str == TEXT("Standard Surface"))
+				{
+					OnShouldFilterAsset = &FInterchangeMaterialXPipelineSettingsCustomization::OnShouldFilterAssetStandardSurface;
+				}
+				else if(Str == TEXT("Standard Surface Transmission"))
 				{
 					OnShouldFilterAsset = &FInterchangeMaterialXPipelineSettingsCustomization::OnShouldFilterAssetStandardSurfaceTransmission;
 				}
@@ -48,6 +52,10 @@ void FInterchangeMaterialXPipelineSettingsCustomization::CustomizeChildren(TShar
 				{
 					OnShouldFilterAsset = &FInterchangeMaterialXPipelineSettingsCustomization::OnShouldFilterAssetSurfaceUnlit;
 				}
+				else if(Str == TEXT("Usd Preview Surface"))
+				{
+					OnShouldFilterAsset = &FInterchangeMaterialXPipelineSettingsCustomization::OnShouldFilterAssetUsdPreviewSurface;
+				}
 			}
 
 			IDetailPropertyRow& SurfaceShaderTypeRow = ChildBuilder.AddProperty(ChildPropertyHandle.ToSharedRef());
@@ -98,3 +106,8 @@ bool FInterchangeMaterialXPipelineSettingsCustomization::OnShouldFilterAssetSurf
 {
 	return FMaterialXPipelineSettings::ShouldFilterAssets(Cast<UMaterialFunction>(InAssetData.GetAsset()), FMaterialXPipelineSettings::SurfaceUnlitInputs, FMaterialXPipelineSettings::SurfaceUnlitOutputs);
 }
+
+bool FInterchangeMaterialXPipelineSettingsCustomization::OnShouldFilterAssetUsdPreviewSurface(const FAssetData& InAssetData)
+{
+	return FMaterialXPipelineSettings::ShouldFilterAssets(Cast<UMaterialFunction>(InAssetData.GetAsset()), FMaterialXPipelineSettings::UsdPreviewSurfaceInputs, FMaterialXPipelineSettings::UsdPreviewSurfaceOutputs);
+}
diff --git a/Engine/Plugins/Interchange/Editor/Source/Pipelines/Private/InterchangeMaterialXPipelineSettingsCustomization.h b/Engine/Plugins/Interchange/Editor/Source/Pipelines/Private/InterchangeMaterialXPipelineSettingsCustomization.h
--- a/Engine/Plugins/Interchange/Editor/Source/Pipelines/Private/InterchangeMaterialXPipelineSettingsCustomization.h
+++ b/Engine/Plugins/Interchange/Editor/Source/Pipelines/Private/InterchangeMaterialXPipelineSettingsCustomization.h
@@ -19,7 +19,9 @@ public:
 
 protected:
 
+	bool OnShouldFilterAsset(const FAssetData&){return false;}
 	bool OnShouldFilterAssetStandardSurface(const FAssetData & InAssetData);
 	bool OnShouldFilterAssetStandardSurfaceTransmission(const FAssetData& InAssetData);
 	bool OnShouldFilterAssetSurfaceUnlit(const FAssetData& InAssetData);
+	bool OnShouldFilterAssetUsdPreviewSurface(const FAssetData& InAssetData);
 };
\ No newline at end of file
diff --git a/Engine/Plugins/Interchange/Runtime/Source/Import/Private/MaterialX/InterchangeMaterialXDefinitions.h b/Engine/Plugins/Interchange/Runtime/Source/Import/Private/MaterialX/InterchangeMaterialXDefinitions.h
--- a/Engine/Plugins/Interchange/Runtime/Source/Import/Private/MaterialX/InterchangeMaterialXDefinitions.h
+++ b/Engine/Plugins/Interchange/Runtime/Source/Import/Private/MaterialX/InterchangeMaterialXDefinitions.h
@@ -6,6 +6,8 @@
 #if WITH_EDITOR
 
 #include "Math/Color.h"
+#include "Math/MathFwd.h"
+#include "Math/Vector.h"
 
 THIRD_PARTY_INCLUDES_START
 #include "MaterialXCore/Library.h"
@@ -147,6 +149,54 @@ MATERIALX_NAMESPACE_BEGIN
 		}
 	}
 
+	namespace UsdPreviewSurface
+	{
+		namespace Input
+		{
+			static constexpr const char* DiffuseColor = "diffuseColor";
+			static constexpr const char* EmissiveColor = "emissiveColor";
+			static constexpr const char* SpecularColor = "specularColor";
+			static constexpr const char* Metallic = "metallic";
+			static constexpr const char* Roughness = "roughness";
+			static constexpr const char* Clearcoat = "clearcoat";
+			static constexpr const char* ClearcoatRoughness = "clearcoatRoughness";
+			static constexpr const char* Opacity = "opacity";
+			static constexpr const char* OpacityThreshold = "opacityThreshold";
+			static constexpr const char* IOR = "ior";
+			static constexpr const char* Normal = "normal";
+			static constexpr const char* Displacement = "displacement";
+			static constexpr const char* Occlusion = "occlusion";
+		}
+
+		namespace DefaultValue
+		{
+			namespace Float
+			{
+				constexpr float Metallic = 0.f;
+				constexpr float Roughness = 0.5f;
+				constexpr float Clearcoat = 0.f;
+				constexpr float ClearcoatRoughness = 0.01f;
+				constexpr float Opacity = 1.f;
+				constexpr float OpacityThreshold = 0.f;
+				constexpr float IOR = 1.5f;
+				constexpr float Displacement = 0.f;
+				constexpr float Occlusion = 1.f;
+			}
+
+			namespace Color3
+			{
+				constexpr FLinearColor DiffuseColor{ 0.18f, 0.18f, 0.18f };
+				constexpr FLinearColor EmissiveColor{ 0, 0, 0 };
+				constexpr FLinearColor SpecularColor{ 0, 0, 0 };
+			}
+			
+			namespace Vector3
+			{
+				static const FVector3f Normal{ 0.f, 0.f, 1.f };
+			}
+		}
+	}
+
 	namespace Lights
 	{
 		//There's no input per se in a Light, but we can find some common inputs among those lights
diff --git a/Engine/Plugins/Interchange/Runtime/Source/Import/Private/MaterialX/MaterialXUtils/MaterialXBase.cpp b/Engine/Plugins/Interchange/Runtime/Source/Import/Private/MaterialX/MaterialXUtils/MaterialXBase.cpp
--- a/Engine/Plugins/Interchange/Runtime/Source/Import/Private/MaterialX/MaterialXUtils/MaterialXBase.cpp
+++ b/Engine/Plugins/Interchange/Runtime/Source/Import/Private/MaterialX/MaterialXUtils/MaterialXBase.cpp
@@ -24,7 +24,8 @@ void FMaterialXBase::UpdateDocumentRecursively(MaterialX::GraphElementPtr Graph)
 		mx::Category::Separate2,
 		mx::Category::Separate3,
 		mx::Category::Separate4,
-		mx::Category::StandardSurface
+		mx::Category::StandardSurface,
+		mx::Category::UsdPreviewSurface
 	};
 
 	auto FilterNode = [](mx::NodePtr Node)
diff --git a/Engine/Plugins/Interchange/Runtime/Source/Import/Private/MaterialX/MaterialXUtils/MaterialXManager.cpp b/Engine/Plugins/Interchange/Runtime/Source/Import/Private/MaterialX/MaterialXUtils/MaterialXManager.cpp
--- a/Engine/Plugins/Interchange/Runtime/Source/Import/Private/MaterialX/MaterialXUtils/MaterialXManager.cpp
+++ b/Engine/Plugins/Interchange/Runtime/Source/Import/Private/MaterialX/MaterialXUtils/MaterialXManager.cpp
@@ -7,6 +7,7 @@
 #include "InterchangeShaderGraphNode.h"
 #include "MaterialX/MaterialXUtils/MaterialXSurfaceShaderAbstract.h"
 #include "MaterialX/MaterialXUtils/MaterialXStandardSurfaceShader.h"
+#include "MaterialX/MaterialXUtils/MaterialXUsdPreviewSurface.h"
 #include "MaterialX/MaterialXUtils/MaterialXSurfaceUnlitShader.h"
 #include "MaterialX/MaterialXUtils/MaterialXPointLightShader.h"
 #include "MaterialX/MaterialXUtils/MaterialXDirectionalLightShader.h"
@@ -199,6 +200,7 @@ FMaterialXManager::FMaterialXManager()
 	, MaterialXContainerDelegates{
 		{mx::Category::SurfaceUnlit, FMaterialXManager::FOnGetMaterialXInstance::CreateStatic(&FMaterialXSurfaceUnlitShader::MakeInstance)},
 		{mx::Category::StandardSurface, FMaterialXManager::FOnGetMaterialXInstance::CreateStatic(&FMaterialXStandardSurfaceShader::MakeInstance)},
+		{mx::Category::UsdPreviewSurface, FMaterialXManager::FOnGetMaterialXInstance::CreateStatic(&FMaterialXUsdPreviewSurface::MakeInstance)},
 		{mx::Category::PointLight, FMaterialXManager::FOnGetMaterialXInstance::CreateStatic(&FMaterialXPointLightShader::MakeInstance)},
 		{mx::Category::DirectionalLight, FMaterialXManager::FOnGetMaterialXInstance::CreateStatic(&FMaterialXDirectionalLightShader::MakeInstance)},
 		{mx::Category::SpotLight, FMaterialXManager::FOnGetMaterialXInstance::CreateStatic(&FMaterialXSpotLightShader::MakeInstance)},
diff --git a/Engine/Plugins/Interchange/Runtime/Source/Import/Private/MaterialX/MaterialXUtils/MaterialXSurfaceShaderAbstract.cpp b/Engine/Plugins/Interchange/Runtime/Source/Import/Private/MaterialX/MaterialXUtils/MaterialXSurfaceShaderAbstract.cpp
--- a/Engine/Plugins/Interchange/Runtime/Source/Import/Private/MaterialX/MaterialXUtils/MaterialXSurfaceShaderAbstract.cpp
+++ b/Engine/Plugins/Interchange/Runtime/Source/Import/Private/MaterialX/MaterialXUtils/MaterialXSurfaceShaderAbstract.cpp
@@ -76,16 +76,13 @@ bool FMaterialXSurfaceShaderAbstract::AddAttributeFromValueOrInterface(MaterialX
 
 bool FMaterialXSurfaceShaderAbstract::AddFloatAttribute(MaterialX::InputPtr Input, const FString& InputChannelName, UInterchangeShaderNode* ShaderNode, float DefaultValue) const
 {
-	if(Input)
+	if(Input && Input->hasValue())
 	{
-		if(Input->hasValue())
-		{
-			float Value = mx::fromValueString<float>(Input->getValueString());
+		float Value = mx::fromValueString<float>(Input->getValueString());
 
-			if(!FMath::IsNearlyEqual(Value, DefaultValue))
-			{
-				return ShaderNode->AddFloatAttribute(UInterchangeShaderPortsAPI::MakeInputValueKey(InputChannelName), Value);
-			}
+		if(!FMath::IsNearlyEqual(Value, DefaultValue))
+		{
+			return ShaderNode->AddFloatAttribute(UInterchangeShaderPortsAPI::MakeInputValueKey(InputChannelName), Value);
 		}
 	}
 
@@ -94,16 +91,28 @@ bool FMaterialXSurfaceShaderAbstract::AddFloatAttribute(MaterialX::InputPtr Inpu
 
 bool FMaterialXSurfaceShaderAbstract::AddLinearColorAttribute(MaterialX::InputPtr Input, const FString& InputChannelName, UInterchangeShaderNode* ShaderNode, const FLinearColor& DefaultValue) const
 {
-	if(Input)
+	if(Input && Input->hasValue())
 	{
-		if(Input->hasValue())
+		const FLinearColor Value = GetLinearColor(Input);
+
+		if(!Value.Equals(DefaultValue))
 		{
-			const FLinearColor Value = GetLinearColor(Input);
+			return ShaderNode->AddLinearColorAttribute(UInterchangeShaderPortsAPI::MakeInputValueKey(InputChannelName), Value);
+		}
+	}
 
-			if(!Value.Equals(DefaultValue))
-			{
-				return ShaderNode->AddLinearColorAttribute(UInterchangeShaderPortsAPI::MakeInputValueKey(InputChannelName), Value);
-			}
+	return false;
+}
+
+bool FMaterialXSurfaceShaderAbstract::AddVectorAttribute(MaterialX::InputPtr Input, const FString& InputChannelName, UInterchangeShaderNode* ShaderNode, const FVector4f& DefaultValue) const
+{
+	if(Input && Input->hasValue())
+	{
+		const FLinearColor Value = GetVector(Input);
+
+		if(!Value.Equals(DefaultValue))
+		{
+			return ShaderNode->AddLinearColorAttribute(UInterchangeShaderPortsAPI::MakeInputValueKey(InputChannelName), Value);
 		}
 	}
 
@@ -182,7 +191,7 @@ void FMaterialXSurfaceShaderAbstract::ConnectNodeCategoryOutputToInput(const Mat
 		{
 			if(mx::NodePtr DownstreamNode = DownstreamElement->asA<mx::Node>())
 			{
-				if(UInterchangeShaderNode** FoundNode = ShaderNodes.Find(GetAttributeParentName(DownstreamNode)))// DownstreamNode->getName().c_str()))
+				if(UInterchangeShaderNode** FoundNode = ShaderNodes.Find(GetAttributeParentName(DownstreamNode)))
 				{
 					ParentShaderNode = *FoundNode;
 				}
diff --git a/Engine/Plugins/Interchange/Runtime/Source/Import/Private/MaterialX/MaterialXUtils/MaterialXUsdPreviewSurface.cpp b/Engine/Plugins/Interchange/Runtime/Source/Import/Private/MaterialX/MaterialXUtils/MaterialXUsdPreviewSurface.cpp
new file mode 100644
--- /dev/null
+++ b/Engine/Plugins/Interchange/Runtime/Source/Import/Private/MaterialX/MaterialXUtils/MaterialXUsdPreviewSurface.cpp
@@ -0,0 +1,102 @@
+// Copyright Epic Games, Inc. All Rights Reserved.
+
+#if WITH_EDITOR
+
+#include "MaterialXUsdPreviewSurface.h"
+
+namespace mx = MaterialX;
+
+FMaterialXUsdPreviewSurface::FMaterialXUsdPreviewSurface(UInterchangeBaseNodeContainer& BaseNodeContainer)
+	: FMaterialXSurfaceShaderAbstract(BaseNodeContainer)
+{
+	NodeDefinition = mx::NodeDefinition::UsdPreviewSurface;
+}
+
+TSharedRef<FMaterialXBase> FMaterialXUsdPreviewSurface::MakeInstance(UInterchangeBaseNodeContainer& BaseNodeContainer)
+{
+	TSharedRef<FMaterialXUsdPreviewSurface> Result = MakeShared<FMaterialXUsdPreviewSurface>(BaseNodeContainer);
+	Result->RegisterConnectNodeOutputToInputDelegates();
+	return Result;
+}
+
+void FMaterialXUsdPreviewSurface::Translate(MaterialX::NodePtr UsdPreviewSurfaceNode)
+{
+	this->SurfaceShaderNode = UsdPreviewSurfaceNode;
+
+	using namespace UE::Interchange::Materials;
+
+	UInterchangeFunctionCallShaderNode* UsdPreviewSurfaceShaderNode;
+
+	const FString NodeUID = UInterchangeShaderNode::MakeNodeUid(ANSI_TO_TCHAR(UsdPreviewSurfaceNode->getName().c_str()), FStringView{});
+
+	if(UsdPreviewSurfaceShaderNode = const_cast<UInterchangeFunctionCallShaderNode*>(Cast<UInterchangeFunctionCallShaderNode>(NodeContainer.GetNode(NodeUID))); !UsdPreviewSurfaceShaderNode)
+	{
+		const FString NodeName = UsdPreviewSurfaceNode->getName().c_str();
+		UsdPreviewSurfaceShaderNode = NewObject<UInterchangeFunctionCallShaderNode>(&NodeContainer);
+		UsdPreviewSurfaceShaderNode->InitializeNode(NodeUID, NodeName, EInterchangeNodeContainerType::TranslatedAsset);
+
+		UsdPreviewSurfaceShaderNode->SetCustomMaterialFunction(TEXT("/Interchange/Functions/MX_UsdPreviewSurface.MX_UsdPreviewSurface"));
+		NodeContainer.AddNode(UsdPreviewSurfaceShaderNode);
+
+		ShaderNodes.Add(NodeName, UsdPreviewSurfaceShaderNode);
+	}
+
+	// Inputs
+	//Diffuse Color
+	ConnectNodeOutputToInput(mx::UsdPreviewSurface::Input::DiffuseColor, UsdPreviewSurfaceShaderNode, UsdPreviewSurface::Parameters::DiffuseColor.ToString(), mx::UsdPreviewSurface::DefaultValue::Color3::DiffuseColor);
+
+	//Emissive Color
+	ConnectNodeOutputToInput(mx::UsdPreviewSurface::Input::EmissiveColor, UsdPreviewSurfaceShaderNode, UsdPreviewSurface::Parameters::EmissiveColor.ToString(), mx::UsdPreviewSurface::DefaultValue::Color3::EmissiveColor);
+
+	//Specular Color
+	ConnectNodeOutputToInput(mx::UsdPreviewSurface::Input::SpecularColor, UsdPreviewSurfaceShaderNode, UsdPreviewSurface::Parameters::SpecularColor.ToString(), mx::UsdPreviewSurface::DefaultValue::Color3::SpecularColor);
+
+	//Metallic
+	ConnectNodeOutputToInput(mx::UsdPreviewSurface::Input::Metallic, UsdPreviewSurfaceShaderNode, UsdPreviewSurface::Parameters::Metallic.ToString(), mx::UsdPreviewSurface::DefaultValue::Float::Metallic);
+
+	//Roughness
+	ConnectNodeOutputToInput(mx::UsdPreviewSurface::Input::Roughness, UsdPreviewSurfaceShaderNode, UsdPreviewSurface::Parameters::Roughness.ToString(), mx::UsdPreviewSurface::DefaultValue::Float::Roughness);
+
+	//Clearcoat
+	ConnectNodeOutputToInput(mx::UsdPreviewSurface::Input::Clearcoat, UsdPreviewSurfaceShaderNode, UsdPreviewSurface::Parameters::Clearcoat.ToString(), mx::UsdPreviewSurface::DefaultValue::Float::Clearcoat);
+
+	//Clearcoat Roughness
+	ConnectNodeOutputToInput(mx::UsdPreviewSurface::Input::ClearcoatRoughness, UsdPreviewSurfaceShaderNode, UsdPreviewSurface::Parameters::ClearcoatRoughness.ToString(), mx::UsdPreviewSurface::DefaultValue::Float::ClearcoatRoughness);
+
+	//Opacity
+	ConnectNodeOutputToInput(mx::UsdPreviewSurface::Input::Opacity, UsdPreviewSurfaceShaderNode, UsdPreviewSurface::Parameters::Opacity.ToString(), mx::UsdPreviewSurface::DefaultValue::Float::Opacity);
+
+	//Opacity Threshold
+	ConnectNodeOutputToInput(mx::UsdPreviewSurface::Input::OpacityThreshold, UsdPreviewSurfaceShaderNode, UsdPreviewSurface::Parameters::OpacityThreshold.ToString(), mx::UsdPreviewSurface::DefaultValue::Float::OpacityThreshold);
+
+	//IOR
+	ConnectNodeOutputToInput(mx::UsdPreviewSurface::Input::IOR, UsdPreviewSurfaceShaderNode, UsdPreviewSurface::Parameters::IOR.ToString(), mx::UsdPreviewSurface::DefaultValue::Float::IOR);
+
+	//Normal
+	ConnectNodeOutputToInput(mx::UsdPreviewSurface::Input::Normal, UsdPreviewSurfaceShaderNode, UsdPreviewSurface::Parameters::Normal.ToString(), mx::UsdPreviewSurface::DefaultValue::Vector3::Normal);
+
+	//Displacement
+	ConnectNodeOutputToInput(mx::UsdPreviewSurface::Input::Displacement, ShaderGraphNode, UsdPreviewSurface::Parameters::Displacement.ToString(), mx::UsdPreviewSurface::DefaultValue::Float::Displacement);
+
+	//Occlusion
+	ConnectNodeOutputToInput(mx::UsdPreviewSurface::Input::Occlusion, UsdPreviewSurfaceShaderNode, UsdPreviewSurface::Parameters::Occlusion.ToString(), mx::UsdPreviewSurface::DefaultValue::Float::Occlusion);
+
+	// Outputs
+	UInterchangeShaderPortsAPI::ConnectOuputToInputByName(ShaderGraphNode, PBR::Parameters::BaseColor.ToString(), UsdPreviewSurfaceShaderNode->GetUniqueID(), PBR::Parameters::BaseColor.ToString());
+	UInterchangeShaderPortsAPI::ConnectOuputToInputByName(ShaderGraphNode, PBR::Parameters::Metallic.ToString(), UsdPreviewSurfaceShaderNode->GetUniqueID(), PBR::Parameters::Metallic.ToString());
+	UInterchangeShaderPortsAPI::ConnectOuputToInputByName(ShaderGraphNode, PBR::Parameters::Specular.ToString(), UsdPreviewSurfaceShaderNode->GetUniqueID(), PBR::Parameters::Specular.ToString());
+	UInterchangeShaderPortsAPI::ConnectOuputToInputByName(ShaderGraphNode, PBR::Parameters::Roughness.ToString(), UsdPreviewSurfaceShaderNode->GetUniqueID(), PBR::Parameters::Roughness.ToString());
+	UInterchangeShaderPortsAPI::ConnectOuputToInputByName(ShaderGraphNode, PBR::Parameters::EmissiveColor.ToString(), UsdPreviewSurfaceShaderNode->GetUniqueID(), PBR::Parameters::EmissiveColor.ToString());
+	UInterchangeShaderPortsAPI::ConnectOuputToInputByName(ShaderGraphNode, PBR::Parameters::Normal.ToString(), UsdPreviewSurfaceShaderNode->GetUniqueID(), PBR::Parameters::Normal.ToString());
+	
+	if(UInterchangeShaderPortsAPI::HasInput(UsdPreviewSurfaceShaderNode, UsdPreviewSurface::Parameters::Opacity))
+	{
+		UInterchangeShaderPortsAPI::ConnectOuputToInputByName(ShaderGraphNode, PBR::Parameters::Opacity.ToString(), UsdPreviewSurfaceShaderNode->GetUniqueID(), PBR::Parameters::Opacity.ToString());
+	}
+
+	UInterchangeShaderPortsAPI::ConnectOuputToInputByName(ShaderGraphNode, PBR::Parameters::Occlusion.ToString(), UsdPreviewSurfaceShaderNode->GetUniqueID(), PBR::Parameters::Occlusion.ToString());
+	UInterchangeShaderPortsAPI::ConnectOuputToInputByName(ShaderGraphNode, PBR::Parameters::Refraction.ToString(), UsdPreviewSurfaceShaderNode->GetUniqueID(), PBR::Parameters::Refraction.ToString());
+	UInterchangeShaderPortsAPI::ConnectOuputToInputByName(ShaderGraphNode, ClearCoat::Parameters::ClearCoat.ToString(), UsdPreviewSurfaceShaderNode->GetUniqueID(), ClearCoat::Parameters::ClearCoat.ToString());
+	UInterchangeShaderPortsAPI::ConnectOuputToInputByName(ShaderGraphNode, ClearCoat::Parameters::ClearCoatRoughness.ToString(), UsdPreviewSurfaceShaderNode->GetUniqueID(), ClearCoat::Parameters::ClearCoatRoughness.ToString());
+}
+#endif
\ No newline at end of file
diff --git a/Engine/Plugins/Interchange/Runtime/Source/Import/Private/MaterialX/MaterialXUtils/MaterialXUsdPreviewSurface.h b/Engine/Plugins/Interchange/Runtime/Source/Import/Private/MaterialX/MaterialXUtils/MaterialXUsdPreviewSurface.h
new file mode 100644
--- /dev/null
+++ b/Engine/Plugins/Interchange/Runtime/Source/Import/Private/MaterialX/MaterialXUtils/MaterialXUsdPreviewSurface.h
@@ -0,0 +1,24 @@
+// Copyright Epic Games, Inc. All Rights Reserved.
+
+#pragma once
+
+#if WITH_EDITOR
+
+#include "MaterialX/MaterialXUtils/MaterialXSurfaceShaderAbstract.h"
+
+class FMaterialXUsdPreviewSurface : public FMaterialXSurfaceShaderAbstract
+{
+protected:
+
+	template<typename T, ESPMode>
+	friend class SharedPointerInternals::TIntrusiveReferenceController;
+
+	FMaterialXUsdPreviewSurface(UInterchangeBaseNodeContainer& BaseNodeContainer);
+
+public:
+
+	static TSharedRef<FMaterialXBase> MakeInstance(UInterchangeBaseNodeContainer& BaseNodeContainer);
+
+	virtual void Translate(MaterialX::NodePtr UsdPreviewSurfaceNode) override;
+};
+#endif
\ No newline at end of file
diff --git a/Engine/Plugins/Interchange/Runtime/Source/Import/Public/MaterialX/MaterialXUtils/MaterialXSurfaceShaderAbstract.h b/Engine/Plugins/Interchange/Runtime/Source/Import/Public/MaterialX/MaterialXUtils/MaterialXSurfaceShaderAbstract.h
--- a/Engine/Plugins/Interchange/Runtime/Source/Import/Public/MaterialX/MaterialXUtils/MaterialXSurfaceShaderAbstract.h
+++ b/Engine/Plugins/Interchange/Runtime/Source/Import/Public/MaterialX/MaterialXUtils/MaterialXSurfaceShaderAbstract.h
@@ -45,7 +45,7 @@ protected:
 	/**
 	 * Add a float attribute to a shader node only if its value taken from the input is not equal to its default value. Return false if the attribute does not exist or if we cannot add it
 	 *
-	 * @param Input - The MaterialX input to retrieve and add the value from, must be of type float/color/vector
+	 * @param Input - The MaterialX input to retrieve and add the value from, must be of type float
 	 * @param InputChannelName - The name of the shader node's input to add the attribute
 	 * @param ShaderNode - The shader node to which we want to add the attribute
 	 * @param DefaultValue - the default value to test the input against
@@ -57,7 +57,7 @@ protected:
 	/**
 	 * Add a FLinearColor attribute to a shader node only if its value taken from the input is not equal to its default value. Return false if the attribute does not exist or if we cannot add it
 	 *
-	 * @param Input - The MaterialX input to retrieve and add the value from, must be of type float/color/vector
+	 * @param Input - The MaterialX input to retrieve and add the value from, must be of type color
 	 * @param InputChannelName - The name of the shader node's input to add the attribute
 	 * @param ShaderNode - The shader node to which we want to add the attribute
 	 * @param DefaultValue - the default value to test the input against
@@ -66,6 +66,18 @@ protected:
 	 */
 	bool AddLinearColorAttribute(MaterialX::InputPtr Input, const FString& InputChannelName, UInterchangeShaderNode* ShaderNode, const FLinearColor& DefaultValue) const;
 
+	/**
+	 * Add a FLinearColor attribute to a shader node only if its value taken from the input is not equal to its default value. Return false if the attribute does not exist or if we cannot add it
+	 *
+	 * @param Input - The MaterialX input to retrieve and add the value from, must be of type vector
+	 * @param InputChannelName - The name of the shader node's input to add the attribute
+	 * @param ShaderNode - The shader node to which we want to add the attribute
+	 * @param DefaultValue - the default value to test the input against
+	 *
+	 * @return true if the attribute was successfully added
+	 */
+	bool AddVectorAttribute(MaterialX::InputPtr Input, const FString& InputChannelName, UInterchangeShaderNode* ShaderNode, const FVector4f& DefaultValue) const;
+
 	/**
 	 * Connect an output either from a node name or a node graph from a MaterialX input to the ShaderNode
 	 * 
@@ -90,7 +102,7 @@ protected:
 			bIsConnected = ConnectNodeNameOutputToInput(Input, ShaderNode, InputShaderName);
 			if(!bIsConnected)
 			{
-				// only handle float and linear color here, for other types, the child should handle them as it is most likely not an input but a parameter to set in Interchange
+				// only handle float, linear color and vector here, for other types, the child should handle them as it is most likely not an input but a parameter to set in Interchange
 				if constexpr(std::is_same_v<decltype(DefaultValue), float>)
 				{
 					bIsConnected = AddFloatAttribute(Input, InputShaderName, ShaderNode, DefaultValue);
@@ -99,6 +111,10 @@ protected:
 				{
 					bIsConnected = AddLinearColorAttribute(Input, InputShaderName, ShaderNode, DefaultValue);
 				}
+				else if constexpr(std::is_same_v<decltype(DefaultValue), FVector4f> || std::is_same_v<decltype(DefaultValue), FVector3f> || std::is_same_v<decltype(DefaultValue), FVector2f>)
+				{					
+					bIsConnected = AddVectorAttribute(Input, InputShaderName, ShaderNode, DefaultValue);
+				}
 			}
 		}
 
diff --git a/Engine/Plugins/Interchange/Runtime/Source/Nodes/Public/InterchangeMaterialDefinitions.h b/Engine/Plugins/Interchange/Runtime/Source/Nodes/Public/InterchangeMaterialDefinitions.h
--- a/Engine/Plugins/Interchange/Runtime/Source/Nodes/Public/InterchangeMaterialDefinitions.h
+++ b/Engine/Plugins/Interchange/Runtime/Source/Nodes/Public/InterchangeMaterialDefinitions.h
@@ -425,6 +425,28 @@ namespace Interchange
 			}
 		}
 
+		namespace UsdPreviewSurface
+		{
+			const FName Name = TEXT("UsdPreviewSurface");
+
+			namespace Parameters
+			{
+				const FName DiffuseColor = TEXT("diffuseColor");
+				const FName EmissiveColor = TEXT("emissiveColor");
+				const FName SpecularColor = TEXT("specularColor");
+				const FName Metallic = TEXT("metallic");
+				const FName Roughness = TEXT("roughness");
+				const FName Clearcoat = TEXT("clearcoat");
+				const FName ClearcoatRoughness = TEXT("clearcoatRoughness");
+				const FName Opacity = TEXT("opacity");
+				const FName OpacityThreshold = TEXT("opacityThreshold");
+				const FName IOR = TEXT("ior");
+				const FName Normal = TEXT("normal");
+				const FName Displacement = TEXT("displacement");
+				const FName Occlusion = TEXT("occlusion");
+			}
+		}
+
 		namespace Unlit
 		{
 			namespace Parameters
diff --git a/Engine/Plugins/Interchange/Runtime/Source/Pipelines/Private/InterchangeGenericMaterialPipeline.cpp b/Engine/Plugins/Interchange/Runtime/Source/Pipelines/Private/InterchangeGenericMaterialPipeline.cpp
--- a/Engine/Plugins/Interchange/Runtime/Source/Pipelines/Private/InterchangeGenericMaterialPipeline.cpp
+++ b/Engine/Plugins/Interchange/Runtime/Source/Pipelines/Private/InterchangeGenericMaterialPipeline.cpp
@@ -82,6 +82,7 @@ FMaterialXPipelineSettings::FMaterialXPipelineSettings()
 	SurfaceShader.FindOrAdd(EInterchangeMaterialXShaders::StandardSurface, FSoftObjectPath{ TEXT("MaterialFunction'/Interchange/Functions/MX_StandardSurface.MX_StandardSurface'") });
 	SurfaceShader.FindOrAdd(EInterchangeMaterialXShaders::StandardSurfaceTransmission, FSoftObjectPath{ TEXT("MaterialFunction'/Interchange/Functions/MX_TransmissionSurface.MX_TransmissionSurface'") });
 	SurfaceShader.FindOrAdd(EInterchangeMaterialXShaders::SurfaceUnlit, FSoftObjectPath{ TEXT("MaterialFunction'/Interchange/Functions/MX_SurfaceUnlit.MX_SurfaceUnlit'") });
+	SurfaceShader.FindOrAdd(EInterchangeMaterialXShaders::UsdPreviewSurface, FSoftObjectPath{ TEXT("MaterialFunction'/Interchange/Functions/MX_UsdPreviewSurface.MX_UsdPreviewSurface'") });
 }
 
 bool FMaterialXPipelineSettings::AreRequiredPackagesLoaded()
@@ -120,7 +121,10 @@ bool FMaterialXPipelineSettings::AreRequiredPackagesLoaded()
 						{
 							bAllLoaded = !ShouldFilterAssets(Cast<UMaterialFunction>(Asset), SurfaceUnlitInputs, SurfaceUnlitOutputs);
 						}
-
+						else if(Pair.Get<EInterchangeMaterialXShaders>() == EInterchangeMaterialXShaders::UsdPreviewSurface)
+						{
+							bAllLoaded = !ShouldFilterAssets(Cast<UMaterialFunction>(Asset), UsdPreviewSurfaceInputs, UsdPreviewSurfaceOutputs);
+						}
 					}
 #endif // WITH_EDITOR
 				}
@@ -264,6 +268,38 @@ TSet<FName> FMaterialXPipelineSettings::SurfaceUnlitOutputs
 	TEXT("OpacityMask")
 };
 
+TSet<FName> FMaterialXPipelineSettings::UsdPreviewSurfaceInputs
+{
+	UE::Interchange::Materials::UsdPreviewSurface::Parameters::DiffuseColor,
+	UE::Interchange::Materials::UsdPreviewSurface::Parameters::EmissiveColor,
+	UE::Interchange::Materials::UsdPreviewSurface::Parameters::SpecularColor,
+	UE::Interchange::Materials::UsdPreviewSurface::Parameters::Metallic,
+	UE::Interchange::Materials::UsdPreviewSurface::Parameters::Roughness,
+	UE::Interchange::Materials::UsdPreviewSurface::Parameters::Clearcoat,
+	UE::Interchange::Materials::UsdPreviewSurface::Parameters::ClearcoatRoughness,
+	UE::Interchange::Materials::UsdPreviewSurface::Parameters::Opacity,
+	UE::Interchange::Materials::UsdPreviewSurface::Parameters::OpacityThreshold,
+	UE::Interchange::Materials::UsdPreviewSurface::Parameters::IOR,
+	UE::Interchange::Materials::UsdPreviewSurface::Parameters::Normal,
+	UE::Interchange::Materials::UsdPreviewSurface::Parameters::Displacement,
+	UE::Interchange::Materials::UsdPreviewSurface::Parameters::Occlusion
+};
+
+TSet<FName> FMaterialXPipelineSettings::UsdPreviewSurfaceOutputs
+{
+	UE::Interchange::Materials::PBR::Parameters::BaseColor,
+	UE::Interchange::Materials::PBR::Parameters::Metallic,
+	UE::Interchange::Materials::PBR::Parameters::Specular,
+	UE::Interchange::Materials::PBR::Parameters::Roughness,
+	UE::Interchange::Materials::PBR::Parameters::EmissiveColor,
+	UE::Interchange::Materials::PBR::Parameters::Opacity,
+	UE::Interchange::Materials::PBR::Parameters::Normal,
+	UE::Interchange::Materials::Common::Parameters::Refraction,
+	UE::Interchange::Materials::Common::Parameters::Occlusion,
+	UE::Interchange::Materials::ClearCoat::Parameters::ClearCoat,
+	UE::Interchange::Materials::ClearCoat::Parameters::ClearCoatRoughness,
+};
+
 bool FMaterialXPipelineSettings::ShouldFilterAssets(UMaterialFunction* Asset, const TSet<FName>& Inputs, const TSet<FName>& Outputs)
 {
 	int32 InputMatches = 0;
diff --git a/Engine/Plugins/Interchange/Runtime/Source/Pipelines/Public/InterchangeGenericMaterialPipeline.h b/Engine/Plugins/Interchange/Runtime/Source/Pipelines/Public/InterchangeGenericMaterialPipeline.h
--- a/Engine/Plugins/Interchange/Runtime/Source/Pipelines/Public/InterchangeGenericMaterialPipeline.h
+++ b/Engine/Plugins/Interchange/Runtime/Source/Pipelines/Public/InterchangeGenericMaterialPipeline.h
@@ -42,7 +42,9 @@ enum class EInterchangeMaterialXShaders : uint8
 	StandardSurfaceTransmission,
 
 	/** Shader used for unlit surface*/
-	SurfaceUnlit
+	SurfaceUnlit,
+
+	UsdPreviewSurface
 };
 
 USTRUCT(BlueprintType)
@@ -73,6 +75,8 @@ private:
 	static TSet<FName> TransmissionSurfaceOutputs;
 	static TSet<FName> SurfaceUnlitInputs;
 	static TSet<FName> SurfaceUnlitOutputs;
+	static TSet<FName> UsdPreviewSurfaceInputs;
+	static TSet<FName> UsdPreviewSurfaceOutputs;
 #endif // WITH_EDITOR
 };
 
