Prevent overriding existing material when importing asset other then materials

#rnx

diff --git a/Engine/Plugins/Experimental/Enterprise/DatasmithInterchange/Source/DatasmithInterchange/Private/InterchangeDatasmithSceneFactory.cpp b/Engine/Plugins/Experimental/Enterprise/DatasmithInterchange/Source/DatasmithInterchange/Private/InterchangeDatasmithSceneFactory.cpp
--- a/Engine/Plugins/Experimental/Enterprise/DatasmithInterchange/Source/DatasmithInterchange/Private/InterchangeDatasmithSceneFactory.cpp
+++ b/Engine/Plugins/Experimental/Enterprise/DatasmithInterchange/Source/DatasmithInterchange/Private/InterchangeDatasmithSceneFactory.cpp
@@ -15,22 +15,23 @@ UClass* UInterchangeDatasmithSceneFactory::GetFactoryClass() const
 }
 
 
-UObject* UInterchangeDatasmithSceneFactory::BeginImportAssetObject_GameThread(const FImportAssetObjectParams& Arguments)
+UInterchangeFactoryBase::FImportAssetResult UInterchangeDatasmithSceneFactory::BeginImportAsset_GameThread(const FImportAssetObjectParams& Arguments)
 {
+	FImportAssetResult ImportAssetResult;
 	UDatasmithScene* DatasmithScene = nullptr;
 
 #if WITH_EDITORONLY_DATA
 	if (!Arguments.AssetNode)
 	{
 		UE_LOG(LogInterchangeDatasmith, Error, TEXT("UInterchangeDatasmithSceneFactory: Asset node parameter is null."));
-		return nullptr;
+		return ImportAssetResult;
 	}
 
 	const UClass* DatasmithSceneClass = Arguments.AssetNode->GetObjectClass();
 	if (!DatasmithSceneClass || !DatasmithSceneClass->IsChildOf(UDatasmithScene::StaticClass()))
 	{
 		UE_LOG(LogInterchangeDatasmith, Error, TEXT("UInterchangeDatasmithSceneFactory: Asset node parameter class doesnt derive from UDatasmithScene."));
-		return nullptr;
+		return ImportAssetResult;
 	}
 
 	// create an asset if it doesn't exist
@@ -50,18 +51,21 @@ UObject* UInterchangeDatasmithSceneFactory::BeginImportAssetObject_GameThread(co
 	if (!DatasmithScene)
 	{
 		UE_LOG(LogInterchangeDatasmith, Warning, TEXT("Could not create Datasmith Scene asset %s"), *Arguments.AssetName);
-		return nullptr;
+		return ImportAssetResult;
 	}
 #endif //WITH_EDITORONLY_DATA
 
-	return DatasmithScene;
+	ImportAssetResult.ImportedObject = DatasmithScene;
+	return ImportAssetResult;
 }
 
-UObject* UInterchangeDatasmithSceneFactory::ImportAssetObject_Async(const FImportAssetObjectParams& Arguments)
+UInterchangeFactoryBase::FImportAssetResult UInterchangeDatasmithSceneFactory::ImportAsset_Async(const FImportAssetObjectParams& Arguments)
 {
+	FImportAssetResult ImportAssetResult;
+
  #if !WITH_EDITORONLY_DATA
  	UE_LOG(LogInterchangeDatasmith, Error, TEXT("Cannot import datasmith scene asset in runtime, this is an editor only feature."));
- 	return nullptr; 
+ 	return ImportAssetResult;
  #else //WITH_EDITORONLY_DATA
 
 	TRACE_CPUPROFILER_EVENT_SCOPE(UInterchangeTextureFactory::CreateAsset);
@@ -69,14 +73,14 @@ UObject* UInterchangeDatasmithSceneFactory::ImportAssetObject_Async(const FImpor
 	if (!Arguments.AssetNode)
 	{
 		UE_LOG(LogInterchangeDatasmith, Error, TEXT("UInterchangeDatasmithSceneFactory: Asset node parameter is null."));
-		return nullptr;
+		return ImportAssetResult;
 	}
 
 	const UClass* DatasmithSceneClass = Arguments.AssetNode->GetObjectClass();
 	if (!DatasmithSceneClass || !DatasmithSceneClass->IsChildOf(UDatasmithScene::StaticClass()))
 	{
 		UE_LOG(LogInterchangeDatasmith, Error, TEXT("UInterchangeDatasmithSceneFactory: Asset node parameter class doesnt derive from UDatasmithScene."));
-		return nullptr;
+		return ImportAssetResult;
 	}
 
 	UClass* SupportedFactoryNodeClass = Arguments.AssetNode->GetClass() == UInterchangeDatasmithSceneFactoryNode::StaticClass()
@@ -85,7 +89,7 @@ UObject* UInterchangeDatasmithSceneFactory::ImportAssetObject_Async(const FImpor
 	if (SupportedFactoryNodeClass == nullptr)
 	{
 		UE_LOG(LogInterchangeDatasmith, Error, TEXT("UInterchangeDatasmithSceneFactory: Asset node parameter is not a UInterchangeDatasmithSceneFactoryNode."));
-		return nullptr;
+		return ImportAssetResult;
 	}
 
 	// create an asset if it doesn't exist
@@ -109,13 +113,14 @@ UObject* UInterchangeDatasmithSceneFactory::ImportAssetObject_Async(const FImpor
 	if (!DatasmithScene)
 	{
 		UE_LOG(LogInterchangeDatasmith, Error, TEXT("UInterchangeDatasmithSceneFactory: Could not create datasmith scene asset %s"), *Arguments.AssetName);
-		return nullptr;
+		return ImportAssetResult;
 	}
 
 	// TODO add InterchangeAssetImportData or DatasmithAssetImportData.
 	// TODO link the created asset to the factory.
 
-	return DatasmithScene;
+	ImportAssetResult.ImportedObject = DatasmithScene;
+	return ImportAssetResult;
 #endif //WITH_EDITORONLY_DATA
 }
 
diff --git a/Engine/Plugins/Experimental/Enterprise/DatasmithInterchange/Source/DatasmithInterchange/Public/InterchangeDatasmithSceneFactory.h b/Engine/Plugins/Experimental/Enterprise/DatasmithInterchange/Source/DatasmithInterchange/Public/InterchangeDatasmithSceneFactory.h
--- a/Engine/Plugins/Experimental/Enterprise/DatasmithInterchange/Source/DatasmithInterchange/Public/InterchangeDatasmithSceneFactory.h
+++ b/Engine/Plugins/Experimental/Enterprise/DatasmithInterchange/Source/DatasmithInterchange/Public/InterchangeDatasmithSceneFactory.h
@@ -21,8 +21,8 @@ public:
 	// Interchange factory base interface begin
 
 	virtual UClass* GetFactoryClass() const override;
-	virtual UObject* BeginImportAssetObject_GameThread(const FImportAssetObjectParams& Arguments) override;
-	virtual UObject* ImportAssetObject_Async(const FImportAssetObjectParams& Arguments) override;
+	virtual FImportAssetResult BeginImportAsset_GameThread(const FImportAssetObjectParams& Arguments) override;
+	virtual FImportAssetResult ImportAsset_Async(const FImportAssetObjectParams& Arguments) override;
 	virtual void SetupObject_GameThread(const FSetupObjectParams& Arguments) override;
 	// Interchange factory base interface end
 	//////////////////////////////////////////////////////////////////////////
diff --git a/Engine/Plugins/Interchange/Runtime/Source/Import/Private/Animation/InterchangeAnimSequenceFactory.cpp b/Engine/Plugins/Interchange/Runtime/Source/Import/Private/Animation/InterchangeAnimSequenceFactory.cpp
--- a/Engine/Plugins/Interchange/Runtime/Source/Import/Private/Animation/InterchangeAnimSequenceFactory.cpp
+++ b/Engine/Plugins/Interchange/Runtime/Source/Import/Private/Animation/InterchangeAnimSequenceFactory.cpp
@@ -767,29 +767,30 @@ UClass* UInterchangeAnimSequenceFactory::GetFactoryClass() const
 	return UAnimSequence::StaticClass();
 }
 
-UObject* UInterchangeAnimSequenceFactory::BeginImportAssetObject_GameThread(const FImportAssetObjectParams& Arguments)
+UInterchangeFactoryBase::FImportAssetResult UInterchangeAnimSequenceFactory::BeginImportAsset_GameThread(const FImportAssetObjectParams& Arguments)
 {
+	FImportAssetResult ImportAssetResult;
 #if !WITH_EDITOR || !WITH_EDITORONLY_DATA
 
 	UE_LOG(LogInterchangeImport, Error, TEXT("Cannot import animsequence asset in runtime, this is an editor only feature."));
-	return nullptr;
+	return ImportAssetResult;
 #else
 	UAnimSequence* AnimSequence = nullptr;
 	if (!Arguments.AssetNode || !Arguments.AssetNode->GetObjectClass()->IsChildOf(GetFactoryClass()))
 	{
-		return nullptr;
+		return ImportAssetResult;
 	}
 
 	const UInterchangeAnimSequenceFactoryNode* AnimSequenceFactoryNode = Cast<UInterchangeAnimSequenceFactoryNode>(Arguments.AssetNode);
 	if (AnimSequenceFactoryNode == nullptr)
 	{
-		return nullptr;
+		return ImportAssetResult;
 	}
 	
 	//Verify if the bone track animation is valid (sequence length versus framerate ...)
 	if(!IsBoneTrackAnimationValid(AnimSequenceFactoryNode, Arguments))
 	{
-		return nullptr;
+		return ImportAssetResult;
 	}
 
 	// create an asset if it doesn't exist
@@ -809,12 +810,13 @@ UObject* UInterchangeAnimSequenceFactory::BeginImportAssetObject_GameThread(cons
 	if (!AnimSequence)
 	{
 		UE_LOG(LogInterchangeImport, Warning, TEXT("Could not create AnimSequence asset %s"), *Arguments.AssetName);
-		return nullptr;
+		return ImportAssetResult;
 	}
 
 	AnimSequence->PreEditChange(nullptr);
 
-	return ImportObjectSourceData(Arguments);
+	ImportAssetResult.ImportedObject = ImportObjectSourceData(Arguments);
+	return ImportAssetResult;
 #endif //else !WITH_EDITOR || !WITH_EDITORONLY_DATA
 }
 
diff --git a/Engine/Plugins/Interchange/Runtime/Source/Import/Private/Animation/InterchangeAnimationTrackSetFactory.cpp b/Engine/Plugins/Interchange/Runtime/Source/Import/Private/Animation/InterchangeAnimationTrackSetFactory.cpp
--- a/Engine/Plugins/Interchange/Runtime/Source/Import/Private/Animation/InterchangeAnimationTrackSetFactory.cpp
+++ b/Engine/Plugins/Interchange/Runtime/Source/Import/Private/Animation/InterchangeAnimationTrackSetFactory.cpp
@@ -652,12 +652,13 @@ UClass* UInterchangeAnimationTrackSetFactory::GetFactoryClass() const
 	return ULevelSequence::StaticClass();
 }
 
-UObject* UInterchangeAnimationTrackSetFactory::BeginImportAssetObject_GameThread(const FImportAssetObjectParams& Arguments)
+UInterchangeFactoryBase::FImportAssetResult UInterchangeAnimationTrackSetFactory::BeginImportAsset_GameThread(const FImportAssetObjectParams& Arguments)
 {
+	FImportAssetResult ImportAssetResult;
 #if !WITH_EDITOR || !WITH_EDITORONLY_DATA
 
 	UE_LOG(LogInterchangeImport, Error, TEXT("Cannot import levelsequence asset in runtime, this is an editor only feature."));
-	return nullptr;
+	return ImportAssetResult;
 #else
 	if (Arguments.ReimportObject)
 	{
@@ -667,25 +668,25 @@ UObject* UInterchangeAnimationTrackSetFactory::BeginImportAssetObject_GameThread
 		Message->AssetType = ULevelSequence::StaticClass();
 		Message->Text = LOCTEXT("CreateEmptyAssetUnsupportedReimport", "Re-import of ULevelSequence not supported yet.");
 
-		return nullptr;
+		return ImportAssetResult;
 	}
 
 	ULevelSequence* LevelSequence = nullptr;
 	if (!Arguments.AssetNode || !Arguments.AssetNode->GetObjectClass()->IsChildOf(GetFactoryClass()))
 	{
-		return nullptr;
+		return ImportAssetResult;
 	}
 
 	const UInterchangeAnimationTrackSetFactoryNode* FactoryNode = Cast<UInterchangeAnimationTrackSetFactoryNode>(Arguments.AssetNode);
 	if (FactoryNode == nullptr)
 	{
-		return nullptr;
+		return ImportAssetResult;
 	}
 
 	if (!UE::Interchange::Private::HasActorToUse(Arguments.NodeContainer, FactoryNode))
 	{
 		UE_LOG(LogInterchangeImport, Warning, TEXT("Level sequence asset, %s, not imported, because all referenced actors are missing."), *FactoryNode->GetDisplayLabel());
-		return nullptr;
+		return ImportAssetResult;
 	}
 
 	// create an asset if it doesn't exist
@@ -705,12 +706,13 @@ UObject* UInterchangeAnimationTrackSetFactory::BeginImportAssetObject_GameThread
 	if (!LevelSequence)
 	{
 		UE_LOG(LogInterchangeImport, Warning, TEXT("Could not create LevelSequence asset %s"), *Arguments.AssetName);
-		return nullptr;
+		return ImportAssetResult;
 	}
 
 	LevelSequence->PreEditChange(nullptr);
 
-	return ImportObjectSourceData(Arguments);
+	ImportAssetResult.ImportedObject = ImportObjectSourceData(Arguments);
+	return ImportAssetResult;
 #endif //else !WITH_EDITOR || !WITH_EDITORONLY_DATA
 }
 
diff --git a/Engine/Plugins/Interchange/Runtime/Source/Import/Private/Material/InterchangeMaterialFactory.cpp b/Engine/Plugins/Interchange/Runtime/Source/Import/Private/Material/InterchangeMaterialFactory.cpp
--- a/Engine/Plugins/Interchange/Runtime/Source/Import/Private/Material/InterchangeMaterialFactory.cpp
+++ b/Engine/Plugins/Interchange/Runtime/Source/Import/Private/Material/InterchangeMaterialFactory.cpp
@@ -9,6 +9,7 @@
 #include "InterchangeSourceData.h"
 #include "InterchangeTextureNode.h"
 #include "InterchangeTextureFactoryNode.h"
+#include "InterchangeTranslatorBase.h"
 #include "Materials/Material.h"
 #include "Materials/MaterialExpression.h"
 #include "Materials/MaterialExpressionClearCoatNormalCustomOutput.h"
@@ -464,17 +465,20 @@ UClass* UInterchangeMaterialFactory::GetFactoryClass() const
 	return UMaterialInterface::StaticClass();
 }
 
-UObject* UInterchangeMaterialFactory::BeginImportAssetObject_GameThread(const FImportAssetObjectParams& Arguments)
+UInterchangeFactoryBase::FImportAssetResult UInterchangeMaterialFactory::BeginImportAsset_GameThread(const FImportAssetObjectParams& Arguments)
 {
+	FImportAssetResult ImportAssetResult;
 	UObject* Material = nullptr;
 
-	auto CouldNotCreateMaterialLog = [this, &Arguments](const FText& Info)
+	auto CouldNotCreateMaterialLog = [this, &Arguments, &ImportAssetResult](const FText& Info)
 	{
 		UInterchangeResultError_Generic* Message = AddMessage<UInterchangeResultError_Generic>();
 		Message->SourceAssetName = Arguments.SourceData->GetFilename();
 		Message->DestinationAssetName = Arguments.AssetName;
 		Message->AssetType = GetFactoryClass();
 		Message->Text = FText::Format(LOCTEXT("MatFactory_CouldNotCreateMat", "Could not create Material asset %s. Reason: %s"), FText::FromString(Arguments.AssetName), Info);
+		bSkipImport = true;
+		ImportAssetResult.bIsFactorySkipAsset = true;
 	};
 	
 	const FText MissMatchClassText = LOCTEXT("MatFactory_CouldNotCreateMat_MissMatchClass", "Missmatch between interchange material factory node class and factory class.");
@@ -482,25 +486,27 @@ UObject* UInterchangeMaterialFactory::BeginImportAssetObject_GameThread(const FI
 	if (!Arguments.AssetNode || !Arguments.AssetNode->GetObjectClass()->IsChildOf(GetFactoryClass()))
 	{
 		CouldNotCreateMaterialLog(MissMatchClassText);
-		return nullptr;
+		return ImportAssetResult;
 	}
 
 	const UInterchangeBaseMaterialFactoryNode* MaterialFactoryNode = Cast<UInterchangeBaseMaterialFactoryNode>(Arguments.AssetNode);
 	if (MaterialFactoryNode == nullptr)
 	{
 		CouldNotCreateMaterialLog(LOCTEXT("MatFactory_CouldNotCreateMat_CannotCastFactoryNode", "Cannot cast interchange factory node to UInterchangeBaseMaterialFactoryNode."));
-		return nullptr;
+		return ImportAssetResult;
 	}
 
 	const UClass* MaterialClass = MaterialFactoryNode->GetObjectClass();
 	if (!ensure(MaterialClass && MaterialClass->IsChildOf(GetFactoryClass())))
 	{
 		CouldNotCreateMaterialLog(MissMatchClassText);
-		return nullptr;
+		return ImportAssetResult;
 	}
 
-	// create an asset if it doesn't exist
-	UObject* ExistingAsset = StaticFindObject(nullptr, Arguments.Parent, *Arguments.AssetName);
+	const bool bIsReimport = Arguments.ReimportObject != nullptr;
+
+	// Find existing asset
+	UObject* ExistingAsset = bIsReimport ? Arguments.ReimportObject : StaticFindObject(nullptr, Arguments.Parent, *Arguments.AssetName);
 
 	// create a new material or overwrite existing asset, if possible
 	if (!ExistingAsset)
@@ -526,16 +532,26 @@ UObject* UInterchangeMaterialFactory::BeginImportAssetObject_GameThread(const FI
 	{
 		//This is a reimport, we are just re-updating the source data
 		Material = ExistingAsset;
+		//We allow override of existing materials only if the translator is a pure material translator or the user directly ask to re-import this object
+		if (!bIsReimport && Arguments.Translator->GetSupportedAssetTypes() != EInterchangeTranslatorAssetType::Materials)
+		{
+			//Do not override the material asset
+			ImportAssetResult.bIsFactorySkipAsset = true;
+			bSkipImport = true;
+		}
 	}
 
 	if (!Material)
 	{
 		CouldNotCreateMaterialLog(LOCTEXT("MatFactory_CouldNotCreateMat_MaterialCreationFail", "Material creation fail."));
-		return nullptr;
+		return ImportAssetResult;
 	}
 
 #if WITH_EDITOR
-	Material->PreEditChange(nullptr);
+	if (!bSkipImport)
+	{
+		Material->PreEditChange(nullptr);
+	}
 #endif //WITH_EDITOR
 
 	// Setup material instance based on reimport policy if applicable
@@ -594,20 +610,24 @@ UObject* UInterchangeMaterialFactory::BeginImportAssetObject_GameThread(const FI
 		}
 	}
 
-	return Material;
+	ImportAssetResult.ImportedObject = Material;
+	return ImportAssetResult;
 }
 
-UObject* UInterchangeMaterialFactory::ImportAssetObject_Async(const FImportAssetObjectParams& Arguments)
+UInterchangeFactoryBase::FImportAssetResult UInterchangeMaterialFactory::ImportAsset_Async(const FImportAssetObjectParams& Arguments)
 {
+	FImportAssetResult ImportAssetResult;
+	ImportAssetResult.bIsFactorySkipAsset = bSkipImport;
+
 	if (!Arguments.AssetNode || !Arguments.AssetNode->GetObjectClass()->IsChildOf(GetFactoryClass()))
 	{
-		return nullptr;
+		return ImportAssetResult;
 	}
 
 	UInterchangeBaseMaterialFactoryNode* MaterialFactoryNode = Cast<UInterchangeBaseMaterialFactoryNode>(Arguments.AssetNode);
 	if (MaterialFactoryNode == nullptr)
 	{
-		return nullptr;
+		return ImportAssetResult;
 	}
 
 	const UClass* MaterialClass = MaterialFactoryNode->GetObjectClass();
@@ -616,6 +636,13 @@ UObject* UInterchangeMaterialFactory::ImportAssetObject_Async(const FImportAsset
 	// create an asset if it doesn't exist
 	UObject* ExistingAsset = StaticFindObject(nullptr, Arguments.Parent, *Arguments.AssetName);
 
+	//Do not override an asset we skip
+	if (bSkipImport)
+	{
+		ImportAssetResult.ImportedObject = ExistingAsset;
+		return ImportAssetResult;
+	}
+
 	UObject* MaterialObject = nullptr;
 	// create a new material or overwrite existing asset, if possible
 	if (!ExistingAsset)
@@ -629,7 +656,7 @@ UObject* UInterchangeMaterialFactory::ImportAssetObject_Async(const FImportAsset
 		else
 		{
 			UE_LOG(LogInterchangeImport, Error, TEXT("Could not create Material asset [%s] outside of the game thread"), *Arguments.AssetName);
-			return nullptr;
+			return ImportAssetResult;
 		}
 	}
 	else if(ExistingAsset->GetClass()->IsChildOf(MaterialClass))
@@ -641,7 +668,7 @@ UObject* UInterchangeMaterialFactory::ImportAssetObject_Async(const FImportAsset
 	if (!MaterialObject)
 	{
 		UE_LOG(LogInterchangeImport, Warning, TEXT("Could not create Material asset %s"), *Arguments.AssetName);
-		return nullptr;
+		return ImportAssetResult;
 	}
 
 	// Currently re-import of UMaterial will not touch the material at all
@@ -681,8 +708,8 @@ UObject* UInterchangeMaterialFactory::ImportAssetObject_Async(const FImportAsset
 	Arguments.SourceData->GetFileContentHash();
 
 	//The interchange completion task (call in the GameThread after the factories pass), will call PostEditChange which will trig another asynchronous system that will build all material in parallel
-
-	return MaterialObject;
+	ImportAssetResult.ImportedObject = MaterialObject;
+	return ImportAssetResult;
 }
 
 /* This function is call in the completion task on the main thread, use it to call main thread post creation step for your assets*/
@@ -691,6 +718,11 @@ void UInterchangeMaterialFactory::SetupObject_GameThread(const FSetupObjectParam
 	check(IsInGameThread());
 	Super::SetupObject_GameThread(Arguments);
 
+	if (bSkipImport)
+	{
+		return;
+	}
+
 	if (ensure(Arguments.ImportedObject && Arguments.SourceData))
 	{
 		//We must call the Update of the asset source file in the main thread because UAssetImportData::Update execute some delegate we do not control
@@ -1348,17 +1380,20 @@ UClass* UInterchangeMaterialFunctionFactory::GetFactoryClass() const
 	return UMaterialFunctionInterface::StaticClass();
 }
 
-UObject* UInterchangeMaterialFunctionFactory::BeginImportAssetObject_GameThread(const FImportAssetObjectParams& Arguments)
+UInterchangeFactoryBase::FImportAssetResult UInterchangeMaterialFunctionFactory::BeginImportAsset_GameThread(const FImportAssetObjectParams& Arguments)
 {
+	FImportAssetResult ImportAssetResult;
 	UObject* Material = nullptr;
 
-	auto CouldNotCreateMaterialLog = [this, &Arguments](const FText& Info)
+	auto CouldNotCreateMaterialLog = [this, &Arguments, &ImportAssetResult](const FText& Info)
 	{
 		UInterchangeResultError_Generic* Message = AddMessage<UInterchangeResultError_Generic>();
 		Message->SourceAssetName = Arguments.SourceData->GetFilename();
 		Message->DestinationAssetName = Arguments.AssetName;
 		Message->AssetType = GetFactoryClass();
 		Message->Text = FText::Format(LOCTEXT("MatFunc_CouldNotCreateMat", "Could not create Material asset %s. Reason: %s"), FText::FromString(Arguments.AssetName), Info);
+		bSkipImport = true;
+		ImportAssetResult.bIsFactorySkipAsset = true;
 	};
 
 	const FText MissMatchClassText = LOCTEXT("MatFunc_CouldNotCreateMat_MissMatchClass", "Missmatch between interchange material factory node class and factory class.");
@@ -1366,25 +1401,28 @@ UObject* UInterchangeMaterialFunctionFactory::BeginImportAssetObject_GameThread(
 	if (!Arguments.AssetNode || !Arguments.AssetNode->GetObjectClass()->IsChildOf(GetFactoryClass()))
 	{
 		CouldNotCreateMaterialLog(MissMatchClassText);
-		return nullptr;
+		return ImportAssetResult;
 	}
 
 	const UInterchangeMaterialFunctionFactoryNode* MaterialFactoryNode = Cast<UInterchangeMaterialFunctionFactoryNode>(Arguments.AssetNode);
 	if (!ensure(MaterialFactoryNode))
 	{
 		CouldNotCreateMaterialLog(LOCTEXT("MatFunc_CouldNotCreateMat_CannotCastFactoryNode", "Cannot cast interchange factory node to UInterchangeBaseMaterialFactoryNode."));
-		return nullptr;
+		return ImportAssetResult;
 	}
 
 	const UClass* MaterialClass = MaterialFactoryNode->GetObjectClass();
 	if (!ensure(MaterialClass && MaterialClass->IsChildOf(GetFactoryClass())))
 	{
 		CouldNotCreateMaterialLog(MissMatchClassText);
-		return nullptr;
+		return ImportAssetResult;
 	}
 
 	// create an asset if it doesn't exist
-	UObject* ExistingAsset = StaticFindObject(nullptr, Arguments.Parent, *Arguments.AssetName);
+	const bool bIsReimport = Arguments.ReimportObject != nullptr;
+
+	// Find existing asset
+	UObject* ExistingAsset = bIsReimport ? Arguments.ReimportObject : StaticFindObject(nullptr, Arguments.Parent, *Arguments.AssetName);
 
 	// create a new material or overwrite existing asset, if possible
 	if (!ExistingAsset)
@@ -1395,32 +1433,46 @@ UObject* UInterchangeMaterialFunctionFactory::BeginImportAssetObject_GameThread(
 	{
 		//This is a reimport, we are just re-updating the source data
 		Material = ExistingAsset;
+		//We allow override of existing materials only if the translator is a pure material translator
+		if (!bIsReimport && Arguments.Translator->GetSupportedAssetTypes() != EInterchangeTranslatorAssetType::Materials)
+		{
+			//Do not override the asset
+			bSkipImport = true;
+			ImportAssetResult.bIsFactorySkipAsset = bSkipImport;
+		}
 	}
 
 	if (!Material)
 	{
 		CouldNotCreateMaterialLog(LOCTEXT("MatFunc_CouldNotCreateMat_MaterialCreationFail", "Material creation fail."));
-		return nullptr;
+		return ImportAssetResult;
 	}
 
 #if WITH_EDITOR
-	Material->PreEditChange(nullptr);
+	if (!bSkipImport)
+	{
+		Material->PreEditChange(nullptr);
+	}
 #endif //WITH_EDITOR
 
-	return Material;
+	ImportAssetResult.ImportedObject = Material;
+	return ImportAssetResult;
 }
 
-UObject* UInterchangeMaterialFunctionFactory::ImportAssetObject_Async(const FImportAssetObjectParams& Arguments)
+UInterchangeFactoryBase::FImportAssetResult UInterchangeMaterialFunctionFactory::ImportAsset_Async(const FImportAssetObjectParams& Arguments)
 {
+	FImportAssetResult ImportAssetResult;
+	ImportAssetResult.bIsFactorySkipAsset = bSkipImport;
+
 	if (!Arguments.AssetNode || !Arguments.AssetNode->GetObjectClass()->IsChildOf(GetFactoryClass()))
 	{
-		return nullptr;
+		return ImportAssetResult;
 	}
 
 	const UInterchangeMaterialFunctionFactoryNode* MaterialFactoryNode = Cast<UInterchangeMaterialFunctionFactoryNode>(Arguments.AssetNode);
 	if (!ensure(MaterialFactoryNode))
 	{
-		return nullptr;
+		return ImportAssetResult;
 	}
 
 	const UClass* MaterialClass = MaterialFactoryNode->GetObjectClass();
@@ -1429,6 +1481,13 @@ UObject* UInterchangeMaterialFunctionFactory::ImportAssetObject_Async(const FImp
 	// create an asset if it doesn't exist
 	UObject* ExistingAsset = StaticFindObject(nullptr, Arguments.Parent, *Arguments.AssetName);
 
+	//Do not override an asset we skip
+	if (bSkipImport)
+	{
+		ImportAssetResult.ImportedObject = ExistingAsset;
+		return ImportAssetResult;
+	}
+
 	UObject* MaterialObject = nullptr;
 	// create a new material or overwrite existing asset, if possible
 	if (!ExistingAsset)
@@ -1442,7 +1501,7 @@ UObject* UInterchangeMaterialFunctionFactory::ImportAssetObject_Async(const FImp
 		else
 		{
 			UE_LOG(LogInterchangeImport, Error, TEXT("Could not create Material asset [%s] outside of the game thread"), *Arguments.AssetName);
-			return nullptr;
+			return ImportAssetResult;
 		}
 	}
 	else if (ExistingAsset->GetClass()->IsChildOf(MaterialClass))
@@ -1454,7 +1513,7 @@ UObject* UInterchangeMaterialFunctionFactory::ImportAssetObject_Async(const FImp
 	if (!MaterialObject)
 	{
 		UE_LOG(LogInterchangeImport, Warning, TEXT("Could not create Material asset %s"), *Arguments.AssetName);
-		return nullptr;
+		return ImportAssetResult;
 	}
 
 	if (MaterialObject)
@@ -1479,7 +1538,8 @@ UObject* UInterchangeMaterialFunctionFactory::ImportAssetObject_Async(const FImp
 		//The interchange completion task (call in the GameThread after the factories pass), will call PostEditChange which will trig another asynchronous system that will build all material in parallel
 	}
 
-	return MaterialObject;
+	ImportAssetResult.ImportedObject = MaterialObject;
+	return ImportAssetResult;
 }
 
 /* This function is call in the completion task on the main thread, use it to call main thread post creation step for your assets*/
@@ -1487,6 +1547,11 @@ void UInterchangeMaterialFunctionFactory::SetupObject_GameThread(const FSetupObj
 {
 	Super::SetupObject_GameThread(Arguments);
 
+	if (bSkipImport)
+	{
+		return;
+	}
+
 #if WITH_EDITORONLY_DATA
 	if (ensure(Arguments.ImportedObject && Arguments.SourceData))
 	{
diff --git a/Engine/Plugins/Interchange/Runtime/Source/Import/Private/Mesh/InterchangePhysicsAssetFactory.cpp b/Engine/Plugins/Interchange/Runtime/Source/Import/Private/Mesh/InterchangePhysicsAssetFactory.cpp
--- a/Engine/Plugins/Interchange/Runtime/Source/Import/Private/Mesh/InterchangePhysicsAssetFactory.cpp
+++ b/Engine/Plugins/Interchange/Runtime/Source/Import/Private/Mesh/InterchangePhysicsAssetFactory.cpp
@@ -16,20 +16,21 @@ UClass* UInterchangePhysicsAssetFactory::GetFactoryClass() const
 	return UPhysicsAsset::StaticClass();
 }
 
-UObject* UInterchangePhysicsAssetFactory::BeginImportAssetObject_GameThread(const FImportAssetObjectParams& Arguments)
+UInterchangeFactoryBase::FImportAssetResult UInterchangePhysicsAssetFactory::BeginImportAsset_GameThread(const FImportAssetObjectParams& Arguments)
 {
+	FImportAssetResult ImportAssetResult;
 	UObject* PhysicsAsset = nullptr;
 
 #if WITH_EDITORONLY_DATA
 	if (!Arguments.AssetNode || !Arguments.AssetNode->GetObjectClass()->IsChildOf(GetFactoryClass()))
 	{
-		return nullptr;
+		return ImportAssetResult;
 	}
 
 	UInterchangePhysicsAssetFactoryNode* PhysicsAssetNode = Cast<UInterchangePhysicsAssetFactoryNode>(Arguments.AssetNode);
 	if (PhysicsAssetNode == nullptr)
 	{
-		return nullptr;
+		return ImportAssetResult;
 	}
 
 	// create an asset if it doesn't exist
@@ -49,33 +50,36 @@ UObject* UInterchangePhysicsAssetFactory::BeginImportAssetObject_GameThread(cons
 	if (!PhysicsAsset)
 	{
 		UE_LOG(LogInterchangeImport, Warning, TEXT("Could not create PhysicsAsset asset %s"), *Arguments.AssetName);
-		return nullptr;
+		return ImportAssetResult;
 	}
 	PhysicsAssetNode->SetCustomReferenceObject(FSoftObjectPath(PhysicsAsset));
 
 	PhysicsAsset->PreEditChange(nullptr);
 #endif //WITH_EDITORONLY_DATA
-	return PhysicsAsset;
+
+	ImportAssetResult.ImportedObject = PhysicsAsset;
+	return ImportAssetResult;
 }
 
-UObject* UInterchangePhysicsAssetFactory::ImportAssetObject_Async(const FImportAssetObjectParams& Arguments)
+UInterchangeFactoryBase::FImportAssetResult UInterchangePhysicsAssetFactory::ImportAsset_Async(const FImportAssetObjectParams& Arguments)
 {
+	FImportAssetResult ImportAssetResult;
 #if !WITH_EDITORONLY_DATA
 
 	UE_LOG(LogInterchangeImport, Error, TEXT("Cannot import PhysicsAsset asset in runtime, this is an editor only feature."));
-	return nullptr;
+	return ImportAssetResult;
 
 #else
 
 	if (!Arguments.AssetNode || !Arguments.AssetNode->GetObjectClass()->IsChildOf(GetFactoryClass()))
 	{
-		return nullptr;
+		return ImportAssetResult;
 	}
 
 	UInterchangePhysicsAssetFactoryNode* PhysicsAssetNode = Cast<UInterchangePhysicsAssetFactoryNode>(Arguments.AssetNode);
 	if (PhysicsAssetNode == nullptr)
 	{
-		return nullptr;
+		return ImportAssetResult;
 	}
 
 	const UClass* PhysicsAssetClass = PhysicsAssetNode->GetObjectClass();
@@ -97,7 +101,7 @@ UObject* UInterchangePhysicsAssetFactory::ImportAssetObject_Async(const FImportA
 		else
 		{
 			UE_LOG(LogInterchangeImport, Error, TEXT("Could not create Physics asset [%s] outside of the game thread"), *Arguments.AssetName);
-			return nullptr;
+			return ImportAssetResult;
 		}
 	}
 	else if(ExistingAsset->GetClass()->IsChildOf(PhysicsAssetClass))
@@ -109,7 +113,7 @@ UObject* UInterchangePhysicsAssetFactory::ImportAssetObject_Async(const FImportA
 	if (!PhysicsAssetObject)
 	{
 		UE_LOG(LogInterchangeImport, Warning, TEXT("Could not create PhysicsAsset asset %s"), *Arguments.AssetName);
-		return nullptr;
+		return ImportAssetResult;
 	}
 
 	if (PhysicsAssetObject)
@@ -122,7 +126,7 @@ UObject* UInterchangePhysicsAssetFactory::ImportAssetObject_Async(const FImportA
 			if (!ensure(PhysicsAsset))
 			{
 				UE_LOG(LogInterchangeImport, Warning, TEXT("Could not create PhysicsAsset asset %s"), *Arguments.AssetName);
-				return nullptr;
+				return ImportAssetResult;
 			}
 			
 			PhysicsAssetNode->SetCustomReferenceObject(FSoftObjectPath(PhysicsAsset));
@@ -139,7 +143,8 @@ UObject* UInterchangePhysicsAssetFactory::ImportAssetObject_Async(const FImportA
 		PhysicsAssetObject->RemoveFromRoot();
 		PhysicsAssetObject->MarkAsGarbage();
 	}
-	return PhysicsAssetObject;
+	ImportAssetResult.ImportedObject = PhysicsAssetObject;
+	return ImportAssetResult;
 #endif
 }
 
diff --git a/Engine/Plugins/Interchange/Runtime/Source/Import/Private/Mesh/InterchangeSkeletalMeshFactory.cpp b/Engine/Plugins/Interchange/Runtime/Source/Import/Private/Mesh/InterchangeSkeletalMeshFactory.cpp
--- a/Engine/Plugins/Interchange/Runtime/Source/Import/Private/Mesh/InterchangeSkeletalMeshFactory.cpp
+++ b/Engine/Plugins/Interchange/Runtime/Source/Import/Private/Mesh/InterchangeSkeletalMeshFactory.cpp
@@ -885,25 +885,28 @@ UClass* UInterchangeSkeletalMeshFactory::GetFactoryClass() const
 	return USkeletalMesh::StaticClass();
 }
 
-UObject* UInterchangeSkeletalMeshFactory::BeginImportAssetObject_GameThread(const FImportAssetObjectParams& Arguments)
+UInterchangeFactoryBase::FImportAssetResult UInterchangeSkeletalMeshFactory::BeginImportAsset_GameThread(const FImportAssetObjectParams& Arguments)
 {
 	TRACE_CPUPROFILER_EVENT_SCOPE("UInterchangeSkeletalMeshFactory::BeginImportAssetObject_GameThread")
+
+	FImportAssetResult ImportAssetResult;
+
 #if !WITH_EDITOR || !WITH_EDITORONLY_DATA
 
 	UE_LOG(LogInterchangeImport, Error, TEXT("Cannot import skeletalMesh asset in runtime, this is an editor only feature."));
-	return nullptr;
+	return ImportAssetResult;
 
 #else
 	USkeletalMesh* SkeletalMesh = nullptr;
 	if (!Arguments.AssetNode || !Arguments.AssetNode->GetObjectClass()->IsChildOf(GetFactoryClass()))
 	{
-		return nullptr;
+		return ImportAssetResult;
 	}
 
 	const UInterchangeSkeletalMeshFactoryNode* SkeletalMeshFactoryNode = Cast<UInterchangeSkeletalMeshFactoryNode>(Arguments.AssetNode);
 	if (SkeletalMeshFactoryNode == nullptr)
 	{
-		return nullptr;
+		return ImportAssetResult;
 	}
 
 	// create an asset if it doesn't exist
@@ -926,7 +929,7 @@ UObject* UInterchangeSkeletalMeshFactory::BeginImportAssetObject_GameThread(cons
 		{
 			UE_LOG(LogInterchangeImport, Warning, TEXT("Could not create SkeletalMesh asset %s"), *Arguments.AssetName);
 		}
-		return nullptr;
+		return ImportAssetResult;
 	}
 	
 	SkeletalMesh->PreEditChange(nullptr);
@@ -940,36 +943,40 @@ UObject* UInterchangeSkeletalMeshFactory::BeginImportAssetObject_GameThread(cons
 		SkeletalMesh->LockPropertiesUntil(SkeletalMeshLockPropertiesEvent);
 	}
 
-	return SkeletalMesh;
+	ImportAssetResult.ImportedObject = SkeletalMesh;
+	return ImportAssetResult;
+
 #endif //else !WITH_EDITOR || !WITH_EDITORONLY_DATA
 }
 
-UObject* UInterchangeSkeletalMeshFactory::ImportAssetObject_Async(const FImportAssetObjectParams& Arguments)
+UInterchangeFactoryBase::FImportAssetResult UInterchangeSkeletalMeshFactory::ImportAsset_Async(const FImportAssetObjectParams& Arguments)
 {
 	TRACE_CPUPROFILER_EVENT_SCOPE("UInterchangeSkeletalMeshFactory::CreateAsset")
 
+	FImportAssetResult ImportAssetResult;
+
 #if !WITH_EDITOR || !WITH_EDITORONLY_DATA
 
 	UE_LOG(LogInterchangeImport, Error, TEXT("Cannot import skeletalMesh asset in runtime, this is an editor only feature."));
-	return nullptr;
+	return ImportAssetResult;
 
 #else
 	if (!Arguments.AssetNode || !Arguments.AssetNode->GetObjectClass()->IsChildOf(GetFactoryClass()))
 	{
-		return nullptr;
+		return ImportAssetResult;
 	}
 
 	UInterchangeSkeletalMeshFactoryNode* SkeletalMeshFactoryNode = Cast<UInterchangeSkeletalMeshFactoryNode>(Arguments.AssetNode);
 	if (SkeletalMeshFactoryNode == nullptr)
 	{
-		return nullptr;
+		return ImportAssetResult;
 	}
 
 	const IInterchangeMeshPayloadInterface* MeshTranslatorPayloadInterface = Cast<IInterchangeMeshPayloadInterface>(Arguments.Translator);
 	if (!MeshTranslatorPayloadInterface)
 	{
 		UE_LOG(LogInterchangeImport, Error, TEXT("Cannot import skeletalMesh, the translator do not implement the IInterchangeSkeletalMeshPayloadInterface."));
-		return nullptr;
+		return ImportAssetResult;
 	}
 
 	const UClass* SkeletalMeshClass = SkeletalMeshFactoryNode->GetObjectClass();
@@ -991,7 +998,7 @@ UObject* UInterchangeSkeletalMeshFactory::ImportAssetObject_Async(const FImportA
 		else
 		{
 			UE_LOG(LogInterchangeImport, Error, TEXT("Could not create SkeletalMesh asset [%s] outside of the game thread"), *Arguments.AssetName);
-			return nullptr;
+			return ImportAssetResult;
 		}
 	}
 	else if(ExistingAsset->GetClass()->IsChildOf(SkeletalMeshClass))
@@ -1006,7 +1013,7 @@ UObject* UInterchangeSkeletalMeshFactory::ImportAssetObject_Async(const FImportA
 		{
 			UE_LOG(LogInterchangeImport, Error, TEXT("Could not create SkeletalMesh asset %s"), *Arguments.AssetName);
 		}
-		return nullptr;
+		return ImportAssetResult;
 	}
 
 	USkeletalMesh* SkeletalMesh = Cast<USkeletalMesh>(SkeletalMeshObject);
@@ -1020,7 +1027,7 @@ UObject* UInterchangeSkeletalMeshFactory::ImportAssetObject_Async(const FImportA
 		{
 			UE_LOG(LogInterchangeImport, Error, TEXT("Could not find reimported skeletalMesh asset %s"), *Arguments.AssetName);
 		}
-		return nullptr;
+		return ImportAssetResult;
 	}
 
 	//Make sure we can modify the skeletalmesh properties
@@ -1467,32 +1474,34 @@ UObject* UInterchangeSkeletalMeshFactory::ImportAssetObject_Async(const FImportA
 	//Getting the file Hash will cache it into the source data
 	Arguments.SourceData->GetFileContentHash();
 
-	return SkeletalMeshObject;
+	ImportAssetResult.ImportedObject = SkeletalMeshObject;
+	return ImportAssetResult;
 
 #endif //else !WITH_EDITOR || !WITH_EDITORONLY_DATA
 }
 
-UObject* UInterchangeSkeletalMeshFactory::EndImportAssetObject_GameThread(const FImportAssetObjectParams& Arguments)
+UInterchangeFactoryBase::FImportAssetResult UInterchangeSkeletalMeshFactory::EndImportAsset_GameThread(const FImportAssetObjectParams& Arguments)
 {
 	TRACE_CPUPROFILER_EVENT_SCOPE("UInterchangeSkeletalMeshFactory::EndImportAssetObject_GameThread")
 	
 	check(IsInGameThread());
+	FImportAssetResult ImportAssetResult;
 
 #if !WITH_EDITOR || !WITH_EDITORONLY_DATA
 
 	UE_LOG(LogInterchangeImport, Error, TEXT("Cannot import skeletalMesh asset in runtime, this is an editor only feature."));
-	return nullptr;
+	return ImportAssetResult;
 
 #else
 	if (!Arguments.AssetNode || !Arguments.AssetNode->GetObjectClass()->IsChildOf(GetFactoryClass()))
 	{
-		return nullptr;
+		return ImportAssetResult;
 	}
 
 	UInterchangeSkeletalMeshFactoryNode* SkeletalMeshFactoryNode = Cast<UInterchangeSkeletalMeshFactoryNode>(Arguments.AssetNode);
 	if (SkeletalMeshFactoryNode == nullptr)
 	{
-		return nullptr;
+		return ImportAssetResult;
 	}
 
 	const UClass* SkeletalMeshClass = SkeletalMeshFactoryNode->GetObjectClass();
@@ -1510,7 +1519,7 @@ UObject* UInterchangeSkeletalMeshFactory::EndImportAssetObject_GameThread(const
 		{
 			UE_LOG(LogInterchangeImport, Error, TEXT("Could not find reimported skeletalMesh asset %s"), *Arguments.AssetName);
 		}
-		return nullptr;
+		return ImportAssetResult;
 	}
 
 	//Make sure we can modify the skeletalmesh properties
@@ -1678,7 +1687,9 @@ UObject* UInterchangeSkeletalMeshFactory::EndImportAssetObject_GameThread(const
 		UE::Interchange::FFactoryCommon::ApplyReimportStrategyToAsset(SkeletalMesh, PreviousNode, CurrentNode, SkeletalMeshFactoryNode);
 	}
 
-	return SkeletalMesh;
+	ImportAssetResult.ImportedObject = SkeletalMesh;
+	return ImportAssetResult;
+
 #endif //WITH_EDITOR
 }
 
diff --git a/Engine/Plugins/Interchange/Runtime/Source/Import/Private/Mesh/InterchangeSkeletonFactory.cpp b/Engine/Plugins/Interchange/Runtime/Source/Import/Private/Mesh/InterchangeSkeletonFactory.cpp
--- a/Engine/Plugins/Interchange/Runtime/Source/Import/Private/Mesh/InterchangeSkeletonFactory.cpp
+++ b/Engine/Plugins/Interchange/Runtime/Source/Import/Private/Mesh/InterchangeSkeletonFactory.cpp
@@ -17,146 +17,76 @@ UClass* UInterchangeSkeletonFactory::GetFactoryClass() const
 	return USkeleton::StaticClass();
 }
 
-UObject* UInterchangeSkeletonFactory::BeginImportAssetObject_GameThread(const FImportAssetObjectParams& Arguments)
+UInterchangeFactoryBase::FImportAssetResult UInterchangeSkeletonFactory::BeginImportAsset_GameThread(const FImportAssetObjectParams& Arguments)
 {
-	UObject* Skeleton = nullptr;
+	FImportAssetResult ImportAssetResult;
+
+	USkeleton* Skeleton = nullptr;
 
 #if WITH_EDITORONLY_DATA
 	if (!Arguments.AssetNode || !Arguments.AssetNode->GetObjectClass()->IsChildOf(GetFactoryClass()))
 	{
-		return nullptr;
+		return ImportAssetResult;
 	}
 
 	UInterchangeSkeletonFactoryNode* SkeletonNode = Cast<UInterchangeSkeletonFactoryNode>(Arguments.AssetNode);
 	if (SkeletonNode == nullptr)
 	{
-		return nullptr;
+		return ImportAssetResult;
 	}
 
+	bool bIsReImport = Arguments.ReimportObject != nullptr;
+
 	// create an asset if it doesn't exist
 	UObject* ExistingAsset = StaticFindObject(nullptr, Arguments.Parent, *Arguments.AssetName);
 
 	// create a new skeleton or overwrite existing asset, if possible
 	if (!ExistingAsset)
 	{
-		Skeleton = NewObject<UObject>(Arguments.Parent, USkeleton::StaticClass(), *Arguments.AssetName, RF_Public | RF_Standalone);
+		Skeleton = NewObject<USkeleton>(Arguments.Parent, USkeleton::StaticClass(), *Arguments.AssetName, RF_Public | RF_Standalone);
 	}
 	else if (ExistingAsset->GetClass()->IsChildOf(USkeleton::StaticClass()))
 	{
 		//This is a reimport, we are just re-updating the source data
-		Skeleton = ExistingAsset;
+		Skeleton = Cast<USkeleton>(ExistingAsset);
 	}
 
 	if (!Skeleton)
 	{
 		UE_LOG(LogInterchangeImport, Warning, TEXT("Could not create Skeleton asset %s"), *Arguments.AssetName);
-		return nullptr;
+		return ImportAssetResult;
 	}
 	SkeletonNode->SetCustomReferenceObject(FSoftObjectPath(Skeleton));
 
 	Skeleton->PreEditChange(nullptr);
-#endif //WITH_EDITORONLY_DATA
-	return Skeleton;
-}
-
-UObject* UInterchangeSkeletonFactory::ImportAssetObject_Async(const FImportAssetObjectParams& Arguments)
-{
-#if !WITH_EDITORONLY_DATA
-
-	UE_LOG(LogInterchangeImport, Error, TEXT("Cannot import Skeleton asset in runtime, this is an editor only feature."));
-	return nullptr;
-
-#else
-
-	if (!Arguments.AssetNode || !Arguments.AssetNode->GetObjectClass()->IsChildOf(GetFactoryClass()))
-	{
-		return nullptr;
-	}
 
-	UInterchangeSkeletonFactoryNode* SkeletonNode = Cast<UInterchangeSkeletonFactoryNode>(Arguments.AssetNode);
-	if (SkeletonNode == nullptr)
+	FString RootJointUid;
+	if (!SkeletonNode->GetCustomRootJointUid(RootJointUid))
 	{
-		return nullptr;
+		UE_LOG(LogInterchangeImport, Warning, TEXT("Could not create Skeleton asset %s, because there is no valid root joint node"), *Arguments.AssetName);
+		return ImportAssetResult;
 	}
-
-	const UClass* SkeletonClass = SkeletonNode->GetObjectClass();
-	check(SkeletonClass && SkeletonClass->IsChildOf(GetFactoryClass()));
-
-	// create an asset if it doesn't exist
-	UObject* ExistingAsset = StaticFindObject(nullptr, Arguments.Parent, *Arguments.AssetName);
-	
-	bool bIsReImport = Arguments.ReimportObject != nullptr;
-	
-	UObject* SkeletonObject = nullptr;
-	// create a new skeleton or overwrite existing asset, if possible
-	if (!ExistingAsset)
+	if (bIsReImport)
 	{
-		//NewObject is not thread safe, the asset registry directory watcher tick on the main thread can trig before we finish initializing the UObject and will crash
-		//The UObject should have been create by calling CreateEmptyAsset on the main thread.
-		if (!IsInGameThread())
-		{
-			UE_LOG(LogInterchangeImport, Error, TEXT("Could not create Skeleton asset [%s] outside of the game thread"), *Arguments.AssetName);
-			return nullptr;
-		}
-		if (bIsReImport)
+		if (const UInterchangeBaseNode* RootJointNode = Arguments.NodeContainer->GetNode(RootJointUid))
 		{
-			UE_LOG(LogInterchangeImport, Error, TEXT("Could not create Skeleton asset [%s] when re-importing"), *Arguments.AssetName);
-			return nullptr;
+			FName RootJointName = FName(*RootJointNode->GetDisplayLabel());
+			if (RootJointName != Skeleton->GetReferenceSkeleton().GetBoneName(0))
+			{
+				return ImportAssetResult;
+			}
 		}
-		SkeletonObject = NewObject<UObject>(Arguments.Parent, SkeletonClass, *Arguments.AssetName, RF_Public | RF_Standalone);
-	}
-	else if(ExistingAsset->GetClass()->IsChildOf(SkeletonClass))
-	{
-		//This is a reimport, we are just re-updating the source data
-		SkeletonObject = ExistingAsset;
 	}
+	//The joint will be added by the skeletalmesh factory since we need a valid skeletalmesh to add joint to a skeleton
+	SkeletonNode->SetCustomReferenceObject(FSoftObjectPath(Skeleton));
 
-	if (!SkeletonObject)
-	{
-		UE_LOG(LogInterchangeImport, Warning, TEXT("Could not create Skeleton asset %s"), *Arguments.AssetName);
-		return nullptr;
-	}
+	//Getting the file Hash will cache it into the source data
+	Arguments.SourceData->GetFileContentHash();
 
-	if (SkeletonObject)
-	{
-		USkeleton* Skeleton = Cast<USkeleton>(SkeletonObject);
-		if (!ensure(Skeleton))
-		{
-			UE_LOG(LogInterchangeImport, Warning, TEXT("Could not create Skeleton asset %s"), *Arguments.AssetName);
-			return nullptr;
-		}
-		FString RootJointUid;
-		if (!SkeletonNode->GetCustomRootJointUid(RootJointUid))
-		{
-			UE_LOG(LogInterchangeImport, Warning, TEXT("Could not create Skeleton asset %s, because there is no valid root joint node"), *Arguments.AssetName);
-			return nullptr;
-		}
-		if (bIsReImport)
-		{
-			if (const UInterchangeBaseNode* RootJointNode = Arguments.NodeContainer->GetNode(RootJointUid))
-			{
-				FName RootJointName = FName(*RootJointNode->GetDisplayLabel());
-				if (RootJointName != Skeleton->GetReferenceSkeleton().GetBoneName(0))
-				{
-					return nullptr;
-				}
-			}
-		}
-		//The joint will be added by the skeletalmesh factory since we need a valid skeletalmesh to add joint to a skeleton
-		SkeletonNode->SetCustomReferenceObject(FSoftObjectPath(Skeleton));
-		
-		//Getting the file Hash will cache it into the source data
-		Arguments.SourceData->GetFileContentHash();
+	//The interchange completion task (call in the GameThread after the factories pass), will call PostEditChange which will trig another asynchronous system that will build all skeleton in parallel
 
-		//The interchange completion task (call in the GameThread after the factories pass), will call PostEditChange which will trig another asynchronous system that will build all skeleton in parallel
-	}
-	else
-	{
-		//The skeleton is not a USkeleton
-		SkeletonObject->RemoveFromRoot();
-		SkeletonObject->MarkAsGarbage();
-	}
-	return SkeletonObject;
-#endif
-}
+#endif //WITH_EDITORONLY_DATA
+	ImportAssetResult.ImportedObject = Skeleton;
 
+	return ImportAssetResult;
+}
diff --git a/Engine/Plugins/Interchange/Runtime/Source/Import/Private/Mesh/InterchangeStaticMeshFactory.cpp b/Engine/Plugins/Interchange/Runtime/Source/Import/Private/Mesh/InterchangeStaticMeshFactory.cpp
--- a/Engine/Plugins/Interchange/Runtime/Source/Import/Private/Mesh/InterchangeStaticMeshFactory.cpp
+++ b/Engine/Plugins/Interchange/Runtime/Source/Import/Private/Mesh/InterchangeStaticMeshFactory.cpp
@@ -144,18 +144,19 @@ namespace UE::Interchange::Private::StaticMesh
 	}
 } //ns UE::Interchange::Private::StaticMesh
 
-UObject* UInterchangeStaticMeshFactory::BeginImportAssetObject_GameThread(const FImportAssetObjectParams& Arguments)
+UInterchangeFactoryBase::FImportAssetResult UInterchangeStaticMeshFactory::BeginImportAsset_GameThread(const FImportAssetObjectParams& Arguments)
 {
+	FImportAssetResult ImportAssetResult;
 	UStaticMesh* StaticMesh = nullptr;
 	if (!Arguments.AssetNode || !Arguments.AssetNode->GetObjectClass()->IsChildOf(GetFactoryClass()))
 	{
-		return nullptr;
+		return ImportAssetResult;
 	}
 
 	const UInterchangeStaticMeshFactoryNode* StaticMeshFactoryNode = Cast<UInterchangeStaticMeshFactoryNode>(Arguments.AssetNode);
 	if (StaticMeshFactoryNode == nullptr)
 	{
-		return nullptr;
+		return ImportAssetResult;
 	}
 
 	// create an asset if it doesn't exist
@@ -190,7 +191,7 @@ UObject* UInterchangeStaticMeshFactory::BeginImportAssetObject_GameThread(const
 		{
 			UE_LOG(LogInterchangeImport, Warning, TEXT("Could not create StaticMesh asset %s"), *Arguments.AssetName);
 		}
-		return nullptr;
+		return ImportAssetResult;
 	}
 
 	// create the BodySetup on the game thread
@@ -200,20 +201,22 @@ UObject* UInterchangeStaticMeshFactory::BeginImportAssetObject_GameThread(const
 	StaticMesh->PreEditChange(nullptr);	
 #endif // WITH_EDITOR
 
-	return StaticMesh;
+	ImportAssetResult.ImportedObject = StaticMesh;
+	return ImportAssetResult;
 }
 
-UObject* UInterchangeStaticMeshFactory::ImportAssetObject_Async(const FImportAssetObjectParams& Arguments)
+UInterchangeFactoryBase::FImportAssetResult UInterchangeStaticMeshFactory::ImportAsset_Async(const FImportAssetObjectParams& Arguments)
 {
+	FImportAssetResult ImportAssetResult;
 	if (!Arguments.AssetNode || !Arguments.AssetNode->GetObjectClass()->IsChildOf(GetFactoryClass()))
 	{
-		return nullptr;
+		return ImportAssetResult;
 	}
 
 	UInterchangeStaticMeshFactoryNode* StaticMeshFactoryNode = Cast<UInterchangeStaticMeshFactoryNode>(Arguments.AssetNode);
 	if (StaticMeshFactoryNode == nullptr)
 	{
-		return nullptr;
+		return ImportAssetResult;
 	}
 
 	const UClass* StaticMeshClass = StaticMeshFactoryNode->GetObjectClass();
@@ -237,7 +240,7 @@ UObject* UInterchangeStaticMeshFactory::ImportAssetObject_Async(const FImportAss
 		else
 		{
 			UE_LOG(LogInterchangeImport, Error, TEXT("Could not create StaticMesh asset [%s] outside of the game thread"), *Arguments.AssetName);
-			return nullptr;
+			return ImportAssetResult;
 		}
 	}
 	else if(ExistingAsset->GetClass()->IsChildOf(StaticMeshClass))
@@ -252,14 +255,14 @@ UObject* UInterchangeStaticMeshFactory::ImportAssetObject_Async(const FImportAss
 		{
 			UE_LOG(LogInterchangeImport, Error, TEXT("Could not create StaticMesh asset %s"), *Arguments.AssetName);
 		}
-		return nullptr;
+		return ImportAssetResult;
 	}
 
 	UStaticMesh* StaticMesh = Cast<UStaticMesh>(StaticMeshObject);
 	if (!ensure(StaticMesh))
 	{
 		UE_LOG(LogInterchangeImport, Error, TEXT("Could not create StaticMesh asset %s"), *Arguments.AssetName);
-		return nullptr;
+		return ImportAssetResult;
 	}
 
 	ensure(!StaticMesh->AreRenderingResourcesInitialized());
@@ -489,20 +492,22 @@ UObject* UInterchangeStaticMeshFactory::ImportAssetObject_Async(const FImportAss
 	// Getting the file Hash will cache it into the source data
 	Arguments.SourceData->GetFileContentHash();
 
-	return StaticMeshObject;
+	ImportAssetResult.ImportedObject = StaticMeshObject;
+	return ImportAssetResult;
 }
 
-UObject* UInterchangeStaticMeshFactory::EndImportAssetObject_GameThread(const FImportAssetObjectParams& Arguments)
+UInterchangeFactoryBase::FImportAssetResult UInterchangeStaticMeshFactory::EndImportAsset_GameThread(const FImportAssetObjectParams& Arguments)
 {
+	FImportAssetResult ImportAssetResult;
 	if (!Arguments.AssetNode || !Arguments.AssetNode->GetObjectClass()->IsChildOf(GetFactoryClass()))
 	{
-		return nullptr;
+		return ImportAssetResult;
 	}
 
 	UInterchangeStaticMeshFactoryNode* StaticMeshFactoryNode = Cast<UInterchangeStaticMeshFactoryNode>(Arguments.AssetNode);
 	if (StaticMeshFactoryNode == nullptr)
 	{
-		return nullptr;
+		return ImportAssetResult;
 	}
 
 	const UClass* StaticMeshClass = StaticMeshFactoryNode->GetObjectClass();
@@ -517,7 +522,7 @@ UObject* UInterchangeStaticMeshFactory::EndImportAssetObject_GameThread(const FI
 	if (!ensure(StaticMesh))
 	{
 		UE_LOG(LogInterchangeImport, Error, TEXT("Could not create StaticMesh asset %s"), *Arguments.AssetName);
-		return nullptr;
+		return ImportAssetResult;
 	}
 
 	for (int32 LodIndex = 0; LodIndex < ImportAssetObjectData.LodMeshDescriptions.Num(); ++LodIndex)
@@ -676,7 +681,9 @@ UObject* UInterchangeStaticMeshFactory::EndImportAssetObject_GameThread(const FI
 	{
 		UE::Interchange::Private::StaticMesh::ReorderMaterialSlotToBaseLod(StaticMesh);
 	}
-	return StaticMesh;
+
+	ImportAssetResult.ImportedObject = StaticMesh;
+	return ImportAssetResult;
 }
 
 void UInterchangeStaticMeshFactory::CommitMeshDescriptions(UStaticMesh& StaticMesh, TArray<FMeshDescription>&& LodMeshDescriptions)
diff --git a/Engine/Plugins/Interchange/Runtime/Source/Import/Private/Scene/InterchangeSceneVariantSetsFactory.cpp b/Engine/Plugins/Interchange/Runtime/Source/Import/Private/Scene/InterchangeSceneVariantSetsFactory.cpp
--- a/Engine/Plugins/Interchange/Runtime/Source/Import/Private/Scene/InterchangeSceneVariantSetsFactory.cpp
+++ b/Engine/Plugins/Interchange/Runtime/Source/Import/Private/Scene/InterchangeSceneVariantSetsFactory.cpp
@@ -294,12 +294,13 @@ UClass* UInterchangeSceneVariantSetsFactory::GetFactoryClass() const
 	return ULevelVariantSets::StaticClass();
 }
 
-UObject* UInterchangeSceneVariantSetsFactory::BeginImportAssetObject_GameThread(const FImportAssetObjectParams& Arguments)
+UInterchangeFactoryBase::FImportAssetResult UInterchangeSceneVariantSetsFactory::BeginImportAsset_GameThread(const FImportAssetObjectParams& Arguments)
 {
+	FImportAssetResult ImportAssetResult;
 #if !WITH_EDITOR || !WITH_EDITORONLY_DATA
 
 	UE_LOG(LogInterchangeImport, Error, TEXT("Cannot import levelsequence asset in runtime, this is an editor only feature."));
-	return nullptr;
+	return ImportAssetResult;
 #else
 	if (Arguments.ReimportObject)
 	{
@@ -309,25 +310,25 @@ UObject* UInterchangeSceneVariantSetsFactory::BeginImportAssetObject_GameThread(
 		Message->AssetType = ULevelVariantSets::StaticClass();
 		Message->Text = LOCTEXT("CreateEmptyAssetUnsupportedReimport", "Re-import of ULevelVariantSets not supported yet.");
 
-		return nullptr;
+		return ImportAssetResult;
 	}
 
 	if (!Cast<IInterchangeVariantSetPayloadInterface>(Arguments.Translator))
 	{
 		UE_LOG(LogInterchangeImport, Error, TEXT("Cannot import LevelVariantSets, the translator do not implement the IInterchangeVariantSetPayloadInterface interface."));
-		return nullptr;
+		return ImportAssetResult;
 	}
 
 	ULevelVariantSets* LevelVariantSets = nullptr;
 	if (!Arguments.AssetNode || !Arguments.AssetNode->GetObjectClass()->IsChildOf(GetFactoryClass()))
 	{
-		return nullptr;
+		return ImportAssetResult;
 	}
 
 	const UInterchangeSceneVariantSetsFactoryNode* FactoryNode = Cast<UInterchangeSceneVariantSetsFactoryNode>(Arguments.AssetNode);
 	if (FactoryNode == nullptr)
 	{
-		return nullptr;
+		return ImportAssetResult;
 	}
 
 	// create an asset if it doesn't exist
@@ -347,10 +348,11 @@ UObject* UInterchangeSceneVariantSetsFactory::BeginImportAssetObject_GameThread(
 	if (!LevelVariantSets)
 	{
 		UE_LOG(LogInterchangeImport, Warning, TEXT("Could not create LevelVariantSets asset %s"), *Arguments.AssetName);
-		return nullptr;
+		return ImportAssetResult;
 	}
 
-	return ImportObjectSourceData(Arguments);
+	ImportAssetResult.ImportedObject = ImportObjectSourceData(Arguments);
+	return ImportAssetResult;
 #endif //else !WITH_EDITOR || !WITH_EDITORONLY_DATA
 }
 
diff --git a/Engine/Plugins/Interchange/Runtime/Source/Import/Private/Texture/InterchangeTextureFactory.cpp b/Engine/Plugins/Interchange/Runtime/Source/Import/Private/Texture/InterchangeTextureFactory.cpp
--- a/Engine/Plugins/Interchange/Runtime/Source/Import/Private/Texture/InterchangeTextureFactory.cpp
+++ b/Engine/Plugins/Interchange/Runtime/Source/Import/Private/Texture/InterchangeTextureFactory.cpp
@@ -1459,39 +1459,41 @@ UClass* UInterchangeTextureFactory::GetFactoryClass() const
 	return UTexture::StaticClass();
 }
 
-UObject* UInterchangeTextureFactory::BeginImportAssetObject_GameThread(const FImportAssetObjectParams& Arguments)
+UInterchangeFactoryBase::FImportAssetResult UInterchangeTextureFactory::BeginImportAsset_GameThread(const FImportAssetObjectParams& Arguments)
 {
 	using namespace  UE::Interchange::Private::InterchangeTextureFactory;
-
+	FImportAssetResult ImportAssetResult;
 	UObject* Texture = nullptr;
 
-	auto CouldNotCreateTextureLog = [this, &Arguments](const FText& Info)
+	auto CouldNotCreateTextureLog = [this, &Arguments, &ImportAssetResult](const FText& Info)
 	{
 		UInterchangeResultError_Generic* Message = AddMessage<UInterchangeResultError_Generic>();
 		Message->SourceAssetName = Arguments.SourceData->GetFilename();
 		Message->DestinationAssetName = Arguments.AssetName;
 		Message->AssetType = GetFactoryClass();
 		Message->Text = FText::Format(LOCTEXT("TexFact_CouldNotCreateMat", "UInterchangeTextureFactory: Could not create texture asset %s. Reason: %s"), FText::FromString(Arguments.AssetName), Info);
+		bSkipImport = true;
+		ImportAssetResult.bIsFactorySkipAsset = true;
 	};
 
 	if (!Arguments.AssetNode)
 	{
 		CouldNotCreateTextureLog(LOCTEXT("TextureFactory_AssetNodeNull", "Asset node parameter is null."));
-		return nullptr;
+		return ImportAssetResult;
 	}
 
 	const UClass* TextureClass = Arguments.AssetNode->GetObjectClass();
 	if (!TextureClass || !TextureClass->IsChildOf(UTexture::StaticClass()))
 	{
 		CouldNotCreateTextureLog(LOCTEXT("TextureFactory_NodeClassMissmatch", "Asset node parameter class doesnt derive from UTexture."));
-		return nullptr;
+		return ImportAssetResult;
 	}
 
 	UClass* SupportedFactoryNodeClass = GetSupportedFactoryNodeClass(Arguments.AssetNode);
 	if (SupportedFactoryNodeClass == nullptr)
 	{
 		CouldNotCreateTextureLog(LOCTEXT("TextureFactory_NodeClassWrong", "Asset node parameter is not a UInterchangeTextureFactoryNode or UInterchangeTextureCubeFactoryNode."));
-		return nullptr;
+		return ImportAssetResult;
 	}
 
 
@@ -1501,17 +1503,19 @@ UObject* UInterchangeTextureFactory::BeginImportAssetObject_GameThread(const FIm
 		FText Info = FText::Format(LOCTEXT("TextureFactory_InvalidTextureTranslated", "Asset factory node (%s) do not reference a valid texture translated node.")
 			, FText::FromString(SupportedFactoryNodeClass->GetAuthoredName()));
 		CouldNotCreateTextureLog(Info);
-		return nullptr;
+		return ImportAssetResult;
 	}
 
 	if (!HasPayloadKey(TextureNodeVariant))
 	{
 		CouldNotCreateTextureLog(LOCTEXT("TextureFactory_InvalidPayloadKey", "Texture translated node doesnt have a payload key."));
-		return nullptr;
+		return ImportAssetResult;
 	}
 
-	// create an asset if it doesn't exist
-	UObject* ExistingAsset = StaticFindObject(nullptr, Arguments.Parent, *Arguments.AssetName);
+	const bool bIsReimport = Arguments.ReimportObject != nullptr;
+
+	// Find existing asset
+	UObject* ExistingAsset = bIsReimport ? Arguments.ReimportObject : StaticFindObject(nullptr, Arguments.Parent, *Arguments.AssetName);
 
 	// create a new texture or overwrite existing asset, if possible
 	if (!ExistingAsset)
@@ -1533,6 +1537,13 @@ UObject* UInterchangeTextureFactory::BeginImportAssetObject_GameThread(const FIm
 	{
 		//This is a reimport, we are just re-updating the source data
 		Texture = ExistingAsset;
+		//We allow override of existing Textures only if the translator is a pure texture translator or the user directly ask to re-import this object
+		if (!bIsReimport && Arguments.Translator->GetSupportedAssetTypes() != EInterchangeTranslatorAssetType::Textures)
+		{
+			//Do not override the material asset
+			ImportAssetResult.bIsFactorySkipAsset = true;
+			bSkipImport = true;
+		}
 	}
 	else
 	{
@@ -1547,24 +1558,27 @@ UObject* UInterchangeTextureFactory::BeginImportAssetObject_GameThread(const FIm
 		const FText FolderName = FText::FromString(FPaths::GetPath(Arguments.Parent->GetPathName()));
 		Message->Text = FText::Format(NSLOCTEXT("InterchangeTextureFactory", "ClassMismatch", "You cannot create a '{0}' asset named '{1}' in '{2}', as there is already a '{3}' asset with the same name in this folder."), TargetClassName, AssetName, FolderName, ExistingClassName);
 		
-		return nullptr;
+		return ImportAssetResult;
 	}
 
 	if (!Texture)
 	{
 		CouldNotCreateTextureLog(LOCTEXT("TextureFactory_TextureCreateFail", "Texture creation fail."));
-		return nullptr;
+		return ImportAssetResult;
 	}
 
-	return Texture;
+	ImportAssetResult.ImportedObject = Texture;
+	return ImportAssetResult;
 }
 
 // The payload fetching and the heavy operations are done here
-UObject* UInterchangeTextureFactory::ImportAssetObject_Async(const FImportAssetObjectParams& Arguments)
+UInterchangeFactoryBase::FImportAssetResult UInterchangeTextureFactory::ImportAsset_Async(const FImportAssetObjectParams& Arguments)
 {
-	TRACE_CPUPROFILER_EVENT_SCOPE(UInterchangeTextureFactory::CreateAsset);
+	TRACE_CPUPROFILER_EVENT_SCOPE(UInterchangeTextureFactory::ImportAsset_Async);
 
 	using namespace UE::Interchange::Private::InterchangeTextureFactory;
+	FImportAssetResult ImportAssetResult;
+	ImportAssetResult.bIsFactorySkipAsset = bSkipImport;
 
 	auto ImportTextureErrorLog = [this, &Arguments](const FText& Info)
 	{
@@ -1578,21 +1592,32 @@ UObject* UInterchangeTextureFactory::ImportAssetObject_Async(const FImportAssetO
 	if (!Arguments.AssetNode)
 	{
 		ImportTextureErrorLog(LOCTEXT("TextureFactory_Async_AssetNodeNull", "UInterchangeTextureFactory: Asset node parameter is null."));
-		return nullptr;
+		return ImportAssetResult;
 	}
 
 	const UClass* TextureClass = Arguments.AssetNode->GetObjectClass();
 	if (!TextureClass || !TextureClass->IsChildOf(UTexture::StaticClass()))
 	{
 		ImportTextureErrorLog(LOCTEXT("TextureFactory_Async_MissMatchClass", "UInterchangeTextureFactory: Asset node parameter class doesnt derive from UTexture."));
-		return nullptr;
+		return ImportAssetResult;
 	}
 
 	UClass* SupportedFactoryNodeClass = GetSupportedFactoryNodeClass(Arguments.AssetNode);
 	if (SupportedFactoryNodeClass == nullptr)
 	{
 		ImportTextureErrorLog(LOCTEXT("TextureFactory_Async_NodeWrongClass", "UInterchangeTextureFactory: Asset node parameter is not a child of UInterchangeTextureFactoryNode."));
-		return nullptr;
+		return ImportAssetResult;
+	}
+
+	// create an asset if it doesn't exist
+	// typically ExistingAsset should already exist, made by CreateEmptyAsset on the main thread
+	UObject* ExistingAsset = StaticFindObject(nullptr, Arguments.Parent, *Arguments.AssetName);
+
+	//Do not override an asset we skip
+	if (bSkipImport)
+	{
+		ImportAssetResult.ImportedObject = ExistingAsset;
+		return ImportAssetResult;
 	}
 
 	FTextureFactoryNodeVariant TextureFactoryNodeVariant = GetAsTextureFactoryNodeVariant(Arguments.AssetNode, SupportedFactoryNodeClass);
@@ -1602,14 +1627,14 @@ UObject* UInterchangeTextureFactory::ImportAssetObject_Async(const FImportAssetO
 		FText Info = FText::Format(LOCTEXT("TextureFactory_Async_InvalidTextureTranslated", "UInterchangeTextureFactory: Asset factory node (%s) do not reference a valid texture translated node.")
 			, FText::FromString(SupportedFactoryNodeClass->GetAuthoredName()));
 		ImportTextureErrorLog(Info);
-		return nullptr;
+		return ImportAssetResult;
 	}
 
 	const TOptional<FString>& PayLoadKey = GetPayloadKey(TextureNodeVariant);
 	if (!PayLoadKey.IsSet())
 	{
 		ImportTextureErrorLog(LOCTEXT("TextureFactory_Async_InvalidPayloadKey", "UInterchangeTextureFactory: Texture translated node (UInterchangeTexture2DNode) doesnt have a payload key."));
-		return nullptr;
+		return ImportAssetResult;
 	}
 
 	AlternateTexturePath.Reset();
@@ -1618,13 +1643,9 @@ UObject* UInterchangeTextureFactory::ImportAssetObject_Async(const FImportAssetO
 	if(TexturePayload.IsType<FEmptyVariantState>())
 	{
 		ImportTextureErrorLog(LOCTEXT("TextureFactory_Async_CannotRetrievePayload", "UInterchangeTextureFactory: Invalid translator couldn't retrive a payload."));
-		return nullptr;
+		return ImportAssetResult;
 	}
 
-	// create an asset if it doesn't exist
-	// typically ExistingAsset should already exist, made by CreateEmptyAsset on the main thread
-	UObject* ExistingAsset = StaticFindObject(nullptr, Arguments.Parent, *Arguments.AssetName);
-
 	UTexture* Texture = nullptr;
 	// create a new texture or overwrite existing asset, if possible
 	if (!ExistingAsset)
@@ -1638,7 +1659,7 @@ UObject* UInterchangeTextureFactory::ImportAssetObject_Async(const FImportAssetO
 		else
 		{
 			ImportTextureErrorLog(LOCTEXT("TextureFactory_Async_CannotCreateAsync", "UInterchangeTextureFactory: Could not create Texture asset outside of the game thread."));
-			return nullptr;
+			return ImportAssetResult;
 		}
 	}
 	else if(ExistingAsset->GetClass()->IsChildOf(TextureClass))
@@ -1651,7 +1672,7 @@ UObject* UInterchangeTextureFactory::ImportAssetObject_Async(const FImportAssetO
 	if (!Texture)
 	{
 		ImportTextureErrorLog(LOCTEXT("TextureFactory_Async_FailCreatingAsset", "UInterchangeTextureFactory: Could not create texture asset."));
-		return nullptr;
+		return ImportAssetResult;
 	}
 
 	// Check if the imported texture(s) has a valid resolution
@@ -1685,7 +1706,8 @@ UObject* UInterchangeTextureFactory::ImportAssetObject_Async(const FImportAssetO
 	if (!bCanSetup)
 	{
 		LogErrorInvalidPayload(*this, Texture->GetClass(), Arguments.SourceData->GetFilename(), Texture->GetName());
-		return Texture;
+		ImportAssetResult.ImportedObject = Texture;
+		return ImportAssetResult;
 	}
 
 	FGraphEventArray TasksToDo;
@@ -1704,8 +1726,8 @@ UObject* UInterchangeTextureFactory::ImportAssetObject_Async(const FImportAssetO
 	FTaskGraphInterface::Get().WaitUntilTasksComplete(TasksToDo, NamedThread);
 
 	//The interchange completion task (call in the GameThread after the factories pass), will call PostEditChange which will trig another asynchronous system that will build all texture in parallel
-
-	return Texture;
+	ImportAssetResult.ImportedObject = Texture;
+	return ImportAssetResult;
 }
 
 /* This function is call in the completion task on the main thread, use it to call main thread post creation step for your assets*/
@@ -1715,6 +1737,11 @@ void UInterchangeTextureFactory::SetupObject_GameThread(const FSetupObjectParams
 
 	check(IsInGameThread());
 
+	if (bSkipImport)
+	{
+		return;
+	}
+
 	UTexture* Texture = Cast<UTexture>(Arguments.ImportedObject);
 
 	// Finish the import on the game thread by doing the setup on the texture here
diff --git a/Engine/Plugins/Interchange/Runtime/Source/Import/Public/Animation/InterchangeAnimSequenceFactory.h b/Engine/Plugins/Interchange/Runtime/Source/Import/Public/Animation/InterchangeAnimSequenceFactory.h
--- a/Engine/Plugins/Interchange/Runtime/Source/Import/Public/Animation/InterchangeAnimSequenceFactory.h
+++ b/Engine/Plugins/Interchange/Runtime/Source/Import/Public/Animation/InterchangeAnimSequenceFactory.h
@@ -23,7 +23,7 @@ public:
 
 	virtual UClass* GetFactoryClass() const override;
 	virtual EInterchangeFactoryAssetType GetFactoryAssetType() override { return EInterchangeFactoryAssetType::Animations; }
-	virtual UObject* BeginImportAssetObject_GameThread(const FImportAssetObjectParams& Arguments) override;
+	virtual FImportAssetResult BeginImportAsset_GameThread(const FImportAssetObjectParams& Arguments) override;
 	virtual void SetupObject_GameThread(const FSetupObjectParams& Arguments) override;
 
 	virtual bool GetSourceFilenames(const UObject* Object, TArray<FString>& OutSourceFilenames) const override;
diff --git a/Engine/Plugins/Interchange/Runtime/Source/Import/Public/Animation/InterchangeAnimationTrackSetFactory.h b/Engine/Plugins/Interchange/Runtime/Source/Import/Public/Animation/InterchangeAnimationTrackSetFactory.h
--- a/Engine/Plugins/Interchange/Runtime/Source/Import/Public/Animation/InterchangeAnimationTrackSetFactory.h
+++ b/Engine/Plugins/Interchange/Runtime/Source/Import/Public/Animation/InterchangeAnimationTrackSetFactory.h
@@ -20,7 +20,7 @@ public:
 
 	virtual UClass* GetFactoryClass() const override;
 	virtual EInterchangeFactoryAssetType GetFactoryAssetType() override { return EInterchangeFactoryAssetType::Animations; }
-	virtual UObject* BeginImportAssetObject_GameThread(const FImportAssetObjectParams& Arguments) override;
+	virtual FImportAssetResult BeginImportAsset_GameThread(const FImportAssetObjectParams& Arguments) override;
 	virtual void SetupObject_GameThread(const FSetupObjectParams& Arguments) override;
 
 	// Interchange factory base interface end
diff --git a/Engine/Plugins/Interchange/Runtime/Source/Import/Public/Material/InterchangeMaterialFactory.h b/Engine/Plugins/Interchange/Runtime/Source/Import/Public/Material/InterchangeMaterialFactory.h
--- a/Engine/Plugins/Interchange/Runtime/Source/Import/Public/Material/InterchangeMaterialFactory.h
+++ b/Engine/Plugins/Interchange/Runtime/Source/Import/Public/Material/InterchangeMaterialFactory.h
@@ -29,8 +29,8 @@ public:
 
 	virtual UClass* GetFactoryClass() const override;
 	virtual EInterchangeFactoryAssetType GetFactoryAssetType() override { return EInterchangeFactoryAssetType::Materials; }
-	virtual UObject* BeginImportAssetObject_GameThread(const FImportAssetObjectParams& Arguments) override;
-	virtual UObject* ImportAssetObject_Async(const FImportAssetObjectParams& Arguments) override;
+	virtual FImportAssetResult BeginImportAsset_GameThread(const FImportAssetObjectParams& Arguments) override;
+	virtual FImportAssetResult ImportAsset_Async(const FImportAssetObjectParams& Arguments) override;
 	virtual void SetupObject_GameThread(const FSetupObjectParams& Arguments) override;
 	virtual bool GetSourceFilenames(const UObject* Object, TArray<FString>& OutSourceFilenames) const override;
 	virtual bool SetSourceFilename(const UObject* Object, const FString& SourceFilename, int32 SourceIndex) const override;
@@ -44,6 +44,9 @@ private:
 
 	void SetupMaterialInstance(UMaterialInstance& MaterialInstance, const UInterchangeBaseNodeContainer& NodeContainer, const UInterchangeMaterialInstanceFactoryNode& FactoryNode, bool bResetInstance);
 	void SetupReimportedMaterialInstance(UMaterialInstance& MaterialInstance, const UInterchangeBaseNodeContainer& NodeContainer, const UInterchangeMaterialInstanceFactoryNode& FactoryNode, const UInterchangeMaterialInstanceFactoryNode& PreviousFactoryNode);
+
+	//If we import without a pure material translator, we should not override an existing material and we must skip the import. See the implementation of BeginImportAssetObject_GameThread function.
+	bool bSkipImport = false;
 };
 
 UCLASS(BlueprintType)
@@ -57,8 +60,8 @@ public:
 
 	virtual UClass* GetFactoryClass() const override;
 	virtual EInterchangeFactoryAssetType GetFactoryAssetType() override { return EInterchangeFactoryAssetType::Materials; }
-	virtual UObject* BeginImportAssetObject_GameThread(const FImportAssetObjectParams& Arguments) override;
-	virtual UObject* ImportAssetObject_Async(const FImportAssetObjectParams& Arguments) override;
+	virtual FImportAssetResult BeginImportAsset_GameThread(const FImportAssetObjectParams& Arguments) override;
+	virtual FImportAssetResult ImportAsset_Async(const FImportAssetObjectParams& Arguments) override;
 	virtual void SetupObject_GameThread(const FSetupObjectParams& Arguments) override;
 
 	// Interchange factory base interface end
@@ -68,6 +71,7 @@ private:
 #if WITH_EDITOR
 	void SetupMaterial(class UMaterialFunction* Material, const FImportAssetObjectParams& Arguments, const class UInterchangeMaterialFunctionFactoryNode* MaterialFactoryNode);
 #endif
+	bool bSkipImport = false;
 };
 
 /**
diff --git a/Engine/Plugins/Interchange/Runtime/Source/Import/Public/Mesh/InterchangePhysicsAssetFactory.h b/Engine/Plugins/Interchange/Runtime/Source/Import/Public/Mesh/InterchangePhysicsAssetFactory.h
--- a/Engine/Plugins/Interchange/Runtime/Source/Import/Public/Mesh/InterchangePhysicsAssetFactory.h
+++ b/Engine/Plugins/Interchange/Runtime/Source/Import/Public/Mesh/InterchangePhysicsAssetFactory.h
@@ -20,8 +20,8 @@ public:
 
 	virtual UClass* GetFactoryClass() const override;
 	virtual EInterchangeFactoryAssetType GetFactoryAssetType() override { return EInterchangeFactoryAssetType::Physics; }
-	virtual UObject* BeginImportAssetObject_GameThread(const FImportAssetObjectParams& Arguments) override;
-	virtual UObject* ImportAssetObject_Async(const FImportAssetObjectParams& Arguments) override;
+	virtual FImportAssetResult BeginImportAsset_GameThread(const FImportAssetObjectParams& Arguments) override;
+	virtual FImportAssetResult ImportAsset_Async(const FImportAssetObjectParams& Arguments) override;
 
 	// Interchange factory base interface end
 	//////////////////////////////////////////////////////////////////////////
diff --git a/Engine/Plugins/Interchange/Runtime/Source/Import/Public/Mesh/InterchangeSkeletalMeshFactory.h b/Engine/Plugins/Interchange/Runtime/Source/Import/Public/Mesh/InterchangeSkeletalMeshFactory.h
--- a/Engine/Plugins/Interchange/Runtime/Source/Import/Public/Mesh/InterchangeSkeletalMeshFactory.h
+++ b/Engine/Plugins/Interchange/Runtime/Source/Import/Public/Mesh/InterchangeSkeletalMeshFactory.h
@@ -40,9 +40,9 @@ public:
 
 	virtual UClass* GetFactoryClass() const override;
 	virtual EInterchangeFactoryAssetType GetFactoryAssetType() override { return EInterchangeFactoryAssetType::Meshes; }
-	virtual UObject* BeginImportAssetObject_GameThread(const FImportAssetObjectParams& Arguments) override;
-	virtual UObject* ImportAssetObject_Async(const FImportAssetObjectParams& Arguments) override;
-	virtual UObject* EndImportAssetObject_GameThread(const FImportAssetObjectParams& Arguments) override;
+	virtual FImportAssetResult BeginImportAsset_GameThread(const FImportAssetObjectParams& Arguments) override;
+	virtual FImportAssetResult ImportAsset_Async(const FImportAssetObjectParams& Arguments) override;
+	virtual FImportAssetResult EndImportAsset_GameThread(const FImportAssetObjectParams& Arguments) override;
 	virtual void Cancel() override;
 	virtual void SetupObject_GameThread(const FSetupObjectParams& Arguments) override;
 	virtual bool GetSourceFilenames(const UObject* Object, TArray<FString>& OutSourceFilenames) const override;
diff --git a/Engine/Plugins/Interchange/Runtime/Source/Import/Public/Mesh/InterchangeSkeletonFactory.h b/Engine/Plugins/Interchange/Runtime/Source/Import/Public/Mesh/InterchangeSkeletonFactory.h
--- a/Engine/Plugins/Interchange/Runtime/Source/Import/Public/Mesh/InterchangeSkeletonFactory.h
+++ b/Engine/Plugins/Interchange/Runtime/Source/Import/Public/Mesh/InterchangeSkeletonFactory.h
@@ -20,8 +20,7 @@ public:
 
 	virtual UClass* GetFactoryClass() const override;
 	virtual EInterchangeFactoryAssetType GetFactoryAssetType() override { return EInterchangeFactoryAssetType::Meshes; }
-	virtual UObject* BeginImportAssetObject_GameThread(const FImportAssetObjectParams& Arguments) override;
-	virtual UObject* ImportAssetObject_Async(const FImportAssetObjectParams& Arguments) override;
+	virtual FImportAssetResult BeginImportAsset_GameThread(const FImportAssetObjectParams& Arguments) override;
 
 	// Interchange factory base interface end
 	//////////////////////////////////////////////////////////////////////////
diff --git a/Engine/Plugins/Interchange/Runtime/Source/Import/Public/Mesh/InterchangeStaticMeshFactory.h b/Engine/Plugins/Interchange/Runtime/Source/Import/Public/Mesh/InterchangeStaticMeshFactory.h
--- a/Engine/Plugins/Interchange/Runtime/Source/Import/Public/Mesh/InterchangeStaticMeshFactory.h
+++ b/Engine/Plugins/Interchange/Runtime/Source/Import/Public/Mesh/InterchangeStaticMeshFactory.h
@@ -30,9 +30,9 @@ public:
 
 	virtual UClass* GetFactoryClass() const override;
 	virtual EInterchangeFactoryAssetType GetFactoryAssetType() override { return EInterchangeFactoryAssetType::Meshes; }
-	virtual UObject* BeginImportAssetObject_GameThread(const FImportAssetObjectParams& Arguments) override;
-	virtual UObject* ImportAssetObject_Async(const FImportAssetObjectParams& Arguments) override;
-	virtual UObject* EndImportAssetObject_GameThread(const FImportAssetObjectParams& Arguments) override;
+	virtual FImportAssetResult BeginImportAsset_GameThread(const FImportAssetObjectParams& Arguments) override;
+	virtual FImportAssetResult ImportAsset_Async(const FImportAssetObjectParams& Arguments) override;
+	virtual FImportAssetResult EndImportAsset_GameThread(const FImportAssetObjectParams& Arguments) override;
 	virtual void SetupObject_GameThread(const FSetupObjectParams& Arguments) override;
 	virtual bool GetSourceFilenames(const UObject* Object, TArray<FString>& OutSourceFilenames) const override;
 	virtual bool SetSourceFilename(const UObject* Object, const FString& SourceFilename, int32 SourceIndex) const override;
diff --git a/Engine/Plugins/Interchange/Runtime/Source/Import/Public/Scene/InterchangeSceneVariantSetsFactory.h b/Engine/Plugins/Interchange/Runtime/Source/Import/Public/Scene/InterchangeSceneVariantSetsFactory.h
--- a/Engine/Plugins/Interchange/Runtime/Source/Import/Public/Scene/InterchangeSceneVariantSetsFactory.h
+++ b/Engine/Plugins/Interchange/Runtime/Source/Import/Public/Scene/InterchangeSceneVariantSetsFactory.h
@@ -23,7 +23,7 @@ public:
 	// Interchange factory base interface begin
 
 	virtual UClass* GetFactoryClass() const override;
-	virtual UObject* BeginImportAssetObject_GameThread(const FImportAssetObjectParams& Arguments) override;
+	virtual FImportAssetResult BeginImportAsset_GameThread(const FImportAssetObjectParams& Arguments) override;
 	
 	virtual void SetupObject_GameThread(const FSetupObjectParams& Arguments) override;
 
diff --git a/Engine/Plugins/Interchange/Runtime/Source/Import/Public/Texture/InterchangeTextureFactory.h b/Engine/Plugins/Interchange/Runtime/Source/Import/Public/Texture/InterchangeTextureFactory.h
--- a/Engine/Plugins/Interchange/Runtime/Source/Import/Public/Texture/InterchangeTextureFactory.h
+++ b/Engine/Plugins/Interchange/Runtime/Source/Import/Public/Texture/InterchangeTextureFactory.h
@@ -56,8 +56,8 @@ public:
 
 	virtual UClass* GetFactoryClass() const override;
 	virtual EInterchangeFactoryAssetType GetFactoryAssetType() override { return EInterchangeFactoryAssetType::Textures; }
-	virtual UObject* BeginImportAssetObject_GameThread(const FImportAssetObjectParams& Arguments) override;
-	virtual UObject* ImportAssetObject_Async(const FImportAssetObjectParams& Arguments) override;
+	virtual FImportAssetResult BeginImportAsset_GameThread(const FImportAssetObjectParams& Arguments) override;
+	virtual FImportAssetResult ImportAsset_Async(const FImportAssetObjectParams& Arguments) override;
 	virtual void SetupObject_GameThread(const FSetupObjectParams& Arguments) override;
 	virtual bool GetSourceFilenames(const UObject* Object, TArray<FString>& OutSourceFilenames) const override;
 	virtual bool SetSourceFilename(const UObject* Object, const FString& SourceFilename, int32 SourceIndex) const override;
@@ -78,6 +78,9 @@ private:
 	//  The data for the source files will be stored here during the import
 	TArray<FAssetImportInfo::FSourceFile> SourceFiles;
 #endif // WITH_EDITORONLY_DATA
+
+	//If we import without a pure texture translator, we should not override an existing texture and we must skip the import. See the implementation of BeginImportAssetObject_GameThread function.
+	bool bSkipImport = false;
 };
 
 
diff --git a/Engine/Source/Runtime/Interchange/Core/Public/InterchangeFactoryBase.h b/Engine/Source/Runtime/Interchange/Core/Public/InterchangeFactoryBase.h
--- a/Engine/Source/Runtime/Interchange/Core/Public/InterchangeFactoryBase.h
+++ b/Engine/Source/Runtime/Interchange/Core/Public/InterchangeFactoryBase.h
@@ -34,14 +34,14 @@ enum class EInterchangeFactoryAssetType : uint8
  * 
  * The first 3 steps use the interchange factory node to import/re-import the UObject
  * 
- * 1. BeginImportAssetObject_GameThread - Create the asset UPackage and the asset UObject. You can also import source data (retrieve payloads) and setup properties on the game thread.
- * 2. ImportAssetObject_Async - Import source data (retrieve payloads) and setup properties asynchronously on any thread.
- * 3. EndImportAssetObject_GameThread - Anything you need to do on the game thread to finalize the import source data and setup properties. (Conflict resolution that need UI)
+ * 1. BeginImportAsset_GameThread - Create the asset UObject. You can also import source data (retrieve payloads) and setup properties on the game thread.
+ * 2. ImportAsset_Async - Import source data (retrieve payloads) and setup properties asynchronously on any thread.
+ * 3. EndImportAsset_GameThread - Anything you need to do on the game thread to finalize the import source data and setup properties. (Conflict resolution that need UI)
  * 
- * The last two steps are helpful to change the imported/re-imported UObject before and after the asset is built.
+ * The last two steps are helpful to change the imported/re-imported UObject before and after the PostEditChange (render data build) is call on the asset.
  * 
  * 4. SetupObject_GameThread - Do any UObject setup required before the build (before PostEditChange), the UObject dependencies should exist and have all the source data and properties imported.
- * 5. FinalizeObject_GameThread - Do any final UObject setup after the build (after PostEditChange)
+ * 5. FinalizeObject_GameThread - Do any final UObject setup after the build (after PostEditChange). Note that the build of an asset can be asynchronous and not finish yet.
  * 
  * 
  * Scene factory implementation
@@ -99,7 +99,17 @@ public:
 		UObject* ReimportObject = nullptr;
 	};
 
-	UE_DEPRECATED(5.3, "This function is replaced by BeginImportAssetObject_GameThread.")
+	struct FImportAssetResult
+	{
+		//If the factory set this to true, the interchange task import object should skip this asset
+		//Asset can be skip if it already exist and its not a re-import. Example, if we import a static mesh and we find an existing material.
+		//We don't want to override the material in this case since the UE material often contains game play graph.
+		bool bIsFactorySkipAsset = false;
+		//Return the UObject import or re-import by the factory, leave it to nullptr if there was an error.
+		UObject* ImportedObject = nullptr;
+	};
+
+	UE_DEPRECATED(5.3, "This function is replaced by BeginImportAsset_GameThread.")
 	virtual UObject* ImportAssetObject_GameThread(const FImportAssetObjectParams& Arguments)
 	{
 		return nullptr;
@@ -110,27 +120,23 @@ public:
 	 * Create the asset package on the game thread since its not thread safe.
 	 *
 	 * @param Arguments - The structure containing all necessary arguments to import the UObject and the package, see the structure definition for the documentation.
-	 * @return the imported UObject or nullptr if there is an error.
+	 * @return the FImportAssetResult, see the structure to access the documentation.
 	 *
 	 * @Note Mandatory to override this function to create the asset UObject package. Not mandatory for level actor, use CreateSceneObject instead.
 	 */
-	virtual UObject* BeginImportAssetObject_GameThread(const FImportAssetObjectParams& Arguments)
+	virtual FImportAssetResult BeginImportAsset_GameThread(const FImportAssetObjectParams& Arguments)
 	{
+		FImportAssetResult ImportAssetResult;
 		PRAGMA_DISABLE_DEPRECATION_WARNINGS
-		return ImportAssetObject_GameThread(Arguments);
+		ImportAssetResult.ImportedObject = ImportAssetObject_GameThread(Arguments);
 		PRAGMA_ENABLE_DEPRECATION_WARNINGS
+		return ImportAssetResult;
 	}
 
-	/**
-	 * Override it to import/re-import the UObject source data and configure the properties asynchronously. Do not create the asset UObject asynchronously.
-	 * Helpful to get all the payload in parallel or do any long import task
-	 *
-	 * @param Arguments - The structure containing all necessary data to imported the UObject, see the structure definition for the documentation.
-	 * @return the imported UObject or nullptr if there is an error.
-	 */
+	UE_DEPRECATED(5.3, "This function is replaced by ImportAsset_Async.")
 	virtual UObject* ImportAssetObject_Async(const FImportAssetObjectParams& Arguments)
 	{
-		//By default simply return the UObject created by ImportAssetObject_GameThread that was store in the asset node.
+		//By default simply return the UObject created by BeginImportAsset_GameThread that was store in the asset node.
 		FSoftObjectPath ReferenceObject;
 		if (Arguments.AssetNode && Arguments.AssetNode->GetCustomReferenceObject(ReferenceObject))
 		{
@@ -139,22 +145,39 @@ public:
 		return nullptr;
 	}
 
+	/**
+	 * Override it to import/re-import the UObject source data and configure the properties asynchronously. Do not create the asset UObject asynchronously.
+	 * Helpful to get all the payload in parallel or do any long import task
+	 *
+	 * @param Arguments - The structure containing all necessary data to imported the UObject, see the structure definition for the documentation.
+	 * @return the FImportAssetResult, see the structure to access the documentation.
+	 */
+	virtual FImportAssetResult ImportAsset_Async(const FImportAssetObjectParams& Arguments)
+	{
+		FImportAssetResult ImportAssetResult;
+		PRAGMA_DISABLE_DEPRECATION_WARNINGS
+		ImportAssetResult.ImportedObject = ImportAssetObject_Async(Arguments);
+		PRAGMA_ENABLE_DEPRECATION_WARNINGS
+		return ImportAssetResult;
+	}
+
 	/**
 	 * Override it to end import/re-import on the game thread. Helpful if you need to display ui (i.e. re-import material conflict resolution) or if you
 	 * need to do anything not thread safe to complete the import/re-import.
 	 *
 	 * @param Arguments - The structure containing all necessary data to imported the UObject, see the structure definition for the documentation.
-	 * @return the imported UObject or nullptr if there is an error.
+	 * @return the FImportAssetResult, see the structure to access the documentation.
 	 */
-	virtual UObject* EndImportAssetObject_GameThread(const FImportAssetObjectParams& Arguments)
+	virtual FImportAssetResult EndImportAsset_GameThread(const FImportAssetObjectParams& Arguments)
 	{
+		FImportAssetResult ImportAssetResult;
 		//By default simply return the UObject created by BeginImportAssetObject_GameThread that was store in the asset node.
 		FSoftObjectPath ReferenceObject;
 		if (Arguments.AssetNode && Arguments.AssetNode->GetCustomReferenceObject(ReferenceObject))
 		{
-			return ReferenceObject.TryLoad();
+			ImportAssetResult.ImportedObject = ReferenceObject.TryLoad();
 		}
-		return nullptr;
+		return ImportAssetResult;
 	}
 
 	/**
diff --git a/Engine/Source/Runtime/Interchange/Engine/Private/Tasks/InterchangeTaskImportObject.cpp b/Engine/Source/Runtime/Interchange/Engine/Private/Tasks/InterchangeTaskImportObject.cpp
--- a/Engine/Source/Runtime/Interchange/Engine/Private/Tasks/InterchangeTaskImportObject.cpp
+++ b/Engine/Source/Runtime/Interchange/Engine/Private/Tasks/InterchangeTaskImportObject.cpp
@@ -159,18 +159,19 @@ namespace UE::Interchange::Private
 		return true;
 	}
 
-	UObject* InternalImportObjectStartup(TSharedPtr<FImportAsyncHelper, ESPMode::ThreadSafe> AsyncHelper
+	UInterchangeFactoryBase::FImportAssetResult InternalImportObjectStartup(TSharedPtr<FImportAsyncHelper, ESPMode::ThreadSafe> AsyncHelper
 		, UInterchangeFactoryBaseNode* FactoryNode
 		, int32 SourceIndex
 		, const FString& PackageBasePath
-		, TUniqueFunction<UObject*(UInterchangeFactoryBase::FImportAssetObjectParams&)>FactoryOperationLambda)
+		, TFunction<UInterchangeFactoryBase::FImportAssetResult(UInterchangeFactoryBase::FImportAssetObjectParams&)>FactoryOperationLambda)
 	{
 		check(AsyncHelper.IsValid());
+		UInterchangeFactoryBase::FImportAssetResult ErrorImportAssetResult;
 
 		//Verify if the task was cancel
 		if (AsyncHelper->bCancel || !FactoryNode || !Private::CanImportClass(*AsyncHelper, *FactoryNode, SourceIndex))
 		{
-			return nullptr;
+			return ErrorImportAssetResult;
 		}
 
 		UInterchangeFactoryBase* Factory = nullptr;
@@ -217,7 +218,7 @@ namespace UE::Interchange::Private
 				Message->AssetType = FactoryNode->GetObjectClass();
 				Message->Text = NSLOCTEXT("Interchange", "BadPackage", "It was not possible to create the asset as its package was not created correctly.");
 
-				return nullptr;
+				return ErrorImportAssetResult;
 			}
 
 			if (!AsyncHelper->SourceDatas.IsValidIndex(SourceIndex) || !AsyncHelper->Translators.IsValidIndex(SourceIndex))
@@ -227,7 +228,7 @@ namespace UE::Interchange::Private
 				Message->AssetType = FactoryNode->GetObjectClass();
 				Message->Text = NSLOCTEXT("Interchange", "SourceDataOrTranslatorInvalid", "It was not possible to create the asset as its translator was not created correctly.");
 
-				return nullptr;
+				return ErrorImportAssetResult;
 			}
 
 			Pkg = *PkgPtr;
@@ -251,7 +252,7 @@ namespace UE::Interchange::Private
 
 			return FactoryOperationLambda(CreateAssetParams);
 		}
-		return nullptr;
+		return ErrorImportAssetResult;
 	}
 }//ns UE::Interchange::Private
 
@@ -315,7 +316,7 @@ void UE::Interchange::FTaskImportObject_GameThread::DoTask(ENamedThreads::Type C
 			CreateAssetParams.ReimportObject = ReimportObject;
 			FactoryNode->SetCustomReferenceObject(FSoftObjectPath(ReimportObject));
 			//We call CreateEmptyAsset to ensure any resource use by an existing UObject is released on the game thread
-			Factory->BeginImportAssetObject_GameThread(CreateAssetParams);
+			Factory->BeginImportAsset_GameThread(CreateAssetParams);
 		}
 		else
 		{
@@ -339,20 +340,34 @@ void UE::Interchange::FTaskImportObject_GameThread::DoTask(ENamedThreads::Type C
 			return;
 		}
 
-		Pkg = CreatePackage(*PackageName);
-		if (Pkg == nullptr)
+		//If the package already exist we must load it so factory can find any existing asset and decide or not to override it
 		{
-			UInterchangeResultError_Generic* Message = Factory->AddMessage<UInterchangeResultError_Generic>();
-			Message->SourceAssetName = AsyncHelper->SourceDatas[SourceIndex]->GetFilename();
-			Message->DestinationAssetName = AssetName;
-			Message->AssetType = FactoryNode->GetObjectClass();
-			Message->Text = FText::Format(NSLOCTEXT("Interchange", "CouldntCreatePackage", "It was not possible to create a package named '{0}'; the asset will not be imported."), FText::FromString(PackageName));
+			//Try to find the package in memory
+			Pkg = FindPackage(nullptr, *PackageName);
+			if (!Pkg)
+			{
+				//Try to load the package from disk
+				Pkg = LoadPackage(nullptr, *PackageName, LOAD_NoWarn | LOAD_Quiet);
+				if (!Pkg)
+				{
+					//Create the package
+					Pkg = CreatePackage(*PackageName);
+					Pkg->SetPackageFlags(PKG_NewlyCreated);
+				}
+			}
 
-			//Skip this asset
-			return;
-		}
+			if (Pkg == nullptr)
+			{
+				UInterchangeResultError_Generic* Message = Factory->AddMessage<UInterchangeResultError_Generic>();
+				Message->SourceAssetName = AsyncHelper->SourceDatas[SourceIndex]->GetFilename();
+				Message->DestinationAssetName = AssetName;
+				Message->AssetType = FactoryNode->GetObjectClass();
+				Message->Text = FText::Format(NSLOCTEXT("Interchange", "CouldntCreatePackage", "It was not possible to create a package named '{0}'; the asset will not be imported."), FText::FromString(PackageName));
 
-		Pkg->SetPackageFlags(PKG_NewlyCreated);
+				//Skip this asset
+				return;
+			}
+		}
 
 		//Import Asset describe by the node
 		UInterchangeFactoryBase::FImportAssetObjectParams CreateAssetParams;
@@ -367,23 +382,27 @@ void UE::Interchange::FTaskImportObject_GameThread::DoTask(ENamedThreads::Type C
 		}
 		CreateAssetParams.ReimportObject = ReimportObject;
 		//Make sure the asset UObject is created with the correct type on the main thread
-		UObject* NodeAsset = Factory->BeginImportAssetObject_GameThread(CreateAssetParams);
-		if (NodeAsset)
+		UInterchangeFactoryBase::FImportAssetResult ImportAssetResult = Factory->BeginImportAsset_GameThread(CreateAssetParams);
+		if (ImportAssetResult.ImportedObject)
 		{
-			if (!NodeAsset->HasAnyInternalFlags(EInternalObjectFlags::Async))
+			//If the factory skip the asset, we simply set the node custom reference object
+			if (!ImportAssetResult.bIsFactorySkipAsset)
 			{
-				//Since the async flag is not set we must be in the game thread
-				ensure(IsInGameThread());
-				NodeAsset->SetInternalFlags(EInternalObjectFlags::Async);
+				if (!ImportAssetResult.ImportedObject->HasAnyInternalFlags(EInternalObjectFlags::Async))
+				{
+					//Since the async flag is not set we must be in the game thread
+					ensure(IsInGameThread());
+					ImportAssetResult.ImportedObject->SetInternalFlags(EInternalObjectFlags::Async);
+				}
+				FScopeLock Lock(&AsyncHelper->ImportedAssetsPerSourceIndexLock);
+				TArray<FImportAsyncHelper::FImportedObjectInfo>& ImportedInfos = AsyncHelper->ImportedAssetsPerSourceIndex.FindOrAdd(SourceIndex);
+				FImportAsyncHelper::FImportedObjectInfo& AssetInfo = ImportedInfos.AddDefaulted_GetRef();
+				AssetInfo.ImportedObject = ImportAssetResult.ImportedObject;
+				AssetInfo.Factory = Factory;
+				AssetInfo.FactoryNode = FactoryNode;
+				AssetInfo.bIsReimport = bool(ReimportObject != nullptr);
 			}
-			FScopeLock Lock(&AsyncHelper->ImportedAssetsPerSourceIndexLock);
-			TArray<FImportAsyncHelper::FImportedObjectInfo>& ImportedInfos = AsyncHelper->ImportedAssetsPerSourceIndex.FindOrAdd(SourceIndex);
-			FImportAsyncHelper::FImportedObjectInfo& AssetInfo = ImportedInfos.AddDefaulted_GetRef();
-			AssetInfo.ImportedObject = NodeAsset;
-			AssetInfo.Factory = Factory;
-			AssetInfo.FactoryNode = FactoryNode;
-			AssetInfo.bIsReimport = bool(ReimportObject != nullptr);
-			FactoryNode->SetCustomReferenceObject(FSoftObjectPath(NodeAsset));
+			FactoryNode->SetCustomReferenceObject(FSoftObjectPath(ImportAssetResult.ImportedObject));
 		}
 	}
 
@@ -424,7 +443,7 @@ void UE::Interchange::FTaskImportObject_Async::DoTask(ENamedThreads::Type Curren
 		, PackageBasePath
 		, [&Factory](UInterchangeFactoryBase::FImportAssetObjectParams& ImportAssetObjectParams)
 		{
-			return Factory->ImportAssetObject_Async(ImportAssetObjectParams);
+			return Factory->ImportAsset_Async(ImportAssetObjectParams);
 		});
 }
 
@@ -449,50 +468,53 @@ void UE::Interchange::FTaskImportObjectFinalize_GameThread::DoTask(ENamedThreads
 		FScopeLock Lock(&AsyncHelper->CreatedFactoriesLock);
 		Factory = AsyncHelper->CreatedFactories.FindChecked(FactoryNode->GetUniqueID());
 	}
-
-	UObject* NodeAsset = Private::InternalImportObjectStartup(AsyncHelper
+	UInterchangeFactoryBase::FImportAssetResult ImportAssetResult = Private::InternalImportObjectStartup(AsyncHelper
 		, FactoryNode
 		, SourceIndex
 		, PackageBasePath
 		, [&Factory](UInterchangeFactoryBase::FImportAssetObjectParams& ImportAssetObjectParams)
 		{
-			return Factory->EndImportAssetObject_GameThread(ImportAssetObjectParams);
+			return Factory->EndImportAsset_GameThread(ImportAssetObjectParams);
 		});
 
-	if (NodeAsset)
+	if (ImportAssetResult.ImportedObject)
 	{
-		FScopeLock Lock(&AsyncHelper->ImportedAssetsPerSourceIndexLock);
-		TArray<FImportAsyncHelper::FImportedObjectInfo>& ImportedInfos = AsyncHelper->ImportedAssetsPerSourceIndex.FindOrAdd(SourceIndex);
-		FImportAsyncHelper::FImportedObjectInfo* AssetInfoPtr = ImportedInfos.FindByPredicate([NodeAsset](const FImportAsyncHelper::FImportedObjectInfo& CurInfo)
-			{
-				return CurInfo.ImportedObject == NodeAsset;
-			});
-
-		if (!AssetInfoPtr)
+		//If the factory skip the asset, we simply set the node custom refernece object
+		if (!ImportAssetResult.bIsFactorySkipAsset)
 		{
-			FImportAsyncHelper::FImportedObjectInfo& AssetInfo = ImportedInfos.AddDefaulted_GetRef();
-			AssetInfo.ImportedObject = NodeAsset;
-			AssetInfo.Factory = Factory;
-			AssetInfo.FactoryNode = FactoryNode;
-			AssetInfo.bIsReimport = bool(AsyncHelper->TaskData.ReimportObject != nullptr);
-		}
+			FScopeLock Lock(&AsyncHelper->ImportedAssetsPerSourceIndexLock);
+			TArray<FImportAsyncHelper::FImportedObjectInfo>& ImportedInfos = AsyncHelper->ImportedAssetsPerSourceIndex.FindOrAdd(SourceIndex);
+			FImportAsyncHelper::FImportedObjectInfo* AssetInfoPtr = ImportedInfos.FindByPredicate([&ImportAssetResult](const FImportAsyncHelper::FImportedObjectInfo& CurInfo)
+				{
+					return CurInfo.ImportedObject == ImportAssetResult.ImportedObject;
+				});
 
-		// Fill in destination asset and type in any results which have been added previously by a translator or pipeline, now that we have a corresponding factory.
-		UInterchangeResultsContainer* Results = AsyncHelper->AssetImportResult->GetResults();
-		for (UInterchangeResult* Result : Results->GetResults())
-		{
-			if (!Result->InterchangeKey.IsEmpty() && (Result->DestinationAssetName.IsEmpty() || Result->AssetType == nullptr))
+			if (!AssetInfoPtr)
 			{
-				TArray<FString> TargetAssets;
-				FactoryNode->GetTargetNodeUids(TargetAssets);
-				if (TargetAssets.Contains(Result->InterchangeKey))
+				FImportAsyncHelper::FImportedObjectInfo& AssetInfo = ImportedInfos.AddDefaulted_GetRef();
+				AssetInfo.ImportedObject = ImportAssetResult.ImportedObject;
+				AssetInfo.Factory = Factory;
+				AssetInfo.FactoryNode = FactoryNode;
+				AssetInfo.bIsReimport = bool(AsyncHelper->TaskData.ReimportObject != nullptr);
+			}
+
+			// Fill in destination asset and type in any results which have been added previously by a translator or pipeline, now that we have a corresponding factory.
+			UInterchangeResultsContainer* Results = AsyncHelper->AssetImportResult->GetResults();
+			for (UInterchangeResult* Result : Results->GetResults())
+			{
+				if (!Result->InterchangeKey.IsEmpty() && (Result->DestinationAssetName.IsEmpty() || Result->AssetType == nullptr))
 				{
-					Result->DestinationAssetName = NodeAsset->GetPathName();
-					Result->AssetType = NodeAsset->GetClass();
+					TArray<FString> TargetAssets;
+					FactoryNode->GetTargetNodeUids(TargetAssets);
+					if (TargetAssets.Contains(Result->InterchangeKey))
+					{
+						Result->DestinationAssetName = ImportAssetResult.ImportedObject->GetPathName();
+						Result->AssetType = ImportAssetResult.ImportedObject->GetClass();
+					}
 				}
 			}
 		}
 
-		FactoryNode->SetCustomReferenceObject(FSoftObjectPath(NodeAsset));
+		FactoryNode->SetCustomReferenceObject(FSoftObjectPath(ImportAssetResult.ImportedObject));
 	}
 }
\ No newline at end of file
diff --git a/Engine/Source/Runtime/Interchange/Engine/Public/InterchangeImportCommon.h b/Engine/Source/Runtime/Interchange/Engine/Public/InterchangeImportCommon.h
--- a/Engine/Source/Runtime/Interchange/Engine/Public/InterchangeImportCommon.h
+++ b/Engine/Source/Runtime/Interchange/Engine/Public/InterchangeImportCommon.h
@@ -17,96 +17,93 @@ class UInterchangeSourceData;
 class UObject;
 template <typename FuncType> class TFunctionRef;
 
-namespace UE
+namespace UE::Interchange
 {
-	namespace Interchange
+	/**
+	* All the code we cannot put in the base factory class because of dependencies (like Engine dep)
+	* Will be available here.
+	*/
+	class INTERCHANGEENGINE_API FFactoryCommon
 	{
-		/**
-		 * All the code we cannot put in the base factory class because of dependencies (like Engine dep)
-		 * Will be available here.
-		 */
-		class INTERCHANGEENGINE_API FFactoryCommon
-		{
-		public:
+	public:
 
-			struct INTERCHANGEENGINE_API FUpdateImportAssetDataParameters
-			{
-				UObject* AssetImportDataOuter = nullptr;
-				UAssetImportData* AssetImportData = nullptr;
-				const UInterchangeSourceData* SourceData = nullptr;
-				FString NodeUniqueID;
-				UInterchangeBaseNodeContainer* NodeContainer;
-				const TArray<UObject*> Pipelines;
+		struct INTERCHANGEENGINE_API FUpdateImportAssetDataParameters
+		{
+			UObject* AssetImportDataOuter = nullptr;
+			UAssetImportData* AssetImportData = nullptr;
+			const UInterchangeSourceData* SourceData = nullptr;
+			FString NodeUniqueID;
+			UInterchangeBaseNodeContainer* NodeContainer;
+			const TArray<UObject*> Pipelines;
 
-				FUpdateImportAssetDataParameters(UObject* InAssetImportDataOuter
-																	, UAssetImportData* InAssetImportData
-																	, const UInterchangeSourceData* InSourceData
-																	, FString InNodeUniqueID
-																	, UInterchangeBaseNodeContainer* InNodeContainer
-																	, const TArray<UObject*>& InPipelines);
-			};
+			FUpdateImportAssetDataParameters(UObject* InAssetImportDataOuter
+																, UAssetImportData* InAssetImportData
+																, const UInterchangeSourceData* InSourceData
+																, FString InNodeUniqueID
+																, UInterchangeBaseNodeContainer* InNodeContainer
+																, const TArray<UObject*>& InPipelines);
+		};
 
-			/**
-			 * Update the AssetImportData source file of the specified asset in the parameters. Also update the node container and the node unique id.
-			 * If the AssetImportData is null it will create one. If the AssetImportData is not an UInterchangeAssetImportData it will create a new one.
-			 * @return The source data that should be stored on the asset or nullptr if a parameter is invalid
-			 */
-			static UAssetImportData* UpdateImportAssetData(FUpdateImportAssetDataParameters& Parameters);
+		/**
+		* Update the AssetImportData source file of the specified asset in the parameters. Also update the node container and the node unique id.
+		* If the AssetImportData is null it will create one. If the AssetImportData is not an UInterchangeAssetImportData it will create a new one.
+		* @return The source data that should be stored on the asset or nullptr if a parameter is invalid
+		*/
+		static UAssetImportData* UpdateImportAssetData(FUpdateImportAssetDataParameters& Parameters);
 
-			/**
-			 * Update the AssetImportData of the specified asset in the parameters. Also update the node container and the node unique id.
-			 * If the AssetImportData is null it will create one. If the AssetImportData is not an UInterchangeAssetImportData it will create a new one.
-			 * The file source update is done by calling the function parameter CustomFileSourceUpdate, so its the client responsability to properly update the file source.
-			 * @return The source data that should be stored on the asset or nullptr if a parameter is invalid
-			 */
-			static UAssetImportData* UpdateImportAssetData(FUpdateImportAssetDataParameters& Parameters, TFunctionRef<void(UInterchangeAssetImportData*)> CustomFileSourceUpdate);
+		/**
+		* Update the AssetImportData of the specified asset in the parameters. Also update the node container and the node unique id.
+		* If the AssetImportData is null it will create one. If the AssetImportData is not an UInterchangeAssetImportData it will create a new one.
+		* The file source update is done by calling the function parameter CustomFileSourceUpdate, so its the client responsability to properly update the file source.
+		* @return The source data that should be stored on the asset or nullptr if a parameter is invalid
+		*/
+		static UAssetImportData* UpdateImportAssetData(FUpdateImportAssetDataParameters& Parameters, TFunctionRef<void(UInterchangeAssetImportData*)> CustomFileSourceUpdate);
 
 #if WITH_EDITORONLY_DATA
-			struct INTERCHANGEENGINE_API FSetImportAssetDataParameters : public FUpdateImportAssetDataParameters
-			{
-				// Allow the factory to provide is own list of source files.
-				TArray<FAssetImportInfo::FSourceFile> SourceFiles;
+		struct INTERCHANGEENGINE_API FSetImportAssetDataParameters : public FUpdateImportAssetDataParameters
+		{
+			// Allow the factory to provide is own list of source files.
+			TArray<FAssetImportInfo::FSourceFile> SourceFiles;
 
-				FSetImportAssetDataParameters(UObject* InAssetImportDataOuter
-					, UAssetImportData* InAssetImportData
-					, const UInterchangeSourceData* InSourceData
-					, FString InNodeUniqueID
-					, UInterchangeBaseNodeContainer* InNodeContainer
-					, const TArray<UObject*>& InPipelines);
-			};
+			FSetImportAssetDataParameters(UObject* InAssetImportDataOuter
+				, UAssetImportData* InAssetImportData
+				, const UInterchangeSourceData* InSourceData
+				, FString InNodeUniqueID
+				, UInterchangeBaseNodeContainer* InNodeContainer
+				, const TArray<UObject*>& InPipelines);
+		};
 
-			/**
-			 * Set the AssetImportData source file of the specified asset in the parameters. Also update the node container and the node unique id.
-			 * If the AssetImportData is null it will create one. If the AssetImportData is not an UInterchangeAssetImportData it will create a new one.
-			 * @return The source data that should be stored on the asset or nullptr if a parameter is invalid
-			 */
-			static UAssetImportData* SetImportAssetData(FSetImportAssetDataParameters& Parameters);
+		/**
+		* Set the AssetImportData source file of the specified asset in the parameters. Also update the node container and the node unique id.
+		* If the AssetImportData is null it will create one. If the AssetImportData is not an UInterchangeAssetImportData it will create a new one.
+		* @return The source data that should be stored on the asset or nullptr if a parameter is invalid
+		*/
+		static UAssetImportData* SetImportAssetData(FSetImportAssetDataParameters& Parameters);
 
-			/**
-			 * Fills the OutSourceFilenames array with the list of source files contained in the asset source data.
-			 * Returns true if the operation was successful.
-			 */
-			static bool GetSourceFilenames(const UAssetImportData* AssetImportData, TArray<FString>& OutSourceFilenames);
+		/**
+		* Fills the OutSourceFilenames array with the list of source files contained in the asset source data.
+		* Returns true if the operation was successful.
+		*/
+		static bool GetSourceFilenames(const UAssetImportData* AssetImportData, TArray<FString>& OutSourceFilenames);
 
-			/**
-			 * Sets the SourceFileName value at the specified index.
-			 */
-			static bool SetSourceFilename(UAssetImportData* AssetImportData, const FString& SourceFilename, int32 SourceIndex, const FString& SourceLabel = FString());
+		/**
+		* Sets the SourceFileName value at the specified index.
+		*/
+		static bool SetSourceFilename(UAssetImportData* AssetImportData, const FString& SourceFilename, int32 SourceIndex, const FString& SourceLabel = FString());
 
-			/**
-			 * Set the object's reimport source at the specified index value.
-			 */
-			static bool SetReimportSourceIndex(const UObject* Object, UAssetImportData* AssetImportData, int32 SourceIndex);
+		/**
+		* Set the object's reimport source at the specified index value.
+		*/
+		static bool SetReimportSourceIndex(const UObject* Object, UAssetImportData* AssetImportData, int32 SourceIndex);
 
 #endif // WITH_EDITORONLY_DATA
 
-			/**
-			 * Apply the current strategy to the PipelineAssetNode
-			 */
-			static void ApplyReimportStrategyToAsset(UObject* Asset
-											  , UInterchangeFactoryBaseNode* PreviousAssetNode
-											  , UInterchangeFactoryBaseNode* CurrentAssetNode
-											  , UInterchangeFactoryBaseNode* PipelineAssetNode);
-		};
-	}
-}
+		/**
+		* Apply the current strategy to the PipelineAssetNode
+		*/
+		static void ApplyReimportStrategyToAsset(UObject* Asset
+											, UInterchangeFactoryBaseNode* PreviousAssetNode
+											, UInterchangeFactoryBaseNode* CurrentAssetNode
+											, UInterchangeFactoryBaseNode* PipelineAssetNode);
+	};
+} //ns UE::Interchange
