Fixed compilation in InterchangeTextureFactory.cpp

diff --git a/Engine/Plugins/Interchange/Runtime/Source/Import/Private/Texture/InterchangeTextureFactory.cpp b/Engine/Plugins/Interchange/Runtime/Source/Import/Private/Texture/InterchangeTextureFactory.cpp
--- a/Engine/Plugins/Interchange/Runtime/Source/Import/Private/Texture/InterchangeTextureFactory.cpp
+++ b/Engine/Plugins/Interchange/Runtime/Source/Import/Private/Texture/InterchangeTextureFactory.cpp
@@ -1152,10 +1152,8 @@ namespace UE::Interchange::Private::InterchangeTextureFactory
 			// Allocate first mipmap.
 			int32 NumBlocksX = ImportImage.SizeX / GPixelFormats[PixelFormat].BlockSizeX;
 			int32 NumBlocksY = ImportImage.SizeY / GPixelFormats[PixelFormat].BlockSizeY;
-			FTexture2DMipMap* Mip = new FTexture2DMipMap();
+			FTexture2DMipMap* Mip = new FTexture2DMipMap(ImportImage.SizeX, ImportImage.SizeY);
 			Texture2D->GetPlatformData()->Mips.Add(Mip);
-			Mip->SizeX = ImportImage.SizeX;
-			Mip->SizeY = ImportImage.SizeY;
 			Mip->BulkData.Lock(LOCK_READ_WRITE);
 			Mip->BulkData.Realloc(NumBlocksX * NumBlocksY * GPixelFormats[PixelFormat].BlockBytes);
 			Mip->BulkData.Unlock();
