[Interchange] Use predefined GLTF material assets at runtime and on demand
Added TexCoord support for GLTF MaterialInstances.
Added usage of ClearCoatNormalScale to GLTF MaterialInstances.
Added Opacity Mask Clip support for GLTF MaterialInstances.
Added UnlitTranslucent support for GLTF MaterialInstances.

diff --git a/Engine/Plugins/Interchange/Runtime/Source/Import/Private/Gltf/InterchangeGLTFMaterialInstance.cpp b/Engine/Plugins/Interchange/Runtime/Source/Import/Private/Gltf/InterchangeGLTFMaterialInstance.cpp
new file mode 100644
--- /dev/null
+++ b/Engine/Plugins/Interchange/Runtime/Source/Import/Private/Gltf/InterchangeGLTFMaterialInstance.cpp
@@ -0,0 +1,989 @@
+// Copyright Epic Games, Inc. All Rights Reserved.
+#pragma once
+
+#include "Gltf/InterchangeGLTFMaterialInstance.h"
+#include "GLTFMaterial.h"
+#include "GLTFTexture.h"
+
+#include "Nodes/InterchangeBaseNodeContainer.h"
+#include "InterchangeTextureNode.h"
+#include "InterchangeShaderGraphNode.h"
+#include "InterchangeMaterialInstanceNode.h"
+#include "InterchangeShaderGraphNode.h"
+
+#include "Materials/MaterialInstance.h"
+#include "Materials/MaterialInstanceConstant.h"
+#include "Materials/MaterialInstanceDynamic.h"
+
+#include "Async/Async.h"
+
+const FString GltfMaterialAttributeIdentifier = TEXT("Gltf_MI_AttributeIdentifier_");
+
+#define INTERCHANGE_GLTF_STRINGIFY(x) #x
+#define DECLARE_INTERCHANGE_GLTF_MI_MAP(MapName) const FString MapName##Texture = TEXT(INTERCHANGE_GLTF_STRINGIFY(MapName ## Texture)); \
+												 const FString MapName##Texture_OffsetScale = TEXT(INTERCHANGE_GLTF_STRINGIFY(MapName ## Texture_OffsetScale)); \
+												 const FString MapName##Texture_Rotation = TEXT(INTERCHANGE_GLTF_STRINGIFY(MapName ## Texture_Rotation)); \
+												 const FString MapName##Texture_TexCoord = TEXT(INTERCHANGE_GLTF_STRINGIFY(MapName ## Texture_TexCoord));
+
+namespace Parameters
+{
+	//MetalRoughness specific:
+	DECLARE_INTERCHANGE_GLTF_MI_MAP(BaseColor)
+	const FString BaseColorFactor = TEXT("BaseColorFactor");
+
+	DECLARE_INTERCHANGE_GLTF_MI_MAP(MetallicRoughness)
+	const FString MetallicFactor = TEXT("MetallicFactor");
+	const FString RoughnessFactor = TEXT("RoughnessFactor");
+
+	DECLARE_INTERCHANGE_GLTF_MI_MAP(Specular)
+	const FString SpecularFactor = TEXT("SpecularFactor");
+	
+
+	//SpecularGlossiness specific
+	DECLARE_INTERCHANGE_GLTF_MI_MAP(Diffuse)
+	const FString DiffuseFactor = TEXT("DiffuseFactor");
+
+	DECLARE_INTERCHANGE_GLTF_MI_MAP(SpecularGlossiness)
+	const FString SpecFactor = TEXT("SpecFactor");
+	const FString GlossinessFactor = TEXT("GlossinessFactor");
+	
+
+	//Generic:
+	DECLARE_INTERCHANGE_GLTF_MI_MAP(Normal)
+	const FString NormalScale = TEXT("NormalScale");
+
+	DECLARE_INTERCHANGE_GLTF_MI_MAP(Emissive)
+	const FString EmissiveFactor = TEXT("EmissiveFactor");
+	const FString EmissiveStrength = TEXT("EmissiveStrength");
+
+	DECLARE_INTERCHANGE_GLTF_MI_MAP(Occlusion)
+	const FString OcclusionStrength = TEXT("OcclusionStrength");
+
+	const FString IOR = TEXT("IOR");
+
+	const FString AlphaCutoff = TEXT("AlphaCutoff");
+	
+
+	//ClearCoat specific:
+	DECLARE_INTERCHANGE_GLTF_MI_MAP(ClearCoat)
+	const FString ClearCoatFactor = TEXT("ClearCoatFactor");
+
+	DECLARE_INTERCHANGE_GLTF_MI_MAP(ClearCoatRoughness)
+	const FString ClearCoatRoughnessFactor = TEXT("ClearCoatRoughnessFactor");
+
+	DECLARE_INTERCHANGE_GLTF_MI_MAP(ClearCoatNormal)
+	const FString ClearCoatNormalScale = TEXT("ClearCoatNormalScale");
+	
+
+	//Sheen specific:
+	DECLARE_INTERCHANGE_GLTF_MI_MAP(SheenColor)
+	const FString SheenColorFactor = TEXT("SheenColorFactor");
+
+	DECLARE_INTERCHANGE_GLTF_MI_MAP(SheenRoughness)
+	const FString SheenRoughnessFactor = TEXT("SheenRoughnessFactor");
+	
+
+	//Transmission specific:
+	DECLARE_INTERCHANGE_GLTF_MI_MAP(Transmission)
+	const FString TransmissionFactor = TEXT("TransmissionFactor");
+};
+
+enum EShadingModel : uint8
+{
+	DEFAULT = 0, //MetalRoughness
+	UNLIT,
+	CLEARCOAT,
+	SHEEN,
+	TRANSMISSION,
+	SPECULARGLOSSINESS,
+	SHADINGMODELCOUNT
+};
+
+FString ToString(EShadingModel ShadingModel)
+{
+	switch (ShadingModel)
+	{
+	case DEFAULT:
+		return TEXT("Default");
+	case UNLIT:
+		return TEXT("Unlit");
+	case CLEARCOAT:
+		return TEXT("ClearCoat");
+	case SHEEN:
+		return TEXT("Sheen");
+	case TRANSMISSION:
+		return TEXT("Transmission");
+	case SPECULARGLOSSINESS:
+		return TEXT("SpecularGlossiness");
+	default:
+		return TEXT("Unexpected");
+	}
+}
+
+FString ToString(GLTF::FMaterial::EAlphaMode AlphaMode)
+{
+	switch (AlphaMode)
+	{
+	case GLTF::FMaterial::EAlphaMode::Opaque:
+		return TEXT("Opaque");
+	case GLTF::FMaterial::EAlphaMode::Mask:
+		return TEXT("Mask");
+	case GLTF::FMaterial::EAlphaMode::Blend:
+		return TEXT("Blend");
+	default:
+		return TEXT("Unexpected");
+	}
+}
+
+FString GetIdentifier(EShadingModel ShadingModel, GLTF::FMaterial::EAlphaMode AlphaMode, bool bTwoSided)
+{
+	FString Name = TEXT("MI_");
+
+	Name += ToString(ShadingModel);
+
+	if (ShadingModel != EShadingModel::TRANSMISSION)
+	{
+		Name += TEXT("_") + ToString(AlphaMode);
+	}
+
+	Name += bTwoSided ? TEXT("_DS") : TEXT("");
+
+	return Name;
+}
+
+
+FGLTFMaterialInstanceSettings::FGLTFMaterialInstanceSettings()
+{
+	static const FString AssetFolder = TEXT("/Interchange/gltf/MaterialInstances/");
+
+	auto AddMaterialParent = [=](EShadingModel ShadingModel, GLTF::FMaterial::EAlphaMode AlphaMode, bool bTwoSided)
+	{
+		const FString Identifier = GetIdentifier(ShadingModel, AlphaMode, bTwoSided);
+		const FString AssetName = Identifier + TEXT(".") + Identifier;
+		const FString AssetPath = AssetFolder + AssetName;
+
+		MaterialParents.FindOrAdd(Identifier, FSoftObjectPath{ *AssetPath });
+	};
+
+	bool bTwoSided = false;
+	for (size_t ShadingModelCounter = 0; ShadingModelCounter < EShadingModel::SHADINGMODELCOUNT; ShadingModelCounter++)
+	{
+		AddMaterialParent(EShadingModel(ShadingModelCounter), GLTF::FMaterial::EAlphaMode::Opaque, bTwoSided);
+
+		if (ShadingModelCounter == EShadingModel::TRANSMISSION) continue;
+		AddMaterialParent(EShadingModel(ShadingModelCounter), GLTF::FMaterial::EAlphaMode::Mask, bTwoSided);
+		AddMaterialParent(EShadingModel(ShadingModelCounter), GLTF::FMaterial::EAlphaMode::Blend, bTwoSided);
+	}
+
+	bTwoSided = true;
+	for (size_t ShadingModelCounter = 0; ShadingModelCounter < EShadingModel::SHADINGMODELCOUNT; ShadingModelCounter++)
+	{
+		AddMaterialParent(EShadingModel(ShadingModelCounter), GLTF::FMaterial::EAlphaMode::Opaque, bTwoSided);
+
+		if (ShadingModelCounter == EShadingModel::TRANSMISSION) continue;
+		AddMaterialParent(EShadingModel(ShadingModelCounter), GLTF::FMaterial::EAlphaMode::Mask, bTwoSided);
+		AddMaterialParent(EShadingModel(ShadingModelCounter), GLTF::FMaterial::EAlphaMode::Blend, bTwoSided);
+	}
+
+	//save out the identifiers:
+	MaterialParents.GenerateKeyArray(ExpectedMaterialInstanceIdentifiers);
+}
+
+TSet<FString> FGLTFMaterialInstanceSettings::GenerateExpectedParametersList(const FString& Identifier) const
+{
+	TSet<FString> ExpectedParameters;
+
+	if (Identifier.Contains(TEXT("_Unlit")))
+	{
+		ExpectedParameters.Add(Parameters::BaseColorTexture);
+		ExpectedParameters.Add(Parameters::BaseColorTexture_OffsetScale);
+		ExpectedParameters.Add(Parameters::BaseColorTexture_Rotation);
+		ExpectedParameters.Add(Parameters::BaseColorTexture_TexCoord);
+		ExpectedParameters.Add(Parameters::BaseColorFactor);
+
+		return ExpectedParameters;
+	}
+
+	//Generic ones:
+	{
+		ExpectedParameters.Add(Parameters::NormalTexture);
+		ExpectedParameters.Add(Parameters::NormalTexture_OffsetScale);
+		ExpectedParameters.Add(Parameters::NormalTexture_Rotation);
+		ExpectedParameters.Add(Parameters::NormalTexture_TexCoord);
+		ExpectedParameters.Add(Parameters::NormalScale);
+
+		if (!Identifier.Contains(TEXT("Transmission")))
+		{
+			ExpectedParameters.Add(Parameters::EmissiveTexture);
+			ExpectedParameters.Add(Parameters::EmissiveTexture_OffsetScale);
+			ExpectedParameters.Add(Parameters::EmissiveTexture_Rotation);
+			ExpectedParameters.Add(Parameters::EmissiveTexture_TexCoord);
+			ExpectedParameters.Add(Parameters::EmissiveFactor);
+			ExpectedParameters.Add(Parameters::EmissiveStrength);
+		}
+
+		ExpectedParameters.Add(Parameters::OcclusionTexture);
+		ExpectedParameters.Add(Parameters::OcclusionTexture_OffsetScale);
+		ExpectedParameters.Add(Parameters::OcclusionTexture_Rotation);
+		ExpectedParameters.Add(Parameters::OcclusionTexture_TexCoord);
+		ExpectedParameters.Add(Parameters::OcclusionStrength);
+
+		ExpectedParameters.Add(Parameters::IOR);
+	}
+
+	//Based on ShadingModel:
+
+	if (Identifier.Contains(TEXT("Default")))
+	{
+		//MetalRoughness Specific:
+
+		ExpectedParameters.Add(Parameters::BaseColorTexture);
+		ExpectedParameters.Add(Parameters::BaseColorTexture_OffsetScale);
+		ExpectedParameters.Add(Parameters::BaseColorTexture_Rotation);
+		ExpectedParameters.Add(Parameters::BaseColorTexture_TexCoord);
+		ExpectedParameters.Add(Parameters::BaseColorFactor);
+
+		ExpectedParameters.Add(Parameters::MetallicRoughnessTexture);
+		ExpectedParameters.Add(Parameters::MetallicRoughnessTexture_OffsetScale);
+		ExpectedParameters.Add(Parameters::MetallicRoughnessTexture_Rotation);
+		ExpectedParameters.Add(Parameters::MetallicRoughnessTexture_TexCoord);
+		ExpectedParameters.Add(Parameters::MetallicFactor);
+		ExpectedParameters.Add(Parameters::RoughnessFactor);
+
+		ExpectedParameters.Add(Parameters::SpecularTexture);
+		ExpectedParameters.Add(Parameters::SpecularTexture_OffsetScale);
+		ExpectedParameters.Add(Parameters::SpecularTexture_Rotation);
+		ExpectedParameters.Add(Parameters::SpecularTexture_TexCoord);
+		ExpectedParameters.Add(Parameters::SpecularFactor);
+
+	}
+	else if (Identifier.Contains(TEXT("ClearCoat")))
+	{
+		ExpectedParameters.Add(Parameters::ClearCoatTexture);
+		ExpectedParameters.Add(Parameters::ClearCoatTexture_OffsetScale);
+		ExpectedParameters.Add(Parameters::ClearCoatTexture_Rotation);
+		ExpectedParameters.Add(Parameters::ClearCoatTexture_TexCoord);
+		ExpectedParameters.Add(Parameters::ClearCoatFactor);
+
+		ExpectedParameters.Add(Parameters::ClearCoatRoughnessTexture);
+		ExpectedParameters.Add(Parameters::ClearCoatRoughnessTexture_OffsetScale);
+		ExpectedParameters.Add(Parameters::ClearCoatRoughnessTexture_Rotation);
+		ExpectedParameters.Add(Parameters::ClearCoatRoughnessTexture_TexCoord);
+		ExpectedParameters.Add(Parameters::ClearCoatRoughnessFactor);
+
+		ExpectedParameters.Add(Parameters::ClearCoatNormalTexture);
+		ExpectedParameters.Add(Parameters::ClearCoatNormalTexture_OffsetScale);
+		ExpectedParameters.Add(Parameters::ClearCoatNormalTexture_Rotation);
+		ExpectedParameters.Add(Parameters::ClearCoatNormalTexture_TexCoord);
+		ExpectedParameters.Add(Parameters::ClearCoatNormalScale);
+	}
+	else if (Identifier.Contains(TEXT("Sheen")))
+	{
+		ExpectedParameters.Add(Parameters::SheenColorTexture);
+		ExpectedParameters.Add(Parameters::SheenColorTexture_OffsetScale);
+		ExpectedParameters.Add(Parameters::SheenColorTexture_Rotation);
+		ExpectedParameters.Add(Parameters::SheenColorTexture_TexCoord);
+		ExpectedParameters.Add(Parameters::SheenColorFactor);
+
+		ExpectedParameters.Add(Parameters::SheenRoughnessTexture);
+		ExpectedParameters.Add(Parameters::SheenRoughnessTexture_OffsetScale);
+		ExpectedParameters.Add(Parameters::SheenRoughnessTexture_Rotation);
+		ExpectedParameters.Add(Parameters::SheenRoughnessTexture_TexCoord);
+		ExpectedParameters.Add(Parameters::SheenRoughnessFactor);
+	}
+	else if (Identifier.Contains(TEXT("Transmission")))
+	{
+		ExpectedParameters.Add(Parameters::TransmissionTexture);
+		ExpectedParameters.Add(Parameters::TransmissionTexture_OffsetScale);
+		ExpectedParameters.Add(Parameters::TransmissionTexture_Rotation);
+		ExpectedParameters.Add(Parameters::TransmissionTexture_TexCoord);
+		ExpectedParameters.Add(Parameters::TransmissionFactor);
+	}
+	else if (Identifier.Contains(TEXT("SpecularGlossiness")))
+	{
+		ExpectedParameters.Add(Parameters::DiffuseTexture);
+		ExpectedParameters.Add(Parameters::DiffuseTexture_OffsetScale);
+		ExpectedParameters.Add(Parameters::DiffuseTexture_Rotation);
+		ExpectedParameters.Add(Parameters::DiffuseTexture_TexCoord);
+		ExpectedParameters.Add(Parameters::DiffuseFactor);
+
+		ExpectedParameters.Add(Parameters::SpecularGlossinessTexture);
+		ExpectedParameters.Add(Parameters::SpecularGlossinessTexture_OffsetScale);
+		ExpectedParameters.Add(Parameters::SpecularGlossinessTexture_Rotation);
+		ExpectedParameters.Add(Parameters::SpecularGlossinessTexture_TexCoord);
+		ExpectedParameters.Add(Parameters::SpecFactor);
+		ExpectedParameters.Add(Parameters::GlossinessFactor);
+	}
+
+	return ExpectedParameters;
+}
+
+TArray<FString> FGLTFMaterialInstanceSettings::ValidateMaterialInstancesAndParameters() const
+{
+	//TODO:
+	//	swap to the used MaterialParents (aka to the overrides when they are exposed)
+	const TMap<FString, FSoftObjectPath>& MaterialParentsUsed = MaterialParents;
+
+	TArray<FString> NotCoveredIdentifiersParameters;
+
+	//Check if all Material variations are covered:
+	TArray<FString> ExpectedIdentifiers = ExpectedMaterialInstanceIdentifiers;
+	TArray<FString> IdentifiersUsed;
+	MaterialParentsUsed.GetKeys(IdentifiersUsed);
+	for (const FString& Identifier : IdentifiersUsed)
+	{
+		ExpectedIdentifiers.Remove(Identifier);
+	}
+	for (const FString& ExpectedIdentifier : ExpectedIdentifiers)
+	{
+		NotCoveredIdentifiersParameters.Add(TEXT("[") + ExpectedIdentifier + TEXT("]: MaterialInstance not found for Identifier."));
+	}
+
+	for (const TPair<FString, FSoftObjectPath>& MaterialParent : MaterialParentsUsed)
+	{
+		TSet<FString> ExpectedParameters = GenerateExpectedParametersList(MaterialParent.Key);
+
+		if (UMaterialInstance* ParentMaterialInstance = Cast<UMaterialInstance>(MaterialParent.Value.TryLoad()))
+		{
+			TArray<FGuid> ParameterIds;
+			TArray<FMaterialParameterInfo> ScalarParameterInfos;
+			TArray<FMaterialParameterInfo> VectorParameterInfos;
+			TArray<FMaterialParameterInfo> TextureParameterInfos;
+			ParentMaterialInstance->GetAllScalarParameterInfo(ScalarParameterInfos, ParameterIds);
+			ParentMaterialInstance->GetAllVectorParameterInfo(VectorParameterInfos, ParameterIds);
+			ParentMaterialInstance->GetAllTextureParameterInfo(TextureParameterInfos, ParameterIds);
+
+			for (const FMaterialParameterInfo& ParameterInfo : ScalarParameterInfos)
+			{
+				ExpectedParameters.Remove(ParameterInfo.Name.ToString());
+			}
+			for (const FMaterialParameterInfo& ParameterInfo : VectorParameterInfos)
+			{
+				ExpectedParameters.Remove(ParameterInfo.Name.ToString());
+			}
+			for (const FMaterialParameterInfo& ParameterInfo : TextureParameterInfos)
+			{
+				ExpectedParameters.Remove(ParameterInfo.Name.ToString());
+			}
+		}
+
+		for (const FString& ExpectedParameter : ExpectedParameters)
+		{
+			NotCoveredIdentifiersParameters.Add(TEXT("[") + MaterialParent.Key + TEXT("]: Does not cover expected parameter: ") + ExpectedParameter + TEXT("."));
+		}
+	}
+
+	return NotCoveredIdentifiersParameters;
+}
+
+//Presedence order based on the logic in the InterchangeGenericMaterialPipeline: ClearCoat > Sheen > unlit
+EShadingModel GetShadingModel(GLTF::FMaterial::EShadingModel GltfShadingModel, bool bHasClearCoat, bool bHasSheen, bool bUnlit, bool bHasTransmission)
+{
+	if (GltfShadingModel == GLTF::FMaterial::EShadingModel::SpecularGlossiness)
+	{
+		return EShadingModel::SPECULARGLOSSINESS;
+	}
+
+	if (bHasClearCoat)
+	{
+		return EShadingModel::CLEARCOAT;
+	}
+	
+	if (bHasSheen)
+	{
+		return EShadingModel::SHEEN;
+	}
+
+	if (bUnlit)
+	{
+		return EShadingModel::UNLIT;
+	}
+
+	if (bHasTransmission)
+	{
+		return EShadingModel::TRANSMISSION;
+	}
+
+	return EShadingModel::DEFAULT;
+}
+
+bool FGLTFMaterialInstanceSettings::ProcessGltfMaterial(UInterchangeBaseNodeContainer& NodeContainer, const GLTF::FMaterial& GltfMaterial, const TArray<GLTF::FTexture>& Textures) const
+{
+	FString ParentIdentifier = GetIdentifier(GetShadingModel(GltfMaterial.ShadingModel, GltfMaterial.bHasClearCoat, GltfMaterial.bHasSheen, GltfMaterial.bIsUnlitShadingModel, GltfMaterial.bHasTransmission), GltfMaterial.AlphaMode, GltfMaterial.bIsDoubleSided);
+	FString Parent;
+	if (const FSoftObjectPath* ObjectPath = MaterialParents.Find(ParentIdentifier))
+	{
+		Parent = ObjectPath->GetAssetPathString();
+	}
+	else
+	{
+		return false;
+	}
+
+	UInterchangeMaterialInstanceNode* MaterialInstanceNode = NewObject<UInterchangeMaterialInstanceNode>(&NodeContainer);
+	
+	FString MaterialNodeUid = UInterchangeShaderGraphNode::MakeNodeUid(GltfMaterial.UniqueId);
+	MaterialInstanceNode->InitializeNode(MaterialNodeUid, GltfMaterial.Name, EInterchangeNodeContainerType::TranslatedAsset);
+	NodeContainer.AddNode(MaterialInstanceNode);
+
+	MaterialInstanceNode->SetCustomParent(Parent);
+
+	auto SetMap = [=](const GLTF::FTextureMap& TextureMap, const FString& Name)
+	{
+		if (Textures.IsValidIndex(TextureMap.TextureIndex))
+		{
+			MaterialInstanceNode->AddTextureParameterValue(Name, UInterchangeTextureNode::MakeNodeUid(Textures[TextureMap.TextureIndex].UniqueId));
+
+			//TexCoord decision in MaterialInstances work based on a Switch Node
+			//Currently Only supports UV0 and UV1
+			// [0...1) -> UV0
+			// [1...2) -> UV1
+			// else    -> UV0 (defaults to 0)
+			MaterialInstanceNode->AddScalarParameterValue(Name + TEXT("_TexCoord"), TextureMap.TexCoord);
+		}
+
+		if (TextureMap.bHasTextureTransform)
+		{
+			FVector4f OffsetScale(TextureMap.TextureTransform.Offset[0], TextureMap.TextureTransform.Offset[1], TextureMap.TextureTransform.Scale[0], TextureMap.TextureTransform.Scale[1]);
+			MaterialInstanceNode->AddVectorParameterValue(Name + TEXT("_OffsetScale"), OffsetScale);
+
+			MaterialInstanceNode->AddScalarParameterValue(Name + TEXT("_Rotation"), TextureMap.TextureTransform.Rotation);
+		}
+	};
+
+	if (GltfMaterial.bIsUnlitShadingModel)
+	{
+		//BaseColorTexture
+		//BaseColorTexture_OffsetScale
+		//BaseColorTExture_Rotation
+		//BaseColorTExture_TexCoord
+		SetMap(GltfMaterial.BaseColor, Parameters::BaseColorTexture);
+
+		//BaseColorFactor
+		MaterialInstanceNode->AddVectorParameterValue(Parameters::BaseColorFactor, GltfMaterial.BaseColorFactor);
+
+		return true;
+	}
+
+	if (GltfMaterial.AlphaMode == GLTF::FMaterial::EAlphaMode::Mask)
+	{
+		MaterialInstanceNode->AddScalarParameterValue(Parameters::AlphaCutoff, GltfMaterial.AlphaCutoff);
+	}
+
+	if (GltfMaterial.ShadingModel == GLTF::FMaterial::EShadingModel::MetallicRoughness)
+	{
+		{
+			//BaseColorTexture
+			//BaseColorTexture_OffsetScale
+			//BaseColorTexture_Rotation
+			//BaseColorTexture_TexCoord
+			SetMap(GltfMaterial.BaseColor, Parameters::BaseColorTexture);
+
+			//BaseColorFactor
+			MaterialInstanceNode->AddVectorParameterValue(Parameters::BaseColorFactor, GltfMaterial.BaseColorFactor);
+		}
+		
+		{
+			//MetallicRoughnessTexture
+			//MetallicRoughnessTexture_OffsetScale
+			//MetallicRoughnessTexture_Rotation
+			//MetallicRoughnessTexture_TexCoord
+			SetMap(GltfMaterial.MetallicRoughness.Map, Parameters::MetallicRoughnessTexture);
+
+			//MetallicFactor
+			MaterialInstanceNode->AddScalarParameterValue(Parameters::MetallicFactor, GltfMaterial.MetallicRoughness.MetallicFactor);
+			
+			//RoughnessFactor
+			MaterialInstanceNode->AddScalarParameterValue(Parameters::RoughnessFactor, GltfMaterial.MetallicRoughness.RoughnessFactor);
+		}
+
+		if (GltfMaterial.bHasSpecular)
+		{
+			//SpecularTexture
+			//SpecularTexture_OffsetScale
+			//SpecularTexture_Rotation
+			//SpecularTexture_TexCoord
+			SetMap(GltfMaterial.Specular.SpecularMap, Parameters::SpecularTexture);
+
+			//SpecularFactor
+			MaterialInstanceNode->AddScalarParameterValue(Parameters::SpecularFactor, GltfMaterial.Specular.SpecularFactor);
+		}
+	}
+	else if (GltfMaterial.ShadingModel == GLTF::FMaterial::EShadingModel::SpecularGlossiness)
+	{
+		//////////////////////
+		//specular glossiness:
+		//////////////////////
+		
+		{
+			//DiffuseTexture
+			//DiffuseTexture_OffsetScale
+			//DiffuseTexture_Rotation
+			//DiffuseTexture_TexCoord
+			SetMap(GltfMaterial.BaseColor, Parameters::DiffuseTexture);
+
+			//DiffuseFactor
+			MaterialInstanceNode->AddVectorParameterValue(Parameters::DiffuseFactor, GltfMaterial.BaseColorFactor);
+		}
+		
+		{
+			//SpecularGlossinessTexture
+			//SpecularGlossinessTexture_OffsetScale
+			//SpecularGlossinessTexture_Rotation
+			//SpecularGlossinessTexture_TexCoord
+			SetMap(GltfMaterial.SpecularGlossiness.Map, Parameters::SpecularGlossinessTexture);
+
+			//SpecFactor
+			MaterialInstanceNode->AddVectorParameterValue(Parameters::SpecFactor, FVector4f(GltfMaterial.SpecularGlossiness.SpecularFactor[0],
+																							GltfMaterial.SpecularGlossiness.SpecularFactor[1],
+																							GltfMaterial.SpecularGlossiness.SpecularFactor[2]));
+
+			//GlossinessFactor
+			MaterialInstanceNode->AddScalarParameterValue(Parameters::GlossinessFactor, GltfMaterial.SpecularGlossiness.GlossinessFactor);
+		}
+	}
+
+	{
+		//NormalTexture
+		//NormalTexture_OffsetScale
+		//NormalTexture_Rotation
+		//NormalTexture_TexCoord
+		SetMap(GltfMaterial.Normal, Parameters::NormalTexture);
+
+		//NormalScale
+		MaterialInstanceNode->AddScalarParameterValue(Parameters::NormalScale, GltfMaterial.NormalScale);
+	}
+
+	{
+		//EmissiveTexture
+		//EmissiveTexture_OffsetScale
+		//EmissiveTexture_Rotation
+		//EmissiveTexture_TexCoord
+		SetMap(GltfMaterial.Emissive, Parameters::EmissiveTexture);
+
+		//EmissiveFactor
+		MaterialInstanceNode->AddVectorParameterValue(Parameters::EmissiveFactor, GltfMaterial.EmissiveFactor);
+		
+		//EmissiveStrength
+		MaterialInstanceNode->AddScalarParameterValue(Parameters::EmissiveStrength, GltfMaterial.bHasEmissiveStrength ? GltfMaterial.EmissiveStrength : 1.0f);
+	}
+
+	{
+		//OcclusionTexture
+		//OcclusionTexture_OffsetScale
+		//OcclusionTexture_Rotation
+		//OcclusionTexture_TexCoord
+		SetMap(GltfMaterial.Occlusion, Parameters::OcclusionTexture);
+
+		//OcclusionStrength
+		MaterialInstanceNode->AddScalarParameterValue(Parameters::OcclusionStrength, GltfMaterial.OcclusionStrength);
+	}
+	
+	if (GltfMaterial.bHasIOR)
+	{
+		//ior
+		MaterialInstanceNode->AddScalarParameterValue(Parameters::IOR, GltfMaterial.IOR);
+	}
+
+	if (GltfMaterial.bHasClearCoat)
+	{
+		{
+			//ClearCoatTexture
+			//ClearCoatTexture_OffsetScale
+			//ClearCoatTexture_Rotation
+			//ClearCoatTexture_TexCoord
+			SetMap(GltfMaterial.ClearCoat.ClearCoatMap, Parameters::ClearCoatTexture);
+
+			//ClearCoatFactor
+			MaterialInstanceNode->AddScalarParameterValue(Parameters::ClearCoatFactor, GltfMaterial.ClearCoat.ClearCoatFactor);
+		}
+
+		{
+			//ClearCoatRoughnessTexture
+			//ClearCoatRoughnessTexture_OffsetScale
+			//ClearCoatRoughnessTexture_Rotation
+			//ClearCoatRoughnessTexture_TexCoord
+			SetMap(GltfMaterial.ClearCoat.RoughnessMap, Parameters::ClearCoatRoughnessTexture);
+
+			//ClearCoatRoughnessFactor
+			MaterialInstanceNode->AddScalarParameterValue(Parameters::ClearCoatRoughnessFactor, GltfMaterial.ClearCoat.Roughness);
+		}
+
+		{
+			//ClearCoatNormalTexture
+			//ClearCoatNormalTexture_OffsetScale
+			//ClearCoatNormalTexture_Rotation
+			//ClearCoatNormalTexture_TexCoord
+			SetMap(GltfMaterial.ClearCoat.NormalMap, Parameters::ClearCoatNormalTexture);
+
+			//ClearCoatNormalFactor
+			MaterialInstanceNode->AddScalarParameterValue(Parameters::ClearCoatNormalScale, GltfMaterial.ClearCoat.NormalMapUVScale);
+		}
+	}
+	else if (GltfMaterial.bHasSheen)
+	{
+		{
+			//SheenColorTexture
+			//SheenColorTexture_OffsetScale
+			//SheenColorTexture_Rotation
+			//SheenColorTexture_TexCoord
+			SetMap(GltfMaterial.Sheen.SheenColorMap, Parameters::SheenColorTexture);
+
+			//SheenColorFactor
+			MaterialInstanceNode->AddVectorParameterValue(Parameters::SheenColorFactor, FVector4f(GltfMaterial.Sheen.SheenColorFactor[0], GltfMaterial.Sheen.SheenColorFactor[1], GltfMaterial.Sheen.SheenColorFactor[2]));
+
+			//SheenRoughnessTexture
+			//SheenRoughnessTexture_OffsetScale
+			//SheenRoughnessTexture_Rotation
+			//SheenRoughnessTexture_TexCoord
+			SetMap(GltfMaterial.Sheen.SheenRoughnessMap, Parameters::SheenRoughnessTexture);
+
+			//SheenRoughnessFactor
+			MaterialInstanceNode->AddScalarParameterValue(Parameters::SheenRoughnessFactor, GltfMaterial.Sheen.SheenRoughnessFactor);
+		}
+	}
+
+	if (GltfMaterial.bHasTransmission)
+	{
+		{
+			//TransmissionTexture
+			//TransmissionTexture_OffsetScale
+			//TransmissionTexture_Rotation
+			//TransmissionTexture_TexCoord
+			SetMap(GltfMaterial.Transmission.TransmissionMap, Parameters::TransmissionTexture);
+
+			//TransmissionFactor
+			MaterialInstanceNode->AddScalarParameterValue(Parameters::TransmissionFactor, GltfMaterial.Transmission.TransmissionFactor);
+		}
+	}
+
+	return true;
+}
+
+bool FGLTFMaterialInstanceSettings::AddGltfMaterialValuesToShaderGraphNode(const GLTF::FMaterial& GltfMaterial, const TArray<GLTF::FTexture>& Textures, UInterchangeShaderGraphNode* ShaderGraphNode) const
+{
+	FString ParentIdentifier = GetIdentifier(GetShadingModel(GltfMaterial.ShadingModel, GltfMaterial.bHasClearCoat, GltfMaterial.bHasSheen, GltfMaterial.bIsUnlitShadingModel, GltfMaterial.bHasTransmission), GltfMaterial.AlphaMode, GltfMaterial.bIsDoubleSided);
+	FString Parent;
+	if (const FSoftObjectPath* ObjectPath = MaterialParents.Find(ParentIdentifier))
+	{
+		Parent = ObjectPath->GetAssetPathString();
+	}
+	else
+	{
+		return false;
+	}
+
+	ShaderGraphNode->AddStringAttribute(*(GltfMaterialAttributeIdentifier + TEXT("ParentIdentifier")), Parent);
+
+	auto SetMap = [=](const GLTF::FTextureMap& TextureMap, const FString& Name)
+	{
+		if (Textures.IsValidIndex(TextureMap.TextureIndex))
+		{
+			ShaderGraphNode->AddStringAttribute(*(GltfMaterialAttributeIdentifier + Name), UInterchangeTextureNode::MakeNodeUid(Textures[TextureMap.TextureIndex].UniqueId));
+
+			//TexCoord decision in MaterialInstances work based on a Switch Node
+			//Currently Only supports UV0 and UV1
+			// [0...1) -> UV0
+			// [1...2) -> UV1
+			// else    -> UV0 (defaults to 0)
+			ShaderGraphNode->AddFloatAttribute(*(GltfMaterialAttributeIdentifier + Name + TEXT("_TexCoord")), TextureMap.TexCoord);
+		}
+
+		if (TextureMap.bHasTextureTransform)
+		{
+			FVector4f OffsetScale(TextureMap.TextureTransform.Offset[0], TextureMap.TextureTransform.Offset[1], TextureMap.TextureTransform.Scale[0], TextureMap.TextureTransform.Scale[1]);
+			ShaderGraphNode->AddLinearColorAttribute(*(GltfMaterialAttributeIdentifier + Name + TEXT("_OffsetScale")), OffsetScale);
+
+			ShaderGraphNode->AddFloatAttribute(*(GltfMaterialAttributeIdentifier + Name + TEXT("_Rotation")), TextureMap.TextureTransform.Rotation);
+		}
+	};
+
+	if (GltfMaterial.bIsUnlitShadingModel)
+	{
+		//BaseColorTexture
+		//BaseColorTexture_OffsetScale
+		//BaseColorTExture_Rotation
+		//BaseColorTExture_TexCoord
+		SetMap(GltfMaterial.BaseColor, Parameters::BaseColorTexture);
+
+		//BaseColorFactor
+		ShaderGraphNode->AddLinearColorAttribute(*(GltfMaterialAttributeIdentifier + Parameters::BaseColorFactor), GltfMaterial.BaseColorFactor);
+
+		return true;
+	}
+
+	if (GltfMaterial.AlphaMode == GLTF::FMaterial::EAlphaMode::Mask)
+	{
+		//AlphaCutoff
+		ShaderGraphNode->AddFloatAttribute(*(GltfMaterialAttributeIdentifier + Parameters::AlphaCutoff), GltfMaterial.AlphaCutoff);
+	}
+
+	if (GltfMaterial.ShadingModel == GLTF::FMaterial::EShadingModel::MetallicRoughness)
+	{
+		{
+			//BaseColorTexture
+			//BaseColorTexture_OffsetScale
+			//BaseColorTexture_Rotation
+			//BaseColorTexture_TexCoord
+			SetMap(GltfMaterial.BaseColor, Parameters::BaseColorTexture);
+
+			//BaseColorFactor
+			ShaderGraphNode->AddLinearColorAttribute(*(GltfMaterialAttributeIdentifier + Parameters::BaseColorFactor), GltfMaterial.BaseColorFactor);
+		}
+
+		{
+			//MetallicRoughnessTexture
+			//MetallicRoughnessTexture_OffsetScale
+			//MetallicRoughnessTexture_Rotation
+			//MetallicRoughnessTexture_TexCoord
+			SetMap(GltfMaterial.MetallicRoughness.Map, Parameters::MetallicRoughnessTexture);
+
+			//MetallicFactor
+			ShaderGraphNode->AddFloatAttribute(*(GltfMaterialAttributeIdentifier + Parameters::MetallicFactor), GltfMaterial.MetallicRoughness.MetallicFactor);
+
+			//RoughnessFactor
+			ShaderGraphNode->AddFloatAttribute(*(GltfMaterialAttributeIdentifier + Parameters::RoughnessFactor), GltfMaterial.MetallicRoughness.RoughnessFactor);
+		}
+
+		if (GltfMaterial.bHasSpecular)
+		{
+			//SpecularTexture
+			//SpecularTexture_OffsetScale
+			//SpecularTexture_Rotation
+			//SpecularTexture_TexCoord
+			SetMap(GltfMaterial.Specular.SpecularMap, Parameters::SpecularTexture);
+
+			//SpecularFactor
+			ShaderGraphNode->AddFloatAttribute(*(GltfMaterialAttributeIdentifier + Parameters::SpecularFactor), GltfMaterial.Specular.SpecularFactor);
+		}
+	}
+	else if (GltfMaterial.ShadingModel == GLTF::FMaterial::EShadingModel::SpecularGlossiness)
+	{
+		//////////////////////
+		//specular glossiness:
+		//////////////////////
+
+		{
+			//DiffuseTexture
+			//DiffuseTexture_OffsetScale
+			//DiffuseTexture_Rotation
+			//DiffuseTexture_TexCoord
+			SetMap(GltfMaterial.BaseColor, Parameters::DiffuseTexture);
+
+			//DiffuseFactor
+			ShaderGraphNode->AddLinearColorAttribute(*(GltfMaterialAttributeIdentifier + Parameters::DiffuseFactor), GltfMaterial.BaseColorFactor);
+		}
+
+		{
+			//SpecularGlossinessTexture
+			//SpecularGlossinessTexture_OffsetScale
+			//SpecularGlossinessTexture_Rotation
+			//SpecularGlossinessTexture_TexCoord
+			SetMap(GltfMaterial.SpecularGlossiness.Map, Parameters::SpecularGlossinessTexture);
+
+			//SpecFactor
+			ShaderGraphNode->AddLinearColorAttribute(*(GltfMaterialAttributeIdentifier + Parameters::SpecFactor), FVector4f(GltfMaterial.SpecularGlossiness.SpecularFactor[0],
+																												GltfMaterial.SpecularGlossiness.SpecularFactor[1],
+																												GltfMaterial.SpecularGlossiness.SpecularFactor[2]));
+
+			//GlossinessFactor
+			ShaderGraphNode->AddFloatAttribute(*(GltfMaterialAttributeIdentifier + Parameters::GlossinessFactor), GltfMaterial.SpecularGlossiness.GlossinessFactor);
+		}
+	}
+
+	{
+		//NormalTexture
+		//NormalTexture_OffsetScale
+		//NormalTexture_Rotation
+		//NormalTexture_TexCoord
+		SetMap(GltfMaterial.Normal, Parameters::NormalTexture);
+
+		//NormalScale
+		ShaderGraphNode->AddFloatAttribute(*(GltfMaterialAttributeIdentifier + Parameters::NormalScale), GltfMaterial.NormalScale);
+	}
+
+	{
+		//EmissiveTexture
+		//EmissiveTexture_OffsetScale
+		//EmissiveTexture_Rotation
+		//EmissiveTexture_TexCoord
+		SetMap(GltfMaterial.Emissive, Parameters::EmissiveTexture);
+
+		//EmissiveFactor
+		ShaderGraphNode->AddLinearColorAttribute(*(GltfMaterialAttributeIdentifier + Parameters::EmissiveFactor), GltfMaterial.EmissiveFactor);
+
+		//EmissiveStrength
+		ShaderGraphNode->AddFloatAttribute(*(GltfMaterialAttributeIdentifier + Parameters::EmissiveStrength), GltfMaterial.bHasEmissiveStrength ? GltfMaterial.EmissiveStrength : 1.0f);
+	}
+
+	{
+		//OcclusionTexture
+		//OcclusionTexture_OffsetScale
+		//OcclusionTexture_Rotation
+		//OcclusionTexture_TexCoord
+		SetMap(GltfMaterial.Occlusion, Parameters::OcclusionTexture);
+
+		//OcclusionStrength
+		ShaderGraphNode->AddFloatAttribute(*(GltfMaterialAttributeIdentifier + Parameters::OcclusionStrength), GltfMaterial.OcclusionStrength);
+	}
+
+	if (GltfMaterial.bHasIOR)
+	{
+		//ior
+		ShaderGraphNode->AddFloatAttribute(*(GltfMaterialAttributeIdentifier + Parameters::IOR), GltfMaterial.IOR);
+	}
+
+	if (GltfMaterial.bHasClearCoat)
+	{
+		{
+			//ClearCoatTexture
+			//ClearCoatTexture_OffsetScale
+			//ClearCoatTexture_Rotation
+			//ClearCoatTexture_TexCoord
+			SetMap(GltfMaterial.ClearCoat.ClearCoatMap, Parameters::ClearCoatTexture);
+
+			//ClearCoatFactor
+			ShaderGraphNode->AddFloatAttribute(*(GltfMaterialAttributeIdentifier + Parameters::ClearCoatFactor), GltfMaterial.ClearCoat.ClearCoatFactor);
+		}
+
+		{
+			//ClearCoatRoughnessTexture
+			//ClearCoatRoughnessTexture_OffsetScale
+			//ClearCoatRoughnessTexture_Rotation
+			//ClearCoatRoughnessTexture_TexCoord
+			SetMap(GltfMaterial.ClearCoat.RoughnessMap, Parameters::ClearCoatRoughnessTexture);
+
+			//ClearCoatRoughnessFactor
+			ShaderGraphNode->AddFloatAttribute(*(GltfMaterialAttributeIdentifier + Parameters::ClearCoatRoughnessFactor), GltfMaterial.ClearCoat.Roughness);
+		}
+
+		{
+			//ClearCoatNormalTexture
+			//ClearCoatNormalTexture_OffsetScale
+			//ClearCoatNormalTexture_Rotation
+			//ClearCoatNormalTexture_TexCoord
+			SetMap(GltfMaterial.ClearCoat.NormalMap, Parameters::ClearCoatNormalTexture);
+
+			//ClearCoatNormalFactor
+			ShaderGraphNode->AddFloatAttribute(*(GltfMaterialAttributeIdentifier + Parameters::ClearCoatNormalScale), GltfMaterial.ClearCoat.NormalMapUVScale);
+		}
+	}
+	else if (GltfMaterial.bHasSheen)
+	{
+		{
+			//SheenColorTexture
+			//SheenColorTexture_OffsetScale
+			//SheenColorTexture_Rotation
+			//SheenColorTexture_TexCoord
+			SetMap(GltfMaterial.Sheen.SheenColorMap, Parameters::SheenColorTexture);
+
+			//SheenColorFactor
+			ShaderGraphNode->AddLinearColorAttribute(*(GltfMaterialAttributeIdentifier + Parameters::SheenColorFactor), FVector4f(GltfMaterial.Sheen.SheenColorFactor[0], GltfMaterial.Sheen.SheenColorFactor[1], GltfMaterial.Sheen.SheenColorFactor[2]));
+
+			//SheenRoughnessTexture
+			//SheenRoughnessTexture_OffsetScale
+			//SheenRoughnessTexture_Rotation
+			//SheenRoughnessTexture_TexCoord
+			SetMap(GltfMaterial.Sheen.SheenRoughnessMap, Parameters::SheenRoughnessTexture);
+
+			//SheenRoughnessFactor
+			ShaderGraphNode->AddFloatAttribute(*(GltfMaterialAttributeIdentifier + Parameters::SheenRoughnessFactor), GltfMaterial.Sheen.SheenRoughnessFactor);
+		}
+	}
+
+	if (GltfMaterial.bHasTransmission)
+	{
+		{
+			//TransmissionTexture
+			//TransmissionTexture_OffsetScale
+			//TransmissionTexture_Rotation
+			//TransmissionTexture_TexCoord
+			SetMap(GltfMaterial.Transmission.TransmissionMap, Parameters::TransmissionTexture);
+
+			//TransmissionFactor
+			ShaderGraphNode->AddFloatAttribute(*(GltfMaterialAttributeIdentifier + Parameters::TransmissionFactor), GltfMaterial.Transmission.TransmissionFactor);
+		}
+	}
+
+	return true;
+}
+
+void FGLTFMaterialInstanceSettings::CreateMaterialInstanceFromShaderGraphNode(UInterchangeBaseNodeContainer& NodeContainer, UInterchangeShaderGraphNode* ShaderGraphNode)
+{
+	TArray<UE::Interchange::FAttributeKey> AttributeKeys;
+	ShaderGraphNode->GetAttributeKeys(AttributeKeys);
+
+	TMap<FString, UE::Interchange::FAttributeKey> GltfAttributeKeys;
+	for (const UE::Interchange::FAttributeKey& AttributeKey : AttributeKeys)
+	{
+		if (AttributeKey.ToString().Contains(GltfMaterialAttributeIdentifier))
+		{
+			GltfAttributeKeys.Add(AttributeKey.ToString().Replace(*GltfMaterialAttributeIdentifier, TEXT(""), ESearchCase::CaseSensitive), AttributeKey);
+		}
+	}
+	if (GltfAttributeKeys.Num() == 0)
+	{
+		return;
+	}
+
+	FString ParentIdentifier;
+	if (!ShaderGraphNode->GetStringAttribute(*(GltfMaterialAttributeIdentifier + TEXT("ParentIdentifier")), ParentIdentifier))
+	{
+		return;
+	}
+	GltfAttributeKeys.Remove(TEXT("ParentIdentifier"));
+
+	UInterchangeMaterialInstanceNode* MaterialInstanceNode = NewObject<UInterchangeMaterialInstanceNode>(&NodeContainer);
+
+	//TODO: once the GltfPipeline is in place we will need to make sure that
+	// if the MaterialInstance is desired, that the MaterialInstance UID are used across the Meshes, instead of the MaterialUIDs (ShaderGraphNode's UIDs)
+	// For now, for testing, for this scenario the ShaderGraphNode Uid and Name have prefix "Gltf_MI_AttributeIdentifier_". 
+	FString MaterialNodeUid = ShaderGraphNode->GetUniqueID();
+	FString MaterialNodeName = ShaderGraphNode->GetDisplayLabel();
+	MaterialNodeUid = MaterialNodeUid.Replace(*GltfMaterialAttributeIdentifier, TEXT(""), ESearchCase::CaseSensitive);
+	MaterialNodeName = MaterialNodeName.Replace(*GltfMaterialAttributeIdentifier, TEXT(""), ESearchCase::CaseSensitive);
+	
+	MaterialInstanceNode->InitializeNode(MaterialNodeUid, MaterialNodeName, EInterchangeNodeContainerType::TranslatedAsset);
+	NodeContainer.AddNode(MaterialInstanceNode);
+
+	MaterialInstanceNode->SetCustomParent(ParentIdentifier);
+
+	for (const TPair<FString, UE::Interchange::FAttributeKey>& GltfAttributeKey : GltfAttributeKeys)
+	{
+		UE::Interchange::EAttributeTypes AttributeType = ShaderGraphNode->GetAttributeType(GltfAttributeKey.Value);
+
+		//we are only using 3 attribute types:
+		switch (AttributeType)
+		{
+			case UE::Interchange::EAttributeTypes::Float:
+				{
+					float Value;
+					if (ShaderGraphNode->GetFloatAttribute(GltfAttributeKey.Value.Key, Value))
+					{
+						MaterialInstanceNode->AddScalarParameterValue(GltfAttributeKey.Key, Value);
+					}
+				}
+				break;
+			case UE::Interchange::EAttributeTypes::LinearColor:
+				{
+					FLinearColor Value;
+					if (ShaderGraphNode->GetLinearColorAttribute(GltfAttributeKey.Value.Key, Value))
+					{
+						MaterialInstanceNode->AddVectorParameterValue(GltfAttributeKey.Key, Value);
+					}
+				}
+				break;
+			case UE::Interchange::EAttributeTypes::String:
+				{
+					FString Value;
+					if (ShaderGraphNode->GetStringAttribute(GltfAttributeKey.Value.Key, Value))
+					{
+						MaterialInstanceNode->AddTextureParameterValue(GltfAttributeKey.Key, Value);
+					}
+				}
+				break;
+			default:
+				break;
+		}
+	}
+}
\ No newline at end of file
diff --git a/Engine/Plugins/Interchange/Runtime/Source/Import/Private/Gltf/InterchangeGltfTranslator.cpp b/Engine/Plugins/Interchange/Runtime/Source/Import/Private/Gltf/InterchangeGltfTranslator.cpp
--- a/Engine/Plugins/Interchange/Runtime/Source/Import/Private/Gltf/InterchangeGltfTranslator.cpp
+++ b/Engine/Plugins/Interchange/Runtime/Source/Import/Private/Gltf/InterchangeGltfTranslator.cpp
@@ -52,6 +52,13 @@ static FAutoConsoleVariableRef CCvarInterchangeEnableGLTFImport(
 	TEXT("Whether glTF support is enabled."),
 	ECVF_Default);
 
+static int GInterchangeGLTFUseMaterialInstancing = 0;
+static FAutoConsoleVariableRef CCvarInterchangeGLTFUseMaterialInstancing(
+	TEXT("Interchange.FeatureFlags.Import.GLTF.UseMaterialInstancing"),
+	GInterchangeGLTFUseMaterialInstancing,
+	TEXT("Whether glTF uses Material Instances."),
+	ECVF_Default);
+
 static const TArray<FString> ImporterSupportedExtensions = {
 	/* Lights */
 	GLTF::ToString(GLTF::EExtension::KHR_LightsPunctual),
@@ -514,15 +521,6 @@ void UInterchangeGltfTranslator::HandleGltfMaterial( UInterchangeBaseNodeContain
 
 	ShaderGraphNode.SetCustomTwoSided( GltfMaterial.bIsDoubleSided );
 
-	//Based on the gltf specification the basecolor and emissive textures have SRGB colors:
-	SetTextureSRGB(NodeContainer, GltfMaterial.BaseColor);
-	SetTextureSRGB(NodeContainer, GltfMaterial.Emissive);
-	//According to GLTF documentation the normal maps are right handed (following OpenGL convention),
-	//however UE expects left handed normal maps, this can be resolved by flipping the green channel of the normal textures:
-	//(based on [REDACTED_PROFILE_URL]/glTF-Sample-Models/tree/master/2.0/NormalTangentTest#problem-flipped-y-axis-or-flipped-green-channel)
-	SetTextureFlipGreenChannel(NodeContainer, GltfMaterial.Normal);
-	SetTextureFlipGreenChannel(NodeContainer, GltfMaterial.ClearCoat.NormalMap);
-
 	if (GltfMaterial.bIsUnlitShadingModel)
 	{
 		// Base Color
@@ -990,7 +988,7 @@ bool UInterchangeGltfTranslator::Translate( UInterchangeBaseNodeContainer& NodeC
 		SendAnalytics(TranslationResult::INPUT_FILE_NOTFOUND);
 		return false;
 	}
-
+	
 	GLTF::FFileReader GltfFileReader;
 
 	const bool bLoadImageData = false;
@@ -1120,11 +1118,74 @@ bool UInterchangeGltfTranslator::Translate( UInterchangeBaseNodeContainer& NodeC
 		int32 MaterialIndex = 0;
 		for ( const GLTF::FMaterial& GltfMaterial : GltfAsset.Materials )
 		{
-			UInterchangeShaderGraphNode* ShaderGraphNode = UInterchangeShaderGraphNode::Create(&NodeContainer, GltfMaterial.UniqueId);
-			ShaderGraphNode->SetDisplayLabel(GltfMaterial.Name);
+			//Based on the gltf specification the basecolor and emissive textures have SRGB colors:
+			SetTextureSRGB(NodeContainer, GltfMaterial.BaseColor, true);
+			SetTextureSRGB(NodeContainer, GltfMaterial.Emissive, true);
+			//Textures that are expected to use Scalar outputs we want to set them as SRGB false explicitly, based on UInterchangeGenericMaterialPipeline::HandleTextureNode
+			SetTextureSRGB(NodeContainer, GltfMaterial.MetallicRoughness.Map, false);
+			SetTextureSRGB(NodeContainer, GltfMaterial.Occlusion, false);
+			SetTextureSRGB(NodeContainer, GltfMaterial.ClearCoat.ClearCoatMap, false);
+			SetTextureSRGB(NodeContainer, GltfMaterial.ClearCoat.RoughnessMap, false);
+			SetTextureSRGB(NodeContainer, GltfMaterial.Transmission.TransmissionMap, false);
+
+			//According to GLTF documentation the normal maps are right handed (following OpenGL convention),
+			//however UE expects left handed normal maps, this can be resolved by flipping the green channel of the normal textures:
+			//(based on [REDACTED_PROFILE_URL]/glTF-Sample-Models/tree/master/2.0/NormalTangentTest#problem-flipped-y-axis-or-flipped-green-channel)
+			SetTextureFlipGreenChannel(NodeContainer, GltfMaterial.Normal);
+			SetTextureFlipGreenChannel(NodeContainer, GltfMaterial.ClearCoat.NormalMap);
+
+			const FString ShaderGraphNodeUid = UInterchangeShaderGraphNode::MakeNodeUid(GltfMaterial.UniqueId);
+			bool bUseVertexColor = MaterialsUsedOnMeshesWithVertexColor.Contains(ShaderGraphNodeUid);
+
+			if (FApp::IsGame() || GInterchangeGLTFUseMaterialInstancing == 1)
+			{
+				//Only create MaterialInstances
+				if (!MaterialInstanceSettings.ProcessGltfMaterial(NodeContainer, GltfMaterial, GltfAsset.Textures))
+				{
+					UE_LOG(LogInterchangeImport, Warning, TEXT("GLTF MaterialInstance Creation failed to find appropriate ParentMaterial. Material.UniqueID: %s"), *GltfMaterial.UniqueId);
+				}
+			}
+			else if (GInterchangeGLTFUseMaterialInstancing == 2)
+			{
+				//Create ShaderGraphNodes and convert to MaterialInstances
+				//Use the correct UIDs on MaterialInstances so that they get applied onto meshes
+				UInterchangeShaderGraphNode* ShaderGraphNode = UInterchangeShaderGraphNode::Create(&NodeContainer, TEXT("Gltf_MI_AttributeIdentifier_") + GltfMaterial.UniqueId);
+				ShaderGraphNode->SetDisplayLabel(TEXT("Gltf_MI_AttributeIdentifier_") + GltfMaterial.Name);
+
+				HandleGltfMaterial(NodeContainer, GltfMaterial, *ShaderGraphNode, bUseVertexColor);
+
+				if (!MaterialInstanceSettings.AddGltfMaterialValuesToShaderGraphNode(GltfMaterial, GltfAsset.Textures, ShaderGraphNode))
+				{
+					UE_LOG(LogInterchangeImport, Warning, TEXT("GLTF MaterialInstance Creation failed to find appropriate ParentMaterial. Material.UniqueID: %s"), *GltfMaterial.UniqueId);
+				}
+				else
+				{
+					FGLTFMaterialInstanceSettings::CreateMaterialInstanceFromShaderGraphNode(NodeContainer, ShaderGraphNode);
+				}
+			}
+			else if (GInterchangeGLTFUseMaterialInstancing == 3)
+			{
+				//Store Gltf Material Attribute values in ShaderGraphNode but do not create MaterialInstances
+				//Test in preparation for Gltf Pipeline:
+				UInterchangeShaderGraphNode* ShaderGraphNode = UInterchangeShaderGraphNode::Create(&NodeContainer, GltfMaterial.UniqueId);
+				ShaderGraphNode->SetDisplayLabel(GltfMaterial.Name);
 
-			bool bUseVertexColor = MaterialsUsedOnMeshesWithVertexColor.Contains(ShaderGraphNode->GetUniqueID());
-			HandleGltfMaterial( NodeContainer, GltfMaterial, *ShaderGraphNode, bUseVertexColor );
+				HandleGltfMaterial(NodeContainer, GltfMaterial, *ShaderGraphNode, bUseVertexColor);
+
+				if (!MaterialInstanceSettings.AddGltfMaterialValuesToShaderGraphNode(GltfMaterial, GltfAsset.Textures, ShaderGraphNode))
+				{
+					UE_LOG(LogInterchangeImport, Warning, TEXT("GLTF MaterialInstance Creation failed to find appropriate ParentMaterial. Material.UniqueID: %s"), *GltfMaterial.UniqueId);
+				}
+			}
+			else
+			{
+				//Only create ShaderGraphNodes:
+				UInterchangeShaderGraphNode* ShaderGraphNode = UInterchangeShaderGraphNode::Create(&NodeContainer, GltfMaterial.UniqueId);
+				ShaderGraphNode->SetDisplayLabel(GltfMaterial.Name);
+
+				HandleGltfMaterial(NodeContainer, GltfMaterial, *ShaderGraphNode, bUseVertexColor);
+			}
+			
 			++MaterialIndex;
 		}
 	}
@@ -1704,14 +1765,14 @@ void UInterchangeGltfTranslator::HandleGltfAnimation(UInterchangeBaseNodeContain
 	NodeContainer.AddNode(TrackSetNode);
 }
 
-void UInterchangeGltfTranslator::SetTextureSRGB(UInterchangeBaseNodeContainer& NodeContainer, const GLTF::FTextureMap& TextureMap) const
+void UInterchangeGltfTranslator::SetTextureSRGB(UInterchangeBaseNodeContainer& NodeContainer, const GLTF::FTextureMap& TextureMap, bool bSRGB) const
 {
 	if (GltfAsset.Textures.IsValidIndex(TextureMap.TextureIndex))
 	{
 		const FString TextureUid = UInterchangeTextureNode::MakeNodeUid(GltfAsset.Textures[TextureMap.TextureIndex].UniqueId);
 		if (UInterchangeTextureNode* TextureNode = const_cast<UInterchangeTextureNode*>(Cast<UInterchangeTextureNode>(NodeContainer.GetNode(TextureUid))))
 		{
-			TextureNode->SetCustomSRGB(true);
+			TextureNode->SetCustomSRGB(bSRGB);
 		}
 	}
 }
@@ -2173,4 +2234,19 @@ UInterchangeMeshNode* UInterchangeGltfTranslator::HandleGltfMesh(UInterchangeBas
 	return MeshNode;
 }
 
+UInterchangeGltfTranslator::UInterchangeGltfTranslator()
+{
+	if (!HasAllFlags(RF_ClassDefaultObject))
+	{
+		if (FApp::IsGame() || GInterchangeGLTFUseMaterialInstancing)
+		{
+			TArray<FString> MaterialInstanceIssues = MaterialInstanceSettings.ValidateMaterialInstancesAndParameters();
+			for (const FString& MaterialInstanceIssue : MaterialInstanceIssues)
+			{
+				UE_LOG(LogInterchangeImport, Warning, TEXT("%s"), *MaterialInstanceIssue);
+			}
+		}
+	}
+}
+
 #undef LOCTEXT_NAMESPACE
\ No newline at end of file
diff --git a/Engine/Plugins/Interchange/Runtime/Source/Import/Public/Gltf/InterchangeGLTFMaterialInstance.h b/Engine/Plugins/Interchange/Runtime/Source/Import/Public/Gltf/InterchangeGLTFMaterialInstance.h
new file mode 100644
--- /dev/null
+++ b/Engine/Plugins/Interchange/Runtime/Source/Import/Public/Gltf/InterchangeGLTFMaterialInstance.h
@@ -0,0 +1,44 @@
+// Copyright Epic Games, Inc. All Rights Reserved.
+#pragma once
+
+#include "CoreMinimal.h"
+
+#include "InterchangeGLTFMaterialInstance.generated.h"
+
+namespace GLTF { 
+	struct FMaterial; 
+	struct FTexture;
+}
+class UInterchangeBaseNodeContainer;
+class UInterchangeMaterialInstanceNode;
+class UInterchangeShaderGraphNode;
+
+// Note: Material Instances have the following known issues:
+// - Translucent Unlit material config is not supported
+// - Multi UV not supported
+// - ClearCoat.NormalMapUVScale not supported
+
+USTRUCT(BlueprintType)
+struct FGLTFMaterialInstanceSettings
+{
+	GENERATED_USTRUCT_BODY()
+
+public:
+	FGLTFMaterialInstanceSettings();
+
+	UPROPERTY(EditAnywhere, BlueprintReadWrite, Category = "GLTF Material", meta = (AllowedClasses = "/Script/Engine.MaterialFunction"))
+	TMap<FString, FSoftObjectPath> MaterialParents;
+
+	TArray<FString> ValidateMaterialInstancesAndParameters() const; //At the moment this is not used, because MaterialParents override is not exposed just yet.
+
+	bool ProcessGltfMaterial(UInterchangeBaseNodeContainer& NodeContainer, const GLTF::FMaterial& GltfMaterial, const TArray<GLTF::FTexture>& Textures) const;
+
+	bool AddGltfMaterialValuesToShaderGraphNode(const GLTF::FMaterial& GltfMaterial, const TArray<GLTF::FTexture>& Textures, UInterchangeShaderGraphNode* ShaderGraphNode) const;
+
+	static void CreateMaterialInstanceFromShaderGraphNode(UInterchangeBaseNodeContainer& NodeContainer, UInterchangeShaderGraphNode* ShaderGraphNode);
+
+private:
+	TSet<FString> GenerateExpectedParametersList(const FString& Identifier) const;
+
+	TArray<FString> ExpectedMaterialInstanceIdentifiers; //Default MaterialInstance' identifiers
+};
\ No newline at end of file
diff --git a/Engine/Plugins/Interchange/Runtime/Source/Import/Public/Gltf/InterchangeGltfTranslator.h b/Engine/Plugins/Interchange/Runtime/Source/Import/Public/Gltf/InterchangeGltfTranslator.h
--- a/Engine/Plugins/Interchange/Runtime/Source/Import/Public/Gltf/InterchangeGltfTranslator.h
+++ b/Engine/Plugins/Interchange/Runtime/Source/Import/Public/Gltf/InterchangeGltfTranslator.h
@@ -12,6 +12,7 @@
 #include "Texture/InterchangeTexturePayloadInterface.h"
 #include "UObject/Object.h"
 #include "UObject/ObjectMacros.h"
+#include "Gltf/InterchangeGLTFMaterialInstance.h"
 
 #include "InterchangeGltfTranslator.generated.h"
 
@@ -32,6 +33,8 @@ class UInterchangeGltfTranslator : public UInterchangeTranslatorBase
 	GENERATED_BODY()
 
 public:
+	UInterchangeGltfTranslator();
+
 	/** Begin UInterchangeTranslatorBase API*/
 	virtual EInterchangeTranslatorType GetTranslatorType() const override;
 	virtual EInterchangeTranslatorAssetType GetSupportedAssetTypes() const override;
@@ -59,6 +62,9 @@ public:
 	virtual TFuture<TOptional<UE::Interchange::FVariantSetPayloadData>> GetVariantSetPayloadData(const FString& PayloadKey) const override;
 	/* IInterchangeVariantSetPayloadInterface End */
 
+	UPROPERTY(EditAnywhere, BlueprintReadWrite, Category = "GLTF", DisplayName = "Material Settings")
+	FGLTFMaterialInstanceSettings MaterialInstanceSettings;
+
 protected:
 	using FNodeUidMap = TMap<const GLTF::FNode*, FString>;
 
@@ -86,7 +92,7 @@ protected:
 	bool GetVariantSetPayloadData(UE::Interchange::FVariantSetPayloadData& PayloadData) const;
 
 private:
-	void SetTextureSRGB(UInterchangeBaseNodeContainer& NodeContainer, const GLTF::FTextureMap& TextureMap) const;
+	void SetTextureSRGB(UInterchangeBaseNodeContainer& NodeContainer, const GLTF::FTextureMap& TextureMap, bool bSRGB) const;
 	void SetTextureFlipGreenChannel(UInterchangeBaseNodeContainer& NodeContainer, const GLTF::FTextureMap& TextureMap) const;
 
 	GLTF::FAsset GltfAsset;
diff --git a/Engine/Plugins/Interchange/Runtime/Source/Pipelines/Private/InterchangeGenericMaterialPipeline.cpp b/Engine/Plugins/Interchange/Runtime/Source/Pipelines/Private/InterchangeGenericMaterialPipeline.cpp
--- a/Engine/Plugins/Interchange/Runtime/Source/Pipelines/Private/InterchangeGenericMaterialPipeline.cpp
+++ b/Engine/Plugins/Interchange/Runtime/Source/Pipelines/Private/InterchangeGenericMaterialPipeline.cpp
@@ -19,6 +19,7 @@
 #include "Materials/Material.h"
 #include "Misc/PackageName.h"
 #include "Nodes/InterchangeBaseNode.h"
+#include "InterchangeMaterialInstanceNode.h"
 
 #include "Materials/MaterialExpressionAdd.h"
 #include "Materials/MaterialExpressionComponentMask.h"
@@ -573,13 +574,11 @@ void UInterchangeGenericMaterialPipeline::ExecutePipeline(UInterchangeBaseNodeCo
 		SourceNode->GetCustomImportUnusedMaterial(bImportUnusedMaterial);
 	}
 
-#if !WITH_EDITOR
 	// Can't import materials at runtime, fallback to instances
-	if (MaterialImport == EInterchangeMaterialImportOption::ImportAsMaterials)
+	if (FApp::IsGame() && MaterialImport == EInterchangeMaterialImportOption::ImportAsMaterials)
 	{
 		MaterialImport = EInterchangeMaterialImportOption::ImportAsMaterialInstances;
 	}
-#endif // !WITH_EDITOR
 
 	if (MaterialImport == EInterchangeMaterialImportOption::ImportAsMaterials)
 	{
@@ -613,6 +612,111 @@ void UInterchangeGenericMaterialPipeline::ExecutePipeline(UInterchangeBaseNodeCo
 			}
 		}
 	}
+
+	TArray<UInterchangeMaterialInstanceNode*> MaterialInstanceNodes;
+	BaseNodeContainer->IterateNodesOfType<UInterchangeMaterialInstanceNode>([&MaterialInstanceNodes](const FString& NodeUid, UInterchangeMaterialInstanceNode* MaterialNode)
+		{
+			MaterialInstanceNodes.Add(MaterialNode);
+		});
+
+	for (UInterchangeMaterialInstanceNode* MaterialNode : MaterialInstanceNodes)
+	{
+		FString ParentPath;
+
+		if (!MaterialNode->GetCustomParent(ParentPath) || ParentPath.IsEmpty())
+		{
+			continue;
+		}
+
+
+		UInterchangeMaterialInstanceFactoryNode* MaterialFactoryNode = nullptr;
+		FString DisplayLabel = MaterialNode->GetDisplayLabel();
+		const FString NodeUid = UInterchangeFactoryBaseNode::BuildFactoryNodeUid(MaterialNode->GetUniqueID());
+		if (BaseNodeContainer->IsNodeUidValid(NodeUid))
+		{
+			//The node already exist, just return it
+			MaterialFactoryNode = Cast<UInterchangeMaterialInstanceFactoryNode>(BaseNodeContainer->GetFactoryNode(NodeUid));
+			if (!MaterialFactoryNode)
+			{
+				continue;
+			}
+		}
+		else
+		{
+			MaterialFactoryNode = NewObject<UInterchangeMaterialInstanceFactoryNode>(BaseNodeContainer);
+			if (!ensure(MaterialFactoryNode))
+			{
+				MaterialFactoryNode = nullptr;
+			}
+			//Creating a Material
+			MaterialFactoryNode->InitializeNode(NodeUid, DisplayLabel, EInterchangeNodeContainerType::FactoryData);
+
+			BaseNodeContainer->AddNode(MaterialFactoryNode);
+			MaterialFactoryNodes.Add(MaterialFactoryNode);
+			MaterialFactoryNode->AddTargetNodeUid(MaterialNode->GetUniqueID());
+			MaterialNode->AddTargetNodeUid(MaterialFactoryNode->GetUniqueID());
+		}
+
+
+		// Set MaterialFactoryNode's display label to MaterialNode's uniqueID
+		// to reconcile mesh's slot names and material assets
+		MaterialFactoryNode->SetDisplayLabel(MaterialNode->GetAssetName());
+		MaterialFactoryNode->SetCustomParent(ParentPath);
+
+		const UClass* MaterialClass = FApp::IsGame() ? UMaterialInstanceDynamic::StaticClass() : UMaterialInstanceConstant::StaticClass();
+		MaterialFactoryNode->SetCustomInstanceClassName(MaterialClass->GetPathName());
+
+		TArray<FString> Inputs;
+		UInterchangeShaderPortsAPI::GatherInputs(MaterialNode, Inputs);
+
+		for (const FString& InputName : Inputs)
+		{
+			FName InputValueKey = UInterchangeShaderPortsAPI::MakeInputValueKey(InputName);
+
+			switch (UInterchangeShaderPortsAPI::GetInputType(MaterialNode, InputName))
+			{
+			case UE::Interchange::EAttributeTypes::Bool:
+			{
+				bool AttributeValue = false;
+				MaterialNode->GetBooleanAttribute(InputValueKey, AttributeValue);
+				MaterialFactoryNode->AddBooleanAttribute(InputValueKey, AttributeValue);
+			}
+			break;
+			case UE::Interchange::EAttributeTypes::Int32:
+			{
+				int32 AttributeValue = 0;
+				MaterialNode->GetInt32Attribute(InputValueKey, AttributeValue);
+				MaterialFactoryNode->AddInt32Attribute(InputValueKey, AttributeValue);
+			}
+			break;
+			case UE::Interchange::EAttributeTypes::Float:
+			{
+				float AttributeValue = 0.f;
+				MaterialNode->GetFloatAttribute(InputValueKey, AttributeValue);
+				MaterialFactoryNode->AddFloatAttribute(InputValueKey, AttributeValue);
+			}
+			break;
+			case UE::Interchange::EAttributeTypes::LinearColor:
+			{
+				FLinearColor AttributeValue = FLinearColor::White;
+				MaterialNode->GetLinearColorAttribute(InputValueKey, AttributeValue);
+				MaterialFactoryNode->AddLinearColorAttribute(InputValueKey, AttributeValue);
+			}
+			break;
+			case UE::Interchange::EAttributeTypes::String:
+			{
+				FString TextureUid;
+				MaterialNode->GetStringAttribute(InputValueKey, TextureUid);
+
+				FString FactoryTextureUid = UInterchangeFactoryBaseNode::BuildFactoryNodeUid(TextureUid);
+
+				MaterialFactoryNode->AddStringAttribute(InputValueKey, FactoryTextureUid);
+				MaterialFactoryNode->AddFactoryDependencyUid(FactoryTextureUid);
+			}
+			break;
+			}
+		}
+	}
 }
 
 void UInterchangeGenericMaterialPipeline::ExecutePostFactoryPipeline(const UInterchangeBaseNodeContainer* InBaseNodeContainer, const FString& NodeKey, UObject* CreatedAsset, bool bIsAReimport)
