Change Attribute key to use FString instead of FName since there is a lot of attribute key that are auto generated and it will flood the FName array.

#rnx

diff --git a/Engine/Plugins/Experimental/Enterprise/DatasmithInterchange/Source/DatasmithInterchange/Private/InterchangeDatasmithMaterialNode.cpp b/Engine/Plugins/Experimental/Enterprise/DatasmithInterchange/Source/DatasmithInterchange/Private/InterchangeDatasmithMaterialNode.cpp
--- a/Engine/Plugins/Experimental/Enterprise/DatasmithInterchange/Source/DatasmithInterchange/Private/InterchangeDatasmithMaterialNode.cpp
+++ b/Engine/Plugins/Experimental/Enterprise/DatasmithInterchange/Source/DatasmithInterchange/Private/InterchangeDatasmithMaterialNode.cpp
@@ -30,19 +30,21 @@
 #include "UObject/ObjectRedirector.h"
 #include "InterchangeMaterialInstanceNode.h"
 
-const FName UInterchangeDatasmithPbrMaterialNode::ShadingModelAttrName(TEXT("Datasmith:UEPbrMaterial:ShadingModel"));
-const FName UInterchangeDatasmithPbrMaterialNode::BlendModeAttrName(TEXT("Datasmith:UEPbrMaterial:BlendMode"));
-const FName UInterchangeDatasmithPbrMaterialNode::OpacityMaskClipValueAttrName(TEXT("Datasmith:UEPbrMaterial:OpacityMaskClipValue"));
-const FName UInterchangeDatasmithPbrMaterialNode::TranslucencyLightingModeAttrName(TEXT("Datasmith:UEPbrMaterial:TranslucencyLightingMode"));
+
+const FString UInterchangeDatasmithPbrMaterialNode::ShadingModelAttrName(TEXT("Datasmith:UEPbrMaterial:ShadingModel"));
+const FString UInterchangeDatasmithPbrMaterialNode::BlendModeAttrName(TEXT("Datasmith:UEPbrMaterial:BlendMode"));
+const FString UInterchangeDatasmithPbrMaterialNode::OpacityMaskClipValueAttrName(TEXT("Datasmith:UEPbrMaterial:OpacityMaskClipValue"));
+const FString UInterchangeDatasmithPbrMaterialNode::TranslucencyLightingModeAttrName(TEXT("Datasmith:UEPbrMaterial:TranslucencyLightingMode"));
 const FString UInterchangeDatasmithPbrMaterialNode::MaterialFunctionsDependenciesKey(TEXT("Datasmith:UEPbrMaterial:MaterialFunctionsDependencies"));
 
 namespace UE::DatasmithInterchange::MaterialUtils
 {
-	const FName DefaultOutputIndexAttrName(TEXT("Datasmith:MaterialExpression:DefaultOutputIndex"));
+	const FString MaterialFunctionPathAttrName(TEXT("Datasmith:Material:FunctionCall:MaterialFunctionPath"));
+	const FString DefaultOutputIndexAttrName(TEXT("Datasmith:MaterialExpression:DefaultOutputIndex"));
 
-	const FName MaterialTypeAttrName(TEXT("Datasmith:ReferenceMaterial:MaterialType"));
-	const FName MaterialQualityAttrName(TEXT("Datasmith:ReferenceMaterial:MaterialQuality"));
-	const FName MaterialParentAttrName(TEXT("Datasmith:ReferenceMaterial:MaterialParent"));
+	const FString MaterialTypeAttrName(TEXT("Datasmith:ReferenceMaterial:MaterialType"));
+	const FString MaterialQualityAttrName(TEXT("Datasmith:ReferenceMaterial:MaterialQuality"));
+	const FString MaterialParentAttrName(TEXT("Datasmith:ReferenceMaterial:MaterialParent"));
 
 	class FPbrMaterialHelper
 	{
@@ -237,7 +239,7 @@ namespace UE::DatasmithInterchange::MaterialUtils
 			const FString TextureUid = NodeUtils::TexturePrefix + FDatasmithUtils::SanitizeObjectName(DatasmithExpression.GetTexturePathName());
 			if (NodeContainer.GetNode(TextureUid))
 			{
-				const FName ValueKey = UInterchangeShaderPortsAPI::MakeInputValueKey(TextureSample::Inputs::Texture.ToString());
+				const FString ValueKey = UInterchangeShaderPortsAPI::MakeInputValueKey(TextureSample::Inputs::Texture.ToString());
 				ExpressionNode->AddStringAttribute(ValueKey, TextureUid);
 			}
 			else
@@ -276,17 +278,17 @@ namespace UE::DatasmithInterchange::MaterialUtils
 		ExpressionNode->SetCustomShaderType(TextureCoordinate::Name.ToString());
 
 		{
-			const FName ValueKey = UInterchangeShaderPortsAPI::MakeInputValueKey(TextureCoordinate::Inputs::Index.ToString());
+			const FString ValueKey = UInterchangeShaderPortsAPI::MakeInputValueKey(TextureCoordinate::Inputs::Index.ToString());
 			ExpressionNode->AddInt32Attribute(ValueKey, DatasmithExpression.GetCoordinateIndex());
 		}
 
 		{
-			const FName ValueKey = UInterchangeShaderPortsAPI::MakeInputValueKey(TextureCoordinate::Inputs::UTiling.ToString());
+			const FString ValueKey = UInterchangeShaderPortsAPI::MakeInputValueKey(TextureCoordinate::Inputs::UTiling.ToString());
 			ExpressionNode->AddFloatAttribute(ValueKey, DatasmithExpression.GetUTiling());
 		}
 
 		{
-			const FName ValueKey = UInterchangeShaderPortsAPI::MakeInputValueKey(TextureCoordinate::Inputs::VTiling.ToString());
+			const FString ValueKey = UInterchangeShaderPortsAPI::MakeInputValueKey(TextureCoordinate::Inputs::VTiling.ToString());
 			ExpressionNode->AddFloatAttribute(ValueKey, DatasmithExpression.GetVTiling());
 		}
 
@@ -569,19 +571,19 @@ namespace UE::DatasmithInterchange::MaterialUtils
 		if (MaterialExpression->IsSubType(EDatasmithMaterialExpressionType::ConstantBool))
 		{
 			const IDatasmithMaterialExpressionBool* BoolExpression = static_cast<const IDatasmithMaterialExpressionBool*>(MaterialExpression);
-			const FName InputName(UInterchangeShaderPortsAPI::MakeInputValueKey(ConnectionData.InputName));
+			const FString InputName(UInterchangeShaderPortsAPI::MakeInputValueKey(ConnectionData.InputName));
 			ConnectionData.InputNode->AddBooleanAttribute(InputName, BoolExpression->GetBool());
 		}
 		else if (MaterialExpression->IsSubType(EDatasmithMaterialExpressionType::ConstantScalar))
 		{
 			const IDatasmithMaterialExpressionScalar* ScalarExpression = static_cast<const IDatasmithMaterialExpressionScalar*>(MaterialExpression);
-			const FName InputName(UInterchangeShaderPortsAPI::MakeInputValueKey(ConnectionData.InputName));
+			const FString InputName(UInterchangeShaderPortsAPI::MakeInputValueKey(ConnectionData.InputName));
 			ConnectionData.InputNode->AddFloatAttribute(InputName, ScalarExpression->GetScalar());
 		}
 		else if (MaterialExpression->IsSubType(EDatasmithMaterialExpressionType::ConstantColor))
 		{
 			const IDatasmithMaterialExpressionColor* ColorExpression = static_cast<const IDatasmithMaterialExpressionColor*>(MaterialExpression);
-			const FName InputName(UInterchangeShaderPortsAPI::MakeInputValueKey(ConnectionData.InputName));
+			const FString InputName(UInterchangeShaderPortsAPI::MakeInputValueKey(ConnectionData.InputName));
 			ConnectionData.InputNode->AddLinearColorAttribute(InputName, ColorExpression->GetColor());
 		}
 		else if (MaterialExpression->IsSubType(EDatasmithMaterialExpressionType::Texture))
diff --git a/Engine/Plugins/Experimental/Enterprise/DatasmithInterchange/Source/DatasmithInterchange/Private/InterchangeDatasmithMaterialNode.h b/Engine/Plugins/Experimental/Enterprise/DatasmithInterchange/Source/DatasmithInterchange/Private/InterchangeDatasmithMaterialNode.h
--- a/Engine/Plugins/Experimental/Enterprise/DatasmithInterchange/Source/DatasmithInterchange/Private/InterchangeDatasmithMaterialNode.h
+++ b/Engine/Plugins/Experimental/Enterprise/DatasmithInterchange/Source/DatasmithInterchange/Private/InterchangeDatasmithMaterialNode.h
@@ -115,10 +115,10 @@ public:
 	}
 
 private:
-	static const FName ShadingModelAttrName;
-	static const FName BlendModeAttrName;
-	static const FName OpacityMaskClipValueAttrName;
-	static const FName TranslucencyLightingModeAttrName;
+	static const FString ShadingModelAttrName;
+	static const FString BlendModeAttrName;
+	static const FString OpacityMaskClipValueAttrName;
+	static const FString TranslucencyLightingModeAttrName;
 	static const FString MaterialFunctionsDependenciesKey;
 
 	UE::Interchange::TArrayAttributeHelper<FString> MaterialFunctionsDependencies;
@@ -126,11 +126,12 @@ private:
 
 namespace UE::DatasmithInterchange::MaterialUtils
 {
-	extern const FName DefaultOutputIndexAttrName;
+	extern const FString DefaultOutputIndexAttrName;
+	extern const FString MaterialFunctionPathAttrName;
 
-	extern const FName MaterialTypeAttrName;
-	extern const FName MaterialQualityAttrName;
-	extern const FName MaterialParentAttrName;
+	extern const FString MaterialTypeAttrName;
+	extern const FString MaterialQualityAttrName;
+	extern const FString MaterialParentAttrName;
 
 	void ProcessMaterialElements(TArray<TSharedPtr<IDatasmithBaseMaterialElement>>& InOutMaterialElements);
 	UInterchangeBaseNode* AddMaterialNode(const TSharedPtr<IDatasmithBaseMaterialElement>& MaterialElement, UInterchangeBaseNodeContainer& NodeContainer);
diff --git a/Engine/Plugins/Experimental/Enterprise/DatasmithInterchange/Source/DatasmithInterchange/Private/InterchangeDatasmithMaterialPipeline.cpp b/Engine/Plugins/Experimental/Enterprise/DatasmithInterchange/Source/DatasmithInterchange/Private/InterchangeDatasmithMaterialPipeline.cpp
--- a/Engine/Plugins/Experimental/Enterprise/DatasmithInterchange/Source/DatasmithInterchange/Private/InterchangeDatasmithMaterialPipeline.cpp
+++ b/Engine/Plugins/Experimental/Enterprise/DatasmithInterchange/Source/DatasmithInterchange/Private/InterchangeDatasmithMaterialPipeline.cpp
@@ -150,7 +150,7 @@ void UInterchangeDatasmithMaterialPipeline::PreImportMaterialNode(UInterchangeBa
 
 	for (const FString& InputName : Inputs)
 	{
-		FName InputValueKey = UInterchangeShaderPortsAPI::MakeInputValueKey(InputName);
+		FString InputValueKey = UInterchangeShaderPortsAPI::MakeInputValueKey(InputName);
 
 		switch (UInterchangeShaderPortsAPI::GetInputType(MaterialNode, InputName))
 		{
@@ -245,7 +245,7 @@ void UInterchangeDatasmithMaterialPipeline::PostImportMaterialInstanceFactoryNod
 
 	for (const FString& InputName : Inputs)
 	{
-		FName InputValueKey = UInterchangeShaderPortsAPI::MakeInputValueKey(InputName);
+		FString InputValueKey = UInterchangeShaderPortsAPI::MakeInputValueKey(InputName);
 
 		if (UInterchangeShaderPortsAPI::GetInputType(FactoryNode, InputName) == UE::Interchange::EAttributeTypes::String)
 		{
diff --git a/Engine/Plugins/Experimental/Enterprise/DatasmithInterchange/Source/DatasmithInterchange/Private/InterchangeDatasmithTextureData.cpp b/Engine/Plugins/Experimental/Enterprise/DatasmithInterchange/Source/DatasmithInterchange/Private/InterchangeDatasmithTextureData.cpp
--- a/Engine/Plugins/Experimental/Enterprise/DatasmithInterchange/Source/DatasmithInterchange/Private/InterchangeDatasmithTextureData.cpp
+++ b/Engine/Plugins/Experimental/Enterprise/DatasmithInterchange/Source/DatasmithInterchange/Private/InterchangeDatasmithTextureData.cpp
@@ -11,15 +11,15 @@
 
 namespace UE::DatasmithInterchange
 {
-	const FName FInterchangeDatasmithTextureDataConst::DatasmithTextureDataKey = FName(TEXT("DatasmithTextureData"));
-	const FName FInterchangeDatasmithTextureDataConst::FileKey = FName(TEXT("File"));
-	const FName FInterchangeDatasmithTextureDataConst::TextureModeKey = FName(TEXT("TextureMode"));
-	const FName FInterchangeDatasmithTextureDataConst::TextureFilterKey = FName(TEXT("TextureFilter"));
-	const FName FInterchangeDatasmithTextureDataConst::TextureAddressXKey = FName(TEXT("TextureAddressX"));
-	const FName FInterchangeDatasmithTextureDataConst::TextureAddressYKey = FName(TEXT("TextureAddressY"));
-	const FName FInterchangeDatasmithTextureDataConst::AllowResizeKey = FName(TEXT("AllowResize"));
-	const FName FInterchangeDatasmithTextureDataConst::RGBCurveKey = FName(TEXT("RGBCurve"));
-	const FName FInterchangeDatasmithTextureDataConst::SRGBKey = FName(TEXT("SRGB"));
+	const FString FInterchangeDatasmithTextureDataConst::DatasmithTextureDataKey = TEXT("DatasmithTextureData");
+	const FString FInterchangeDatasmithTextureDataConst::FileKey = TEXT("File");
+	const FString FInterchangeDatasmithTextureDataConst::TextureModeKey = TEXT("TextureMode");
+	const FString FInterchangeDatasmithTextureDataConst::TextureFilterKey = TEXT("TextureFilter");
+	const FString FInterchangeDatasmithTextureDataConst::TextureAddressXKey = TEXT("TextureAddressX");
+	const FString FInterchangeDatasmithTextureDataConst::TextureAddressYKey = TEXT("TextureAddressY");
+	const FString FInterchangeDatasmithTextureDataConst::AllowResizeKey = TEXT("AllowResize");
+	const FString FInterchangeDatasmithTextureDataConst::RGBCurveKey = TEXT("RGBCurve");
+	const FString FInterchangeDatasmithTextureDataConst::SRGBKey = TEXT("SRGB");
 
 	FInterchangeDatasmithTextureDataConst::FInterchangeDatasmithTextureDataConst(const UInterchangeBaseNode* InNode)
 		: ConstBaseNode(InNode)
diff --git a/Engine/Plugins/Experimental/Enterprise/DatasmithInterchange/Source/DatasmithInterchange/Private/InterchangeDatasmithTextureData.h b/Engine/Plugins/Experimental/Enterprise/DatasmithInterchange/Source/DatasmithInterchange/Private/InterchangeDatasmithTextureData.h
--- a/Engine/Plugins/Experimental/Enterprise/DatasmithInterchange/Source/DatasmithInterchange/Private/InterchangeDatasmithTextureData.h
+++ b/Engine/Plugins/Experimental/Enterprise/DatasmithInterchange/Source/DatasmithInterchange/Private/InterchangeDatasmithTextureData.h
@@ -39,15 +39,15 @@ namespace UE::DatasmithInterchange
 	protected:
 		const UInterchangeBaseNode* ConstBaseNode;
 
-		static const FName DatasmithTextureDataKey;
-		static const FName FileKey;
-		static const FName TextureModeKey;
-		static const FName TextureFilterKey;
-		static const FName TextureAddressXKey;
-		static const FName TextureAddressYKey;
-		static const FName AllowResizeKey;
-		static const FName RGBCurveKey;
-		static const FName SRGBKey;
+		static const FString DatasmithTextureDataKey;
+		static const FString FileKey;
+		static const FString TextureModeKey;
+		static const FString TextureFilterKey;
+		static const FString TextureAddressXKey;
+		static const FString TextureAddressYKey;
+		static const FString AllowResizeKey;
+		static const FString RGBCurveKey;
+		static const FString SRGBKey;
 	};
 
 	struct FInterchangeDatasmithTextureData : public FInterchangeDatasmithTextureDataConst
diff --git a/Engine/Plugins/Experimental/Enterprise/DatasmithInterchange/Source/DatasmithInterchange/Private/InterchangeDatasmithUtils.cpp b/Engine/Plugins/Experimental/Enterprise/DatasmithInterchange/Source/DatasmithInterchange/Private/InterchangeDatasmithUtils.cpp
--- a/Engine/Plugins/Experimental/Enterprise/DatasmithInterchange/Source/DatasmithInterchange/Private/InterchangeDatasmithUtils.cpp
+++ b/Engine/Plugins/Experimental/Enterprise/DatasmithInterchange/Source/DatasmithInterchange/Private/InterchangeDatasmithUtils.cpp
@@ -72,7 +72,7 @@ namespace UE::DatasmithInterchange::NodeUtils
 
 namespace UE::DatasmithInterchange::MeshUtils
 {
-	const FName MeshMaterialAttrName(TEXT("Datasmith:Mesh:MaterialOverride"));
+	const FString MeshMaterialAttrName(TEXT("Datasmith:Mesh:MaterialOverride"));
 }
 
 namespace UE::DatasmithInterchange::VariantSetUtils
diff --git a/Engine/Plugins/Experimental/Enterprise/DatasmithInterchange/Source/DatasmithInterchange/Private/InterchangeDatasmithUtils.h b/Engine/Plugins/Experimental/Enterprise/DatasmithInterchange/Source/DatasmithInterchange/Private/InterchangeDatasmithUtils.h
--- a/Engine/Plugins/Experimental/Enterprise/DatasmithInterchange/Source/DatasmithInterchange/Private/InterchangeDatasmithUtils.h
+++ b/Engine/Plugins/Experimental/Enterprise/DatasmithInterchange/Source/DatasmithInterchange/Private/InterchangeDatasmithUtils.h
@@ -148,7 +148,7 @@ namespace UE::DatasmithInterchange::NodeUtils
 
 namespace UE::DatasmithInterchange::MeshUtils
 {
-	extern const FName MeshMaterialAttrName;
+	extern const FString MeshMaterialAttrName;
 }
 
 namespace UE::DatasmithInterchange::AnimUtils
diff --git a/Engine/Plugins/Interchange/Runtime/Source/Nodes/Private/InterchangeShaderGraphNode.cpp b/Engine/Plugins/Interchange/Runtime/Source/Nodes/Private/InterchangeShaderGraphNode.cpp
--- a/Engine/Plugins/Interchange/Runtime/Source/Nodes/Private/InterchangeShaderGraphNode.cpp
+++ b/Engine/Plugins/Interchange/Runtime/Source/Nodes/Private/InterchangeShaderGraphNode.cpp
@@ -9,7 +9,7 @@ const TCHAR* UInterchangeShaderPortsAPI::InputPrefix = TEXT("Inputs");
 const TCHAR* UInterchangeShaderPortsAPI::InputSeparator = TEXT(":");
 const TCHAR* UInterchangeShaderPortsAPI::OutputByIndex = TEXT("__ByIndex__");
 
-FName UInterchangeShaderPortsAPI::MakeInputConnectionKey(const FString& InputName)
+FString UInterchangeShaderPortsAPI::MakeInputConnectionKey(const FString& InputName)
 {
 	TStringBuilder<128> StringBuilder;
 	StringBuilder.Append(InputPrefix);
@@ -18,10 +18,10 @@ FName UInterchangeShaderPortsAPI::MakeInputConnectionKey(const FString& InputNam
 	StringBuilder.Append(InputSeparator);
 	StringBuilder.Append(TEXT("Connect"));
 
-	return FName(StringBuilder.ToString());
+	return StringBuilder.ToString();
 }
 
-FName UInterchangeShaderPortsAPI::MakeInputValueKey(const FString& InputName)
+FString UInterchangeShaderPortsAPI::MakeInputValueKey(const FString& InputName)
 {
 	TStringBuilder<128> StringBuilder;
 	StringBuilder.Append(InputPrefix);
@@ -30,7 +30,7 @@ FName UInterchangeShaderPortsAPI::MakeInputValueKey(const FString& InputName)
 	StringBuilder.Append(InputSeparator);
 	StringBuilder.Append(TEXT("Value"));
 	
-	return FName(StringBuilder.ToString());
+	return StringBuilder.ToString();
 }
 
 FString UInterchangeShaderPortsAPI::MakeInputName(const FString& InputKey)
diff --git a/Engine/Plugins/Interchange/Runtime/Source/Nodes/Public/InterchangeShaderGraphNode.h b/Engine/Plugins/Interchange/Runtime/Source/Nodes/Public/InterchangeShaderGraphNode.h
--- a/Engine/Plugins/Interchange/Runtime/Source/Nodes/Public/InterchangeShaderGraphNode.h
+++ b/Engine/Plugins/Interchange/Runtime/Source/Nodes/Public/InterchangeShaderGraphNode.h
@@ -23,13 +23,13 @@ public:
 	 * Makes an attribute key to represent a node being connected to an input (ie: Inputs:InputName:Connect).
 	 */
 	UFUNCTION(BlueprintCallable, Category = "Interchange | Node | Material")
-	static FName MakeInputConnectionKey(const FString& InputName);
+	static FString MakeInputConnectionKey(const FString& InputName);
  
 	/**
 	 * Makes an attribute key to represent a value being given to an input (ie: Inputs:InputName:Value).
 	 */
 	UFUNCTION(BlueprintCallable, Category = "Interchange | Node | Material")
-	static FName MakeInputValueKey(const FString& InputName);
+	static FString MakeInputValueKey(const FString& InputName);
  
 	/**
 	 * From an attribute key associated with an input (ie: Inputs:InputName:Value), retrieves the input name from it.
diff --git a/Engine/Plugins/Interchange/Runtime/Source/Pipelines/Private/InterchangeGenericMaterialPipeline.cpp b/Engine/Plugins/Interchange/Runtime/Source/Pipelines/Private/InterchangeGenericMaterialPipeline.cpp
--- a/Engine/Plugins/Interchange/Runtime/Source/Pipelines/Private/InterchangeGenericMaterialPipeline.cpp
+++ b/Engine/Plugins/Interchange/Runtime/Source/Pipelines/Private/InterchangeGenericMaterialPipeline.cpp
@@ -924,8 +924,8 @@ bool UInterchangeGenericMaterialPipeline::HandlePhongModel(const UInterchangeSha
 
 					const float ShininessScale = 100.f; // Divide shininess by 100 to bring it into a 0-1 range for roughness.
 					const FName ShininessScaleParameterName = GET_MEMBER_NAME_CHECKED(UMaterialExpressionDivide, ConstB);
-					DivideShininessNode->AddFloatAttribute(ShininessScaleParameterName, ShininessScale);
-					DivideShininessNode->AddApplyAndFillDelegates<float>(ShininessScaleParameterName, UMaterialExpressionDivide::StaticClass(),  GET_MEMBER_NAME_CHECKED(UMaterialExpressionDivide, ConstB));
+					DivideShininessNode->AddFloatAttribute(ShininessScaleParameterName.ToString(), ShininessScale);
+					DivideShininessNode->AddApplyAndFillDelegates<float>(ShininessScaleParameterName.ToString(), UMaterialExpressionDivide::StaticClass(), GET_MEMBER_NAME_CHECKED(UMaterialExpressionDivide, ConstB));
 
 					// Connect Shininess to Divide
 					UInterchangeShaderPortsAPI::ConnectOuputToInputByName(DivideShininessNode, GET_MEMBER_NAME_CHECKED(UMaterialExpressionDivide, A).ToString(),
@@ -1421,8 +1421,8 @@ void UInterchangeGenericMaterialPipeline::HandleCommonParameters(const UIntercha
 
 				const float AirIOR = 1.f;
 				const FName ConstAMemberName = GET_MEMBER_NAME_CHECKED(UMaterialExpressionLinearInterpolate, ConstA);
-				IORLerp->AddFloatAttribute(ConstAMemberName, AirIOR);
-				IORLerp->AddApplyAndFillDelegates<float>(ConstAMemberName, UMaterialExpressionLinearInterpolate::StaticClass(), ConstAMemberName);
+				IORLerp->AddFloatAttribute(ConstAMemberName.ToString(), AirIOR);
+				IORLerp->AddApplyAndFillDelegates<float>(ConstAMemberName.ToString(), UMaterialExpressionLinearInterpolate::StaticClass(), ConstAMemberName);
 
 				UInterchangeShaderPortsAPI::ConnectOuputToInputByName(IORLerp, GET_MEMBER_NAME_CHECKED(UMaterialExpressionLinearInterpolate, B).ToString(),
 					ExpressionFactoryNode.Get<0>()->GetUniqueID(), ExpressionFactoryNode.Get<1>());
@@ -1783,8 +1783,8 @@ void UInterchangeGenericMaterialPipeline::HandleTextureCoordinateNode(const UInt
 		if (ShaderNode->GetInt32Attribute(UInterchangeShaderPortsAPI::MakeInputValueKey(Nodes::TextureCoordinate::Inputs::Index.ToString()), CoordIndex))
 		{
 			const FName CoordinateIndexMemberName = GET_MEMBER_NAME_CHECKED(UMaterialExpressionTextureCoordinate, CoordinateIndex);
-			TexCoordFactoryNode->AddInt32Attribute(CoordinateIndexMemberName, CoordIndex);
-			TexCoordFactoryNode->AddApplyAndFillDelegates<int32>(CoordinateIndexMemberName, UMaterialExpressionTextureCoordinate::StaticClass(), CoordinateIndexMemberName);
+			TexCoordFactoryNode->AddInt32Attribute(CoordinateIndexMemberName.ToString(), CoordIndex);
+			TexCoordFactoryNode->AddApplyAndFillDelegates<int32>(CoordinateIndexMemberName.ToString(), UMaterialExpressionTextureCoordinate::StaticClass(), CoordinateIndexMemberName);
 		}
 	}
 
@@ -1795,8 +1795,8 @@ void UInterchangeGenericMaterialPipeline::HandleTextureCoordinateNode(const UInt
 		if (UTilingValue.IsType<float>())
 		{
 			const FName UTilingMemberName = GET_MEMBER_NAME_CHECKED(UMaterialExpressionTextureCoordinate, UTiling);
-			TexCoordFactoryNode->AddFloatAttribute(UTilingMemberName, UTilingValue.Get<float>());
-			TexCoordFactoryNode->AddApplyAndFillDelegates<float>(UTilingMemberName, UMaterialExpressionTextureCoordinate::StaticClass(), UTilingMemberName);
+			TexCoordFactoryNode->AddFloatAttribute(UTilingMemberName.ToString(), UTilingValue.Get<float>());
+			TexCoordFactoryNode->AddApplyAndFillDelegates<float>(UTilingMemberName.ToString(), UMaterialExpressionTextureCoordinate::StaticClass(), UTilingMemberName);
 		}
 	}
 
@@ -1807,8 +1807,8 @@ void UInterchangeGenericMaterialPipeline::HandleTextureCoordinateNode(const UInt
 		if (VTilingValue.IsType<float>())
 		{
 			const FName VTilingMemberName = GET_MEMBER_NAME_CHECKED(UMaterialExpressionTextureCoordinate, VTiling);
-			TexCoordFactoryNode->AddFloatAttribute(VTilingMemberName, VTilingValue.Get<float>());
-			TexCoordFactoryNode->AddApplyAndFillDelegates<float>(VTilingMemberName, UMaterialExpressionTextureCoordinate::StaticClass(), VTilingMemberName);
+			TexCoordFactoryNode->AddFloatAttribute(VTilingMemberName.ToString(), VTilingValue.Get<float>());
+			TexCoordFactoryNode->AddApplyAndFillDelegates<float>(VTilingMemberName.ToString(), UMaterialExpressionTextureCoordinate::StaticClass(), VTilingMemberName);
 		}
 	}
 
@@ -1929,13 +1929,13 @@ void UInterchangeGenericMaterialPipeline::HandleMaskNode(const UInterchangeShade
 	MaskFactoryNode->SetCustomExpressionClassName(UMaterialExpressionComponentMask::StaticClass()->GetName());
 
 	bool bRChannel = false;
-	ShaderNode->GetBooleanAttribute(Mask::Attributes::R, bRChannel);
+	ShaderNode->GetBooleanAttribute(Mask::Attributes::R.ToString(), bRChannel);
 	bool bGChannel = false;
-	ShaderNode->GetBooleanAttribute(Mask::Attributes::G, bGChannel);
+	ShaderNode->GetBooleanAttribute(Mask::Attributes::G.ToString(), bGChannel);
 	bool bBChannel = false;
-	ShaderNode->GetBooleanAttribute(Mask::Attributes::B, bBChannel);
+	ShaderNode->GetBooleanAttribute(Mask::Attributes::B.ToString(), bBChannel);
 	bool bAChannel = false;
-	ShaderNode->GetBooleanAttribute(Mask::Attributes::A, bAChannel);
+	ShaderNode->GetBooleanAttribute(Mask::Attributes::A.ToString(), bAChannel);
 	bool bIsAnyMaskChannelSet = bRChannel || bGChannel || bBChannel || bAChannel;
 
 	if(bIsAnyMaskChannelSet)
@@ -1943,29 +1943,29 @@ void UInterchangeGenericMaterialPipeline::HandleMaskNode(const UInterchangeShade
 		// R
 		{
 			const FName RMemberName = GET_MEMBER_NAME_CHECKED(UMaterialExpressionComponentMask, R);
-			MaskFactoryNode->AddBooleanAttribute(RMemberName, bRChannel);
-			MaskFactoryNode->AddApplyAndFillDelegates<bool>(RMemberName, UMaterialExpressionComponentMask::StaticClass(), RMemberName);
+			MaskFactoryNode->AddBooleanAttribute(RMemberName.ToString(), bRChannel);
+			MaskFactoryNode->AddApplyAndFillDelegates<bool>(RMemberName.ToString(), UMaterialExpressionComponentMask::StaticClass(), RMemberName);
 		}
 
 		// G
 		{
 			const FName GMemberName = GET_MEMBER_NAME_CHECKED(UMaterialExpressionComponentMask, G);
-			MaskFactoryNode->AddBooleanAttribute(GMemberName, bGChannel);
-			MaskFactoryNode->AddApplyAndFillDelegates<bool>(GMemberName, UMaterialExpressionComponentMask::StaticClass(), GMemberName);
+			MaskFactoryNode->AddBooleanAttribute(GMemberName.ToString(), bGChannel);
+			MaskFactoryNode->AddApplyAndFillDelegates<bool>(GMemberName.ToString(), UMaterialExpressionComponentMask::StaticClass(), GMemberName);
 		}
 
 		// B
 		{
 			const FName BMemberName = GET_MEMBER_NAME_CHECKED(UMaterialExpressionComponentMask, B);
-			MaskFactoryNode->AddBooleanAttribute(BMemberName, bBChannel);
-			MaskFactoryNode->AddApplyAndFillDelegates<bool>(BMemberName, UMaterialExpressionComponentMask::StaticClass(), BMemberName);
+			MaskFactoryNode->AddBooleanAttribute(BMemberName.ToString(), bBChannel);
+			MaskFactoryNode->AddApplyAndFillDelegates<bool>(BMemberName.ToString(), UMaterialExpressionComponentMask::StaticClass(), BMemberName);
 		}
 
 		// A
 		{
 			const FName AMemberName = GET_MEMBER_NAME_CHECKED(UMaterialExpressionComponentMask, A);
-			MaskFactoryNode->AddBooleanAttribute(AMemberName, bAChannel);
-			MaskFactoryNode->AddApplyAndFillDelegates<bool>(AMemberName, UMaterialExpressionComponentMask::StaticClass(), AMemberName);
+			MaskFactoryNode->AddBooleanAttribute(AMemberName.ToString(), bAChannel);
+			MaskFactoryNode->AddApplyAndFillDelegates<bool>(AMemberName.ToString(), UMaterialExpressionComponentMask::StaticClass(), AMemberName);
 		}
 	}
 
@@ -2414,8 +2414,8 @@ UInterchangeMaterialExpressionFactoryNode* UInterchangeGenericMaterialPipeline::
 	if (ShaderNode->GetFloatAttribute(UInterchangeShaderPortsAPI::MakeInputValueKey(InputName), InputValue))
 	{
 		const FName DefaultValueMemberName = GET_MEMBER_NAME_CHECKED(UMaterialExpressionScalarParameter, DefaultValue);
-		MaterialExpressionFactoryNode->AddFloatAttribute(DefaultValueMemberName, InputValue);
-		MaterialExpressionFactoryNode->AddApplyAndFillDelegates<float>(DefaultValueMemberName, UMaterialExpressionScalarParameter::StaticClass(), DefaultValueMemberName);
+		MaterialExpressionFactoryNode->AddFloatAttribute(DefaultValueMemberName.ToString(), InputValue);
+		MaterialExpressionFactoryNode->AddApplyAndFillDelegates<float>(DefaultValueMemberName.ToString(), UMaterialExpressionScalarParameter::StaticClass(), DefaultValueMemberName);
 	}
 
 	return MaterialExpressionFactoryNode;
@@ -2429,8 +2429,8 @@ UInterchangeMaterialExpressionFactoryNode* UInterchangeGenericMaterialPipeline::
 	if (ShaderNode->GetLinearColorAttribute(UInterchangeShaderPortsAPI::MakeInputValueKey(InputName), InputValue))
 	{
 		const FName DefaultValueMemberName = GET_MEMBER_NAME_CHECKED(UMaterialExpressionVectorParameter, DefaultValue);
-		MaterialExpressionFactoryNode->AddLinearColorAttribute(DefaultValueMemberName, InputValue);
-		MaterialExpressionFactoryNode->AddApplyAndFillDelegates<FLinearColor>(DefaultValueMemberName, UMaterialExpressionVectorParameter::StaticClass(), DefaultValueMemberName);
+		MaterialExpressionFactoryNode->AddLinearColorAttribute(DefaultValueMemberName.ToString(), InputValue);
+		MaterialExpressionFactoryNode->AddApplyAndFillDelegates<FLinearColor>(DefaultValueMemberName.ToString(), UMaterialExpressionVectorParameter::StaticClass(), DefaultValueMemberName);
 	}
 
 	return MaterialExpressionFactoryNode;
@@ -2444,8 +2444,8 @@ UInterchangeMaterialExpressionFactoryNode* UInterchangeGenericMaterialPipeline::
 		UInterchangeMaterialExpressionFactoryNode* VectorParameterFactoryNode = CreateExpressionNode(InputName, ParentUid, UMaterialExpressionVectorParameter::StaticClass());
 
 		const FName DefaultValueMemberName = GET_MEMBER_NAME_CHECKED(UMaterialExpressionVectorParameter, DefaultValue);
-		VectorParameterFactoryNode->AddLinearColorAttribute(DefaultValueMemberName, FLinearColor(InputValue.X, InputValue.Y, 0.f));
-		VectorParameterFactoryNode->AddApplyAndFillDelegates<FLinearColor>(DefaultValueMemberName, UMaterialExpressionVectorParameter::StaticClass(), DefaultValueMemberName);
+		VectorParameterFactoryNode->AddLinearColorAttribute(DefaultValueMemberName.ToString(), FLinearColor(InputValue.X, InputValue.Y, 0.f));
+		VectorParameterFactoryNode->AddApplyAndFillDelegates<FLinearColor>(DefaultValueMemberName.ToString(), UMaterialExpressionVectorParameter::StaticClass(), DefaultValueMemberName);
 
 		// Defaults to R&G
 		UInterchangeMaterialExpressionFactoryNode* ComponentMaskFactoryNode = CreateExpressionNode(InputName + TEXT("_Mask"), ParentUid, UMaterialExpressionComponentMask::StaticClass());
@@ -2917,8 +2917,8 @@ bool UInterchangeGenericMaterialPipeline::HandleBxDFInput(const UInterchangeShad
 	// Make sure the bUseMaterialAttributes property of the material is set to true
 	static const FName UseMaterialAttributesMemberName = GET_MEMBER_NAME_CHECKED(UMaterial, bUseMaterialAttributes);
 
-	MaterialFactoryNode->AddBooleanAttribute(UseMaterialAttributesMemberName, true);
-	MaterialFactoryNode->AddApplyAndFillDelegates<FString>(UseMaterialAttributesMemberName, UMaterialExpressionMaterialFunctionCall::StaticClass(), UseMaterialAttributesMemberName);
+	MaterialFactoryNode->AddBooleanAttribute(UseMaterialAttributesMemberName.ToString(), true);
+	MaterialFactoryNode->AddApplyAndFillDelegates<FString>(UseMaterialAttributesMemberName.ToString(), UMaterialExpressionMaterialFunctionCall::StaticClass(), UseMaterialAttributesMemberName);
 
 	return true;
 }
diff --git a/Engine/Source/Runtime/Interchange/Core/Private/Nodes/InterchangeBaseNode.cpp b/Engine/Source/Runtime/Interchange/Core/Private/Nodes/InterchangeBaseNode.cpp
--- a/Engine/Source/Runtime/Interchange/Core/Private/Nodes/InterchangeBaseNode.cpp
+++ b/Engine/Source/Runtime/Interchange/Core/Private/Nodes/InterchangeBaseNode.cpp
@@ -183,88 +183,88 @@ void UInterchangeBaseNode::GetAttributeKeys(TArray<UE::Interchange::FAttributeKe
 	Attributes->GetAttributeKeys(AttributeKeys);
 }
 
-bool UInterchangeBaseNode::RemoveAttribute(const FName& NodeAttributeKey)
+bool UInterchangeBaseNode::RemoveAttribute(const FString& NodeAttributeKey)
 {
 	Attributes->UnregisterAttribute(UE::Interchange::FAttributeKey(NodeAttributeKey));
 	return !HasAttribute(UE::Interchange::FAttributeKey(NodeAttributeKey));
 }
 
-bool UInterchangeBaseNode::AddBooleanAttribute(const FName& NodeAttributeKey, const bool& Value)
+bool UInterchangeBaseNode::AddBooleanAttribute(const FString& NodeAttributeKey, const bool& Value)
 {
 	INTERCHANGE_BASE_NODE_ADD_ATTRIBUTE(bool);
 }
 
-bool UInterchangeBaseNode::GetBooleanAttribute(const FName& NodeAttributeKey, bool& OutValue) const
+bool UInterchangeBaseNode::GetBooleanAttribute(const FString& NodeAttributeKey, bool& OutValue) const
 {
 	INTERCHANGE_BASE_NODE_GET_ATTRIBUTE(bool);
 }
 
-bool UInterchangeBaseNode::AddInt32Attribute(const FName& NodeAttributeKey, const int32& Value)
+bool UInterchangeBaseNode::AddInt32Attribute(const FString& NodeAttributeKey, const int32& Value)
 {
 	INTERCHANGE_BASE_NODE_ADD_ATTRIBUTE(int32);
 }
 
-bool UInterchangeBaseNode::GetInt32Attribute(const FName& NodeAttributeKey, int32& OutValue) const
+bool UInterchangeBaseNode::GetInt32Attribute(const FString& NodeAttributeKey, int32& OutValue) const
 {
 	INTERCHANGE_BASE_NODE_GET_ATTRIBUTE(int32);
 }
 
-bool UInterchangeBaseNode::AddFloatAttribute(const FName& NodeAttributeKey, const float& Value)
+bool UInterchangeBaseNode::AddFloatAttribute(const FString& NodeAttributeKey, const float& Value)
 {
 	INTERCHANGE_BASE_NODE_ADD_ATTRIBUTE(float);
 }
 
-bool UInterchangeBaseNode::GetFloatAttribute(const FName& NodeAttributeKey, float& OutValue) const
+bool UInterchangeBaseNode::GetFloatAttribute(const FString& NodeAttributeKey, float& OutValue) const
 {
 	INTERCHANGE_BASE_NODE_GET_ATTRIBUTE(float);
 }
 
-bool UInterchangeBaseNode::AddDoubleAttribute(const FName& NodeAttributeKey, const double& Value)
+bool UInterchangeBaseNode::AddDoubleAttribute(const FString& NodeAttributeKey, const double& Value)
 {
 	INTERCHANGE_BASE_NODE_ADD_ATTRIBUTE(double);
 }
 
-bool UInterchangeBaseNode::GetDoubleAttribute(const FName& NodeAttributeKey, double& OutValue) const
+bool UInterchangeBaseNode::GetDoubleAttribute(const FString& NodeAttributeKey, double& OutValue) const
 {
 	INTERCHANGE_BASE_NODE_GET_ATTRIBUTE(double);
 }
 
-bool UInterchangeBaseNode::AddStringAttribute(const FName& NodeAttributeKey, const FString& Value)
+bool UInterchangeBaseNode::AddStringAttribute(const FString& NodeAttributeKey, const FString& Value)
 {
 	INTERCHANGE_BASE_NODE_ADD_ATTRIBUTE(FString);
 }
 
-bool UInterchangeBaseNode::GetStringAttribute(const FName& NodeAttributeKey, FString& OutValue) const
+bool UInterchangeBaseNode::GetStringAttribute(const FString& NodeAttributeKey, FString& OutValue) const
 {
 	INTERCHANGE_BASE_NODE_GET_ATTRIBUTE(FString);
 }
 
-bool UInterchangeBaseNode::AddGuidAttribute(const FName& NodeAttributeKey, const FGuid& Value)
+bool UInterchangeBaseNode::AddGuidAttribute(const FString& NodeAttributeKey, const FGuid& Value)
 {
 	INTERCHANGE_BASE_NODE_ADD_ATTRIBUTE(FGuid);
 }
 
-bool UInterchangeBaseNode::GetGuidAttribute(const FName& NodeAttributeKey, FGuid& OutValue) const
+bool UInterchangeBaseNode::GetGuidAttribute(const FString& NodeAttributeKey, FGuid& OutValue) const
 {
 	INTERCHANGE_BASE_NODE_GET_ATTRIBUTE(FGuid);
 }
 
-bool UInterchangeBaseNode::AddLinearColorAttribute(const FName& NodeAttributeKey, const FLinearColor& Value)
+bool UInterchangeBaseNode::AddLinearColorAttribute(const FString& NodeAttributeKey, const FLinearColor& Value)
 {
 	INTERCHANGE_BASE_NODE_ADD_ATTRIBUTE(FLinearColor);
 }
 
-bool UInterchangeBaseNode::GetLinearColorAttribute(const FName& NodeAttributeKey, FLinearColor& OutValue) const
+bool UInterchangeBaseNode::GetLinearColorAttribute(const FString& NodeAttributeKey, FLinearColor& OutValue) const
 {
 	INTERCHANGE_BASE_NODE_GET_ATTRIBUTE(FLinearColor);
 }
 
-bool UInterchangeBaseNode::AddVector2Attribute(const FName& NodeAttributeKey, const FVector2f& Value)
+bool UInterchangeBaseNode::AddVector2Attribute(const FString& NodeAttributeKey, const FVector2f& Value)
 {
 	INTERCHANGE_BASE_NODE_ADD_ATTRIBUTE(FVector2f);
 }
 
-bool UInterchangeBaseNode::GetVector2Attribute(const FName& NodeAttributeKey, FVector2f& OutValue) const
+bool UInterchangeBaseNode::GetVector2Attribute(const FString& NodeAttributeKey, FVector2f& OutValue) const
 {
 	INTERCHANGE_BASE_NODE_GET_ATTRIBUTE(FVector2f);
 }
diff --git a/Engine/Source/Runtime/Interchange/Core/Public/Nodes/InterchangeBaseNode.h b/Engine/Source/Runtime/Interchange/Core/Public/Nodes/InterchangeBaseNode.h
--- a/Engine/Source/Runtime/Interchange/Core/Public/Nodes/InterchangeBaseNode.h
+++ b/Engine/Source/Runtime/Interchange/Core/Public/Nodes/InterchangeBaseNode.h
@@ -276,106 +276,106 @@ public:
 	 * Remove any attribute from this node. Return false if we cannot remove it. If the attribute do not exist it will return true.
 	 */
 	UFUNCTION(BlueprintCallable, Category = "Interchange | Node")
-	bool RemoveAttribute(const FName& NodeAttributeKey);
+	bool RemoveAttribute(const FString& NodeAttributeKey);
 
 	/**
 	 * Add a boolean attribute to this node. Return false if the attribute do not exist or if we cannot add it
 	 */
 	UFUNCTION(BlueprintCallable, Category = "Interchange | Node")
-	bool AddBooleanAttribute(const FName& NodeAttributeKey, const bool& Value);
+	bool AddBooleanAttribute(const FString& NodeAttributeKey, const bool& Value);
 	
 	/**
 	 * Get a boolean attribute from this node. Return false if the attribute do not exist
 	 */
 	UFUNCTION(BlueprintCallable, Category = "Interchange | Node")
-	bool GetBooleanAttribute(const FName& NodeAttributeKey, bool& OutValue) const;
+	bool GetBooleanAttribute(const FString& NodeAttributeKey, bool& OutValue) const;
 
 	/**
 	 * Add a int32 attribute to this node. Return false if the attribute do not exist or if we cannot add it
 	 */
 	UFUNCTION(BlueprintCallable, Category = "Interchange | Node")
-	bool AddInt32Attribute(const FName& NodeAttributeKey, const int32& Value);
+	bool AddInt32Attribute(const FString& NodeAttributeKey, const int32& Value);
 	
 	/**
 	 * Get a int32 attribute from this node. Return false if the attribute do not exist
 	 */
 	UFUNCTION(BlueprintCallable, Category = "Interchange | Node")
-	bool GetInt32Attribute(const FName& NodeAttributeKey, int32& OutValue) const;
+	bool GetInt32Attribute(const FString& NodeAttributeKey, int32& OutValue) const;
 
 	/**
 	 * Add a float attribute to this node. Return false if the attribute do not exist or if we cannot add it
 	 */
 	UFUNCTION(BlueprintCallable, Category = "Interchange | Node")
-	bool AddFloatAttribute(const FName& NodeAttributeKey, const float& Value);
+	bool AddFloatAttribute(const FString& NodeAttributeKey, const float& Value);
 	
 	/**
 	 * Get a float attribute from this node. Return false if the attribute do not exist
 	 */
 	UFUNCTION(BlueprintCallable, Category = "Interchange | Node")
-	bool GetFloatAttribute(const FName& NodeAttributeKey, float& OutValue) const;
+	bool GetFloatAttribute(const FString& NodeAttributeKey, float& OutValue) const;
 
 	/**
 	 * Add a double attribute to this node. Return false if the attribute do not exist or if we cannot add it
 	 */
 	UFUNCTION(BlueprintCallable, Category = "Interchange | Node")
-	bool AddDoubleAttribute(const FName& NodeAttributeKey, const double& Value);
+	bool AddDoubleAttribute(const FString& NodeAttributeKey, const double& Value);
 	
 	/**
 	 * Get a double attribute from this node. Return false if the attribute do not exist
 	 */
 	UFUNCTION(BlueprintCallable, Category = "Interchange | Node")
-	bool GetDoubleAttribute(const FName& NodeAttributeKey, double& OutValue) const;
+	bool GetDoubleAttribute(const FString& NodeAttributeKey, double& OutValue) const;
 
 	/**
 	 * Add a string attribute to this node. Return false if the attribute do not exist or if we cannot add it
 	 */
 	UFUNCTION(BlueprintCallable, Category = "Interchange | Node")
-	bool AddStringAttribute(const FName& NodeAttributeKey, const FString& Value);
+	bool AddStringAttribute(const FString& NodeAttributeKey, const FString& Value);
 	
 	/**
 	 * Get a string attribute from this node. Return false if the attribute do not exist
 	 */
 	UFUNCTION(BlueprintCallable, Category = "Interchange | Node")
-	bool GetStringAttribute(const FName& NodeAttributeKey, FString& OutValue) const;
+	bool GetStringAttribute(const FString& NodeAttributeKey, FString& OutValue) const;
 
 	/**
 	 * Add a guid attribute to this node. Return false if the attribute do not exist or if we cannot add it
 	 */
 	UFUNCTION(BlueprintCallable, Category = "Interchange | Node")
-	bool AddGuidAttribute(const FName& NodeAttributeKey, const FGuid& Value);
+	bool AddGuidAttribute(const FString& NodeAttributeKey, const FGuid& Value);
 
 	/**
 	 * Get a guid attribute from this node. Return false if the attribute do not exist
 	 */
 	UFUNCTION(BlueprintCallable, Category = "Interchange | Node")
-	bool GetGuidAttribute(const FName& NodeAttributeKey, FGuid& OutValue) const;
+	bool GetGuidAttribute(const FString& NodeAttributeKey, FGuid& OutValue) const;
 
 	/**
 	 * Add a FLinearColor attribute to this node. Return false if the attribute do not exist or if we cannot add it
 	 */
 	UFUNCTION(BlueprintCallable, Category = "Interchange | Node")
-	bool AddLinearColorAttribute(const FName& NodeAttributeKey, const FLinearColor& Value);
+	bool AddLinearColorAttribute(const FString& NodeAttributeKey, const FLinearColor& Value);
 	
 	/**
 	 * Get a FLinearColor attribute from this node. Return false if the attribute do not exist
 	 */
 	UFUNCTION(BlueprintCallable, Category = "Interchange | Node")
-	bool GetLinearColorAttribute(const FName& NodeAttributeKey, FLinearColor& OutValue) const;
+	bool GetLinearColorAttribute(const FString& NodeAttributeKey, FLinearColor& OutValue) const;
 
 	/**
 	 * Add a Vector2 attribute to this node. Return false if the attribute do not exist or if we cannot add it
 	 */
 	UFUNCTION(BlueprintCallable, Category = "Interchange | Node")
-	bool AddVector2Attribute(const FName& NodeAttributeKey, const FVector2f& Value);
+	bool AddVector2Attribute(const FString& NodeAttributeKey, const FVector2f& Value);
 
 	/**
 	 * Get a Vector2 attribute from this node. Return false if the attribute do not exist
 	 */
 	UFUNCTION(BlueprintCallable, Category = "Interchange | Node")
-	bool GetVector2Attribute(const FName& NodeAttributeKey, FVector2f& OutValue) const;
+	bool GetVector2Attribute(const FString& NodeAttributeKey, FVector2f& OutValue) const;
 
 	template<typename AttributeType>
-	AttributeType GetAttributeChecked(const FName& NodeAttributeKey) const
+	AttributeType GetAttributeChecked(const FString& NodeAttributeKey) const
 	{
 		AttributeType Value = AttributeType();
 		check(HasAttribute(UE::Interchange::FAttributeKey(NodeAttributeKey)));
@@ -386,12 +386,12 @@ public:
 	}
 
 	template<typename AttributeType>
-	bool GetAttribute(const FName& NodeAttributeKey, AttributeType& OutValue) const
+	bool GetAttribute(const FString& NodeAttributeKey, AttributeType& OutValue) const
 	{
 		INTERCHANGE_BASE_NODE_GET_ATTRIBUTE(AttributeType);
 	}
 	template<typename AttributeType>
-	bool SetAttribute(const FName& NodeAttributeKey, const AttributeType& Value)
+	bool SetAttribute(const FString& NodeAttributeKey, const AttributeType& Value)
 	{
 		INTERCHANGE_BASE_NODE_ADD_ATTRIBUTE(AttributeType);
 	}
diff --git a/Engine/Source/Runtime/Interchange/Core/Public/Nodes/InterchangeBaseNodeUtilities.h b/Engine/Source/Runtime/Interchange/Core/Public/Nodes/InterchangeBaseNodeUtilities.h
--- a/Engine/Source/Runtime/Interchange/Core/Public/Nodes/InterchangeBaseNodeUtilities.h
+++ b/Engine/Source/Runtime/Interchange/Core/Public/Nodes/InterchangeBaseNodeUtilities.h
@@ -273,13 +273,13 @@ private:
 	FAttributeKey KeyCount;
 	FAttributeKey GetKeyCount() const
 	{
-		ensure(!KeyCount.Key.IsNone()); return KeyCount;
+		ensure(!KeyCount.Key.IsEmpty()); return KeyCount;
 	}
 
 	FAttributeKey GetIndexKey(int32 Index) const
 	{
 		FString DepIndexKeyString = GetKeyCount().ToString() + IndexKey() + FString::FromInt(Index);
-		return FAttributeKey(*DepIndexKeyString);
+		return FAttributeKey(DepIndexKeyString);
 	}
 };
 
diff --git a/Engine/Source/Runtime/Interchange/Core/Public/Nodes/InterchangeFactoryBaseNode.h b/Engine/Source/Runtime/Interchange/Core/Public/Nodes/InterchangeFactoryBaseNode.h
--- a/Engine/Source/Runtime/Interchange/Core/Public/Nodes/InterchangeFactoryBaseNode.h
+++ b/Engine/Source/Runtime/Interchange/Core/Public/Nodes/InterchangeFactoryBaseNode.h
@@ -214,7 +214,7 @@ public:
 	 * @param PropertyName		The name of the property of the ObjectClass that we will read and write to.
 	 */
 	template<typename AttributeType>
-	inline void AddApplyAndFillDelegates(const FName& NodeAttributeKey, UClass* ObjectClass, const FName PropertyName);
+	inline void AddApplyAndFillDelegates(const FString& NodeAttributeKey, UClass* ObjectClass, const FName PropertyName);
 
 	/**
 	 * Writes an attribute value to a UObject.
@@ -224,7 +224,7 @@ public:
 	 * @result					True if we succeeded.
 	 */
 	template<typename AttributeType>
-	inline bool ApplyAttributeToObject(const FName& NodeAttributeKey, UObject* Object, const FName PropertyName) const;
+	inline bool ApplyAttributeToObject(const FString& NodeAttributeKey, UObject* Object, const FName PropertyName) const;
 
 	/**
 	 * Reads an attribute value from a UObject.
@@ -234,7 +234,7 @@ public:
 	 * @result					True if we succeeded.
 	 */
 	template<typename AttributeType>
-	inline bool FillAttributeFromObject(const FName& NodeAttributeKey, UObject* Object, const FName PropertyName);
+	inline bool FillAttributeFromObject(const FString& NodeAttributeKey, UObject* Object, const FName PropertyName);
 
 	/**
 	 * Each Attribute that was set and have a delegate set for the specified UObject->UClass will
@@ -265,7 +265,7 @@ private:
 };
 
 template<typename AttributeType>
-inline void UInterchangeFactoryBaseNode::AddApplyAndFillDelegates(const FName& NodeAttributeKey, UClass* ObjectClass, const FName PropertyName)
+inline void UInterchangeFactoryBaseNode::AddApplyAndFillDelegates(const FString& NodeAttributeKey, UClass* ObjectClass, const FName PropertyName)
 {
 	TArray<UE::Interchange::FApplyAttributeToAsset>& ApplyDelegates = ApplyCustomAttributeDelegates.FindOrAdd(ObjectClass);
 	ApplyDelegates.Add(
@@ -287,7 +287,7 @@ inline void UInterchangeFactoryBaseNode::AddApplyAndFillDelegates(const FName& N
 }
 
 template<typename AttributeType>
-inline bool UInterchangeFactoryBaseNode::ApplyAttributeToObject(const FName& NodeAttributeKey, UObject* Object, const FName PropertyName) const
+inline bool UInterchangeFactoryBaseNode::ApplyAttributeToObject(const FString& NodeAttributeKey, UObject* Object, const FName PropertyName) const
 {
 	if (!Object)
 	{
@@ -323,7 +323,7 @@ inline bool UInterchangeFactoryBaseNode::ApplyAttributeToObject(const FName& Nod
  * If the target property is a FObjectPropertyBase, treat the string as an object path.
  */
 template<>
-inline bool UInterchangeFactoryBaseNode::ApplyAttributeToObject<FString>(const FName& NodeAttributeKey, UObject* Object, const FName PropertyName) const
+inline bool UInterchangeFactoryBaseNode::ApplyAttributeToObject<FString>(const FString& NodeAttributeKey, UObject* Object, const FName PropertyName) const
 {
 	if (!Object)
 	{
@@ -366,7 +366,7 @@ inline bool UInterchangeFactoryBaseNode::ApplyAttributeToObject<FString>(const F
  * If the target property is a FBoolProperty, treat the propertyg as a bitfield.
  */
 template<>
-inline bool UInterchangeFactoryBaseNode::ApplyAttributeToObject<bool>(const FName& NodeAttributeKey, UObject* Object, const FName PropertyName) const
+inline bool UInterchangeFactoryBaseNode::ApplyAttributeToObject<bool>(const FString& NodeAttributeKey, UObject* Object, const FName PropertyName) const
 {
 	if (!Object)
 	{
@@ -406,7 +406,7 @@ inline bool UInterchangeFactoryBaseNode::ApplyAttributeToObject<bool>(const FNam
 }
 
 template<typename AttributeType>
-inline bool UInterchangeFactoryBaseNode::FillAttributeFromObject(const FName& NodeAttributeKey, UObject* Object, const FName PropertyName)
+inline bool UInterchangeFactoryBaseNode::FillAttributeFromObject(const FString& NodeAttributeKey, UObject* Object, const FName PropertyName)
 {
 	TVariant<UObject*, uint8*> Container;
 	Container.Set<UObject*>(Object);
@@ -442,7 +442,7 @@ inline bool UInterchangeFactoryBaseNode::FillAttributeFromObject(const FName& No
  * If the target property is a FObjectPropertyBase, treat the string as an object path.
  */
 template<>
-inline bool UInterchangeFactoryBaseNode::FillAttributeFromObject<FString>(const FName& NodeAttributeKey, UObject* Object, const FName PropertyName)
+inline bool UInterchangeFactoryBaseNode::FillAttributeFromObject<FString>(const FString& NodeAttributeKey, UObject* Object, const FName PropertyName)
 {
 	TVariant<UObject*, uint8*> Container;
 	Container.Set<UObject*>(Object);
@@ -477,7 +477,7 @@ inline bool UInterchangeFactoryBaseNode::FillAttributeFromObject<FString>(const
  * If the target property is a FBoolProperty, treat the property as a bitfield.
  */
 template<>
-inline bool UInterchangeFactoryBaseNode::FillAttributeFromObject<bool>(const FName& NodeAttributeKey, UObject* Object, const FName PropertyName)
+inline bool UInterchangeFactoryBaseNode::FillAttributeFromObject<bool>(const FString& NodeAttributeKey, UObject* Object, const FName PropertyName)
 {
 	TVariant<UObject*, uint8*> Container;
 	Container.Set<UObject*>(Object);
diff --git a/Engine/Source/Runtime/Interchange/Core/Public/Types/AttributeStorage.h b/Engine/Source/Runtime/Interchange/Core/Public/Types/AttributeStorage.h
--- a/Engine/Source/Runtime/Interchange/Core/Public/Types/AttributeStorage.h
+++ b/Engine/Source/Runtime/Interchange/Core/Public/Types/AttributeStorage.h
@@ -74,7 +74,7 @@ namespace UE
 
 		struct FAttributeKey
 		{
-			FName Key;
+			FString Key;
 
 			FAttributeKey() = default;
 		
@@ -85,42 +85,42 @@ namespace UE
 
 			explicit FAttributeKey(const FName& Other)
 			{
-				Key = Other;
+				Key = Other.ToString();
 			}
 
 			explicit FAttributeKey(const FString& Other)
 			{
-				Key = FName(*Other);
+				Key = Other;
 			}
 
-			explicit FAttributeKey(FName&& Other)
+			explicit FAttributeKey(FString&& Other)
 			{
 				Key = MoveTemp(Other);
 			}
 
 			explicit FAttributeKey(const FText& Other)
 			{
-				Key = FName(*Other.ToString());
+				Key = Other.ToString();
 			}
 
 			explicit FAttributeKey(const TCHAR* Other)
 			{
-				Key = FName(Other);
+				Key = Other;
 			}
 
 			FORCEINLINE FAttributeKey& operator=(const FName& Other)
 			{
-				Key = Other;
+				Key = Other.ToString();
 				return *this;
 			}
 
 			FORCEINLINE FAttributeKey& operator=(const FString& Other)
 			{
-				Key = FName(*Other);
+				Key = Other;
 				return *this;
 			}
 
-			FORCEINLINE FAttributeKey& operator=(FName&& Other)
+			FORCEINLINE FAttributeKey& operator=(FString&& Other)
 			{
 				Key = MoveTemp(Other);
 				return *this;
@@ -128,29 +128,29 @@ namespace UE
 
 			FORCEINLINE FAttributeKey& operator=(const FText& Other)
 			{
-				Key = *(Other.ToString());
+				Key = Other.ToString();
 				return *this;
 			}
 
 			FORCEINLINE FAttributeKey& operator=(const TCHAR* Other)
 			{
-				Key = FName(Other);
+				Key = Other;
 				return *this;
 			}
 		
 			FORCEINLINE bool operator==(const FAttributeKey& Other) const
 			{
-				return Key == Other.Key;
+				return Key.Equals(Other.Key);
 			}
 
 			FORCEINLINE bool operator!=(const FAttributeKey& Other) const
 			{
-				return Key != Other.Key;
+				return !Key.Equals(Other.Key);
 			}
 
 			FORCEINLINE bool operator<(const FAttributeKey& Other) const
 			{
-				return Key.LexicalLess(Other.Key);
+				return Key.Compare(Other.Key) < 0;
 			}
 
 			FORCEINLINE bool operator<=(const FAttributeKey& Other) const
@@ -160,40 +160,28 @@ namespace UE
 
 			FORCEINLINE bool operator>(const FAttributeKey& Other) const
 			{
-				return Other.Key.LexicalLess(Key);
+				return Key.Compare(Other.Key) > 0;
 			}
 
 			FORCEINLINE bool operator>=(const FAttributeKey& Other) const
 			{
-				return Other.Key.LexicalLess(Key) || Key == Other.Key;
+				return Key.Compare(Other.Key) >= 0;
 			}
 
 			friend FArchive& operator<<(FArchive& Ar, FAttributeKey& AttributeKey)
 			{
-				FString SerializeString;
-				//When serializing FName with FString we always must do if(IsLoading()) else to avoid changing the data
-				//When a reference serialization is execute (reference do not have IsLoading or IsSaving set)
-				if (Ar.IsLoading())
-				{
-					Ar << SerializeString;
-					AttributeKey.Key = FName(*SerializeString);
-				}
-				else
-				{
-					SerializeString = AttributeKey.Key.ToString();
-					Ar << SerializeString;
-				}
+				Ar << AttributeKey.Key;
 				return Ar;
 			}
 
 			FORCEINLINE FString ToString() const
 			{
-				return Key.ToString();
+				return Key;
 			}
 
 			FORCEINLINE FName ToName() const
 			{
-				return Key;
+				return FName(*Key);
 			}
 		};
 
