Fix merge conflict and CIS
#rnx

diff --git a/Engine/Plugins/Interchange/Runtime/Source/Pipelines/Private/InterchangeGenericMaterialPipeline.cpp b/Engine/Plugins/Interchange/Runtime/Source/Pipelines/Private/InterchangeGenericMaterialPipeline.cpp
--- a/Engine/Plugins/Interchange/Runtime/Source/Pipelines/Private/InterchangeGenericMaterialPipeline.cpp
+++ b/Engine/Plugins/Interchange/Runtime/Source/Pipelines/Private/InterchangeGenericMaterialPipeline.cpp
@@ -1989,27 +1989,27 @@ void UInterchangeGenericMaterialPipeline::HandleTimeNode(const UInterchangeShade
 	TimeFactoryNode->SetCustomExpressionClassName(UMaterialExpressionTime::StaticClass()->GetName());
 
 	// IgnorePause
-	if(bool bIgnorePause; ShaderNode->GetBooleanAttribute(Time::Attributes::IgnorePause, bIgnorePause))
+	if(bool bIgnorePause; ShaderNode->GetBooleanAttribute(Time::Attributes::IgnorePause.ToString(), bIgnorePause))
 	{
 		const FName IgnorePauseMemberName = GET_MEMBER_NAME_CHECKED(UMaterialExpressionTime, bIgnorePause);
-		TimeFactoryNode->AddBooleanAttribute(IgnorePauseMemberName, bIgnorePause);
-		TimeFactoryNode->AddApplyAndFillDelegates<bool>(IgnorePauseMemberName, UMaterialExpressionTime::StaticClass(), IgnorePauseMemberName);
+		TimeFactoryNode->AddBooleanAttribute(IgnorePauseMemberName.ToString(), bIgnorePause);
+		TimeFactoryNode->AddApplyAndFillDelegates<bool>(IgnorePauseMemberName.ToString(), UMaterialExpressionTime::StaticClass(), IgnorePauseMemberName);
 	}
 
 	// OverridePeriod
-	if(bool bOverridePeriod; ShaderNode->GetBooleanAttribute(Time::Attributes::OverridePeriod, bOverridePeriod))
+	if(bool bOverridePeriod; ShaderNode->GetBooleanAttribute(Time::Attributes::OverridePeriod.ToString(), bOverridePeriod))
 	{
 		const FName OverridePeriodMemberName = GET_MEMBER_NAME_CHECKED(UMaterialExpressionTime, bOverride_Period);
-		TimeFactoryNode->AddBooleanAttribute(OverridePeriodMemberName, bOverridePeriod);
-		TimeFactoryNode->AddApplyAndFillDelegates<bool>(OverridePeriodMemberName, UMaterialExpressionTime::StaticClass(), OverridePeriodMemberName);
+		TimeFactoryNode->AddBooleanAttribute(OverridePeriodMemberName.ToString(), bOverridePeriod);
+		TimeFactoryNode->AddApplyAndFillDelegates<bool>(OverridePeriodMemberName.ToString(), UMaterialExpressionTime::StaticClass(), OverridePeriodMemberName);
 	}
 
 	// Period
-	if(float Period; ShaderNode->GetFloatAttribute(Time::Attributes::Period, Period))
+	if(float Period; ShaderNode->GetFloatAttribute(Time::Attributes::Period.ToString(), Period))
 	{
 		const FName PeriodMemberName = GET_MEMBER_NAME_CHECKED(UMaterialExpressionTime, Period);
-		TimeFactoryNode->AddFloatAttribute(PeriodMemberName, Period);
-		TimeFactoryNode->AddApplyAndFillDelegates<float>(PeriodMemberName, UMaterialExpressionTime::StaticClass(), PeriodMemberName);
+		TimeFactoryNode->AddFloatAttribute(PeriodMemberName.ToString(), Period);
+		TimeFactoryNode->AddApplyAndFillDelegates<float>(PeriodMemberName.ToString(), UMaterialExpressionTime::StaticClass(), PeriodMemberName);
 	}
 }
 
@@ -2032,19 +2032,19 @@ void UInterchangeGenericMaterialPipeline::HandleTransformPositionNode(const UInt
 	}
 
 	// TransformSourceType
-	if(int32 TransformSourceType; ShaderNode->GetInt32Attribute(TransformPosition::Attributes::TransformSourceType, TransformSourceType))
+	if(int32 TransformSourceType; ShaderNode->GetInt32Attribute(TransformPosition::Attributes::TransformSourceType.ToString(), TransformSourceType))
 	{
 		const FName TransformSourceTypeMemberName = GET_MEMBER_NAME_CHECKED(UMaterialExpressionTransformPosition, TransformSourceType);
-		TransformPositionFactoryNode->AddInt32Attribute(TransformSourceTypeMemberName,TransformSourceType);
-		TransformPositionFactoryNode->AddApplyAndFillDelegates<int32>(TransformSourceTypeMemberName, UMaterialExpressionTransformPosition::StaticClass(), TransformSourceTypeMemberName);
+		TransformPositionFactoryNode->AddInt32Attribute(TransformSourceTypeMemberName.ToString(),TransformSourceType);
+		TransformPositionFactoryNode->AddApplyAndFillDelegates<int32>(TransformSourceTypeMemberName.ToString(), UMaterialExpressionTransformPosition::StaticClass(), TransformSourceTypeMemberName);
 	}
 
 	// TransformType
-	if(int32 TransformType; ShaderNode->GetInt32Attribute(TransformPosition::Attributes::TransformType, TransformType))
+	if(int32 TransformType; ShaderNode->GetInt32Attribute(TransformPosition::Attributes::TransformType.ToString(), TransformType))
 	{
 		const FName TransformTypeMemberName = GET_MEMBER_NAME_CHECKED(UMaterialExpressionTransformPosition, TransformType);
-		TransformPositionFactoryNode->AddInt32Attribute(TransformTypeMemberName, TransformType);
-		TransformPositionFactoryNode->AddApplyAndFillDelegates<int32>(TransformTypeMemberName, UMaterialExpressionTransformPosition::StaticClass(), TransformTypeMemberName);
+		TransformPositionFactoryNode->AddInt32Attribute(TransformTypeMemberName.ToString(), TransformType);
+		TransformPositionFactoryNode->AddApplyAndFillDelegates<int32>(TransformTypeMemberName.ToString(), UMaterialExpressionTransformPosition::StaticClass(), TransformTypeMemberName);
 	}
 }
 
@@ -2067,19 +2067,19 @@ void UInterchangeGenericMaterialPipeline::HandleTransformVectorNode(const UInter
 	}
 
 	// TransformSourceType
-	if(int32 TransformSourceType; ShaderNode->GetInt32Attribute(TransformVector::Attributes::TransformSourceType, TransformSourceType))
+	if(int32 TransformSourceType; ShaderNode->GetInt32Attribute(TransformVector::Attributes::TransformSourceType.ToString(), TransformSourceType))
 	{
 		const FName TransformSourceTypeMemberName = GET_MEMBER_NAME_CHECKED(UMaterialExpressionTransform, TransformSourceType);
-		TransformVectorFactoryNode->AddInt32Attribute(TransformSourceTypeMemberName, TransformSourceType);
-		TransformVectorFactoryNode->AddApplyAndFillDelegates<int32>(TransformSourceTypeMemberName, UMaterialExpressionTransform::StaticClass(), TransformSourceTypeMemberName);
+		TransformVectorFactoryNode->AddInt32Attribute(TransformSourceTypeMemberName.ToString(), TransformSourceType);
+		TransformVectorFactoryNode->AddApplyAndFillDelegates<int32>(TransformSourceTypeMemberName.ToString(), UMaterialExpressionTransform::StaticClass(), TransformSourceTypeMemberName);
 	}
 
 	// TransformType
-	if(int32 TransformType; ShaderNode->GetInt32Attribute(TransformVector::Attributes::TransformType, TransformType))
+	if(int32 TransformType; ShaderNode->GetInt32Attribute(TransformVector::Attributes::TransformType.ToString(), TransformType))
 	{
 		const FName TransformTypeMemberName = GET_MEMBER_NAME_CHECKED(UMaterialExpressionTransform, TransformType);
-		TransformVectorFactoryNode->AddInt32Attribute(TransformTypeMemberName, TransformType);
-		TransformVectorFactoryNode->AddApplyAndFillDelegates<int32>(TransformTypeMemberName, UMaterialExpressionTransform::StaticClass(), TransformTypeMemberName);
+		TransformVectorFactoryNode->AddInt32Attribute(TransformTypeMemberName.ToString(), TransformType);
+		TransformVectorFactoryNode->AddApplyAndFillDelegates<int32>(TransformTypeMemberName.ToString(), UMaterialExpressionTransform::StaticClass(), TransformTypeMemberName);
 	}
 }
 
@@ -2114,83 +2114,83 @@ void UInterchangeGenericMaterialPipeline::HandleNoiseNode(const UInterchangeShad
 	}
 
 	// Scale
-	if(float Scale; ShaderNode->GetFloatAttribute(Noise::Attributes::Scale, Scale))
+	if(float Scale; ShaderNode->GetFloatAttribute(Noise::Attributes::Scale.ToString(), Scale))
 	{
 		const FName ScaleMemberName = GET_MEMBER_NAME_CHECKED(UMaterialExpressionNoise, Scale);
-		NoiseFactoryNode->AddFloatAttribute(ScaleMemberName, Scale);
-		NoiseFactoryNode->AddApplyAndFillDelegates<float>(ScaleMemberName, UMaterialExpressionNoise::StaticClass(), ScaleMemberName);
+		NoiseFactoryNode->AddFloatAttribute(ScaleMemberName.ToString(), Scale);
+		NoiseFactoryNode->AddApplyAndFillDelegates<float>(ScaleMemberName.ToString(), UMaterialExpressionNoise::StaticClass(), ScaleMemberName);
 	}
 
 	// Quality
-	if(int32 Quality; ShaderNode->GetInt32Attribute(Noise::Attributes::Quality, Quality))
+	if(int32 Quality; ShaderNode->GetInt32Attribute(Noise::Attributes::Quality.ToString(), Quality))
 	{
 		const FName QualityMemberName = GET_MEMBER_NAME_CHECKED(UMaterialExpressionNoise, Quality);
-		NoiseFactoryNode->AddInt32Attribute(QualityMemberName, Quality);
-		NoiseFactoryNode->AddApplyAndFillDelegates<int32>(QualityMemberName, UMaterialExpressionNoise::StaticClass(), QualityMemberName);
+		NoiseFactoryNode->AddInt32Attribute(QualityMemberName.ToString(), Quality);
+		NoiseFactoryNode->AddApplyAndFillDelegates<int32>(QualityMemberName.ToString(), UMaterialExpressionNoise::StaticClass(), QualityMemberName);
 	}
 
 	// Noise Function
-	if(int32 NoiseFunction; ShaderNode->GetInt32Attribute(Noise::Attributes::Function, NoiseFunction))
+	if(int32 NoiseFunction; ShaderNode->GetInt32Attribute(Noise::Attributes::Function.ToString(), NoiseFunction))
 	{
 		const FName NoiseFunctionMemberName = GET_MEMBER_NAME_CHECKED(UMaterialExpressionNoise, NoiseFunction);
-		NoiseFactoryNode->AddInt32Attribute(NoiseFunctionMemberName, NoiseFunction);
-		NoiseFactoryNode->AddApplyAndFillDelegates<int32>(NoiseFunctionMemberName, UMaterialExpressionNoise::StaticClass(), NoiseFunctionMemberName);
+		NoiseFactoryNode->AddInt32Attribute(NoiseFunctionMemberName.ToString(), NoiseFunction);
+		NoiseFactoryNode->AddApplyAndFillDelegates<int32>(NoiseFunctionMemberName.ToString(), UMaterialExpressionNoise::StaticClass(), NoiseFunctionMemberName);
 	}
 
 	// Turbulence
-	if(bool bTurbulence; ShaderNode->GetBooleanAttribute(Noise::Attributes::Turbulence, bTurbulence))
+	if(bool bTurbulence; ShaderNode->GetBooleanAttribute(Noise::Attributes::Turbulence.ToString(), bTurbulence))
 	{
 		const FName TurbulenceMemberName = GET_MEMBER_NAME_CHECKED(UMaterialExpressionNoise, bTurbulence);
-		NoiseFactoryNode->AddBooleanAttribute(TurbulenceMemberName, bTurbulence);
-		NoiseFactoryNode->AddApplyAndFillDelegates<bool>(TurbulenceMemberName, UMaterialExpressionNoise::StaticClass(), TurbulenceMemberName);
+		NoiseFactoryNode->AddBooleanAttribute(TurbulenceMemberName.ToString(), bTurbulence);
+		NoiseFactoryNode->AddApplyAndFillDelegates<bool>(TurbulenceMemberName.ToString(), UMaterialExpressionNoise::StaticClass(), TurbulenceMemberName);
 	}
 
 	// Levels
-	if(int32 Levels; ShaderNode->GetInt32Attribute(Noise::Attributes::Levels, Levels))
+	if(int32 Levels; ShaderNode->GetInt32Attribute(Noise::Attributes::Levels.ToString(), Levels))
 	{
 		const FName LevelsMemberName = GET_MEMBER_NAME_CHECKED(UMaterialExpressionNoise, Levels);
-		NoiseFactoryNode->AddInt32Attribute(LevelsMemberName, Levels);
-		NoiseFactoryNode->AddApplyAndFillDelegates<int32>(LevelsMemberName, UMaterialExpressionNoise::StaticClass(), LevelsMemberName);
+		NoiseFactoryNode->AddInt32Attribute(LevelsMemberName.ToString(), Levels);
+		NoiseFactoryNode->AddApplyAndFillDelegates<int32>(LevelsMemberName.ToString(), UMaterialExpressionNoise::StaticClass(), LevelsMemberName);
 	}
 
 	// Output Min
-	if(float OutputMin; ShaderNode->GetFloatAttribute(Noise::Attributes::OutputMin, OutputMin))
+	if(float OutputMin; ShaderNode->GetFloatAttribute(Noise::Attributes::OutputMin.ToString(), OutputMin))
 	{
 		const FName OutputMinMemberName = GET_MEMBER_NAME_CHECKED(UMaterialExpressionNoise, OutputMin);
-		NoiseFactoryNode->AddFloatAttribute(OutputMinMemberName, OutputMin);
-		NoiseFactoryNode->AddApplyAndFillDelegates<float>(OutputMinMemberName, UMaterialExpressionNoise::StaticClass(), OutputMinMemberName);
+		NoiseFactoryNode->AddFloatAttribute(OutputMinMemberName.ToString(), OutputMin);
+		NoiseFactoryNode->AddApplyAndFillDelegates<float>(OutputMinMemberName.ToString(), UMaterialExpressionNoise::StaticClass(), OutputMinMemberName);
 	}
 
 	// Output Max
-	if(float OutputMax; ShaderNode->GetFloatAttribute(Noise::Attributes::OutputMax, OutputMax))
+	if(float OutputMax; ShaderNode->GetFloatAttribute(Noise::Attributes::OutputMax.ToString(), OutputMax))
 	{
 		const FName OutputMaxMemberName = GET_MEMBER_NAME_CHECKED(UMaterialExpressionNoise, OutputMax);
-		NoiseFactoryNode->AddFloatAttribute(OutputMaxMemberName, OutputMax);
-		NoiseFactoryNode->AddApplyAndFillDelegates<float>(OutputMaxMemberName, UMaterialExpressionNoise::StaticClass(), OutputMaxMemberName);
+		NoiseFactoryNode->AddFloatAttribute(OutputMaxMemberName.ToString(), OutputMax);
+		NoiseFactoryNode->AddApplyAndFillDelegates<float>(OutputMaxMemberName.ToString(), UMaterialExpressionNoise::StaticClass(), OutputMaxMemberName);
 	}
 
 	// Level Scale
-	if(float LevelScale; ShaderNode->GetFloatAttribute(Noise::Attributes::LevelScale, LevelScale))
+	if(float LevelScale; ShaderNode->GetFloatAttribute(Noise::Attributes::LevelScale.ToString(), LevelScale))
 	{
 		const FName LevelScaleMemberName = GET_MEMBER_NAME_CHECKED(UMaterialExpressionNoise, LevelScale);
-		NoiseFactoryNode->AddFloatAttribute(LevelScaleMemberName, LevelScale);
-		NoiseFactoryNode->AddApplyAndFillDelegates<float>(LevelScaleMemberName, UMaterialExpressionNoise::StaticClass(), LevelScaleMemberName);
+		NoiseFactoryNode->AddFloatAttribute(LevelScaleMemberName.ToString(), LevelScale);
+		NoiseFactoryNode->AddApplyAndFillDelegates<float>(LevelScaleMemberName.ToString(), UMaterialExpressionNoise::StaticClass(), LevelScaleMemberName);
 	}
 
 	// Tiling
-	if(bool bTiling; ShaderNode->GetBooleanAttribute(Noise::Attributes::Tiling, bTiling))
+	if(bool bTiling; ShaderNode->GetBooleanAttribute(Noise::Attributes::Tiling.ToString(), bTiling))
 	{
 		const FName TilingMemberName = GET_MEMBER_NAME_CHECKED(UMaterialExpressionNoise, bTiling);
-		NoiseFactoryNode->AddBooleanAttribute(TilingMemberName, bTiling);
-		NoiseFactoryNode->AddApplyAndFillDelegates<bool>(TilingMemberName, UMaterialExpressionNoise::StaticClass(), TilingMemberName);
+		NoiseFactoryNode->AddBooleanAttribute(TilingMemberName.ToString(), bTiling);
+		NoiseFactoryNode->AddApplyAndFillDelegates<bool>(TilingMemberName.ToString(), UMaterialExpressionNoise::StaticClass(), TilingMemberName);
 	}
 
 	// Levels
-	if(int32 RepeatSize; ShaderNode->GetInt32Attribute(Noise::Attributes::RepeatSize, RepeatSize))
+	if(int32 RepeatSize; ShaderNode->GetInt32Attribute(Noise::Attributes::RepeatSize.ToString(), RepeatSize))
 	{
 		const FName RepeatSizeMemberName = GET_MEMBER_NAME_CHECKED(UMaterialExpressionNoise, RepeatSize);
-		NoiseFactoryNode->AddInt32Attribute(RepeatSizeMemberName, RepeatSize);
-		NoiseFactoryNode->AddApplyAndFillDelegates<int32>(RepeatSizeMemberName, UMaterialExpressionNoise::StaticClass(), RepeatSizeMemberName);
+		NoiseFactoryNode->AddInt32Attribute(RepeatSizeMemberName.ToString(), RepeatSize);
+		NoiseFactoryNode->AddApplyAndFillDelegates<int32>(RepeatSizeMemberName.ToString(), UMaterialExpressionNoise::StaticClass(), RepeatSizeMemberName);
 	}
 }
 
@@ -2213,35 +2213,35 @@ void UInterchangeGenericMaterialPipeline::HandleVectorNoiseNode(const UInterchan
 	}
 
 	// Noise Function
-	if(int32 NoiseFunction; ShaderNode->GetInt32Attribute(VectorNoise::Attributes::Function, NoiseFunction))
+	if(int32 NoiseFunction; ShaderNode->GetInt32Attribute(VectorNoise::Attributes::Function.ToString(), NoiseFunction))
 	{
 		const FName NoiseFunctionMemberName = GET_MEMBER_NAME_CHECKED(UMaterialExpressionVectorNoise, NoiseFunction);
-		NoiseFactoryNode->AddInt32Attribute(NoiseFunctionMemberName, NoiseFunction);
-		NoiseFactoryNode->AddApplyAndFillDelegates<int32>(NoiseFunctionMemberName, UMaterialExpressionVectorNoise::StaticClass(), NoiseFunctionMemberName);
+		NoiseFactoryNode->AddInt32Attribute(NoiseFunctionMemberName.ToString(), NoiseFunction);
+		NoiseFactoryNode->AddApplyAndFillDelegates<int32>(NoiseFunctionMemberName.ToString(), UMaterialExpressionVectorNoise::StaticClass(), NoiseFunctionMemberName);
 	}
 
 	// Quality
-	if(int32 Quality; ShaderNode->GetInt32Attribute(VectorNoise::Attributes::Quality, Quality))
+	if(int32 Quality; ShaderNode->GetInt32Attribute(VectorNoise::Attributes::Quality.ToString(), Quality))
 	{
 		const FName QualityMemberName = GET_MEMBER_NAME_CHECKED(UMaterialExpressionVectorNoise, Quality);
-		NoiseFactoryNode->AddInt32Attribute(QualityMemberName, Quality);
-		NoiseFactoryNode->AddApplyAndFillDelegates<int32>(QualityMemberName, UMaterialExpressionVectorNoise::StaticClass(), QualityMemberName);
+		NoiseFactoryNode->AddInt32Attribute(QualityMemberName.ToString(), Quality);
+		NoiseFactoryNode->AddApplyAndFillDelegates<int32>(QualityMemberName.ToString(), UMaterialExpressionVectorNoise::StaticClass(), QualityMemberName);
 	}
 
 	// Tiling
-	if(bool bTiling; ShaderNode->GetBooleanAttribute(VectorNoise::Attributes::Tiling, bTiling))
+	if(bool bTiling; ShaderNode->GetBooleanAttribute(VectorNoise::Attributes::Tiling.ToString(), bTiling))
 	{
 		const FName TilingMemberName = GET_MEMBER_NAME_CHECKED(UMaterialExpressionVectorNoise, bTiling);
-		NoiseFactoryNode->AddBooleanAttribute(TilingMemberName, bTiling);
-		NoiseFactoryNode->AddApplyAndFillDelegates<bool>(TilingMemberName, UMaterialExpressionVectorNoise::StaticClass(), TilingMemberName);
+		NoiseFactoryNode->AddBooleanAttribute(TilingMemberName.ToString(), bTiling);
+		NoiseFactoryNode->AddApplyAndFillDelegates<bool>(TilingMemberName.ToString(), UMaterialExpressionVectorNoise::StaticClass(), TilingMemberName);
 	}
 
 	// Tile Size
-	if(int32 TileSize; ShaderNode->GetInt32Attribute(VectorNoise::Attributes::Function, TileSize))
+	if(int32 TileSize; ShaderNode->GetInt32Attribute(VectorNoise::Attributes::Function.ToString(), TileSize))
 	{
 		const FName TileSizeMemberName = GET_MEMBER_NAME_CHECKED(UMaterialExpressionVectorNoise, TileSize);
-		NoiseFactoryNode->AddInt32Attribute(TileSizeMemberName, TileSize);
-		NoiseFactoryNode->AddApplyAndFillDelegates<int32>(TileSizeMemberName, UMaterialExpressionVectorNoise::StaticClass(), TileSizeMemberName);
+		NoiseFactoryNode->AddInt32Attribute(TileSizeMemberName.ToString(), TileSize);
+		NoiseFactoryNode->AddApplyAndFillDelegates<int32>(TileSizeMemberName.ToString(), UMaterialExpressionVectorNoise::StaticClass(), TileSizeMemberName);
 	}
 }
 
