Fix CIS build from a merge issue
#rnx

diff --git a/Engine/Plugins/Interchange/Runtime/Source/Import/Private/Material/InterchangeMaterialFactory.cpp b/Engine/Plugins/Interchange/Runtime/Source/Import/Private/Material/InterchangeMaterialFactory.cpp
--- a/Engine/Plugins/Interchange/Runtime/Source/Import/Private/Material/InterchangeMaterialFactory.cpp
+++ b/Engine/Plugins/Interchange/Runtime/Source/Import/Private/Material/InterchangeMaterialFactory.cpp
@@ -1284,8 +1284,8 @@ void UInterchangeMaterialFactory::SetupReimportedMaterialInstance(UMaterialInsta
 
 	for (const FString& InputName : Inputs)
 	{
-		const FName ParameterName = *InputName;
-		const FName AttributKey = UInterchangeShaderPortsAPI::MakeInputValueKey(InputName);
+		const FString ParameterName = *InputName;
+		const FString AttributKey = UInterchangeShaderPortsAPI::MakeInputValueKey(InputName);
 
 		FGuid Uid;
 		switch (UInterchangeShaderPortsAPI::GetInputType(&FactoryNode, InputName))
diff --git a/Engine/Plugins/Interchange/Runtime/Source/Import/Private/MaterialX/MaterialXUtils/MaterialXSurfaceShaderAbstract.cpp b/Engine/Plugins/Interchange/Runtime/Source/Import/Private/MaterialX/MaterialXUtils/MaterialXSurfaceShaderAbstract.cpp
--- a/Engine/Plugins/Interchange/Runtime/Source/Import/Private/MaterialX/MaterialXUtils/MaterialXSurfaceShaderAbstract.cpp
+++ b/Engine/Plugins/Interchange/Runtime/Source/Import/Private/MaterialX/MaterialXUtils/MaterialXSurfaceShaderAbstract.cpp
@@ -260,10 +260,10 @@ void FMaterialXSurfaceShaderAbstract::ConnectExtractInputToOutput(MaterialX::Nod
 		const int32 Index = mx::fromValueString<int>(InputIndex->getValueString());
 		switch(Index)
 		{
-		case 0: MaskShaderNode->AddBooleanAttribute(Mask::Attributes::R, true); break;
-		case 1: MaskShaderNode->AddBooleanAttribute(Mask::Attributes::G, true); break;
-		case 2: MaskShaderNode->AddBooleanAttribute(Mask::Attributes::B, true); break;
-		case 3: MaskShaderNode->AddBooleanAttribute(Mask::Attributes::A, true); break;
+		case 0: MaskShaderNode->AddBooleanAttribute(Mask::Attributes::R.ToString(), true); break;
+		case 1: MaskShaderNode->AddBooleanAttribute(Mask::Attributes::G.ToString(), true); break;
+		case 2: MaskShaderNode->AddBooleanAttribute(Mask::Attributes::B.ToString(), true); break;
+		case 3: MaskShaderNode->AddBooleanAttribute(Mask::Attributes::A.ToString(), true); break;
 		default:
 			UE_LOG(LogInterchangeImport, Warning, TEXT("Wrong index number for extract node, values are from [0-3]"));
 			break;
@@ -645,8 +645,8 @@ void FMaterialXSurfaceShaderAbstract::ConnectPositionInputToOutput(MaterialX::No
 		UInterchangeShaderNode* TransformNode = CreateShaderNode(UpstreamNode->getName().c_str() + FString(TEXT("_Transform")), TransformPosition::Name.ToString());
 		NodeToConnectTo = TransformNode;
 		InputToConnectTo = TransformPosition::Inputs::Input.ToString();
-		TransformNode->AddInt32Attribute(TransformPosition::Attributes::TransformSourceType, EMaterialPositionTransformSource::TRANSFORMPOSSOURCE_World);
-		TransformNode->AddInt32Attribute(TransformPosition::Attributes::TransformType, EMaterialPositionTransformSource::TRANSFORMPOSSOURCE_Local);
+		TransformNode->AddInt32Attribute(TransformPosition::Attributes::TransformSourceType.ToString(), EMaterialPositionTransformSource::TRANSFORMPOSSOURCE_World);
+		TransformNode->AddInt32Attribute(TransformPosition::Attributes::TransformType.ToString(), EMaterialPositionTransformSource::TRANSFORMPOSSOURCE_Local);
 		UInterchangeShaderPortsAPI::ConnectDefaultOuputToInput(ParentShaderNode, InputChannelName, TransformNode->GetUniqueID());
 	}
 
@@ -682,8 +682,8 @@ void FMaterialXSurfaceShaderAbstract::ConnectNormalInputToOutput(MaterialX::Node
 		UInterchangeShaderNode* TransformNode = CreateShaderNode(UpstreamNode->getName().c_str() + FString(TEXT("_Transform")), TransformVector::Name.ToString());
 		NodeToConnectTo = TransformNode;
 		InputToConnectTo = TransformVector::Inputs::Input.ToString();
-		TransformNode->AddInt32Attribute(TransformVector::Attributes::TransformSourceType, EMaterialVectorCoordTransformSource::TRANSFORMSOURCE_World);
-		TransformNode->AddInt32Attribute(TransformVector::Attributes::TransformType, EMaterialVectorCoordTransform::TRANSFORM_Local);
+		TransformNode->AddInt32Attribute(TransformVector::Attributes::TransformSourceType.ToString(), EMaterialVectorCoordTransformSource::TRANSFORMSOURCE_World);
+		TransformNode->AddInt32Attribute(TransformVector::Attributes::TransformType.ToString(), EMaterialVectorCoordTransform::TRANSFORM_Local);
 		UInterchangeShaderPortsAPI::ConnectDefaultOuputToInput(ParentShaderNode, InputChannelName, TransformNode->GetUniqueID());
 	}
 
@@ -719,8 +719,8 @@ void FMaterialXSurfaceShaderAbstract::ConnectTangentInputToOutput(MaterialX::Nod
 		UInterchangeShaderNode* TransformNode = CreateShaderNode(UpstreamNode->getName().c_str() + FString(TEXT("_Transform")), TransformVector::Name.ToString());
 		NodeToConnectTo = TransformNode;
 		InputToConnectTo = TransformVector::Inputs::Input.ToString();
-		TransformNode->AddInt32Attribute(TransformVector::Attributes::TransformSourceType, EMaterialVectorCoordTransformSource::TRANSFORMSOURCE_World);
-		TransformNode->AddInt32Attribute(TransformVector::Attributes::TransformType, EMaterialVectorCoordTransform::TRANSFORM_Local);
+		TransformNode->AddInt32Attribute(TransformVector::Attributes::TransformSourceType.ToString(), EMaterialVectorCoordTransformSource::TRANSFORMSOURCE_World);
+		TransformNode->AddInt32Attribute(TransformVector::Attributes::TransformType.ToString(), EMaterialVectorCoordTransform::TRANSFORM_Local);
 		UInterchangeShaderPortsAPI::ConnectDefaultOuputToInput(ParentShaderNode, InputChannelName, TransformNode->GetUniqueID());
 	}
 
@@ -763,8 +763,8 @@ void FMaterialXSurfaceShaderAbstract::ConnectBitangentInputToOutput(MaterialX::N
 		UInterchangeShaderNode* TransformNode = CreateShaderNode(UpstreamNode->getName().c_str() + FString(TEXT("_Transform")), TransformVector::Name.ToString());
 		NodeToConnectTo = TransformNode;
 		InputToConnectTo = TransformVector::Inputs::Input.ToString();
-		TransformNode->AddInt32Attribute(TransformVector::Attributes::TransformSourceType, EMaterialVectorCoordTransformSource::TRANSFORMSOURCE_World);
-		TransformNode->AddInt32Attribute(TransformVector::Attributes::TransformType, EMaterialVectorCoordTransform::TRANSFORM_Local);
+		TransformNode->AddInt32Attribute(TransformVector::Attributes::TransformSourceType.ToString(), EMaterialVectorCoordTransformSource::TRANSFORMSOURCE_World);
+		TransformNode->AddInt32Attribute(TransformVector::Attributes::TransformType.ToString(), EMaterialVectorCoordTransform::TRANSFORM_Local);
 		UInterchangeShaderPortsAPI::ConnectDefaultOuputToInput(ParentShaderNode, InputChannelName, TransformNode->GetUniqueID());
 	}
 
@@ -775,7 +775,7 @@ void FMaterialXSurfaceShaderAbstract::ConnectTimeInputToOutput(MaterialX::NodePt
 {
 	using namespace UE::Interchange::Materials::Standard::Nodes;
 	UInterchangeShaderNode* TimeNode = CreateShaderNode(UpstreamNode->getName().c_str(), TEXT("Time"));
-	TimeNode->AddBooleanAttribute(Time::Attributes::OverridePeriod, true);
+	TimeNode->AddBooleanAttribute(Time::Attributes::OverridePeriod.ToString(), true);
 
 	float FPS;
 	mx::InputPtr Input = UpstreamNode->getInput("fps");
@@ -789,7 +789,7 @@ void FMaterialXSurfaceShaderAbstract::ConnectTimeInputToOutput(MaterialX::NodePt
 	FPS = mx::fromValueString<float>(Input->getValueString());
 
 	//UE is a period
-	TimeNode->AddFloatAttribute(Time::Attributes::Period, 1.f / FPS);
+	TimeNode->AddFloatAttribute(Time::Attributes::Period.ToString(), 1.f / FPS);
 
 	UInterchangeShaderPortsAPI::ConnectDefaultOuputToInput(ParentShaderNode, InputChannelName, TimeNode->GetUniqueID());
 }
@@ -800,9 +800,9 @@ void FMaterialXSurfaceShaderAbstract::ConnectNoise3DInputToOutput(MaterialX::Nod
 
 	// MaterialX defines the Noise3d as Perlin Noise which is multiplied by the Amplitude then Added to Pivot
 	UInterchangeShaderNode* NoiseNode = CreateShaderNode(UpstreamNode->getName().c_str(), Noise::Name.ToString());
-	NoiseNode->AddInt32Attribute(Noise::Attributes::Function, ENoiseFunction::NOISEFUNCTION_GradientTex);
-	NoiseNode->AddBooleanAttribute(Noise::Attributes::Turbulence, false);
-	NoiseNode->AddFloatAttribute(Noise::Attributes::OutputMin, 0);
+	NoiseNode->AddInt32Attribute(Noise::Attributes::Function.ToString(), ENoiseFunction::NOISEFUNCTION_GradientTex);
+	NoiseNode->AddBooleanAttribute(Noise::Attributes::Turbulence.ToString(), false);
+	NoiseNode->AddFloatAttribute(Noise::Attributes::OutputMin.ToString(), 0);
 
 	UInterchangeShaderNode* NodeToConnect = NoiseNode;
 
@@ -848,7 +848,7 @@ void FMaterialXSurfaceShaderAbstract::ConnectCellNoise3DInputToOutput(MaterialX:
 
 	// Let's use a vector noise for this one, the only one that is close to MaterialX implementation
 	UInterchangeShaderNode* NoiseNode = CreateShaderNode(UpstreamNode->getName().c_str(), VectorNoise::Name.ToString());
-	NoiseNode->AddInt32Attribute(VectorNoise::Attributes::Function, EVectorNoiseFunction::VNF_CellnoiseALU);
+	NoiseNode->AddInt32Attribute(VectorNoise::Attributes::Function.ToString(), EVectorNoiseFunction::VNF_CellnoiseALU);
 
 	UInterchangeShaderPortsAPI::ConnectDefaultOuputToInput(ParentShaderNode, InputChannelName, NoiseNode->GetUniqueID());
 }
@@ -859,7 +859,7 @@ void FMaterialXSurfaceShaderAbstract::ConnectWorleyNoise3DInputToOutput(Material
 
 	//Also called Voronoi, the implementation is a bit different in UE, especially we don't have access to the jitter
 	UInterchangeShaderNode* NoiseNode = CreateShaderNode(UpstreamNode->getName().c_str(), Noise::Name.ToString());
-	NoiseNode->AddInt32Attribute(Noise::Attributes::Function, ENoiseFunction::NOISEFUNCTION_VoronoiALU);
+	NoiseNode->AddInt32Attribute(Noise::Attributes::Function.ToString(), ENoiseFunction::NOISEFUNCTION_VoronoiALU);
 
 	UInterchangeShaderPortsAPI::ConnectDefaultOuputToInput(ParentShaderNode, InputChannelName, NoiseNode->GetUniqueID());
 }
@@ -957,14 +957,14 @@ void FMaterialXSurfaceShaderAbstract::ConnectBlurInputToOutput(MaterialX::NodePt
 					// By default TextureSampleBox uses gaussian filter
 					if(InputKernel->getValueString() == "box")
 					{
-						TextureShaderNode->AddInt32Attribute(TextureSampleBlur::Attributes::Filter, 0);
+						TextureShaderNode->AddInt32Attribute(TextureSampleBlur::Attributes::Filter.ToString(), 0);
 					}
 				}
 
 				if(mx::InputPtr InputKernel = UpstreamNode->getInput("size"))
 				{
 					float KernelSize = mx::fromValueString<float>(InputKernel->getValueString());
-					TextureShaderNode->AddFloatAttribute(TextureSampleBlur::Attributes::KernelSize, KernelSize);
+					TextureShaderNode->AddFloatAttribute(TextureSampleBlur::Attributes::KernelSize.ToString(), KernelSize);
 				}
 			}
 			else
@@ -989,10 +989,10 @@ UInterchangeShaderNode* FMaterialXSurfaceShaderAbstract::CreateMaskShaderNode(ui
 	bool bA = (0b0001 & RGBA) >> 0;
 	using namespace UE::Interchange::Materials::Standard::Nodes;
 	UInterchangeShaderNode* MaskShaderNode = CreateShaderNode(NodeName, Mask::Name.ToString());
-	MaskShaderNode->AddBooleanAttribute(Mask::Attributes::R, bR);
-	MaskShaderNode->AddBooleanAttribute(Mask::Attributes::G, bG);
-	MaskShaderNode->AddBooleanAttribute(Mask::Attributes::B, bB);
-	MaskShaderNode->AddBooleanAttribute(Mask::Attributes::A, bA);
+	MaskShaderNode->AddBooleanAttribute(Mask::Attributes::R.ToString(), bR);
+	MaskShaderNode->AddBooleanAttribute(Mask::Attributes::G.ToString(), bG);
+	MaskShaderNode->AddBooleanAttribute(Mask::Attributes::B.ToString(), bB);
+	MaskShaderNode->AddBooleanAttribute(Mask::Attributes::A.ToString(), bA);
 
 	return MaskShaderNode;
 }
diff --git a/Engine/Plugins/Interchange/Runtime/Source/Pipelines/Private/InterchangeGenericMaterialPipeline.cpp b/Engine/Plugins/Interchange/Runtime/Source/Pipelines/Private/InterchangeGenericMaterialPipeline.cpp
--- a/Engine/Plugins/Interchange/Runtime/Source/Pipelines/Private/InterchangeGenericMaterialPipeline.cpp
+++ b/Engine/Plugins/Interchange/Runtime/Source/Pipelines/Private/InterchangeGenericMaterialPipeline.cpp
@@ -418,8 +418,8 @@ namespace UE::Interchange::InterchangeGenericMaterialPipeline::Private
 
 		FunctionCallExpression.SetCustomExpressionClassName(CustomExpressionClass->GetName());
 		
-		FunctionCallExpression.AddStringAttribute(MaterialFunctionMemberName, MaterialFunctionPath);
-		FunctionCallExpression.AddApplyAndFillDelegates<FString>(MaterialFunctionMemberName, CustomExpressionClass, MaterialFunctionMemberName);
+		FunctionCallExpression.AddStringAttribute(MaterialFunctionMemberName.ToString(), MaterialFunctionPath);
+		FunctionCallExpression.AddApplyAndFillDelegates<FString>(MaterialFunctionMemberName.ToString(), CustomExpressionClass, MaterialFunctionMemberName);
 	}
 
 	UInterchangeMaterialExpressionFactoryNode* CreateExpressionWithMaterialFunction(
@@ -671,7 +671,7 @@ void UInterchangeGenericMaterialPipeline::ExecutePipeline(UInterchangeBaseNodeCo
 
 		for (const FString& InputName : Inputs)
 		{
-			FName InputValueKey = UInterchangeShaderPortsAPI::MakeInputValueKey(InputName);
+			FString InputValueKey = UInterchangeShaderPortsAPI::MakeInputValueKey(InputName);
 
 			switch (UInterchangeShaderPortsAPI::GetInputType(MaterialNode, InputName))
 			{
@@ -1737,35 +1737,35 @@ void UInterchangeGenericMaterialPipeline::HandleTextureSampleBlurNode(const UInt
 	using namespace UE::Interchange::Materials::Standard::Nodes;
 
 	// KernelSize
-	if(float KernelSize; ShaderNode->GetFloatAttribute(TextureSampleBlur::Attributes::KernelSize, KernelSize))
+	if(float KernelSize; ShaderNode->GetFloatAttribute(TextureSampleBlur::Attributes::KernelSize.ToString(), KernelSize))
 	{
 		const FName KernelSizeMemberName = GET_MEMBER_NAME_CHECKED(UMaterialExpressionTextureSampleParameterBlur, KernelSize);
-		TextureSampleFactoryNode->AddFloatAttribute(KernelSizeMemberName, KernelSize);
-		TextureSampleFactoryNode->AddApplyAndFillDelegates<float>(KernelSizeMemberName, UMaterialExpressionTextureSampleParameterBlur::StaticClass(), KernelSizeMemberName);
+		TextureSampleFactoryNode->AddFloatAttribute(KernelSizeMemberName.ToString(), KernelSize);
+		TextureSampleFactoryNode->AddApplyAndFillDelegates<float>(KernelSizeMemberName.ToString(), UMaterialExpressionTextureSampleParameterBlur::StaticClass(), KernelSizeMemberName);
 	}
 
 	// FilterSize
-	if(float FilterSize; ShaderNode->GetFloatAttribute(TextureSampleBlur::Attributes::FilterSize, FilterSize))
+	if(float FilterSize; ShaderNode->GetFloatAttribute(TextureSampleBlur::Attributes::FilterSize.ToString(), FilterSize))
 	{
 		const FName FilterSizeMemberName = GET_MEMBER_NAME_CHECKED(UMaterialExpressionTextureSampleParameterBlur, FilterSize);
-		TextureSampleFactoryNode->AddFloatAttribute(FilterSizeMemberName, FilterSize);
-		TextureSampleFactoryNode->AddApplyAndFillDelegates<float>(FilterSizeMemberName, UMaterialExpressionTextureSampleParameterBlur::StaticClass(), FilterSizeMemberName);
+		TextureSampleFactoryNode->AddFloatAttribute(FilterSizeMemberName.ToString(), FilterSize);
+		TextureSampleFactoryNode->AddApplyAndFillDelegates<float>(FilterSizeMemberName.ToString(), UMaterialExpressionTextureSampleParameterBlur::StaticClass(), FilterSizeMemberName);
 	}
 
 	// FilterOffset
-	if(float FilterOffset; ShaderNode->GetFloatAttribute(TextureSampleBlur::Attributes::FilterOffset, FilterOffset))
+	if(float FilterOffset; ShaderNode->GetFloatAttribute(TextureSampleBlur::Attributes::FilterOffset.ToString(), FilterOffset))
 	{
 		const FName FilterOffsetMemberName = GET_MEMBER_NAME_CHECKED(UMaterialExpressionTextureSampleParameterBlur, FilterOffset);
-		TextureSampleFactoryNode->AddFloatAttribute(FilterOffsetMemberName, FilterOffset);
-		TextureSampleFactoryNode->AddApplyAndFillDelegates<float>(FilterOffsetMemberName, UMaterialExpressionTextureSampleParameterBlur::StaticClass(), FilterOffsetMemberName);
+		TextureSampleFactoryNode->AddFloatAttribute(FilterOffsetMemberName.ToString(), FilterOffset);
+		TextureSampleFactoryNode->AddApplyAndFillDelegates<float>(FilterOffsetMemberName.ToString(), UMaterialExpressionTextureSampleParameterBlur::StaticClass(), FilterOffsetMemberName);
 	}
 
 	// Filter
-	if(int32 Filter; ShaderNode->GetInt32Attribute(TextureSampleBlur::Attributes::Filter, Filter))
+	if(int32 Filter; ShaderNode->GetInt32Attribute(TextureSampleBlur::Attributes::Filter.ToString(), Filter))
 	{
 		const FName FilterMemberName = GET_MEMBER_NAME_CHECKED(UMaterialExpressionTextureSampleParameterBlur, Filter);
-		TextureSampleFactoryNode->AddInt32Attribute(FilterMemberName, Filter);
-		TextureSampleFactoryNode->AddApplyAndFillDelegates<int>(FilterMemberName, UMaterialExpressionTextureSampleParameterBlur::StaticClass(), FilterMemberName);
+		TextureSampleFactoryNode->AddInt32Attribute(FilterMemberName.ToString(), Filter);
+		TextureSampleFactoryNode->AddApplyAndFillDelegates<int>(FilterMemberName.ToString(), UMaterialExpressionTextureSampleParameterBlur::StaticClass(), FilterMemberName);
 	}
 }
 
