Implemented ability to reimport a file into a level - #1
  - Added the ability to apply reimport policy on actors based on attributes with delegates
  - Fixed the copy of factory node's attributes from one node to another when the attributes are associated with delegates

diff --git a/Engine/Plugins/Experimental/Enterprise/DatasmithInterchange/Source/DatasmithInterchange/Private/InterchangeDatasmithAreaLightFactory.cpp b/Engine/Plugins/Experimental/Enterprise/DatasmithInterchange/Source/DatasmithInterchange/Private/InterchangeDatasmithAreaLightFactory.cpp
--- a/Engine/Plugins/Experimental/Enterprise/DatasmithInterchange/Source/DatasmithInterchange/Private/InterchangeDatasmithAreaLightFactory.cpp
+++ b/Engine/Plugins/Experimental/Enterprise/DatasmithInterchange/Source/DatasmithInterchange/Private/InterchangeDatasmithAreaLightFactory.cpp
@@ -29,45 +29,43 @@ UClass* UInterchangeDatasmithAreaLightFactory::GetFactoryClass() const
 	return LightShapeBlueprint->GeneratedClass;
 }
 
-UObject* UInterchangeDatasmithAreaLightFactory::ImportSceneObject_GameThread(const UInterchangeFactoryBase::FImportSceneObjectsParams& CreateSceneObjectsParams)
+UObject* UInterchangeDatasmithAreaLightFactory::ProcessActor(AActor& SpawnedActor, const UInterchangeActorFactoryNode& FactoryNode, const UInterchangeBaseNodeContainer& NodeContainer)
 {
 	//using namespace UE::Interchange::ActorHelper;
 
-	UInterchangeDatasmithAreaLightFactoryNode* AreaLightFactoryNode = Cast<UInterchangeDatasmithAreaLightFactoryNode>(CreateSceneObjectsParams.FactoryNode);
+	const UInterchangeDatasmithAreaLightFactoryNode* AreaLightFactoryNode = Cast<UInterchangeDatasmithAreaLightFactoryNode>(&FactoryNode);
 	if (!AreaLightFactoryNode)
 	{
 		return nullptr;
 	}
 
-	ADatasmithAreaLightActor* SpawnedAreaLightActor = Cast<ADatasmithAreaLightActor>(UE::Interchange::ActorHelper::SpawnFactoryActor(CreateSceneObjectsParams));
-	if (!SpawnedAreaLightActor)
+	ADatasmithAreaLightActor* AreaLightActor = Cast<ADatasmithAreaLightActor>(&SpawnedActor);
+	if (!AreaLightActor)
 	{
 		return nullptr;
 	}
 
 	// Update AreaLight properties and rebuild the actor.
 	{
-		SpawnedAreaLightActor->UnregisterAllComponents(true);
+		AreaLightActor->UnregisterAllComponents(true);
 
 		// Find the IES Texture asset and apply it.
 		FString IESTextureUid;
 		if (AreaLightFactoryNode->GetCustomIESTexture(IESTextureUid))
 		{
-			const UInterchangeBaseNodeContainer* NodeContainer = CreateSceneObjectsParams.NodeContainer;
-
-			if (const UInterchangeBaseNode* IESTextureNode = NodeContainer->GetNode(IESTextureUid))
+			if (const UInterchangeBaseNode* IESTextureNode = NodeContainer.GetNode(IESTextureUid))
 			{
 				TArray<FString> TargetNodesUid;
 				IESTextureNode->GetTargetNodeUids(TargetNodesUid);
 				for (const FString& TargetNodeUid : TargetNodesUid)
 				{
-					if (UInterchangeTextureFactoryNode* TextureFactoryNode = Cast<UInterchangeTextureFactoryNode>(NodeContainer->GetFactoryNode(TargetNodeUid)))
+					if (UInterchangeTextureFactoryNode* TextureFactoryNode = Cast<UInterchangeTextureFactoryNode>(NodeContainer.GetFactoryNode(TargetNodeUid)))
 					{
 						FSoftObjectPath TextureReferenceObject;
 						TextureFactoryNode->GetCustomReferenceObject(TextureReferenceObject);
 						if (UTextureLightProfile* TextureLightProfile = Cast<UTextureLightProfile>(TextureReferenceObject.TryLoad()))
 						{
-							SpawnedAreaLightActor->IESTexture = TextureLightProfile;
+							AreaLightActor->IESTexture = TextureLightProfile;
 							break;
 						}
 					}
@@ -76,13 +74,13 @@ UObject* UInterchangeDatasmithAreaLightFactory::ImportSceneObject_GameThread(con
 		}
 
 		// Apply all simple attributes.
-		AreaLightFactoryNode->ApplyAllCustomAttributeToObject(SpawnedAreaLightActor);
+		AreaLightFactoryNode->ApplyAllCustomAttributeToObject(AreaLightActor);
 
-		SpawnedAreaLightActor->RegisterAllComponents();
+		AreaLightActor->RegisterAllComponents();
 #if WITH_EDITOR
-		SpawnedAreaLightActor->RerunConstructionScripts();
+		AreaLightActor->RerunConstructionScripts();
 #endif //WITH_EDITOR
 	}
 
-	return SpawnedAreaLightActor;
+	return AreaLightActor;
 };
\ No newline at end of file
diff --git a/Engine/Plugins/Experimental/Enterprise/DatasmithInterchange/Source/DatasmithInterchange/Private/InterchangeDatasmithAreaLightFactoryNode.cpp b/Engine/Plugins/Experimental/Enterprise/DatasmithInterchange/Source/DatasmithInterchange/Private/InterchangeDatasmithAreaLightFactoryNode.cpp
--- a/Engine/Plugins/Experimental/Enterprise/DatasmithInterchange/Source/DatasmithInterchange/Private/InterchangeDatasmithAreaLightFactoryNode.cpp
+++ b/Engine/Plugins/Experimental/Enterprise/DatasmithInterchange/Source/DatasmithInterchange/Private/InterchangeDatasmithAreaLightFactoryNode.cpp
@@ -213,3 +213,27 @@ bool UInterchangeDatasmithAreaLightFactoryNode::SetCustomSpotlightOuterAngle(con
 {
 	IMPLEMENT_NODE_ATTRIBUTE_SETTER(UInterchangeDatasmithAreaLightFactoryNode, SpotlightOuterAngle, float, ADatasmithAreaLightActor);
 }
+
+void UInterchangeDatasmithAreaLightFactoryNode::CopyWithObject(const UInterchangeFactoryBaseNode* SourceNode, UObject* Object)
+{
+	Super::CopyWithObject(SourceNode, Object);
+
+	if (const UInterchangeDatasmithAreaLightFactoryNode* ActorFactoryNode = Cast<UInterchangeDatasmithAreaLightFactoryNode>(SourceNode))
+	{
+		COPY_NODE_DELEGATES(ActorFactoryNode, LightType, EDatasmithAreaLightActorType, ADatasmithAreaLightActor)
+		COPY_NODE_DELEGATES(ActorFactoryNode, LightShape, EDatasmithAreaLightActorShape, ADatasmithAreaLightActor)
+		COPY_NODE_DELEGATES(ActorFactoryNode, Dimensions, FVector2D, ADatasmithAreaLightActor)
+		COPY_NODE_DELEGATES(ActorFactoryNode, Intensity, float, ADatasmithAreaLightActor)
+		COPY_NODE_DELEGATES(ActorFactoryNode, IntensityUnits, ELightUnits, ADatasmithAreaLightActor)
+		COPY_NODE_DELEGATES(ActorFactoryNode, Color, FLinearColor, ADatasmithAreaLightActor)
+		COPY_NODE_DELEGATES(ActorFactoryNode, Temperature, float, ADatasmithAreaLightActor)
+		COPY_NODE_DELEGATES(ActorFactoryNode, bUseIESBrightness, bool, ADatasmithAreaLightActor)
+		COPY_NODE_DELEGATES(ActorFactoryNode, IESBrightnessScale, float, ADatasmithAreaLightActor)
+		COPY_NODE_DELEGATES(ActorFactoryNode, Rotation, FRotator, ADatasmithAreaLightActor)
+		COPY_NODE_DELEGATES(ActorFactoryNode, SourceRadius, float, ADatasmithAreaLightActor)
+		COPY_NODE_DELEGATES(ActorFactoryNode, SourceLength, float, ADatasmithAreaLightActor)
+		COPY_NODE_DELEGATES(ActorFactoryNode, AttenuationRadius, float, ADatasmithAreaLightActor)
+		COPY_NODE_DELEGATES(ActorFactoryNode, SpotlightInnerAngle, float, ADatasmithAreaLightActor)
+		COPY_NODE_DELEGATES(ActorFactoryNode, SpotlightOuterAngle, float, ADatasmithAreaLightActor)
+	}
+}
diff --git a/Engine/Plugins/Experimental/Enterprise/DatasmithInterchange/Source/DatasmithInterchange/Public/InterchangeDatasmithAreaLightFactory.h b/Engine/Plugins/Experimental/Enterprise/DatasmithInterchange/Source/DatasmithInterchange/Public/InterchangeDatasmithAreaLightFactory.h
--- a/Engine/Plugins/Experimental/Enterprise/DatasmithInterchange/Source/DatasmithInterchange/Public/InterchangeDatasmithAreaLightFactory.h
+++ b/Engine/Plugins/Experimental/Enterprise/DatasmithInterchange/Source/DatasmithInterchange/Public/InterchangeDatasmithAreaLightFactory.h
@@ -21,8 +21,15 @@ public:
 
 	virtual UClass* GetFactoryClass() const override;
 
-	virtual UObject* ImportSceneObject_GameThread(const FImportSceneObjectsParams& Arguments) override;
-
 	// Interchange factory base interface end
 	//////////////////////////////////////////////////////////////////////////
+
+protected:
+	//////////////////////////////////////////////////////////////////////////
+	// Interchange actor factory interface begin
+
+	virtual UObject* ProcessActor(class AActor& SpawnedActor, const UInterchangeActorFactoryNode& FactoryNode, const UInterchangeBaseNodeContainer& NodeContainer) override;
+
+	// Interchange actor factory interface end
+	//////////////////////////////////////////////////////////////////////////
 };
\ No newline at end of file
diff --git a/Engine/Plugins/Experimental/Enterprise/DatasmithInterchange/Source/DatasmithInterchange/Public/InterchangeDatasmithAreaLightFactoryNode.h b/Engine/Plugins/Experimental/Enterprise/DatasmithInterchange/Source/DatasmithInterchange/Public/InterchangeDatasmithAreaLightFactoryNode.h
--- a/Engine/Plugins/Experimental/Enterprise/DatasmithInterchange/Source/DatasmithInterchange/Public/InterchangeDatasmithAreaLightFactoryNode.h
+++ b/Engine/Plugins/Experimental/Enterprise/DatasmithInterchange/Source/DatasmithInterchange/Public/InterchangeDatasmithAreaLightFactoryNode.h
@@ -127,6 +127,8 @@ public:
 	UFUNCTION(BlueprintCallable, Category = "Interchange | Node | Datasmith")
 	bool SetCustomSpotlightOuterAngle(const float& AttributeValue, bool bAddApplyDelegate = true);
 
+	virtual void CopyWithObject(const UInterchangeFactoryBaseNode* SourceNode, UObject* Object) override;
+
 private:
 	static  UE::Interchange::FAttributeKey Macro_CustomLightTypeKey;
 	static  UE::Interchange::FAttributeKey Macro_CustomLightShapeKey;
diff --git a/Engine/Plugins/Interchange/Runtime/Source/FactoryNodes/InterchangeFactoryNodes.Build.cs b/Engine/Plugins/Interchange/Runtime/Source/FactoryNodes/InterchangeFactoryNodes.Build.cs
--- a/Engine/Plugins/Interchange/Runtime/Source/FactoryNodes/InterchangeFactoryNodes.Build.cs
+++ b/Engine/Plugins/Interchange/Runtime/Source/FactoryNodes/InterchangeFactoryNodes.Build.cs
@@ -14,6 +14,7 @@ namespace UnrealBuildTool.Rules
 					"CoreUObject",
 					"Engine",
 					"InterchangeCore",
+					"InterchangeEngine",
 					"InterchangeNodes",
 					"TextureUtilitiesCommon",
 					"VariantManagerContent",
diff --git a/Engine/Plugins/Interchange/Runtime/Source/FactoryNodes/Private/InterchangeActorFactoryNode.cpp b/Engine/Plugins/Interchange/Runtime/Source/FactoryNodes/Private/InterchangeActorFactoryNode.cpp
--- a/Engine/Plugins/Interchange/Runtime/Source/FactoryNodes/Private/InterchangeActorFactoryNode.cpp
+++ b/Engine/Plugins/Interchange/Runtime/Source/FactoryNodes/Private/InterchangeActorFactoryNode.cpp
@@ -5,7 +5,8 @@
 #include UE_INLINE_GENERATED_CPP_BY_NAME(InterchangeActorFactoryNode)
 
 #if WITH_ENGINE
-	#include "GameFramework/Actor.h"
+#include "Components/SceneComponent.h"
+#include "GameFramework/Actor.h"
 #endif
 
 UClass* UInterchangeActorFactoryNode::GetObjectClass() const
@@ -26,3 +27,50 @@ UClass* UInterchangeActorFactoryNode::GetObjectClass() const
 	return nullptr;
 #endif
 }
+
+bool UInterchangeActorFactoryNode::GetCustomGlobalTransform(FTransform& AttributeValue) const
+{
+	IMPLEMENT_NODE_ATTRIBUTE_GETTER(GlobalTransform, FTransform);
+}
+
+bool UInterchangeActorFactoryNode::SetCustomGlobalTransform(const FTransform& AttributeValue, bool bAddApplyDelegate)
+{
+	IMPLEMENT_NODE_ATTRIBUTE_SETTER_WITH_CUSTOM_DELEGATE_WITH_CUSTOM_CLASS(UInterchangeActorFactoryNode, GlobalTransform, FTransform, USceneComponent);
+}
+
+bool UInterchangeActorFactoryNode::ApplyCustomGlobalTransformToAsset(UObject* Asset) const
+{
+	FTransform GlobalTransform;
+	if (GetCustomGlobalTransform(GlobalTransform))
+	{
+		if (USceneComponent* Component = Cast<USceneComponent>(Asset))
+		{
+			Component->SetWorldTransform(GlobalTransform);
+			return true;
+		}
+	}
+
+	return false;
+}
+
+bool UInterchangeActorFactoryNode::FillCustomGlobalTransformFromAsset(UObject* Asset)
+{
+	if (const USceneComponent* Component = Cast<USceneComponent>(Asset))
+	{
+		FTransform Transform = Component->GetComponentToWorld();
+		return this->SetCustomGlobalTransform(Transform, false);
+	}
+
+	return false;
+}
+
+void UInterchangeActorFactoryNode::CopyWithObject(const UInterchangeFactoryBaseNode* SourceNode, UObject* Object)
+{
+	Super::CopyWithObject(SourceNode, Object);
+
+	if (const UInterchangeActorFactoryNode* ActorFactoryNode = Cast<UInterchangeActorFactoryNode>(SourceNode))
+	{
+		COPY_NODE_DELEGATES_WITH_CUSTOM_DELEGATE(ActorFactoryNode, UInterchangeActorFactoryNode, GlobalTransform, FTransform, USceneComponent::StaticClass())
+	}
+}
+
diff --git a/Engine/Plugins/Interchange/Runtime/Source/FactoryNodes/Private/InterchangeLightFactoryNode.cpp b/Engine/Plugins/Interchange/Runtime/Source/FactoryNodes/Private/InterchangeLightFactoryNode.cpp
--- a/Engine/Plugins/Interchange/Runtime/Source/FactoryNodes/Private/InterchangeLightFactoryNode.cpp
+++ b/Engine/Plugins/Interchange/Runtime/Source/FactoryNodes/Private/InterchangeLightFactoryNode.cpp
@@ -45,6 +45,19 @@ bool UInterchangeBaseLightFactoryNode::SetCustomUseTemperature(bool AttributeVal
 	IMPLEMENT_NODE_ATTRIBUTE_SETTER(UInterchangeBaseLightFactoryNode, bUseTemperature, bool, ULightComponent)
 }
 
+void UInterchangeBaseLightFactoryNode::CopyWithObject(const UInterchangeFactoryBaseNode* SourceNode, UObject* Object)
+{
+	Super::CopyWithObject(SourceNode, Object);
+
+	if (const UInterchangeBaseLightFactoryNode* LightFactoryNode = Cast<UInterchangeBaseLightFactoryNode>(SourceNode))
+	{
+		COPY_NODE_DELEGATES(LightFactoryNode, LightColor, FColor, ULightComponent)
+		COPY_NODE_DELEGATES(LightFactoryNode, Intensity, float, ULightComponent)
+		COPY_NODE_DELEGATES(LightFactoryNode, Temperature, float, ULightComponent)
+		COPY_NODE_DELEGATES(LightFactoryNode, bUseTemperature, bool, ULightComponent)
+	}
+}
+
 bool UInterchangeLightFactoryNode::GetCustomIntensityUnits(ELightUnits& AttributeValue) const
 {
 	IMPLEMENT_NODE_ATTRIBUTE_GETTER(IntensityUnits, ELightUnits)
@@ -75,6 +88,17 @@ bool UInterchangeLightFactoryNode::SetCustomIESTexture(const FString& AttributeV
 	IMPLEMENT_NODE_ATTRIBUTE_SETTER_NODELEGATE(IESTexture, FString)
 }
 
+void UInterchangeLightFactoryNode::CopyWithObject(const UInterchangeFactoryBaseNode* SourceNode, UObject* Object)
+{
+	Super::CopyWithObject(SourceNode, Object);
+
+	if (const UInterchangeLightFactoryNode* LightFactoryNode = Cast<UInterchangeLightFactoryNode>(SourceNode))
+	{
+		COPY_NODE_DELEGATES(LightFactoryNode, IntensityUnits, ELightUnits, ULocalLightComponent)
+		COPY_NODE_DELEGATES(LightFactoryNode, AttenuationRadius, float, ULocalLightComponent)
+	}
+}
+
 bool UInterchangePointLightFactoryNode::GetCustomUseInverseSquaredFalloff(bool& AttributeValue) const
 {
 	IMPLEMENT_NODE_ATTRIBUTE_GETTER(bUseInverseSquaredFalloff, bool)
@@ -95,6 +119,17 @@ bool UInterchangePointLightFactoryNode::SetCustomLightFalloffExponent(float Attr
 	IMPLEMENT_NODE_ATTRIBUTE_SETTER(UInterchangePointLightFactoryNode, LightFalloffExponent, float, UPointLightComponent)
 }
 
+void UInterchangePointLightFactoryNode::CopyWithObject(const UInterchangeFactoryBaseNode* SourceNode, UObject* Object)
+{
+	Super::CopyWithObject(SourceNode, Object);
+
+	if (const UInterchangePointLightFactoryNode* LightFactoryNode = Cast<UInterchangePointLightFactoryNode>(SourceNode))
+	{
+		COPY_NODE_DELEGATES(LightFactoryNode, bUseInverseSquaredFalloff, bool, UPointLightComponent)
+		COPY_NODE_DELEGATES(LightFactoryNode, LightFalloffExponent, float, UPointLightComponent)
+	}
+}
+
 bool UInterchangeSpotLightFactoryNode::GetCustomInnerConeAngle(float& AttributeValue) const
 {
 	IMPLEMENT_NODE_ATTRIBUTE_GETTER(InnerConeAngle, float)
@@ -115,6 +150,17 @@ bool UInterchangeSpotLightFactoryNode::SetCustomOuterConeAngle(float AttributeVa
 	IMPLEMENT_NODE_ATTRIBUTE_SETTER(UInterchangeSpotLightFactoryNode, OuterConeAngle, float, USpotLightComponent)
 }
 
+void UInterchangeSpotLightFactoryNode::CopyWithObject(const UInterchangeFactoryBaseNode* SourceNode, UObject* Object)
+{
+	Super::CopyWithObject(SourceNode, Object);
+
+	if (const UInterchangeSpotLightFactoryNode* LightFactoryNode = Cast<UInterchangeSpotLightFactoryNode>(SourceNode))
+	{
+		COPY_NODE_DELEGATES(LightFactoryNode, InnerConeAngle, float, USpotLightComponent)
+		COPY_NODE_DELEGATES(LightFactoryNode, OuterConeAngle, float, USpotLightComponent)
+	}
+}
+
 bool UInterchangeRectLightFactoryNode::GetCustomSourceWidth(float& AttributeValue) const
 {
 	IMPLEMENT_NODE_ATTRIBUTE_GETTER(SourceWidth, float)
@@ -133,4 +179,15 @@ bool UInterchangeRectLightFactoryNode::GetCustomSourceHeight(float& AttributeVal
 bool UInterchangeRectLightFactoryNode::SetCustomSourceHeight(float AttributeValue, bool bAddApplyDelegate)
 {
 	IMPLEMENT_NODE_ATTRIBUTE_SETTER(UInterchangeRectLightFactoryNode, SourceHeight, float, URectLightComponent)
-}
\ No newline at end of file
+}
+
+void UInterchangeRectLightFactoryNode::CopyWithObject(const UInterchangeFactoryBaseNode* SourceNode, UObject* Object)
+{
+	Super::CopyWithObject(SourceNode, Object);
+
+	if (const UInterchangeRectLightFactoryNode* LightFactoryNode = Cast<UInterchangeRectLightFactoryNode>(SourceNode))
+	{
+		COPY_NODE_DELEGATES(LightFactoryNode, SourceWidth, float, URectLightComponent)
+		COPY_NODE_DELEGATES(LightFactoryNode, SourceHeight, float, URectLightComponent)
+	}
+}
diff --git a/Engine/Plugins/Interchange/Runtime/Source/FactoryNodes/Private/InterchangeMeshFactoryNode.cpp b/Engine/Plugins/Interchange/Runtime/Source/FactoryNodes/Private/InterchangeMeshFactoryNode.cpp
--- a/Engine/Plugins/Interchange/Runtime/Source/FactoryNodes/Private/InterchangeMeshFactoryNode.cpp
+++ b/Engine/Plugins/Interchange/Runtime/Source/FactoryNodes/Private/InterchangeMeshFactoryNode.cpp
@@ -339,4 +339,23 @@ bool UInterchangeMeshFactoryNode::ApplyCustomRemoveDegeneratesToAsset(UObject* A
 bool UInterchangeMeshFactoryNode::FillCustomRemoveDegeneratesFromAsset(UObject* Asset)
 {
 	IMPLEMENT_MESH_BUILD_ASSET_TO_VALUE(RemoveDegenerates, bRemoveDegenerates);
-}
\ No newline at end of file
+}
+
+void UInterchangeMeshFactoryNode::CopyWithObject(const UInterchangeFactoryBaseNode* SourceNode, UObject* Object)
+{
+	Super::CopyWithObject(SourceNode, Object);
+
+	if (const UInterchangeMeshFactoryNode* MeshFactoryNode = Cast<UInterchangeMeshFactoryNode>(SourceNode))
+	{
+		UClass* Class = GetObjectClass();
+
+		COPY_NODE_DELEGATES_WITH_CUSTOM_DELEGATE(MeshFactoryNode, UInterchangeMeshFactoryNode, RecomputeNormals, bool, Class)
+		COPY_NODE_DELEGATES_WITH_CUSTOM_DELEGATE(MeshFactoryNode, UInterchangeMeshFactoryNode, RecomputeTangents, bool, Class)
+		COPY_NODE_DELEGATES_WITH_CUSTOM_DELEGATE(MeshFactoryNode, UInterchangeMeshFactoryNode, UseMikkTSpace, bool, Class)
+		COPY_NODE_DELEGATES_WITH_CUSTOM_DELEGATE(MeshFactoryNode, UInterchangeMeshFactoryNode, ComputeWeightedNormals, bool, Class)
+		COPY_NODE_DELEGATES_WITH_CUSTOM_DELEGATE(MeshFactoryNode, UInterchangeMeshFactoryNode, UseHighPrecisionTangentBasis, bool, Class)
+		COPY_NODE_DELEGATES_WITH_CUSTOM_DELEGATE(MeshFactoryNode, UInterchangeMeshFactoryNode, UseFullPrecisionUVs, bool, Class)
+		COPY_NODE_DELEGATES_WITH_CUSTOM_DELEGATE(MeshFactoryNode, UInterchangeMeshFactoryNode, UseBackwardsCompatibleF16TruncUVs, bool, Class)
+		COPY_NODE_DELEGATES_WITH_CUSTOM_DELEGATE(MeshFactoryNode, UInterchangeMeshFactoryNode, RemoveDegenerates, bool, Class)
+	}
+}
diff --git a/Engine/Plugins/Interchange/Runtime/Source/FactoryNodes/Private/InterchangeSkeletalMeshFactoryNode.cpp b/Engine/Plugins/Interchange/Runtime/Source/FactoryNodes/Private/InterchangeSkeletalMeshFactoryNode.cpp
--- a/Engine/Plugins/Interchange/Runtime/Source/FactoryNodes/Private/InterchangeSkeletalMeshFactoryNode.cpp
+++ b/Engine/Plugins/Interchange/Runtime/Source/FactoryNodes/Private/InterchangeSkeletalMeshFactoryNode.cpp
@@ -282,3 +282,19 @@ bool UInterchangeSkeletalMeshFactoryNode::SetNodeClassFromClassAttribute()
 	return bIsNodeClassInitialized;
 }
 
+void UInterchangeSkeletalMeshFactoryNode::CopyWithObject(const UInterchangeFactoryBaseNode* SourceNode, UObject* Object)
+{
+	Super::CopyWithObject(SourceNode, Object);
+
+	if (const UInterchangeSkeletalMeshFactoryNode* MeshFactoryNode = Cast<UInterchangeSkeletalMeshFactoryNode>(SourceNode))
+	{
+		UClass* Class = GetObjectClass();
+
+		COPY_NODE_DELEGATES_WITH_CUSTOM_DELEGATE(MeshFactoryNode, UInterchangeSkeletalMeshFactoryNode, UseHighPrecisionSkinWeights, bool, Class)
+		COPY_NODE_DELEGATES_WITH_CUSTOM_DELEGATE(MeshFactoryNode, UInterchangeSkeletalMeshFactoryNode, ThresholdPosition, float, Class)
+		COPY_NODE_DELEGATES_WITH_CUSTOM_DELEGATE(MeshFactoryNode, UInterchangeSkeletalMeshFactoryNode, ThresholdTangentNormal, float, Class)
+		COPY_NODE_DELEGATES_WITH_CUSTOM_DELEGATE(MeshFactoryNode, UInterchangeSkeletalMeshFactoryNode, ThresholdUV, float, Class)
+		COPY_NODE_DELEGATES_WITH_CUSTOM_DELEGATE(MeshFactoryNode, UInterchangeSkeletalMeshFactoryNode, MorphThresholdPosition, float, Class)
+		COPY_NODE_DELEGATES_WITH_CUSTOM_DELEGATE(MeshFactoryNode, UInterchangeSkeletalMeshFactoryNode, BoneInfluenceLimit, int32, Class)
+	}
+}
diff --git a/Engine/Plugins/Interchange/Runtime/Source/FactoryNodes/Public/InterchangeActorFactoryNode.h b/Engine/Plugins/Interchange/Runtime/Source/FactoryNodes/Public/InterchangeActorFactoryNode.h
--- a/Engine/Plugins/Interchange/Runtime/Source/FactoryNodes/Public/InterchangeActorFactoryNode.h
+++ b/Engine/Plugins/Interchange/Runtime/Source/FactoryNodes/Public/InterchangeActorFactoryNode.h
@@ -15,16 +15,10 @@ public:
 	virtual UClass* GetObjectClass() const override;
 
 	UFUNCTION(BlueprintCallable, Category = "Interchange | Node | ActorFactory")
-	bool GetCustomGlobalTransform(FTransform& AttributeValue) const
-	{
-		IMPLEMENT_NODE_ATTRIBUTE_GETTER(GlobalTransform, FTransform);
-	}
+	bool GetCustomGlobalTransform(FTransform& AttributeValue) const;
 
 	UFUNCTION(BlueprintCallable, Category = "Interchange | Node | ActorFactory")
-	bool SetCustomGlobalTransform(const FTransform& AttributeValue)
-	{
-		IMPLEMENT_NODE_ATTRIBUTE_SETTER_NODELEGATE(GlobalTransform, FTransform);
-	}
+	bool SetCustomGlobalTransform(const FTransform& AttributeValue, bool bAddApplyDelegate = true);
 
 	UFUNCTION(BlueprintCallable, Category = "Interchange | Node | ActorFactory")
 	bool GetCustomActorClassName(FString& AttributeValue) const
@@ -50,7 +44,13 @@ public:
 		IMPLEMENT_NODE_ATTRIBUTE_SETTER_NODELEGATE(Mobility, uint8);
 	}
 
+	virtual void CopyWithObject(const UInterchangeFactoryBaseNode* SourceNode, UObject* Object) override;
+
 private:
+
+	bool ApplyCustomGlobalTransformToAsset(UObject* Asset) const;
+	bool FillCustomGlobalTransformFromAsset(UObject* Asset);
+
 	const UE::Interchange::FAttributeKey Macro_CustomGlobalTransformKey = UE::Interchange::FAttributeKey(TEXT("GlobalTransform"));
 	const UE::Interchange::FAttributeKey Macro_CustomActorClassNameKey = UE::Interchange::FAttributeKey(TEXT("ActorClassName"));
 	const UE::Interchange::FAttributeKey Macro_CustomMobilityKey = UE::Interchange::FAttributeKey(TEXT("Mobility"));
diff --git a/Engine/Plugins/Interchange/Runtime/Source/FactoryNodes/Public/InterchangeCineCameraFactoryNode.h b/Engine/Plugins/Interchange/Runtime/Source/FactoryNodes/Public/InterchangeCineCameraFactoryNode.h
--- a/Engine/Plugins/Interchange/Runtime/Source/FactoryNodes/Public/InterchangeCineCameraFactoryNode.h
+++ b/Engine/Plugins/Interchange/Runtime/Source/FactoryNodes/Public/InterchangeCineCameraFactoryNode.h
@@ -65,6 +65,19 @@ public:
 		IMPLEMENT_NODE_ATTRIBUTE_SETTER_WITH_CUSTOM_DELEGATE_WITH_CUSTOM_CLASS(UInterchangeCineCameraFactoryNode, FocusMethod, ECameraFocusMethod, UCineCameraComponent);
 	}	
 
+	virtual void CopyWithObject(const UInterchangeFactoryBaseNode* SourceNode, UObject* Object) override
+	{
+		Super::CopyWithObject(SourceNode, Object);
+
+		if (const UInterchangeCineCameraFactoryNode* ActorFactoryNode = Cast<UInterchangeCineCameraFactoryNode>(SourceNode))
+		{
+			COPY_NODE_DELEGATES_WITH_CUSTOM_DELEGATE(ActorFactoryNode, UInterchangeCineCameraFactoryNode, FocalLength, float, UCineCameraComponent::StaticClass())
+			COPY_NODE_DELEGATES_WITH_CUSTOM_DELEGATE(ActorFactoryNode, UInterchangeCineCameraFactoryNode, SensorWidth, float, UCineCameraComponent::StaticClass())
+			COPY_NODE_DELEGATES_WITH_CUSTOM_DELEGATE(ActorFactoryNode, UInterchangeCineCameraFactoryNode, SensorHeight, float, UCineCameraComponent::StaticClass())
+			COPY_NODE_DELEGATES_WITH_CUSTOM_DELEGATE(ActorFactoryNode, UInterchangeCineCameraFactoryNode, FocusMethod, ECameraFocusMethod, UCineCameraComponent::StaticClass())
+		}
+	}
+
 private:
 	const UE::Interchange::FAttributeKey Macro_CustomFocalLengthKey = UE::Interchange::FAttributeKey(TEXT("FocalLength"));
 	const UE::Interchange::FAttributeKey Macro_CustomSensorWidthKey = UE::Interchange::FAttributeKey(TEXT("SensorWidth"));
diff --git a/Engine/Plugins/Interchange/Runtime/Source/FactoryNodes/Public/InterchangeLightFactoryNode.h b/Engine/Plugins/Interchange/Runtime/Source/FactoryNodes/Public/InterchangeLightFactoryNode.h
--- a/Engine/Plugins/Interchange/Runtime/Source/FactoryNodes/Public/InterchangeLightFactoryNode.h
+++ b/Engine/Plugins/Interchange/Runtime/Source/FactoryNodes/Public/InterchangeLightFactoryNode.h
@@ -38,6 +38,8 @@ public:
 	UFUNCTION(BlueprintCallable, Category = "Interchange | Node | LightFactory")
 	bool SetCustomUseTemperature(bool AttributeValue, bool bAddApplyDelegate = true);
 
+	virtual void CopyWithObject(const UInterchangeFactoryBaseNode* SourceNode, UObject* Object) override;
+
 private:
 
 	IMPLEMENT_NODE_ATTRIBUTE_KEY(LightColor)
@@ -77,6 +79,8 @@ public:
 	UFUNCTION(BlueprintCallable, Category = "Interchange | Node | LightFactory")
 	bool SetCustomIESTexture(const FString& AttributeValue);
 
+	virtual void CopyWithObject(const UInterchangeFactoryBaseNode* SourceNode, UObject* Object) override;
+
 private:
 
 	IMPLEMENT_NODE_ATTRIBUTE_KEY(IntensityUnits)
@@ -103,6 +107,8 @@ public:
 	UFUNCTION(BlueprintCallable, Category = "Interchange | Node | RectLightFactory")
 	bool SetCustomSourceHeight(float AttributeValue, bool bAddApplyDelegate = true);
 
+	virtual void CopyWithObject(const UInterchangeFactoryBaseNode* SourceNode, UObject* Object) override;
+
 private:
 	IMPLEMENT_NODE_ATTRIBUTE_KEY(SourceWidth)
 	IMPLEMENT_NODE_ATTRIBUTE_KEY(SourceHeight)
@@ -127,6 +133,8 @@ public:
 	UFUNCTION(BlueprintCallable, Category = "Interchange | Node | PointLightFactory")
 	bool SetCustomLightFalloffExponent(float AttributeValue, bool bAddApplyDelegate = true);
 
+	virtual void CopyWithObject(const UInterchangeFactoryBaseNode* SourceNode, UObject* Object) override;
+
 private:
 	IMPLEMENT_NODE_ATTRIBUTE_KEY(bUseInverseSquaredFalloff)
 	IMPLEMENT_NODE_ATTRIBUTE_KEY(LightFalloffExponent)
@@ -151,6 +159,8 @@ public:
 	UFUNCTION(BlueprintCallable, Category = "Interchange | Node | SpotLightFactory")
 	bool SetCustomOuterConeAngle(float AttributeValue, bool bAddApplyDelegate = true);
 
+	virtual void CopyWithObject(const UInterchangeFactoryBaseNode* SourceNode, UObject* Object) override;
+
 private:
 
 	IMPLEMENT_NODE_ATTRIBUTE_KEY(InnerConeAngle)
diff --git a/Engine/Plugins/Interchange/Runtime/Source/FactoryNodes/Public/InterchangeMeshFactoryNode.h b/Engine/Plugins/Interchange/Runtime/Source/FactoryNodes/Public/InterchangeMeshFactoryNode.h
--- a/Engine/Plugins/Interchange/Runtime/Source/FactoryNodes/Public/InterchangeMeshFactoryNode.h
+++ b/Engine/Plugins/Interchange/Runtime/Source/FactoryNodes/Public/InterchangeMeshFactoryNode.h
@@ -163,6 +163,8 @@ public:
 	UFUNCTION(BlueprintCallable, Category = "Interchange | FactoryNode | Mesh")
 	bool SetCustomRemoveDegenerates(const bool& AttributeValue, bool bAddApplyDelegate = true);
 
+	virtual void CopyWithObject(const UInterchangeFactoryBaseNode* SourceNode, UObject* Object) override;
+
 private:
 
 	const UE::Interchange::FAttributeKey Macro_CustomVertexColorReplaceKey = UE::Interchange::FAttributeKey(TEXT("VertexColorReplace"));
diff --git a/Engine/Plugins/Interchange/Runtime/Source/FactoryNodes/Public/InterchangeSkeletalMeshFactoryNode.h b/Engine/Plugins/Interchange/Runtime/Source/FactoryNodes/Public/InterchangeSkeletalMeshFactoryNode.h
--- a/Engine/Plugins/Interchange/Runtime/Source/FactoryNodes/Public/InterchangeSkeletalMeshFactoryNode.h
+++ b/Engine/Plugins/Interchange/Runtime/Source/FactoryNodes/Public/InterchangeSkeletalMeshFactoryNode.h
@@ -148,6 +148,8 @@ public:
 	 * The skeletal mesh thumbnail can have an overlay if the last re-import was geometry only. This thumbnail overlay feature use the metadata to find out if the last import was geometry only.
 	 */
 	virtual void AppendAssetRegistryTags(TArray<FAssetRegistryTag>& OutTags) const override;
+
+	virtual void CopyWithObject(const UInterchangeFactoryBaseNode* SourceNode, UObject* Object) override;
 private:
 
 	virtual void FillAssetClassFromAttribute() override;
diff --git a/Engine/Plugins/Interchange/Runtime/Source/FactoryNodes/Public/InterchangeTexture2DArrayFactoryNode.h b/Engine/Plugins/Interchange/Runtime/Source/FactoryNodes/Public/InterchangeTexture2DArrayFactoryNode.h
--- a/Engine/Plugins/Interchange/Runtime/Source/FactoryNodes/Public/InterchangeTexture2DArrayFactoryNode.h
+++ b/Engine/Plugins/Interchange/Runtime/Source/FactoryNodes/Public/InterchangeTexture2DArrayFactoryNode.h
@@ -67,6 +67,20 @@ public:
 #endif
 	}
 
+	virtual void CopyWithObject(const UInterchangeFactoryBaseNode* SourceNode, UObject* Object) override
+	{
+		Super::CopyWithObject(SourceNode, Object);
+
+#if WITH_EDITORONLY_DATA
+		if (const UInterchangeTexture2DArrayFactoryNode* TextureFactoryNode = Cast<UInterchangeTexture2DArrayFactoryNode>(SourceNode))
+		{
+			COPY_NODE_DELEGATES(TextureFactoryNode, AddressX, uint8, UTexture2DArray)
+			COPY_NODE_DELEGATES(TextureFactoryNode, AddressY, uint8, UTexture2DArray)
+			COPY_NODE_DELEGATES(TextureFactoryNode, AddressZ, uint8, UTexture2DArray)
+		}
+#endif
+	}
+
 private:
 	virtual UClass* GetObjectClass() const override
 	{
diff --git a/Engine/Plugins/Interchange/Runtime/Source/FactoryNodes/Public/InterchangeTexture2DFactoryNode.h b/Engine/Plugins/Interchange/Runtime/Source/FactoryNodes/Public/InterchangeTexture2DFactoryNode.h
--- a/Engine/Plugins/Interchange/Runtime/Source/FactoryNodes/Public/InterchangeTexture2DFactoryNode.h
+++ b/Engine/Plugins/Interchange/Runtime/Source/FactoryNodes/Public/InterchangeTexture2DFactoryNode.h
@@ -183,6 +183,17 @@ public:
 		SourceBlocks.Initialize(Attributes.ToSharedRef(), UE::Interchange::FTexture2DFactoryNodeStaticData::GetBaseSourceBlocksKey().ToString());
 	}
 
+	virtual void CopyWithObject(const UInterchangeFactoryBaseNode* SourceNode, UObject* Object) override
+	{
+		Super::CopyWithObject(SourceNode, Object);
+
+		if (const UInterchangeTexture2DFactoryNode* TextureFactoryNode = Cast<UInterchangeTexture2DFactoryNode>(SourceNode))
+		{
+			COPY_NODE_DELEGATES(TextureFactoryNode, AddressX, TEnumAsByte<TextureAddress>, UTexture2D)
+			COPY_NODE_DELEGATES(TextureFactoryNode, AddressY, TEnumAsByte<TextureAddress>, UTexture2D)
+		}
+	}
+
 private:
 	IMPLEMENT_NODE_ATTRIBUTE_KEY(AddressX);
 	IMPLEMENT_NODE_ATTRIBUTE_KEY(AddressY);
diff --git a/Engine/Plugins/Interchange/Runtime/Source/FactoryNodes/Public/InterchangeTextureFactoryNode.h b/Engine/Plugins/Interchange/Runtime/Source/FactoryNodes/Public/InterchangeTextureFactoryNode.h
--- a/Engine/Plugins/Interchange/Runtime/Source/FactoryNodes/Public/InterchangeTextureFactoryNode.h
+++ b/Engine/Plugins/Interchange/Runtime/Source/FactoryNodes/Public/InterchangeTextureFactoryNode.h
@@ -692,6 +692,52 @@ public:
 		IMPLEMENT_NODE_ATTRIBUTE_SETTER_NODELEGATE(PreferCompressedSourceData, bool)
 	}
 
+	virtual void CopyWithObject(const UInterchangeFactoryBaseNode* SourceNode, UObject* Object) override
+	{
+		Super::CopyWithObject(SourceNode, Object);
+
+		if (const UInterchangeTextureFactoryNode* TextureFactoryNode = Cast<UInterchangeTextureFactoryNode>(SourceNode))
+		{
+#if WITH_EDITORONLY_DATA
+			COPY_NODE_DELEGATES(TextureFactoryNode, AdjustBrightness, float, UTexture)
+			COPY_NODE_DELEGATES(TextureFactoryNode, AdjustBrightnessCurve, float, UTexture)
+			COPY_NODE_DELEGATES(TextureFactoryNode, AdjustVibrance, float, UTexture)
+			COPY_NODE_DELEGATES(TextureFactoryNode, AdjustSaturation, float, UTexture)
+			COPY_NODE_DELEGATES(TextureFactoryNode, AdjustRGBCurve, float, UTexture)
+			COPY_NODE_DELEGATES(TextureFactoryNode, AdjustHue, float, UTexture)
+			COPY_NODE_DELEGATES(TextureFactoryNode, AdjustMinAlpha, float, UTexture)
+			COPY_NODE_DELEGATES(TextureFactoryNode, AdjustMaxAlpha, float, UTexture)
+			COPY_NODE_DELEGATES(TextureFactoryNode, bChromaKeyTexture, bool, UTexture)
+			COPY_NODE_DELEGATES(TextureFactoryNode, ChromaKeyThreshold, float, UTexture)
+			COPY_NODE_DELEGATES(TextureFactoryNode, ChromaKeyColor, FColor, UTexture)
+			COPY_NODE_DELEGATES(TextureFactoryNode, CompressionNoAlpha, bool, UTexture)
+			COPY_NODE_DELEGATES(TextureFactoryNode, DeferCompression, bool, UTexture)
+			COPY_NODE_DELEGATES(TextureFactoryNode, LossyCompressionAmount, uint8, UTexture)
+			COPY_NODE_DELEGATES(TextureFactoryNode, MaxTextureSize, int32, UTexture)
+			COPY_NODE_DELEGATES(TextureFactoryNode, CompressionQuality, uint8, UTexture)
+			COPY_NODE_DELEGATES(TextureFactoryNode, AlphaCoverageThresholds, FVector4, UTexture)
+			COPY_NODE_DELEGATES(TextureFactoryNode, bDoScaleMipsForAlphaCoverage, bool, UTexture)
+			COPY_NODE_DELEGATES(TextureFactoryNode, bFlipGreenChannel, bool, UTexture)
+			COPY_NODE_DELEGATES(TextureFactoryNode, PowerOfTwoMode, uint8, UTexture)
+			COPY_NODE_DELEGATES(TextureFactoryNode, PaddingColor, FColor, UTexture)
+			COPY_NODE_DELEGATES(TextureFactoryNode, bUseLegacyGamma, bool, UTexture)
+			COPY_NODE_DELEGATES(TextureFactoryNode, bPreserveBorder, bool, UTexture)
+			COPY_NODE_DELEGATES(TextureFactoryNode, MipGenSettings, uint8, UTexture)
+			COPY_NODE_DELEGATES(TextureFactoryNode, CompositeTextureMode, uint8, UTexture)
+			COPY_NODE_DELEGATES(TextureFactoryNode, CompositePower, float, UTexture)
+#endif
+			COPY_NODE_DELEGATES(TextureFactoryNode, VirtualTextureStreaming, bool, UTexture)
+			COPY_NODE_DELEGATES(TextureFactoryNode, CompressionSettings, uint8, UTexture)
+			COPY_NODE_DELEGATES_WITH_CUSTOM_DELEGATE(TextureFactoryNode, UInterchangeTextureFactoryNode, Downscale, float, UTexture::StaticClass())
+			COPY_NODE_DELEGATES(TextureFactoryNode, DownscaleOptions, uint8, UTexture)
+			COPY_NODE_DELEGATES(TextureFactoryNode, Filter, uint8, UTexture)
+			COPY_NODE_DELEGATES(TextureFactoryNode, LODBias, int32, UTexture)
+			COPY_NODE_DELEGATES(TextureFactoryNode, LODGroup, uint8, UTexture)
+			COPY_NODE_DELEGATES(TextureFactoryNode, MipLoadOptions, uint8, UTexture)
+			COPY_NODE_DELEGATES(TextureFactoryNode, SRGB, bool, UTexture)
+		}
+	}
+
 private:
 
 	const UE::Interchange::FAttributeKey ClassNameAttributeKey = UE::Interchange::FBaseNodeStaticData::ClassTypeAttributeKey();
diff --git a/Engine/Plugins/Interchange/Runtime/Source/FactoryNodes/Public/InterchangeTextureLightProfileFactoryNode.h b/Engine/Plugins/Interchange/Runtime/Source/FactoryNodes/Public/InterchangeTextureLightProfileFactoryNode.h
--- a/Engine/Plugins/Interchange/Runtime/Source/FactoryNodes/Public/InterchangeTextureLightProfileFactoryNode.h
+++ b/Engine/Plugins/Interchange/Runtime/Source/FactoryNodes/Public/InterchangeTextureLightProfileFactoryNode.h
@@ -54,6 +54,18 @@ public:
 #endif
 	}
 
+	virtual void CopyWithObject(const UInterchangeFactoryBaseNode* SourceNode, UObject* Object) override
+	{
+		Super::CopyWithObject(SourceNode, Object);
+
+#if WITH_EDITORONLY_DATA
+		if (const UInterchangeTextureLightProfileFactoryNode* TextureFactoryNode = Cast<UInterchangeTextureLightProfileFactoryNode>(SourceNode))
+		{
+			COPY_NODE_DELEGATES(TextureFactoryNode, Brightness, float, UTextureLightProfile)
+			COPY_NODE_DELEGATES(TextureFactoryNode, TextureMultiplier, float, UTextureLightProfile)
+		}
+#endif
+	}
 
 private:
 	virtual UClass* GetObjectClass() const override
diff --git a/Engine/Plugins/Interchange/Runtime/Source/Import/Private/Scene/InterchangeActorFactory.cpp b/Engine/Plugins/Interchange/Runtime/Source/Import/Private/Scene/InterchangeActorFactory.cpp
--- a/Engine/Plugins/Interchange/Runtime/Source/Import/Private/Scene/InterchangeActorFactory.cpp
+++ b/Engine/Plugins/Interchange/Runtime/Source/Import/Private/Scene/InterchangeActorFactory.cpp
@@ -1,6 +1,7 @@
 // Copyright Epic Games, Inc. All Rights Reserved. 
 #include "Scene/InterchangeActorFactory.h"
 
+#include "InterchangeActorFactoryNode.h"
 #include "Nodes/InterchangeFactoryBaseNode.h"
 #include "Scene/InterchangeActorHelper.h"
 
@@ -16,17 +17,29 @@ UClass* UInterchangeActorFactory::GetFactoryClass() const
 
 UObject* UInterchangeActorFactory::ImportSceneObject_GameThread(const UInterchangeFactoryBase::FImportSceneObjectsParams& CreateSceneObjectsParams)
 {
-	AActor* SpawnedActor = UE::Interchange::ActorHelper::SpawnFactoryActor(CreateSceneObjectsParams);
+	using namespace UE::Interchange;
 
-	if (!SpawnedActor)
+	UInterchangeActorFactoryNode* FactoryNode = Cast<UInterchangeActorFactoryNode>(CreateSceneObjectsParams.FactoryNode);
+	if (!ensure(FactoryNode) || !CreateSceneObjectsParams.NodeContainer)
 	{
 		return nullptr;
 	}
 
-	if (USceneComponent* RootComponent = SpawnedActor->GetRootComponent())
+	AActor* SpawnedActor = ActorHelper::SpawnFactoryActor(CreateSceneObjectsParams);
+
+	if (SpawnedActor)
 	{
-		CreateSceneObjectsParams.FactoryNode->ApplyAllCustomAttributeToObject(RootComponent);
+		if (UObject* ObjectToUpdate = ProcessActor(*SpawnedActor, *FactoryNode, *CreateSceneObjectsParams.NodeContainer))
+		{
+			ActorHelper::ApplyAllCustomAttributes(CreateSceneObjectsParams, *ObjectToUpdate);
+		}
 	}
 
 	return SpawnedActor;
 }
+
+UObject* UInterchangeActorFactory::ProcessActor(AActor& SpawnedActor, const UInterchangeActorFactoryNode& /*FactoryNode*/, const UInterchangeBaseNodeContainer& /*NodeContainer*/)
+{
+	return SpawnedActor.GetRootComponent();
+}
+
diff --git a/Engine/Plugins/Interchange/Runtime/Source/Import/Private/Scene/InterchangeActorHelper.cpp b/Engine/Plugins/Interchange/Runtime/Source/Import/Private/Scene/InterchangeActorHelper.cpp
--- a/Engine/Plugins/Interchange/Runtime/Source/Import/Private/Scene/InterchangeActorHelper.cpp
+++ b/Engine/Plugins/Interchange/Runtime/Source/Import/Private/Scene/InterchangeActorHelper.cpp
@@ -5,6 +5,7 @@
 #include "Engine/Engine.h"
 #include "Engine/Level.h"
 #include "InterchangeActorFactoryNode.h"
+#include "InterchangeImportCommon.h"
 #include "InterchangeMaterialFactoryNode.h"
 #include "InterchangeMeshActorFactoryNode.h"
 #include "InterchangeMeshNode.h"
@@ -37,12 +38,6 @@ AActor* UE::Interchange::ActorHelper::GetSpawnedParentActor(const UInterchangeBa
 
 AActor* UE::Interchange::ActorHelper::SpawnFactoryActor(const UInterchangeFactoryBase::FImportSceneObjectsParams& CreateSceneObjectsParams)
 {
-	FActorSpawnParameters SpawnParameters;
-	SpawnParameters.Name = FName(*CreateSceneObjectsParams.ObjectName);
-	SpawnParameters.NameMode = FActorSpawnParameters::ESpawnActorNameMode::Requested;
-	SpawnParameters.OverrideLevel = CreateSceneObjectsParams.Level;
-	SpawnParameters.SpawnCollisionHandlingOverride = ESpawnActorCollisionHandlingMethod::AlwaysSpawn;
-
 	const UInterchangeActorFactoryNode* FactoryNode = Cast<UInterchangeActorFactoryNode>(CreateSceneObjectsParams.FactoryNode);
 	const UInterchangeBaseNodeContainer* NodeContainer = CreateSceneObjectsParams.NodeContainer;
 
@@ -51,6 +46,12 @@ AActor* UE::Interchange::ActorHelper::SpawnFactoryActor(const UInterchangeFactor
 		return nullptr;
 	}
 
+	FActorSpawnParameters SpawnParameters;
+	SpawnParameters.Name = FName(*CreateSceneObjectsParams.ObjectName);
+	SpawnParameters.NameMode = FActorSpawnParameters::ESpawnActorNameMode::Requested;
+	SpawnParameters.OverrideLevel = CreateSceneObjectsParams.Level;
+	SpawnParameters.SpawnCollisionHandlingOverride = ESpawnActorCollisionHandlingMethod::AlwaysSpawn;
+
 	AActor* ParentActor = UE::Interchange::ActorHelper::GetSpawnedParentActor(NodeContainer, FactoryNode);
 	UWorld* const World = [&SpawnParameters, &ParentActor]()
 	{
@@ -83,11 +84,34 @@ AActor* UE::Interchange::ActorHelper::SpawnFactoryActor(const UInterchangeFactor
 		return nullptr;
 	}
 
-	FTransform Transform = FTransform::Identity;
-	FactoryNode->GetCustomGlobalTransform(Transform);
-
 	UClass* ActorClass = FactoryNode->GetObjectClass();
-	AActor* SpawnedActor = World->SpawnActor<AActor>(ActorClass, Transform, SpawnParameters);
+	AActor* SpawnedActor = Cast<AActor>(CreateSceneObjectsParams.ReimportObject);
+	if (SpawnedActor)
+	{
+		if (SpawnedActor->GetClass() == ActorClass)
+		{
+			// TODO: Check whether parenting is the same
+		}
+		else
+		{
+			SpawnedActor = nullptr;
+		}
+	}
+	// The related actor has been deleted. Check on reimport policy
+	else if (CreateSceneObjectsParams.ReimportFactoryNode)
+	{
+		// if reimport policy does not prioritize new content, do not recreate the actor
+		const EReimportStrategyFlags ReimportStrategyFlags = CreateSceneObjectsParams.FactoryNode->GetReimportStrategyFlags();
+		if (ReimportStrategyFlags != EReimportStrategyFlags::ApplyPipelineProperties)
+		{
+			return nullptr;
+		}
+	}
+
+	if (!SpawnedActor)
+	{
+		SpawnedActor = World->SpawnActor<AActor>(ActorClass, SpawnParameters);
+	}
 
 	if (SpawnedActor)
 	{
@@ -100,8 +124,6 @@ AActor* UE::Interchange::ActorHelper::SpawnFactoryActor(const UInterchangeFactor
 #if WITH_EDITORONLY_DATA
 			RootComponent->bVisualizeComponent = true;
 #endif
-			RootComponent->SetWorldTransform(Transform);
-
 			SpawnedActor->SetRootComponent(RootComponent);
 			SpawnedActor->AddInstanceComponent(RootComponent);
 		}
@@ -124,7 +146,7 @@ AActor* UE::Interchange::ActorHelper::SpawnFactoryActor(const UInterchangeFactor
 		}
 	}
 
-	if (ParentActor)
+	if (ParentActor || !SpawnedActor->IsAttachedTo(ParentActor))
 	{
 		SpawnedActor->AttachToActor(ParentActor, FAttachmentTransformRules::KeepWorldTransform);
 	}
@@ -188,3 +210,21 @@ void UE::Interchange::ActorHelper::ApplySlotMaterialDependencies(const UIntercha
 		}
 	}
 }
+
+void UE::Interchange::ActorHelper::ApplyAllCustomAttributes(const UInterchangeFactoryBase::FImportSceneObjectsParams& CreateSceneObjectsParams, UObject& ObjectToUpdate)
+{
+	if (CreateSceneObjectsParams.ReimportObject)
+	{
+		if (ObjectToUpdate.GetOuter() == CreateSceneObjectsParams.ReimportObject)
+		{
+			UInterchangeFactoryBaseNode* CurrentNode = UInterchangeFactoryBaseNode::DuplicateWithObject(CreateSceneObjectsParams.FactoryNode, &ObjectToUpdate);
+
+			FFactoryCommon::ApplyReimportStrategyToAsset(&ObjectToUpdate, CreateSceneObjectsParams.ReimportFactoryNode, CurrentNode, CreateSceneObjectsParams.FactoryNode);
+		}
+	}
+	else
+	{
+		CreateSceneObjectsParams.FactoryNode->ApplyAllCustomAttributeToObject(&ObjectToUpdate);
+	}
+}
+#pragma optimize ("", on)
diff --git a/Engine/Plugins/Interchange/Runtime/Source/Import/Private/Scene/InterchangeCineCameraActorFactory.cpp b/Engine/Plugins/Interchange/Runtime/Source/Import/Private/Scene/InterchangeCineCameraActorFactory.cpp
--- a/Engine/Plugins/Interchange/Runtime/Source/Import/Private/Scene/InterchangeCineCameraActorFactory.cpp
+++ b/Engine/Plugins/Interchange/Runtime/Source/Import/Private/Scene/InterchangeCineCameraActorFactory.cpp
@@ -19,21 +19,11 @@ UClass* UInterchangeCineCameraActorFactory::GetFactoryClass() const
 	return ACineCameraActor::StaticClass();
 }
 
-UObject* UInterchangeCineCameraActorFactory::ImportSceneObject_GameThread(const UInterchangeFactoryBase::FImportSceneObjectsParams& CreateSceneObjectsParams)
+UObject* UInterchangeCineCameraActorFactory::ProcessActor(AActor& SpawnedActor, const UInterchangeActorFactoryNode& /*FactoryNode*/, const UInterchangeBaseNodeContainer& /*NodeContainer*/)
 {
-	ACineCameraActor* SpawnedActor = Cast<ACineCameraActor>(UE::Interchange::ActorHelper::SpawnFactoryActor(CreateSceneObjectsParams));
+	ACineCameraActor* CineCameraActor = Cast<ACineCameraActor>(&SpawnedActor);
 
-	if (!SpawnedActor)
-	{
-		return nullptr;
-	}
-
-	if (UCineCameraComponent* CineCameraComponent = SpawnedActor->GetCineCameraComponent())
-	{
-		CreateSceneObjectsParams.FactoryNode->ApplyAllCustomAttributeToObject(CineCameraComponent);
-	}
-
-	return SpawnedActor;
+	return CineCameraActor ? CineCameraActor->GetCineCameraComponent() : nullptr;
 }
 
 
diff --git a/Engine/Plugins/Interchange/Runtime/Source/Import/Private/Scene/InterchangeLightActorFactory.cpp b/Engine/Plugins/Interchange/Runtime/Source/Import/Private/Scene/InterchangeLightActorFactory.cpp
--- a/Engine/Plugins/Interchange/Runtime/Source/Import/Private/Scene/InterchangeLightActorFactory.cpp
+++ b/Engine/Plugins/Interchange/Runtime/Source/Import/Private/Scene/InterchangeLightActorFactory.cpp
@@ -18,36 +18,34 @@ UClass* UInterchangeLightActorFactory::GetFactoryClass() const
 	return ALight::StaticClass();
 }
 
-UObject* UInterchangeLightActorFactory::ImportSceneObject_GameThread(const UInterchangeFactoryBase::FImportSceneObjectsParams& CreateSceneObjectsParams)
+UObject* UInterchangeLightActorFactory::ProcessActor(AActor& SpawnedActor, const UInterchangeActorFactoryNode& FactoryNode, const UInterchangeBaseNodeContainer& NodeContainer)
 {
-	ALight * SpawnedActor = Cast<ALight>(UE::Interchange::ActorHelper::SpawnFactoryActor(CreateSceneObjectsParams));
-
-	if(!SpawnedActor)
-	{
-		return nullptr;
-	}
-
-	if(ULightComponent* LightComponent = SpawnedActor->GetLightComponent())
+	if (ALight* LightActor = Cast<ALight>(&SpawnedActor))
 	{
-		LightComponent->UnregisterComponent();
-		if(const UInterchangeLightFactoryNode* LightFactoryNode = Cast<UInterchangeLightFactoryNode>(CreateSceneObjectsParams.FactoryNode))
+		if (ULightComponent* LightComponent = LightActor->GetLightComponent())
 		{
-			if(FString IESTexture; LightFactoryNode->GetCustomIESTexture(IESTexture))
+			LightComponent->UnregisterComponent();
+
+			if (const UInterchangeLightFactoryNode* LightFactoryNode = Cast<UInterchangeLightFactoryNode>(&FactoryNode))
 			{
-				IESTexture = UInterchangeFactoryBaseNode::BuildFactoryNodeUid(IESTexture);
-				if(const UInterchangeTextureLightProfileFactoryNode* TextureFactoryNode = Cast<UInterchangeTextureLightProfileFactoryNode>(CreateSceneObjectsParams.NodeContainer->GetNode(IESTexture)))
+				if (FString IESTexture; LightFactoryNode->GetCustomIESTexture(IESTexture))
 				{
-					FSoftObjectPath ReferenceObject;
-					TextureFactoryNode->GetCustomReferenceObject(ReferenceObject);
-					if(UTextureLightProfile* Texture = Cast<UTextureLightProfile>(ReferenceObject.TryLoad()))
+					IESTexture = UInterchangeFactoryBaseNode::BuildFactoryNodeUid(IESTexture);
+					if (const UInterchangeTextureLightProfileFactoryNode* TextureFactoryNode = Cast<UInterchangeTextureLightProfileFactoryNode>(NodeContainer.GetNode(IESTexture)))
 					{
-						LightComponent->SetIESTexture(Texture);
+						FSoftObjectPath ReferenceObject;
+						TextureFactoryNode->GetCustomReferenceObject(ReferenceObject);
+						if (UTextureLightProfile* Texture = Cast<UTextureLightProfile>(ReferenceObject.TryLoad()))
+						{
+							LightComponent->SetIESTexture(Texture);
+						}
 					}
 				}
+
+				return LightComponent;
 			}
 		}
-		CreateSceneObjectsParams.FactoryNode->ApplyAllCustomAttributeToObject(LightComponent);
 	}
 
-	return SpawnedActor;
+	return nullptr;
 }
diff --git a/Engine/Plugins/Interchange/Runtime/Source/Import/Private/Scene/InterchangeSkeletalMeshActorFactory.cpp b/Engine/Plugins/Interchange/Runtime/Source/Import/Private/Scene/InterchangeSkeletalMeshActorFactory.cpp
--- a/Engine/Plugins/Interchange/Runtime/Source/Import/Private/Scene/InterchangeSkeletalMeshActorFactory.cpp
+++ b/Engine/Plugins/Interchange/Runtime/Source/Import/Private/Scene/InterchangeSkeletalMeshActorFactory.cpp
@@ -1,7 +1,6 @@
 // Copyright Epic Games, Inc. All Rights Reserved. 
 #include "Scene/InterchangeSkeletalMeshActorFactory.h"
 
-#include "InterchangeActorFactoryNode.h"
 #include "InterchangeMeshActorFactoryNode.h"
 #include "Scene/InterchangeActorHelper.h"
 #include "InterchangeSceneNode.h"
@@ -15,24 +14,23 @@
 #include UE_INLINE_GENERATED_CPP_BY_NAME(InterchangeSkeletalMeshActorFactory)
 
 
-UObject* UInterchangeSkeletalMeshActorFactory::ImportSceneObject_GameThread(const UInterchangeFactoryBase::FImportSceneObjectsParams& CreateSceneObjectsParams)
+UObject* UInterchangeSkeletalMeshActorFactory::ProcessActor(AActor& SpawnedActor, const UInterchangeActorFactoryNode& /*FactoryNode*/, const UInterchangeBaseNodeContainer& /*NodeContainer*/)
 {
-	ASkeletalMeshActor* SpawnedActor = Cast<ASkeletalMeshActor>(UE::Interchange::ActorHelper::SpawnFactoryActor(CreateSceneObjectsParams));
+	ASkeletalMeshActor* SkeletalMeshActor = Cast<ASkeletalMeshActor>(&SpawnedActor);
 
-	if (!SpawnedActor)
+	if (!SkeletalMeshActor)
 	{
 		return nullptr;
 	}
 
-	UInterchangeFactoryBaseNode* FactoryNode = CreateSceneObjectsParams.FactoryNode;
-	SetupSkeletalMeshActor(CreateSceneObjectsParams.NodeContainer, FactoryNode, SpawnedActor);
-
-	if (USkeletalMeshComponent* SkeletalMeshComponent = SpawnedActor->GetSkeletalMeshComponent())
+	if (USkeletalMeshComponent* SkeletalMeshComponent = SkeletalMeshActor->GetSkeletalMeshComponent())
 	{
-		FactoryNode->ApplyAllCustomAttributeToObject(SkeletalMeshComponent);
+		SkeletalMeshComponent->UnregisterComponent();
+
+		return SkeletalMeshComponent;
 	}
 
-	return SpawnedActor;
+	return nullptr;
 };
 
 UClass* UInterchangeSkeletalMeshActorFactory::GetFactoryClass() const
@@ -40,12 +38,6 @@ UClass* UInterchangeSkeletalMeshActorFactory::GetFactoryClass() const
 	return ASkeletalMeshActor::StaticClass();
 }
 
-void UInterchangeSkeletalMeshActorFactory::SetupSkeletalMeshActor(const UInterchangeBaseNodeContainer* NodeContainer, const UInterchangeFactoryBaseNode* ActorFactoryNode, ASkeletalMeshActor* SkeletalMeshActor)
-{
-	USkeletalMeshComponent* SkeletalMeshComponent = SkeletalMeshActor->GetSkeletalMeshComponent();
-	SkeletalMeshComponent->UnregisterComponent();
-}
-
 void UInterchangeSkeletalMeshActorFactory::SetupObject_GameThread(const FSetupObjectParams& Arguments)
 {
 	if (ASkeletalMeshActor* SkeletalMeshActor = Cast<ASkeletalMeshActor>(Arguments.ImportedObject))
diff --git a/Engine/Plugins/Interchange/Runtime/Source/Import/Private/Scene/InterchangeStaticMeshActorFactory.cpp b/Engine/Plugins/Interchange/Runtime/Source/Import/Private/Scene/InterchangeStaticMeshActorFactory.cpp
--- a/Engine/Plugins/Interchange/Runtime/Source/Import/Private/Scene/InterchangeStaticMeshActorFactory.cpp
+++ b/Engine/Plugins/Interchange/Runtime/Source/Import/Private/Scene/InterchangeStaticMeshActorFactory.cpp
@@ -1,7 +1,6 @@
 // Copyright Epic Games, Inc. All Rights Reserved. 
 #include "Scene/InterchangeStaticMeshActorFactory.h"
 
-#include "InterchangeActorFactoryNode.h"
 #include "InterchangeMeshActorFactoryNode.h"
 #include "Nodes/InterchangeBaseNodeContainer.h"
 #include "Scene/InterchangeActorHelper.h"
@@ -17,50 +16,40 @@ UClass* UInterchangeStaticMeshActorFactory::GetFactoryClass() const
 	return AStaticMeshActor::StaticClass();
 }
 
-UObject* UInterchangeStaticMeshActorFactory::ImportSceneObject_GameThread(const UInterchangeFactoryBase::FImportSceneObjectsParams& CreateSceneObjectsParams)
+UObject* UInterchangeStaticMeshActorFactory::ProcessActor(AActor& SpawnedActor, const UInterchangeActorFactoryNode& FactoryNode, const UInterchangeBaseNodeContainer& NodeContainer)
 {
-	AStaticMeshActor* SpawnedActor = Cast<AStaticMeshActor>(UE::Interchange::ActorHelper::SpawnFactoryActor(CreateSceneObjectsParams));
+	using namespace UE::Interchange;
 
-	if (!SpawnedActor)
-	{
-		return nullptr;
-	}
-
-	UInterchangeFactoryBaseNode* FactoryNode = CreateSceneObjectsParams.FactoryNode;
-	SetupStaticMeshActor(CreateSceneObjectsParams.NodeContainer, FactoryNode, SpawnedActor);
-
-	if (UStaticMeshComponent* StaticMeshComponent = SpawnedActor->GetStaticMeshComponent())
-	{
-		FactoryNode->ApplyAllCustomAttributeToObject(StaticMeshComponent);
-	}
-
-	return SpawnedActor;
-};
+	AStaticMeshActor* StaticMeshActor = Cast<AStaticMeshActor>(&SpawnedActor);
 
-void UInterchangeStaticMeshActorFactory::SetupStaticMeshActor(const UInterchangeBaseNodeContainer* NodeContainer, const UInterchangeFactoryBaseNode* ActorFactoryNode, AStaticMeshActor* StaticMeshActor)
-{
 	if (!StaticMeshActor)
 	{
-		return;
+		return nullptr;
 	}
 
 	if (UStaticMeshComponent* StaticMeshComponent = StaticMeshActor->GetStaticMeshComponent())
 	{
 		StaticMeshComponent->UnregisterComponent();
 
-		if (const UInterchangeFactoryBaseNode* MeshNode = UE::Interchange::ActorHelper::FindAssetInstanceFactoryNode(NodeContainer, ActorFactoryNode))
+		if (const UInterchangeFactoryBaseNode* MeshNode = ActorHelper::FindAssetInstanceFactoryNode(&NodeContainer, &FactoryNode))
 		{
 			FSoftObjectPath ReferenceObject;
 			MeshNode->GetCustomReferenceObject(ReferenceObject);
 			if (UStaticMesh* StaticMesh = Cast<UStaticMesh>(ReferenceObject.TryLoad()))
 			{
-				StaticMeshComponent->SetStaticMesh(StaticMesh);
-
-				if (const UInterchangeMeshActorFactoryNode* MeshActorFactoryNode = Cast<UInterchangeMeshActorFactoryNode>(ActorFactoryNode))
+				if (StaticMesh != StaticMeshComponent->GetStaticMesh())
 				{
-					UE::Interchange::ActorHelper::ApplySlotMaterialDependencies(*NodeContainer, *MeshActorFactoryNode, *StaticMeshComponent);
+					StaticMeshComponent->SetStaticMesh(StaticMesh);
 				}
 			}
 		}
+		else
+		{
+			// TODO: Warn that new mesh has not been applied
+		}
+
+		return StaticMeshComponent;
 	}
-}
+
+	return nullptr;
+};
diff --git a/Engine/Plugins/Interchange/Runtime/Source/Import/Public/Scene/InterchangeActorFactory.h b/Engine/Plugins/Interchange/Runtime/Source/Import/Public/Scene/InterchangeActorFactory.h
--- a/Engine/Plugins/Interchange/Runtime/Source/Import/Public/Scene/InterchangeActorFactory.h
+++ b/Engine/Plugins/Interchange/Runtime/Source/Import/Public/Scene/InterchangeActorFactory.h
@@ -9,6 +9,10 @@
 
 #include "InterchangeActorFactory.generated.h"
 
+class AActor;
+class UInterchangeActorFactoryNode;
+class USceneComponent;
+
 UCLASS(BlueprintType)
 class INTERCHANGEIMPORT_API UInterchangeActorFactory : public UInterchangeFactoryBase
 {
@@ -20,10 +24,19 @@ public:
 
 	virtual UClass* GetFactoryClass() const override;
 
+private:
 	virtual UObject* ImportSceneObject_GameThread(const UInterchangeFactoryBase::FImportSceneObjectsParams& CreateSceneObjectsParams) override;
 
 	// Interchange factory base interface end
 	//////////////////////////////////////////////////////////////////////////
+
+protected:
+	/**
+	 * Method called in UInterchangeActorFactory::ImportSceneObject_GameThread to allow
+	 * child class to complete the creation of the actor.
+	 * This method is expected to return the UOBject to apply the factory node's custom attributes to.
+	 */
+	virtual UObject* ProcessActor(AActor& SpawnedActor, const UInterchangeActorFactoryNode& FactoryNode, const UInterchangeBaseNodeContainer& NodeContainer);
 };
 
 
diff --git a/Engine/Plugins/Interchange/Runtime/Source/Import/Public/Scene/InterchangeActorHelper.h b/Engine/Plugins/Interchange/Runtime/Source/Import/Public/Scene/InterchangeActorHelper.h
--- a/Engine/Plugins/Interchange/Runtime/Source/Import/Public/Scene/InterchangeActorHelper.h
+++ b/Engine/Plugins/Interchange/Runtime/Source/Import/Public/Scene/InterchangeActorHelper.h
@@ -32,4 +32,10 @@ namespace UE::Interchange::ActorHelper
 	 * Applies material slot dependencies stored in ActorFactoryNode to MeshComponent.
 	 */
 	INTERCHANGEIMPORT_API void ApplySlotMaterialDependencies(const UInterchangeBaseNodeContainer& NodeContainer, const UInterchangeMeshActorFactoryNode& ActorFactoryNode, UMeshComponent& MeshComponent);
+
+	/**
+	 * Applies custom attributes from to ObjectToUpdate based on content of CreateSceneObjectsParams
+	 * Reimport policies will be applied if applicable
+	 */
+	INTERCHANGEIMPORT_API void ApplyAllCustomAttributes(const UInterchangeFactoryBase::FImportSceneObjectsParams& CreateSceneObjectsParams, UObject& ObjectToUpdate);
 }
diff --git a/Engine/Plugins/Interchange/Runtime/Source/Import/Public/Scene/InterchangeCineCameraActorFactory.h b/Engine/Plugins/Interchange/Runtime/Source/Import/Public/Scene/InterchangeCineCameraActorFactory.h
--- a/Engine/Plugins/Interchange/Runtime/Source/Import/Public/Scene/InterchangeCineCameraActorFactory.h
+++ b/Engine/Plugins/Interchange/Runtime/Source/Import/Public/Scene/InterchangeCineCameraActorFactory.h
@@ -3,14 +3,14 @@
 #pragma once
 
 #include "CoreMinimal.h"
-#include "InterchangeActorFactory.h"
+#include "Scene/InterchangeActorFactory.h"
 #include "UObject/Object.h"
 #include "UObject/ObjectMacros.h"
 
 #include "InterchangeCineCameraActorFactory.generated.h"
 
 UCLASS(BlueprintType)
-class INTERCHANGEIMPORT_API UInterchangeCineCameraActorFactory : public UInterchangeFactoryBase
+class INTERCHANGEIMPORT_API UInterchangeCineCameraActorFactory : public UInterchangeActorFactory
 {
 	GENERATED_BODY()
 public:
@@ -20,10 +20,17 @@ public:
 
 	virtual UClass* GetFactoryClass() const override;
 
-	virtual UObject* ImportSceneObject_GameThread(const UInterchangeFactoryBase::FImportSceneObjectsParams& CreateSceneObjectsParams) override;
-
 	// Interchange factory base interface end
 	//////////////////////////////////////////////////////////////////////////
+
+	//////////////////////////////////////////////////////////////////////////
+	// Interchange actor factory interface begin
+
+	virtual UObject* ProcessActor(class AActor& SpawnedActor, const UInterchangeActorFactoryNode& FactoryNode, const UInterchangeBaseNodeContainer& NodeContainer) override;
+
+
+	// Interchange actor factory interface end
+	//////////////////////////////////////////////////////////////////////////
 };
 
 
diff --git a/Engine/Plugins/Interchange/Runtime/Source/Import/Public/Scene/InterchangeLightActorFactory.h b/Engine/Plugins/Interchange/Runtime/Source/Import/Public/Scene/InterchangeLightActorFactory.h
--- a/Engine/Plugins/Interchange/Runtime/Source/Import/Public/Scene/InterchangeLightActorFactory.h
+++ b/Engine/Plugins/Interchange/Runtime/Source/Import/Public/Scene/InterchangeLightActorFactory.h
@@ -2,12 +2,12 @@
 
 #pragma once
 
-#include "InterchangeFactoryBase.h"
+#include "Scene/InterchangeActorFactory.h"
 
 #include "InterchangeLightActorFactory.generated.h"
 
 UCLASS(BlueprintType)
-class INTERCHANGEIMPORT_API UInterchangeLightActorFactory : public UInterchangeFactoryBase
+class INTERCHANGEIMPORT_API UInterchangeLightActorFactory : public UInterchangeActorFactory
 {
 	GENERATED_BODY()
 
@@ -17,8 +17,16 @@ public:
 
 	virtual UClass* GetFactoryClass() const override;
 
-	virtual UObject* ImportSceneObject_GameThread(const UInterchangeFactoryBase::FImportSceneObjectsParams& CreateSceneObjectsParams) override;
-
 	// Interchange factory base interface end
 	//////////////////////////////////////////////////////////////////////////
+
+protected:
+	//////////////////////////////////////////////////////////////////////////
+	// Interchange actor factory interface begin
+
+	virtual UObject* ProcessActor(class AActor& SpawnedActor, const UInterchangeActorFactoryNode& FactoryNode, const UInterchangeBaseNodeContainer& NodeContainer) override;
+
+
+	// Interchange actor factory interface end
+	//////////////////////////////////////////////////////////////////////////
 };
\ No newline at end of file
diff --git a/Engine/Plugins/Interchange/Runtime/Source/Import/Public/Scene/InterchangeSkeletalMeshActorFactory.h b/Engine/Plugins/Interchange/Runtime/Source/Import/Public/Scene/InterchangeSkeletalMeshActorFactory.h
--- a/Engine/Plugins/Interchange/Runtime/Source/Import/Public/Scene/InterchangeSkeletalMeshActorFactory.h
+++ b/Engine/Plugins/Interchange/Runtime/Source/Import/Public/Scene/InterchangeSkeletalMeshActorFactory.h
@@ -24,16 +24,19 @@ public:
 
 	virtual UClass* GetFactoryClass() const override;
 
-	virtual UObject* ImportSceneObject_GameThread(const UInterchangeFactoryBase::FImportSceneObjectsParams& CreateSceneObjectsParams) override;
-
 	virtual void SetupObject_GameThread(const FSetupObjectParams& Arguments) override;
 
 	// Interchange factory base interface end
 	//////////////////////////////////////////////////////////////////////////
 
-private:
+protected:
+	//////////////////////////////////////////////////////////////////////////
+	// Interchange actor factory interface begin
+
+	virtual UObject* ProcessActor(class AActor& SpawnedActor, const UInterchangeActorFactoryNode& FactoryNode, const UInterchangeBaseNodeContainer& NodeContainer) override;
 
-	void SetupSkeletalMeshActor(const UInterchangeBaseNodeContainer* NodeContainer, const UInterchangeFactoryBaseNode* ActorFactoryNode, ASkeletalMeshActor* SkeletalMeshActor);
+	// Interchange actor factory interface end
+	//////////////////////////////////////////////////////////////////////////
 };
 
 
diff --git a/Engine/Plugins/Interchange/Runtime/Source/Import/Public/Scene/InterchangeStaticMeshActorFactory.h b/Engine/Plugins/Interchange/Runtime/Source/Import/Public/Scene/InterchangeStaticMeshActorFactory.h
--- a/Engine/Plugins/Interchange/Runtime/Source/Import/Public/Scene/InterchangeStaticMeshActorFactory.h
+++ b/Engine/Plugins/Interchange/Runtime/Source/Import/Public/Scene/InterchangeStaticMeshActorFactory.h
@@ -3,7 +3,7 @@
 #pragma once
 
 #include "CoreMinimal.h"
-#include "InterchangeActorFactory.h"
+#include "Scene/InterchangeActorFactory.h"
 #include "UObject/Object.h"
 #include "UObject/ObjectMacros.h"
 
@@ -14,7 +14,7 @@ class AStaticMeshActor;
 class UInterchangeActorFactoryNode;
 
 UCLASS(BlueprintType, Experimental)
-class INTERCHANGEIMPORT_API UInterchangeStaticMeshActorFactory : public UInterchangeFactoryBase
+class INTERCHANGEIMPORT_API UInterchangeStaticMeshActorFactory : public UInterchangeActorFactory
 {
 	GENERATED_BODY()
 public:
@@ -24,15 +24,17 @@ public:
 
 	virtual UClass* GetFactoryClass() const override;
 
-	virtual UObject* ImportSceneObject_GameThread(const UInterchangeFactoryBase::FImportSceneObjectsParams& CreateSceneObjectsParams) override;
-
-
 	// Interchange factory base interface end
 	//////////////////////////////////////////////////////////////////////////
 
-private:
+protected:
+	//////////////////////////////////////////////////////////////////////////
+	// Interchange actor factory interface begin
 
-	void SetupStaticMeshActor(const UInterchangeBaseNodeContainer* NodeContainer, const UInterchangeFactoryBaseNode* ActorFactoryNode, AStaticMeshActor* StaticMeshActor);
+	virtual UObject* ProcessActor(class AActor& SpawnedActor, const UInterchangeActorFactoryNode& FactoryNode, const UInterchangeBaseNodeContainer& NodeContainer) override;
+
+	// Interchange actor factory interface end
+	//////////////////////////////////////////////////////////////////////////
 };
 
 
diff --git a/Engine/Source/Runtime/Interchange/Core/Private/Nodes/InterchangeBaseNode.cpp b/Engine/Source/Runtime/Interchange/Core/Private/Nodes/InterchangeBaseNode.cpp
--- a/Engine/Source/Runtime/Interchange/Core/Private/Nodes/InterchangeBaseNode.cpp
+++ b/Engine/Source/Runtime/Interchange/Core/Private/Nodes/InterchangeBaseNode.cpp
@@ -441,7 +441,7 @@ void UInterchangeBaseNode::Serialize(FArchive& Ar)
 	}
 }
 
-void UInterchangeBaseNode::CompareNodeStorage(UInterchangeBaseNode* NodeA, const UInterchangeBaseNode* NodeB, TArray<UE::Interchange::FAttributeKey>& RemovedAttributes, TArray<UE::Interchange::FAttributeKey>& AddedAttributes, TArray<UE::Interchange::FAttributeKey>& ModifiedAttributes)
+void UInterchangeBaseNode::CompareNodeStorage(const UInterchangeBaseNode* NodeA, const UInterchangeBaseNode* NodeB, TArray<UE::Interchange::FAttributeKey>& RemovedAttributes, TArray<UE::Interchange::FAttributeKey>& AddedAttributes, TArray<UE::Interchange::FAttributeKey>& ModifiedAttributes)
 {
 	UE::Interchange::FAttributeStorage::CompareStorage(*(NodeA->Attributes), *(NodeB->Attributes), RemovedAttributes, AddedAttributes, ModifiedAttributes);
 }
diff --git a/Engine/Source/Runtime/Interchange/Core/Private/Nodes/InterchangeFactoryBaseNode.cpp b/Engine/Source/Runtime/Interchange/Core/Private/Nodes/InterchangeFactoryBaseNode.cpp
--- a/Engine/Source/Runtime/Interchange/Core/Private/Nodes/InterchangeFactoryBaseNode.cpp
+++ b/Engine/Source/Runtime/Interchange/Core/Private/Nodes/InterchangeFactoryBaseNode.cpp
@@ -4,6 +4,7 @@
 #include "CoreMinimal.h"
 #include "Types/AttributeStorage.h"
 #include "UObject/Class.h"
+#include "UObject/Package.h"
 #include "UObject/Object.h"
 #include "UObject/ObjectMacros.h"
 
@@ -194,4 +195,15 @@ void UInterchangeFactoryBaseNode::FillAllCustomAttributeFromObject(UObject* Obje
 			}
 		}
 	}
-}
\ No newline at end of file
+}
+
+UInterchangeFactoryBaseNode* UInterchangeFactoryBaseNode::DuplicateWithObject(const UInterchangeFactoryBaseNode* SourceNode, UObject* Object)
+{
+	UInterchangeFactoryBaseNode* CurrentNode = NewObject<UInterchangeFactoryBaseNode>(GetTransientPackage(), SourceNode->GetClass());
+
+	CurrentNode->CopyWithObject(SourceNode, Object);
+
+	CurrentNode->FillAllCustomAttributeFromObject(Object);
+
+	return CurrentNode;
+}
diff --git a/Engine/Source/Runtime/Interchange/Core/Public/InterchangeFactoryBase.h b/Engine/Source/Runtime/Interchange/Core/Public/InterchangeFactoryBase.h
--- a/Engine/Source/Runtime/Interchange/Core/Public/InterchangeFactoryBase.h
+++ b/Engine/Source/Runtime/Interchange/Core/Public/InterchangeFactoryBase.h
@@ -196,6 +196,16 @@ public:
 
 		/** The node container associated with the current source index */
 		const UInterchangeBaseNodeContainer* NodeContainer = nullptr;
+	
+		/**
+		 * If not null, the factory must perform a reimport of the scene node
+		 */
+		UObject* ReimportObject = nullptr;
+	
+		/**
+		 * Factory base node associated with the reimported scene node
+		 */
+		const UInterchangeFactoryBaseNode* ReimportFactoryNode = nullptr;
 	};
 
 	/**
diff --git a/Engine/Source/Runtime/Interchange/Core/Public/Nodes/InterchangeBaseNode.h b/Engine/Source/Runtime/Interchange/Core/Public/Nodes/InterchangeBaseNode.h
--- a/Engine/Source/Runtime/Interchange/Core/Public/Nodes/InterchangeBaseNode.h
+++ b/Engine/Source/Runtime/Interchange/Core/Public/Nodes/InterchangeBaseNode.h
@@ -502,7 +502,7 @@ public:
 	 */
 	virtual void Serialize(FArchive& Ar) override;
 
-	static void CompareNodeStorage(UInterchangeBaseNode* NodeA, const UInterchangeBaseNode* NodeB, TArray<UE::Interchange::FAttributeKey>& RemovedAttributes, TArray<UE::Interchange::FAttributeKey>& AddedAttributes, TArray<UE::Interchange::FAttributeKey>& ModifiedAttributes);
+	static void CompareNodeStorage(const UInterchangeBaseNode* NodeA, const UInterchangeBaseNode* NodeB, TArray<UE::Interchange::FAttributeKey>& RemovedAttributes, TArray<UE::Interchange::FAttributeKey>& AddedAttributes, TArray<UE::Interchange::FAttributeKey>& ModifiedAttributes);
 	
 	static void CopyStorageAttributes(const UInterchangeBaseNode* SourceNode, UInterchangeBaseNode* DestinationNode, TArray<UE::Interchange::FAttributeKey>& AttributeKeys);
 
diff --git a/Engine/Source/Runtime/Interchange/Core/Public/Nodes/InterchangeFactoryBaseNode.h b/Engine/Source/Runtime/Interchange/Core/Public/Nodes/InterchangeFactoryBaseNode.h
--- a/Engine/Source/Runtime/Interchange/Core/Public/Nodes/InterchangeFactoryBaseNode.h
+++ b/Engine/Source/Runtime/Interchange/Core/Public/Nodes/InterchangeFactoryBaseNode.h
@@ -91,6 +91,39 @@ struct FFrame;
 	return false;
 
 
+/*
+ * This macro force the addition of the Fill and Apply delegates on 'this'
+ * if the designated attribute has been set on the source node
+ * Copying the arrays of delegates from the source node to 'this' is not sufficient
+ * because the delegates of the source node have a reference on the source node not 'this'.
+ * Consequently a call to those delegates would always use the source node.
+ */
+#define COPY_NODE_DELEGATES(SourceNode, AttributeName, AttributeType, AssetType)															\
+{																																			\
+	const UE::Interchange::FAttributeKey& AttributeKey = Macro_Custom##AttributeName##Key;													\
+	FString OperationName = GetTypeName() + TEXT(".Get" #AttributeName);																	\
+	AttributeType AttributeValue;																											\
+	if(InterchangePrivateNodeBase::GetCustomAttribute<AttributeType>(*SourceNode->Attributes, AttributeKey, OperationName, AttributeValue)) \
+	{																																		\
+		FName PropertyName = TEXT(""#AttributeName);																						\
+		AddApplyAndFillDelegates<AttributeType>(AttributeKey.Key, AssetType::StaticClass(), PropertyName);									\
+	}																																		\
+}
+
+#define COPY_NODE_DELEGATES_WITH_CUSTOM_DELEGATE(SourceNode, NodeClassName, AttributeName, AttributeType, AssetClass)							\
+{																																				\
+	const UE::Interchange::FAttributeKey& AttributeKey = Macro_Custom##AttributeName##Key;														\
+	FString OperationName = GetTypeName() + TEXT(".Get" #AttributeName);																		\
+	AttributeType AttributeValue;																												\
+	if(InterchangePrivateNodeBase::GetCustomAttribute<AttributeType>(*SourceNode->Attributes, AttributeKey, OperationName, AttributeValue))		\
+	{																																			\
+		TArray<UE::Interchange::FApplyAttributeToAsset>& Delegates = ApplyCustomAttributeDelegates.FindOrAdd(AssetClass);						\
+		Delegates.Add(UE::Interchange::FApplyAttributeToAsset::CreateUObject(this, &NodeClassName::ApplyCustom##AttributeName##ToAsset));		\
+		TArray<UE::Interchange::FFillAttributeToAsset>& FillDelegates = FillCustomAttributeDelegates.FindOrAdd(AssetClass);						\
+		FillDelegates.Add(UE::Interchange::FFillAttributeToAsset::CreateUObject(this, &NodeClassName::FillCustom##AttributeName##FromAsset));	\
+	}																																			\
+}
+
 namespace UE::Interchange
 {
 
@@ -236,6 +269,17 @@ public:
 	template<typename AttributeType>
 	inline bool FillAttributeFromObject(const FString& NodeAttributeKey, UObject* Object, const FName PropertyName);
 
+	/**
+	 * Copies all the attributes from SourceNode to this node.
+	 * @param SourceNode		The source factory node to copy from.
+	 * @param Object			The object to build delegates on custom attributes from, if applicable.
+	 */
+	virtual void CopyWithObject(const UInterchangeFactoryBaseNode* SourceNode, UObject* Object)
+	{
+		// Just copy the attributes storage. Child classes will use Object in their override
+		UInterchangeBaseNode::CopyStorage(SourceNode, this);
+	}
+
 	/**
 	 * Each Attribute that was set and have a delegate set for the specified UObject->UClass will
 	 * get the delegate execute so it apply the attribute to the UObject property.
@@ -245,6 +289,14 @@ public:
 
 	void FillAllCustomAttributeFromObject(UObject* Object) const;
 
+	/**
+	 * Copies all the custom attributes from SourceNode to this node and
+	 * gets the appropriate values from Object.
+	 * @param SourceNode		The source factory node to copy from.
+	 * @param Object			The object to fill custom attributes' values from, if applicable.
+	 */
+	static UInterchangeFactoryBaseNode* DuplicateWithObject(const UInterchangeFactoryBaseNode* SourceNode, UObject* Object);
+
 protected:
 	/**
 	 * Those dependencies are use by the interchange parsing task to make sure the asset are created in the correct order.
diff --git a/Engine/Source/Runtime/Interchange/Engine/Private/InterchangeImportCommon.cpp b/Engine/Source/Runtime/Interchange/Engine/Private/InterchangeImportCommon.cpp
--- a/Engine/Source/Runtime/Interchange/Engine/Private/InterchangeImportCommon.cpp
+++ b/Engine/Source/Runtime/Interchange/Engine/Private/InterchangeImportCommon.cpp
@@ -272,8 +272,8 @@ namespace UE
 
 
 		void FFactoryCommon::ApplyReimportStrategyToAsset(UObject* Asset
-										  , UInterchangeFactoryBaseNode* PreviousAssetNode
-										  , UInterchangeFactoryBaseNode* CurrentAssetNode
+										  , const UInterchangeFactoryBaseNode* PreviousAssetNode
+										  , const UInterchangeFactoryBaseNode* CurrentAssetNode
 										  , UInterchangeFactoryBaseNode* PipelineAssetNode)
 		{
 			TRACE_CPUPROFILER_EVENT_SCOPE("UE::Interchange::FFactoryCommon::ApplyReimportStrategyToAsset")
diff --git a/Engine/Source/Runtime/Interchange/Engine/Public/InterchangeImportCommon.h b/Engine/Source/Runtime/Interchange/Engine/Public/InterchangeImportCommon.h
--- a/Engine/Source/Runtime/Interchange/Engine/Public/InterchangeImportCommon.h
+++ b/Engine/Source/Runtime/Interchange/Engine/Public/InterchangeImportCommon.h
@@ -102,8 +102,8 @@ namespace UE::Interchange
 		* Apply the current strategy to the PipelineAssetNode
 		*/
 		static void ApplyReimportStrategyToAsset(UObject* Asset
-											, UInterchangeFactoryBaseNode* PreviousAssetNode
-											, UInterchangeFactoryBaseNode* CurrentAssetNode
+											, const UInterchangeFactoryBaseNode* PreviousAssetNode
+											, const UInterchangeFactoryBaseNode* CurrentAssetNode
 											, UInterchangeFactoryBaseNode* PipelineAssetNode);
 	};
 } //ns UE::Interchange
