fix Interchange crash in AddMessage when importing unsupported DDS pixel formats
also fix UDIM import when one of the udim blocks failed to load

diff --git a/Engine/Plugins/Interchange/Runtime/Source/Import/Private/Texture/InterchangeTextureFactory.cpp b/Engine/Plugins/Interchange/Runtime/Source/Import/Private/Texture/InterchangeTextureFactory.cpp
--- a/Engine/Plugins/Interchange/Runtime/Source/Import/Private/Texture/InterchangeTextureFactory.cpp
+++ b/Engine/Plugins/Interchange/Runtime/Source/Import/Private/Texture/InterchangeTextureFactory.cpp
@@ -519,12 +519,34 @@ namespace UE::Interchange::Private::InterchangeTextureFactory
 			}
 			, EParallelForFlags::Unbalanced | EParallelForFlags::BackgroundPriority);
 
+		
+		// remove any that failed to import:
+		for (int32 Index = Images.Num()-1; Index >= 0; --Index)
+		{
+			// we will crash later if an image has Format == TSF_Invalid
+			if ( ! Images[Index].IsValid() )
+			{
+				Images.RemoveAt(Index);
+				UDIMsAndSourcesFileArray.RemoveAt(Index);
+			}
+		}
+
+		if ( Images.Num() == 0 )
+		{
+			return {};
+		}
+
+		// InitDataSharedAmongBlocks will not work unless Images[0] is valid :
+		check( Images[0].IsValid() );
+		check( Images.Num() == UDIMsAndSourcesFileArray.Num() );
 
 		UE::Interchange::FImportBlockedImage BlockedImage;
 		// If the image that triggered the import is not valid. We shouldn't import the rest of the images into the UDIM.
 		if (BlockedImage.InitDataSharedAmongBlocks(Images[0]))
 		{
 			BlockedImage.BlocksData.Reserve(Images.Num());
+			// bSRGB starts on :
+			check( BlockedImage.bSRGB );
 
 			bool bMismatchedFormats = false;
 			bool bMismatchedGammaSpace = false;
@@ -544,6 +566,12 @@ namespace UE::Interchange::Private::InterchangeTextureFactory
 				}
 			}
 
+			if ( BlockedImage.bSRGB && ! ERawImageFormat::GetFormatNeedsGammaSpace( FImageCoreUtils::ConvertToRawImageFormat(BlockedImage.Format) ) )
+			{
+				BlockedImage.bSRGB = false;
+				bMismatchedGammaSpace = true;
+			}
+
 			if (bMismatchedGammaSpace || bMismatchedFormats)
 			{
 				UE_LOG(LogInterchangeImport, Warning, TEXT("Mismatched UDIM image %s, converting all to %s/%s ..."), bMismatchedGammaSpace ? TEXT("gamma spaces") : TEXT("pixel formats"),
