deprecate DDSLoader ; use DDSFile instead

diff --git a/Engine/Plugins/Interchange/Runtime/Source/Import/Private/Texture/InterchangeDDSTranslator.cpp b/Engine/Plugins/Interchange/Runtime/Source/Import/Private/Texture/InterchangeDDSTranslator.cpp
--- a/Engine/Plugins/Interchange/Runtime/Source/Import/Private/Texture/InterchangeDDSTranslator.cpp
+++ b/Engine/Plugins/Interchange/Runtime/Source/Import/Private/Texture/InterchangeDDSTranslator.cpp
@@ -3,7 +3,6 @@
 #include "Texture/InterchangeDDSTranslator.h"
 
 #include "DDSFile.h"
-#include "DDSLoader.h"
 #include "Engine/Texture.h"
 #include "Engine/Texture2D.h"
 #include "HAL/FileManager.h"
@@ -39,14 +38,13 @@ namespace UE::Interchange::Private::InterchangeDDSTranslator
 
 		if (FileReaderArchive)
 		{
-			const int64 SizeOfDDSMarker = 4;
 			const int64 SizeOfFile = FileReaderArchive->TotalSize();
-			const int64 MinimalHeaderSize = FDDSLoadHelper::GetDDSHeaderMinimalSize() + SizeOfDDSMarker;
+			const int64 MinimalHeaderSize = UE::DDS::GetDDSHeaderMinimalSize();
 
 			// If the file is not bigger than the smallest header possible then clearly the file is not valid as a dds file.
 			if (SizeOfFile > MinimalHeaderSize)
 			{
-				const int64 MaximalHeaderSize = FDDSLoadHelper::GetDDSHeaderMaximalSize() + SizeOfDDSMarker;
+				const int64 MaximalHeaderSize = UE::DDS::GetDDSHeaderMaximalSize();
 				const int64 BytesToRead = FMath::Min(SizeOfFile, MaximalHeaderSize);
 
 				OutHeader.Reset(BytesToRead);
diff --git a/Engine/Source/Runtime/Engine/Private/DDSLoader.cpp b/Engine/Source/Runtime/Engine/Private/DDSLoader.cpp
--- a/Engine/Source/Runtime/Engine/Private/DDSLoader.cpp
+++ b/Engine/Source/Runtime/Engine/Private/DDSLoader.cpp
@@ -5,6 +5,10 @@
 #include "Engine/Texture2D.h"
 #include "RenderUtils.h"
 
+// deprecated and unused -> delete me !
+// disable warnings about myself :
+PRAGMA_DISABLE_DEPRECATION_WARNINGS
+
 FDDSLoadHelper::FDDSLoadHelper(const uint8* Buffer, uint64 Length) 
 	: DDSHeader(0), DDS10Header(0)
 {
@@ -258,3 +262,4 @@ int64 FDDSLoadHelper::GetDDSHeaderMinimalSize()
 	return sizeof(FDDSFileHeader);
 }
 
+PRAGMA_ENABLE_DEPRECATION_WARNINGS
diff --git a/Engine/Source/Runtime/Engine/Private/HDRLoader.cpp b/Engine/Source/Runtime/Engine/Private/HDRLoader.cpp
--- a/Engine/Source/Runtime/Engine/Private/HDRLoader.cpp
+++ b/Engine/Source/Runtime/Engine/Private/HDRLoader.cpp
@@ -4,6 +4,8 @@
 #include "HDRLoader.h"
 #include "DDSLoader.h"
 
+// deprecated and unused -> delete me !
+
 DEFINE_LOG_CATEGORY_STATIC(LogHDRLoader, Log, All);
 
 FHDRLoadHelper::FHDRLoadHelper(const uint8* Buffer, uint32 BufferLength)
diff --git a/Engine/Source/Runtime/Engine/Public/DDSLoader.h b/Engine/Source/Runtime/Engine/Public/DDSLoader.h
--- a/Engine/Source/Runtime/Engine/Public/DDSLoader.h
+++ b/Engine/Source/Runtime/Engine/Public/DDSLoader.h
@@ -93,9 +93,7 @@ struct FDDS10FileHeader
 };
 #pragma pack(pop)
 
-// @todo Oodle : deprecate old DDSLoader after it is removed from InterchangeDDSTranslator 
-// 
-//class UE_DEPRECATED(5.1, "Use DDSFile instead") FDDSLoadHelper;
+class UE_DEPRECATED(5.2, "Use DDSFile instead") FDDSLoadHelper;
 class FDDSLoadHelper
 {
 public:
diff --git a/Engine/Source/Runtime/ImageCore/Private/DDSFile.cpp b/Engine/Source/Runtime/ImageCore/Private/DDSFile.cpp
--- a/Engine/Source/Runtime/ImageCore/Private/DDSFile.cpp
+++ b/Engine/Source/Runtime/ImageCore/Private/DDSFile.cpp
@@ -737,6 +737,18 @@ static EDDSError ReadPayload(FDDSFile* InDDS, const uint8* InReadCursor, const u
 }
 
 
+int64 GetDDSHeaderMaximalSize()
+{
+	return sizeof(FDDSHeaderWithMagic) + sizeof(FDDSHeaderDX10);
+}
+
+int64 GetDDSHeaderMinimalSize()
+{
+	return sizeof(FDDSHeaderWithMagic);
+}
+
+
+
 /* static */ FDDSFile* FDDSFile::CreateFromDDSInMemory(const uint8* InDDS, int64 InDDSSize, EDDSError *OutError, bool bHeaderOnly)
 {
 	// If no OutError passed in, redirect it to dummy storage on stack.
diff --git a/Engine/Source/Runtime/ImageCore/Public/DDSFile.h b/Engine/Source/Runtime/ImageCore/Public/DDSFile.h
--- a/Engine/Source/Runtime/ImageCore/Public/DDSFile.h
+++ b/Engine/Source/Runtime/ImageCore/Public/DDSFile.h
@@ -337,4 +337,9 @@ namespace UE { namespace DDS
 	IMAGECORE_API EDXGIFormat DXGIFormatFromRawFormat(ERawImageFormat::Type RawFormat,EGammaSpace GammaSpace);
 
 	IMAGECORE_API ERawImageFormat::Type DXGIFormatGetClosestRawFormat(EDXGIFormat fmt, bool * pIsExactMatch = nullptr);
+	
+	// note these sizes include the 4-byte magic word, unlike the older DDSLoader functions
+	IMAGECORE_API int64 GetDDSHeaderMaximalSize();
+	IMAGECORE_API int64 GetDDSHeaderMinimalSize();
+
 } }
