fix "Imported DDS files must either be a single base mip" warning being logged 5 times on each image

diff --git a/Engine/Plugins/Interchange/Runtime/Source/Import/Private/Texture/InterchangeDDSTranslator.cpp b/Engine/Plugins/Interchange/Runtime/Source/Import/Private/Texture/InterchangeDDSTranslator.cpp
--- a/Engine/Plugins/Interchange/Runtime/Source/Import/Private/Texture/InterchangeDDSTranslator.cpp
+++ b/Engine/Plugins/Interchange/Runtime/Source/Import/Private/Texture/InterchangeDDSTranslator.cpp
@@ -71,6 +71,7 @@ namespace UE::Interchange::Private::InterchangeDDSTranslator
 		if (bHeaderOnly)
 		{
 			LoadDDSHeaderFromFile(DDSSourceData, Filename);
+			// no error checking?
 		}
 		else if (!FFileHelper::LoadFileToArray(DDSSourceData, *Filename))
 		{
@@ -87,22 +88,25 @@ namespace UE::Interchange::Private::InterchangeDDSTranslator
 			check(Error != UE::DDS::EDDSError::OK);
 			if (Error != UE::DDS::EDDSError::NotADds && Error != UE::DDS::EDDSError::IoError)
 			{
-				UE_LOG(LogInterchangeImport, Warning, TEXT("Failed to load DDS (Error=%d)"), (int)Error);
+				UE_LOG(LogInterchangeImport, Warning, TEXT("Failed to load DDS (Error=%d) [%s]"), (int)Error, *Filename);
 			}
 			return {};
 		}
 
-		// We require a complete mip chain or a single mip.
-		if (DDS->MipCount > 1)
+		if (!bHeaderOnly)
 		{
-			if (DDS->MipCount != FImageCoreUtils::GetMipCountFromDimensions(DDS->Width, DDS->Height, DDS->Depth, DDS->IsValidTextureVolume()))
+			// this function is called 4 times with bHeaderOnly=true, then once with bHeaderOnly=false
+			// do noisy warnings only in the last call so they aren't repeated
+
+			// We require a complete mip chain or a single mip.
+			if (DDS->MipCount > 1)
 			{
-				UE_LOG(LogInterchangeImport, Warning, TEXT("Imported DDS files must either be a single base mip, or a full mip chain, not a partial mip chain (found %d expected %d)"), DDS->MipCount, FImageCoreUtils::GetMipCountFromDimensions(DDS->Width, DDS->Height, DDS->Depth, DDS->IsValidTextureVolume()));
+				if (DDS->MipCount != FImageCoreUtils::GetMipCountFromDimensions(DDS->Width, DDS->Height, DDS->Depth, DDS->IsValidTextureVolume()))
+				{
+					UE_LOG(LogInterchangeImport, Warning, TEXT("Imported DDS files must either be a single base mip, or a full mip chain, not a partial mip chain (found %d expected %d) [%s]"), DDS->MipCount, FImageCoreUtils::GetMipCountFromDimensions(DDS->Width, DDS->Height, DDS->Depth, DDS->IsValidTextureVolume()), *Filename);
+				}
 			}
-		}
 
-		if (!bHeaderOnly)
-		{
 			// change X8 formats to A8 :	
 			DDS->ConvertRGBXtoRGBA();
 			// change RGBA8 to BGRA8 before DXGIFormatGetClosestRawFormat :
