fix Interchange texture import uses of int32 breaking large texture import

diff --git a/Engine/Plugins/Interchange/Runtime/Source/Import/Private/Texture/InterchangeBlockedTexturePayloadData.cpp b/Engine/Plugins/Interchange/Runtime/Source/Import/Private/Texture/InterchangeBlockedTexturePayloadData.cpp
--- a/Engine/Plugins/Interchange/Runtime/Source/Import/Private/Texture/InterchangeBlockedTexturePayloadData.cpp
+++ b/Engine/Plugins/Interchange/Runtime/Source/Import/Private/Texture/InterchangeBlockedTexturePayloadData.cpp
@@ -11,13 +11,13 @@ namespace UE::Interchange
 		int64 ComputeBlockSize(const FTextureSourceBlock& Block, ETextureSourceFormat Format)
 		{
 			int64 Size = 0;
-			int32 Width = Block.SizeX;
-			int32 Height = Block.SizeY;
+			int64 Width = Block.SizeX;
+			int64 Height = Block.SizeY;
 			for (int32 CurrentMip = 0; CurrentMip < Block.NumMips; ++CurrentMip)
 			{
 				Size += Width * Height * FTextureSource::GetBytesPerPixel(Format) * Block.NumSlices;
-				Width >>= 2;
-				Height >>= 2;
+				Width  = FMath::Max(Width>>1,1);
+				Height = FMath::Max(Height>>1,1);
 			}
 			return Size;
 		}
diff --git a/Engine/Plugins/Interchange/Runtime/Source/Import/Private/Texture/InterchangeImageWrapperTranslator.cpp b/Engine/Plugins/Interchange/Runtime/Source/Import/Private/Texture/InterchangeImageWrapperTranslator.cpp
--- a/Engine/Plugins/Interchange/Runtime/Source/Import/Private/Texture/InterchangeImageWrapperTranslator.cpp
+++ b/Engine/Plugins/Interchange/Runtime/Source/Import/Private/Texture/InterchangeImageWrapperTranslator.cpp
@@ -148,7 +148,7 @@ TOptional<UE::Interchange::FImportImage> UInterchangeImageWrapperTranslator::Get
 	const uint8* Buffer = SourceDataBuffer.GetData();
 	const uint8* BufferEnd = Buffer + SourceDataBuffer.Num();
 
-	const int32 Length = BufferEnd - Buffer;
+	const int64 Length = BufferEnd - Buffer;
 
 	IImageWrapperModule& ImageWrapperModule = FModuleManager::LoadModuleChecked<IImageWrapperModule>(FName("ImageWrapper"));
 
diff --git a/Engine/Plugins/Interchange/Runtime/Source/Import/Private/Texture/InterchangeJPGTranslator.cpp b/Engine/Plugins/Interchange/Runtime/Source/Import/Private/Texture/InterchangeJPGTranslator.cpp
--- a/Engine/Plugins/Interchange/Runtime/Source/Import/Private/Texture/InterchangeJPGTranslator.cpp
+++ b/Engine/Plugins/Interchange/Runtime/Source/Import/Private/Texture/InterchangeJPGTranslator.cpp
@@ -88,7 +88,7 @@ TOptional<UE::Interchange::FImportImage> UInterchangeJPGTranslator::GetTexturePa
 	const uint8* Buffer = SourceDataBuffer.GetData();
 	const uint8* BufferEnd = Buffer + SourceDataBuffer.Num();
 
-	const int32 Length = BufferEnd - Buffer;
+	const int64 Length = BufferEnd - Buffer;
 
 	IImageWrapperModule& ImageWrapperModule = FModuleManager::LoadModuleChecked<IImageWrapperModule>(FName("ImageWrapper"));
 
diff --git a/Engine/Plugins/Interchange/Runtime/Source/Import/Private/Texture/InterchangePCXTranslator.cpp b/Engine/Plugins/Interchange/Runtime/Source/Import/Private/Texture/InterchangePCXTranslator.cpp
--- a/Engine/Plugins/Interchange/Runtime/Source/Import/Private/Texture/InterchangePCXTranslator.cpp
+++ b/Engine/Plugins/Interchange/Runtime/Source/Import/Private/Texture/InterchangePCXTranslator.cpp
@@ -106,7 +106,7 @@ TOptional<UE::Interchange::FImportImage> UInterchangePCXTranslator::GetTexturePa
 	const uint8* Buffer = SourceDataBuffer.GetData();
 	const uint8* BufferEnd = Buffer + SourceDataBuffer.Num();
 
-	const int32 Length = BufferEnd - Buffer;
+	const int64 Length = BufferEnd - Buffer;
 
 	FModuleManager::LoadModuleChecked<IImageWrapperModule>(FName("ImageWrapper"));
 
diff --git a/Engine/Plugins/Interchange/Runtime/Source/Import/Private/Texture/InterchangePSDTranslator.cpp b/Engine/Plugins/Interchange/Runtime/Source/Import/Private/Texture/InterchangePSDTranslator.cpp
--- a/Engine/Plugins/Interchange/Runtime/Source/Import/Private/Texture/InterchangePSDTranslator.cpp
+++ b/Engine/Plugins/Interchange/Runtime/Source/Import/Private/Texture/InterchangePSDTranslator.cpp
@@ -408,7 +408,7 @@ TOptional<UE::Interchange::FImportImage> UInterchangePSDTranslator::GetTexturePa
 	const uint8* Buffer = SourceDataBuffer.GetData();
 	const uint8* BufferEnd = Buffer + SourceDataBuffer.Num();
 
-	const int32 Length = BufferEnd - Buffer;
+	const int64 Length = BufferEnd - Buffer;
 
 	//
 	// PSD File
