Fixed an ensure when the import of a udim report an error from the interchange translator. It's was simply missing the results container to capture the error message.

diff --git a/Engine/Plugins/Interchange/Runtime/Source/Import/Private/Texture/InterchangeTextureFactory.cpp b/Engine/Plugins/Interchange/Runtime/Source/Import/Private/Texture/InterchangeTextureFactory.cpp
--- a/Engine/Plugins/Interchange/Runtime/Source/Import/Private/Texture/InterchangeTextureFactory.cpp
+++ b/Engine/Plugins/Interchange/Runtime/Source/Import/Private/Texture/InterchangeTextureFactory.cpp
@@ -78,7 +78,11 @@ namespace UE::Interchange::Private::InterchangeTextureFactory
 					NewObject<UInterchangeTranslatorBase>(TransientPackage, TranslatorClass, NAME_None)
 					, NewObject<UInterchangeSourceData>(TransientPackage, UInterchangeSourceData::StaticClass(), NAME_None)
 				);
+			}
 
+			for (TPair<UInterchangeTranslatorBase*, UInterchangeSourceData*>& TranslatorAndSourceData : TranslatorsAndSourcesData)
+			{
+				TranslatorAndSourceData.Key->SetResultsContainer(BaseTranslator.Results);
 			}
 		};
 
