Interchange fix some issues with UDIM import of DDS with existing mips

diff --git a/Engine/Plugins/Interchange/Runtime/Source/Import/Private/Texture/InterchangeBlockedTexturePayloadData.cpp b/Engine/Plugins/Interchange/Runtime/Source/Import/Private/Texture/InterchangeBlockedTexturePayloadData.cpp
--- a/Engine/Plugins/Interchange/Runtime/Source/Import/Private/Texture/InterchangeBlockedTexturePayloadData.cpp
+++ b/Engine/Plugins/Interchange/Runtime/Source/Import/Private/Texture/InterchangeBlockedTexturePayloadData.cpp
@@ -49,7 +49,8 @@ namespace UE::Interchange
 			Format = Image.Format;
 			CompressionSettings = Image.CompressionSettings;
 			bSRGB = Image.bSRGB;
-			MipGenSettings = Image.MipGenSettings;
+			//don't set MipGenSettings from Image 0
+			//MipGenSettings = Image.MipGenSettings;
 			return true;
 		}
 
diff --git a/Engine/Plugins/Interchange/Runtime/Source/Import/Private/Texture/InterchangeTextureFactory.cpp b/Engine/Plugins/Interchange/Runtime/Source/Import/Private/Texture/InterchangeTextureFactory.cpp
--- a/Engine/Plugins/Interchange/Runtime/Source/Import/Private/Texture/InterchangeTextureFactory.cpp
+++ b/Engine/Plugins/Interchange/Runtime/Source/Import/Private/Texture/InterchangeTextureFactory.cpp
@@ -549,8 +549,6 @@ namespace UE::Interchange::Private::InterchangeTextureFactory
 		if (BlockedImage.InitDataSharedAmongBlocks(Images[0]))
 		{
 			BlockedImage.BlocksData.Reserve(Images.Num());
-			// bSRGB starts on :
-			check( BlockedImage.bSRGB );
 
 			bool bMismatchedFormats = false;
 			bool bMismatchedGammaSpace = false;
@@ -585,31 +583,71 @@ namespace UE::Interchange::Private::InterchangeTextureFactory
 				{
 					if (Image.bSRGB != BlockedImage.bSRGB || Image.Format != BlockedImage.Format)
 					{
+						check(Image.IsValid());
+
 						ERawImageFormat::Type ImageRawFormat = FImageCoreUtils::ConvertToRawImageFormat(Image.Format);
 						FImageView SourceImage(Image.RawData.GetData(), Image.SizeX, Image.SizeY, 1, ImageRawFormat, Image.bSRGB ? EGammaSpace::sRGB : EGammaSpace::Linear);
-						FImage DestImage(Image.SizeX, Image.SizeY, FImageCoreUtils::ConvertToRawImageFormat(BlockedImage.Format), BlockedImage.bSRGB ? EGammaSpace::sRGB : EGammaSpace::Linear);
+						check( SourceImage.GetImageSizeBytes() <= (int64)Image.RawData.GetSize() );
+						
+						ERawImageFormat::Type BlockedImageRawFormat = FImageCoreUtils::ConvertToRawImageFormat(BlockedImage.Format);
+						FImage DestImage(Image.SizeX, Image.SizeY, BlockedImageRawFormat, BlockedImage.bSRGB ? EGammaSpace::sRGB : EGammaSpace::Linear);
 						FImageCore::CopyImage(SourceImage, DestImage);
 
 						Image.RawData = MakeUniqueBufferFromArray(MoveTemp(DestImage.RawData));
 						Image.Format = BlockedImage.Format;
 						Image.bSRGB = BlockedImage.bSRGB;
+
+						if ( Image.NumMips != 1 )
+						{
+							UE_LOG(LogInterchangeImport, Warning, TEXT("UDIM Image had existing mips; they were discarded in format change") );
+							Image.NumMips = 1; // discard imported mips, they won't be used in UDIM VT anyway
+							Image.MipGenSettings.Reset();
+						}
+
+						check(Image.IsValid());
 					}
 				}
 			}
 
+			bool bAllImagesSameNumMips = true;
 			for (int32 Index = 0; Index < Images.Num(); ++Index)
 			{
 				int32 UDIMIndex = UDIMsAndSourcesFileArray[Index]->Key;
 				int32 BlockX;
 				int32 BlockY;
 				UE::TextureUtilitiesCommon::ExtractUDIMCoordinates(UDIMIndex, BlockX, BlockY);
-				BlockedImage.InitBlockFromImage(BlockX, BlockY, Images[Index]);
+
+				if ( ! BlockedImage.InitBlockFromImage(BlockX, BlockY, Images[Index]) )
+				{
+					UE_LOG(LogInterchangeImport, Warning, TEXT("UDIM InitBlockFromImage failed") );
+					return { };
+				}
+
+				if ( Images[Index].NumMips != Images[0].NumMips )
+				{
+					bAllImagesSameNumMips = false;
+				}
 			}
 
-			BlockedImage.MigrateDataFromImagesToRawData(Images);
-		}
+			if (! BlockedImage.MigrateDataFromImagesToRawData(Images) )
+			{
+				UE_LOG(LogInterchangeImport, Warning, TEXT("UDIM MigrateDataFromImagesToRawData failed") );
+				return { };
+			}
+			
+			// BlockedImage.MipGenSettings is note set from Images[].MipGenSettings
+			//	instead see if all UDIM blocks have existing mips, and if so set LeaveExisting
+			if ( Images[0].NumMips != 1 && bAllImagesSameNumMips )
+			{
+				BlockedImage.MipGenSettings = TextureMipGenSettings::TMGS_LeaveExistingMips;
+			}
 
-		return BlockedImage;
+			return BlockedImage;
+		}
+		else
+		{
+			return { };
+		}
 	}
 
 	FTexturePayloadVariant GetTexturePayload(const UInterchangeSourceData* SourceData, const FString& PayloadKey, const FTextureNodeVariant& TextureNodeVariant, const FTextureFactoryNodeVariant& FactoryNodeVariant, const UInterchangeTranslatorBase* Translator, TOptional<FString>& AlternateTexturePath)
@@ -1402,15 +1440,15 @@ namespace UE::Interchange::Private::InterchangeTextureFactory
 
 		if (CVarVirtualTexturesEnabled->GetValueOnGameThread() && CVarVirtualTexturesAutoImportEnabled->GetValueOnGameThread())
 		{
-			const int VirtualTextureAutoEnableThreshold = GetDefault<UTextureImportSettings>()->AutoVTSize;
-			const int VirtualTextureAutoEnableThresholdPixels = VirtualTextureAutoEnableThreshold * VirtualTextureAutoEnableThreshold;
+			const int64 VirtualTextureAutoEnableThreshold = GetDefault<UTextureImportSettings>()->AutoVTSize;
+			const int64 VirtualTextureAutoEnableThresholdPixels = VirtualTextureAutoEnableThreshold * VirtualTextureAutoEnableThreshold;
 
 			// We do this in pixels so a 8192 x 128 texture won't get VT enabled 
 			// We use the Source size instead of simple Texture2D->GetSizeX() as this uses the size of the platform data
 			// however for a new texture platform data may not be generated yet, and for an reimport of a texture this is the size of the
 			// old texture. 
 			// Using source size gives one small caveat. It looks at the size before mipmap power of two padding adjustment.
-			if (Texture2D->Source.GetSizeX() * Texture2D->Source.GetSizeY() >= VirtualTextureAutoEnableThresholdPixels ||
+			if ( (int64) Texture2D->Source.GetSizeX() * Texture2D->Source.GetSizeY() >= VirtualTextureAutoEnableThresholdPixels ||
 				Texture2D->Source.GetSizeX() > UTexture::GetMaximumDimensionOfNonVT() ||
 				Texture2D->Source.GetSizeY() > UTexture::GetMaximumDimensionOfNonVT())
 			{
diff --git a/Engine/Source/Runtime/Engine/Private/Texture.cpp b/Engine/Source/Runtime/Engine/Private/Texture.cpp
--- a/Engine/Source/Runtime/Engine/Private/Texture.cpp
+++ b/Engine/Source/Runtime/Engine/Private/Texture.cpp
@@ -546,7 +546,7 @@ void UTexture::ValidateSettingsAfterImportOrEdit(bool * pRequiresNotifyMaterials
 			{
 				if ( GetTextureClass() == ETextureClass::TwoD )
 				{
-				UE_LOG(LogTexture, Display, TEXT("Large Texture %s Dimension=%d changed to VT; to disable VT set MaxTextureSize first"), *GetName(),MaxDimension);
+					UE_LOG(LogTexture, Display, TEXT("Large Texture %s Dimension=%d changed to VT; to disable VT set MaxTextureSize first"), *GetName(),MaxDimension);
 					VirtualTextureStreaming = true;
 					bRequiresNotifyMaterials = true;
 				}
@@ -571,7 +571,7 @@ void UTexture::ValidateSettingsAfterImportOrEdit(bool * pRequiresNotifyMaterials
 		{
 			if (!bIsPowerOfTwo)
 			{
-				if ( bLargeTextureMustBeVT )
+				if ( bLargeTextureMustBeVT || Source.GetNumBlocks() > 1 )
 				{
 					UE_LOG(LogTexture, Warning, TEXT("Large VT \"%s\", must be padded to power-of-2 for VT support (%dx%d)"), *GetName(), Source.GetSizeX(),Source.GetSizeY());
 					// VT nonpow2 will fail to build
