- Prefixing MaterialX material expressions classes and names to avoid potential validation errors on some plugins
- Redirecting material expressions in BaseEngine.ini
- Optimizing Length material expression for Scalar inputs by just outputting abs value which is practically free on gpu
- Modifiying UMaterialExpressionTextureSampleParameterBlur, Kernel size should be an enum instead of a float

diff --git a/Engine/Config/BaseEngine.ini b/Engine/Config/BaseEngine.ini
--- a/Engine/Config/BaseEngine.ini
+++ b/Engine/Config/BaseEngine.ini
@@ -1151,6 +1151,47 @@ DefaultWorldPartitionHLODLayer=/Engine/Maps/Templates/HLODs/HLODLayer_Instanced.
 +ClassRedirects=(OldName="XRDeviceVisualizationComponent", NewName="/Script/XRBase.XRDeviceVisualizationComponent")
 +ClassRedirects=(OldName="XRLoadingScreenFunctionLibrary", NewName="/Script/XRBase.XRLoadingScreenFunctionLibrary")
 
+; Rename MaterialX Material Expressions and Enums
++ClassRedirects=(OldName="/Script/InterchangeImport.UMaterialExpressionAppend3Vector", NewName="/Script/InterchangeImport.UMaterialExpressionMaterialXAppend3Vector")
++ClassRedirects=(OldName="/Script/InterchangeImport.UMaterialExpressionAppend4Vector", NewName="/Script/InterchangeImport.UMaterialExpressionMaterialXAppend4Vector")
++ClassRedirects=(OldName="/Script/InterchangeImport.UMaterialExpressionBurn", NewName="/Script/InterchangeImport.UMaterialExpressionMaterialXBurn")
++ClassRedirects=(OldName="/Script/InterchangeImport.UMaterialExpressionDifference", NewName="/Script/InterchangeImport.UMaterialExpressionMaterialXDifference")
++ClassRedirects=(OldName="/Script/InterchangeImport.UMaterialExpressionDisjointOver", NewName="/Script/InterchangeImport.UMaterialExpressionMaterialXDisjointOver")
++ClassRedirects=(OldName="/Script/InterchangeImport.UMaterialExpressionDodge", NewName="/Script/InterchangeImport.UMaterialExpressionMaterialXDodge")
++ClassRedirects=(OldName="/Script/InterchangeImport.UMaterialExpressionExponential", NewName="/Script/InterchangeImport.UMaterialExpressionMaterialXExponential")
++ClassRedirects=(OldName="/Script/InterchangeImport.UMaterialExpressionFractal3D", NewName="/Script/InterchangeImport.UMaterialExpressionMaterialXFractal3D")
++ClassRedirects=(OldName="/Script/InterchangeImport.UMaterialExpressionHsvToRgb", NewName="/Script/InterchangeImport.UMaterialExpressionMaterialXHsvToRgb")
++ClassRedirects=(OldName="/Script/InterchangeImport.UMaterialExpressionIn", NewName="/Script/InterchangeImport.UMaterialExpressionMaterialXIn")
++ClassRedirects=(OldName="/Script/InterchangeImport.UMaterialExpressionLength", NewName="/Script/InterchangeImport.UMaterialExpressionMaterialXLength")
++ClassRedirects=(OldName="/Script/InterchangeImport.UMaterialExpressionLogarithm", NewName="/Script/InterchangeImport.UMaterialExpressionMaterialXLogarithm")
++ClassRedirects=(OldName="/Script/InterchangeImport.UMaterialExpressionLuminance", NewName="/Script/InterchangeImport.UMaterialExpressionMaterialXLuminance")
++ClassRedirects=(OldName="/Script/InterchangeImport.UMaterialExpressionMask", NewName="/Script/InterchangeImport.UMaterialExpressionMaterialXMask")
++ClassRedirects=(OldName="/Script/InterchangeImport.UMaterialExpressionMatte", NewName="/Script/InterchangeImport.UMaterialExpressionMaterialXMatte")
++ClassRedirects=(OldName="/Script/InterchangeImport.UMaterialExpressionMinus", NewName="/Script/InterchangeImport.UMaterialExpressionMaterialXMinus")
++ClassRedirects=(OldName="/Script/InterchangeImport.UMaterialExpressionOut", NewName="/Script/InterchangeImport.UMaterialExpressionMaterialXOut")
++ClassRedirects=(OldName="/Script/InterchangeImport.UMaterialExpressionOver", NewName="/Script/InterchangeImport.UMaterialExpressionMaterialXOver")
++ClassRedirects=(OldName="/Script/InterchangeImport.UMaterialExpressionOverlay", NewName="/Script/InterchangeImport.UMaterialExpressionMaterialXOverlay")
++ClassRedirects=(OldName="/Script/InterchangeImport.UMaterialExpressionPlace2D", NewName="/Script/InterchangeImport.UMaterialExpressionMaterialXPlace2D")
++ClassRedirects=(OldName="/Script/InterchangeImport.UMaterialExpressionPlus", NewName="/Script/InterchangeImport.UMaterialExpressionMaterialXPlus")
++ClassRedirects=(OldName="/Script/InterchangeImport.UMaterialExpressionPremult", NewName="/Script/InterchangeImport.UMaterialExpressionMaterialXPremult")
++ClassRedirects=(OldName="/Script/InterchangeImport.UMaterialExpressionRamp4", NewName="/Script/InterchangeImport.UMaterialExpressionMaterialXRamp4")
++ClassRedirects=(OldName="/Script/InterchangeImport.UMaterialExpressionRampLeftRight", NewName="/Script/InterchangeImport.UMaterialExpressionMaterialXRampLeftRight")
++ClassRedirects=(OldName="/Script/InterchangeImport.UMaterialExpressionRampTopBottom", NewName="/Script/InterchangeImport.UMaterialExpressionMaterialXRampTopBottom")
++ClassRedirects=(OldName="/Script/InterchangeImport.UMaterialExpressionRemap", NewName="/Script/InterchangeImport.UMaterialExpressionMaterialXRemap")
++ClassRedirects=(OldName="/Script/InterchangeImport.UMaterialExpressionRgbToHsv", NewName="/Script/InterchangeImport.UMaterialExpressionMaterialXRgbToHsv")
++ClassRedirects=(OldName="/Script/InterchangeImport.UMaterialExpressionRotate2D", NewName="/Script/InterchangeImport.UMaterialExpressionMaterialXRotate2D")
++ClassRedirects=(OldName="/Script/InterchangeImport.UMaterialExpressionScreen", NewName="/Script/InterchangeImport.UMaterialExpressionMaterialXScreen")
++ClassRedirects=(OldName="/Script/InterchangeImport.UMaterialExpressionSplitLeftRight", NewName="/Script/InterchangeImport.UMaterialExpressionMaterialXSplitLeftRight")
++ClassRedirects=(OldName="/Script/InterchangeImport.UMaterialExpressionSplitTopBottom", NewName="/Script/InterchangeImport.UMaterialExpressionMaterialXSplitTopBottom")
++ClassRedirects=(OldName="/Script/InterchangeImport.UMaterialExpressionTextureSampleParameterBlur", NewName="/Script/InterchangeImport.UMaterialExpressionMaterialXTextureSampleParameterBlur")
++ClassRedirects=(OldName="/Script/InterchangeImport.UMaterialExpressionUnpremult", NewName="/Script/InterchangeImport.UMaterialExpressionMaterialXUnpremult")
++EnumRedirects=(OldName="ETextureSampleBlurFilter", NewName="EMaterialXTextureSampleBlurFilter")
++EnumRedirects=(OldName="ELuminanceMode", NewName="EMaterialXLuminanceMode")
+
+
+
+
+
 [CoreUObject.Metadata]
 ; Note: UnrealHeaderTool should be rerun after making changes to MetadataRedirects to catch any keys specified in class headers
 +MetadataRedirects=(OldKey="K2Protected", NewKey="BlueprintProtected")
diff --git a/Engine/Plugins/Interchange/Runtime/Source/Import/Private/MaterialX/MaterialExpressions/MaterialExpressionAppend3Vector.cpp b/Engine/Plugins/Interchange/Runtime/Source/Import/Private/MaterialX/MaterialExpressions/MaterialExpressionAppend3Vector.cpp
--- a/Engine/Plugins/Interchange/Runtime/Source/Import/Private/MaterialX/MaterialExpressions/MaterialExpressionAppend3Vector.cpp
+++ b/Engine/Plugins/Interchange/Runtime/Source/Import/Private/MaterialX/MaterialExpressions/MaterialExpressionAppend3Vector.cpp
@@ -5,46 +5,40 @@
 
 #include UE_INLINE_GENERATED_CPP_BY_NAME(MaterialExpressionAppend3Vector)
 
-#define LOCTEXT_NAMESPACE "MaterialExpressionAppend3Vector"
+#define LOCTEXT_NAMESPACE "MaterialExpressionMaterialXAppend3Vector"
 
-UMaterialExpressionAppend3Vector::UMaterialExpressionAppend3Vector(const FObjectInitializer& ObjectInitializer)
+UMaterialExpressionMaterialXAppend3Vector::UMaterialExpressionMaterialXAppend3Vector(const FObjectInitializer& ObjectInitializer)
 	: Super(ObjectInitializer)
 {
 	// Structure to hold one-time initialization
 	struct FConstructorStatics
 	{
 		FText NAME_MaterialX;
-		FText NAME_Math;
-		FText NAME_VectorOps;
 		FConstructorStatics()
 			: NAME_MaterialX(LOCTEXT("MaterialX", "MaterialX"))
-			, NAME_Math(LOCTEXT("Math", "Math"))
-			, NAME_VectorOps(LOCTEXT("VectorOps", "VectorOps"))
 		{}
 	};
 	static FConstructorStatics ConstructorStatics;
 
 #if WITH_EDITORONLY_DATA
 	MenuCategories.Add(ConstructorStatics.NAME_MaterialX);
-	MenuCategories.Add(ConstructorStatics.NAME_Math);
-	MenuCategories.Add(ConstructorStatics.NAME_VectorOps);
 #endif
 }
 
 #if WITH_EDITOR
-int32 UMaterialExpressionAppend3Vector::Compile(FMaterialCompiler* Compiler, int32 OutputIndex)
+int32 UMaterialExpressionMaterialXAppend3Vector::Compile(FMaterialCompiler* Compiler, int32 OutputIndex)
 {
 	if(!A.GetTracedInput().Expression)
 	{
-		return Compiler->Errorf(TEXT("Missing AppendVector input A"));
+		return Compiler->Errorf(TEXT("Missing MaterialX Append3 input A"));
 	}
 	else if(!B.GetTracedInput().Expression)
 	{
-		return Compiler->Errorf(TEXT("Missing AppendVector input B"));
+		return Compiler->Errorf(TEXT("Missing MaterialX Append3 input B"));
 	}
 	else if(!C.GetTracedInput().Expression)
 	{
-		return Compiler->Errorf(TEXT("Missing AppendVector input C"));
+		return Compiler->Errorf(TEXT("Missing MaterialX Append3 input C"));
 	}
 	else
 	{
@@ -56,9 +50,9 @@ int32 UMaterialExpressionAppend3Vector::Compile(FMaterialCompiler* Compiler, int
 	}
 }
 
-void UMaterialExpressionAppend3Vector::GetCaption(TArray<FString>& OutCaptions) const
+void UMaterialExpressionMaterialXAppend3Vector::GetCaption(TArray<FString>& OutCaptions) const
 {
-	OutCaptions.Add(TEXT("Append3"));
+	OutCaptions.Add(TEXT("MaterialX Append3"));
 }
 #endif
 
diff --git a/Engine/Plugins/Interchange/Runtime/Source/Import/Private/MaterialX/MaterialExpressions/MaterialExpressionAppend3Vector.h b/Engine/Plugins/Interchange/Runtime/Source/Import/Private/MaterialX/MaterialExpressions/MaterialExpressionAppend3Vector.h
--- a/Engine/Plugins/Interchange/Runtime/Source/Import/Private/MaterialX/MaterialExpressions/MaterialExpressionAppend3Vector.h
+++ b/Engine/Plugins/Interchange/Runtime/Source/Import/Private/MaterialX/MaterialExpressions/MaterialExpressionAppend3Vector.h
@@ -10,7 +10,7 @@
  * A material expression that allows combining 3 channels together to create a vector with more channel than the original
  */
 UCLASS(MinimalAPI, collapsecategories, hidecategories = Object)
-class UMaterialExpressionAppend3Vector : public UMaterialExpression
+class UMaterialExpressionMaterialXAppend3Vector : public UMaterialExpression
 {
 	GENERATED_UCLASS_BODY()
 
@@ -23,11 +23,11 @@ class UMaterialExpressionAppend3Vector : public UMaterialExpression
 	UPROPERTY()
 	FExpressionInput C;
 
-	//~ Begin UMaterialExpression Interface
+	//~ Begin UMaterialExpressionMaterialX Interface
 #if WITH_EDITOR
 	virtual int32 Compile(class FMaterialCompiler* Compiler, int32 OutputIndex) override;
 	virtual void GetCaption(TArray<FString>& OutCaptions) const override;
 #endif
-	//~ End UMaterialExpression Interface
+	//~ End UMaterialExpressionMaterialX Interface
 };
 
diff --git a/Engine/Plugins/Interchange/Runtime/Source/Import/Private/MaterialX/MaterialExpressions/MaterialExpressionAppend4Vector.cpp b/Engine/Plugins/Interchange/Runtime/Source/Import/Private/MaterialX/MaterialExpressions/MaterialExpressionAppend4Vector.cpp
--- a/Engine/Plugins/Interchange/Runtime/Source/Import/Private/MaterialX/MaterialExpressions/MaterialExpressionAppend4Vector.cpp
+++ b/Engine/Plugins/Interchange/Runtime/Source/Import/Private/MaterialX/MaterialExpressions/MaterialExpressionAppend4Vector.cpp
@@ -5,50 +5,44 @@
 
 #include UE_INLINE_GENERATED_CPP_BY_NAME(MaterialExpressionAppend4Vector)
 
-#define LOCTEXT_NAMESPACE "MaterialExpressionAppend4Vector"
+#define LOCTEXT_NAMESPACE "MaterialExpressionMaterialXAppend4Vector"
 
-UMaterialExpressionAppend4Vector::UMaterialExpressionAppend4Vector(const FObjectInitializer& ObjectInitializer)
+UMaterialExpressionMaterialXAppend4Vector::UMaterialExpressionMaterialXAppend4Vector(const FObjectInitializer& ObjectInitializer)
 	: Super(ObjectInitializer)
 {
 	// Structure to hold one-time initialization
 	struct FConstructorStatics
 	{
 		FText NAME_MaterialX;
-		FText NAME_Math;
-		FText NAME_VectorOps;
 		FConstructorStatics()
 			: NAME_MaterialX(LOCTEXT("MaterialX", "MaterialX"))
-			, NAME_Math(LOCTEXT("Math", "Math"))
-			, NAME_VectorOps(LOCTEXT("VectorOps", "VectorOps"))
 		{}
 	};
 	static FConstructorStatics ConstructorStatics;
 
 #if WITH_EDITORONLY_DATA
 	MenuCategories.Add(ConstructorStatics.NAME_MaterialX);
-	MenuCategories.Add(ConstructorStatics.NAME_Math);
-	MenuCategories.Add(ConstructorStatics.NAME_VectorOps);
 #endif
 }
 
 #if WITH_EDITOR
-int32 UMaterialExpressionAppend4Vector::Compile(FMaterialCompiler* Compiler, int32 OutputIndex)
+int32 UMaterialExpressionMaterialXAppend4Vector::Compile(FMaterialCompiler* Compiler, int32 OutputIndex)
 {
 	if(!A.GetTracedInput().Expression)
 	{
-		return Compiler->Errorf(TEXT("Missing AppendVector input A"));
+		return Compiler->Errorf(TEXT("Missing MaterialX Append4 input A"));
 	}
 	else if(!B.GetTracedInput().Expression)
 	{
-		return Compiler->Errorf(TEXT("Missing AppendVector input B"));
+		return Compiler->Errorf(TEXT("Missing MaterialX Append4 input B"));
 	}
 	else if(!C.GetTracedInput().Expression)
 	{
-		return Compiler->Errorf(TEXT("Missing AppendVector input C"));
+		return Compiler->Errorf(TEXT("Missing MaterialX Append4 input C"));
 	}
 	else if(!D.GetTracedInput().Expression)
 	{
-		return Compiler->Errorf(TEXT("Missing AppendVector input D"));
+		return Compiler->Errorf(TEXT("Missing MaterialX Append4 input D"));
 	}
 	else
 	{
@@ -62,9 +56,9 @@ int32 UMaterialExpressionAppend4Vector::Compile(FMaterialCompiler* Compiler, int
 	}
 }
 
-void UMaterialExpressionAppend4Vector::GetCaption(TArray<FString>& OutCaptions) const
+void UMaterialExpressionMaterialXAppend4Vector::GetCaption(TArray<FString>& OutCaptions) const
 {
-	OutCaptions.Add(TEXT("Append4"));
+	OutCaptions.Add(TEXT("MaterialX Append4"));
 }
 #endif
 
diff --git a/Engine/Plugins/Interchange/Runtime/Source/Import/Private/MaterialX/MaterialExpressions/MaterialExpressionAppend4Vector.h b/Engine/Plugins/Interchange/Runtime/Source/Import/Private/MaterialX/MaterialExpressions/MaterialExpressionAppend4Vector.h
--- a/Engine/Plugins/Interchange/Runtime/Source/Import/Private/MaterialX/MaterialExpressions/MaterialExpressionAppend4Vector.h
+++ b/Engine/Plugins/Interchange/Runtime/Source/Import/Private/MaterialX/MaterialExpressions/MaterialExpressionAppend4Vector.h
@@ -12,7 +12,7 @@
  * A material expression that allows combining 4 channels together to create a vector with more channel than the original
  */
 UCLASS(MinimalAPI, collapsecategories, hidecategories = Object)
-class UMaterialExpressionAppend4Vector : public UMaterialExpression
+class UMaterialExpressionMaterialXAppend4Vector : public UMaterialExpression
 {
 	GENERATED_UCLASS_BODY()
 
@@ -28,11 +28,11 @@ class UMaterialExpressionAppend4Vector : public UMaterialExpression
 	UPROPERTY()
 	FExpressionInput D;
 
-	//~ Begin UMaterialExpression Interface
+	//~ Begin UMaterialExpressionMaterialX Interface
 #if WITH_EDITOR
 	virtual int32 Compile(class FMaterialCompiler* Compiler, int32 OutputIndex) override;
 	virtual void GetCaption(TArray<FString>& OutCaptions) const override;
 #endif
-	//~ End UMaterialExpression Interface
+	//~ End UMaterialExpressionMaterialX Interface
 };
 
diff --git a/Engine/Plugins/Interchange/Runtime/Source/Import/Private/MaterialX/MaterialExpressions/MaterialExpressionBurn.cpp b/Engine/Plugins/Interchange/Runtime/Source/Import/Private/MaterialX/MaterialExpressions/MaterialExpressionBurn.cpp
--- a/Engine/Plugins/Interchange/Runtime/Source/Import/Private/MaterialX/MaterialExpressions/MaterialExpressionBurn.cpp
+++ b/Engine/Plugins/Interchange/Runtime/Source/Import/Private/MaterialX/MaterialExpressions/MaterialExpressionBurn.cpp
@@ -4,40 +4,37 @@
 
 #include UE_INLINE_GENERATED_CPP_BY_NAME(MaterialExpressionBurn)
 
-#define LOCTEXT_NAMESPACE "MaterialExpressionBurn"
+#define LOCTEXT_NAMESPACE "MaterialExpressionMaterialXBurn"
 
-UMaterialExpressionBurn::UMaterialExpressionBurn(const FObjectInitializer& ObjectInitializer)
+UMaterialExpressionMaterialXBurn::UMaterialExpressionMaterialXBurn(const FObjectInitializer& ObjectInitializer)
 	: Super(ObjectInitializer)
 {
 	// Structure to hold one-time initialization
 	struct FConstructorStatics
 	{
 		FText NAME_MaterialX;
-		FText NAME_Compositing;
-		FConstructorStatics() 
+		FConstructorStatics()
 			: NAME_MaterialX(LOCTEXT("MaterialX", "MaterialX"))
-			, NAME_Compositing(LOCTEXT("Compositing", "Compositing"))
 		{}
 	};
 	static FConstructorStatics ConstructorStatics;
 
 #if WITH_EDITORONLY_DATA
 	MenuCategories.Add(ConstructorStatics.NAME_MaterialX);
-	MenuCategories.Add(ConstructorStatics.NAME_Compositing);
 #endif
 }
 
 #if WITH_EDITOR
-int32 UMaterialExpressionBurn::Compile(FMaterialCompiler* Compiler, int32 OutputIndex)
+int32 UMaterialExpressionMaterialXBurn::Compile(FMaterialCompiler* Compiler, int32 OutputIndex)
 {
 	if(!A.GetTracedInput().Expression)
 	{
-		return Compiler->Errorf(TEXT("Missing A input"));
+		return Compiler->Errorf(TEXT("Missing MaterialX Burn input A"));
 	}
 
 	if(!B.GetTracedInput().Expression)
 	{
-		return Compiler->Errorf(TEXT("Missing B input"));
+		return Compiler->Errorf(TEXT("Missing MaterialX Burn input B"));
 	}
 
 	int32 IndexAlpha = Alpha.GetTracedInput().Expression ? Alpha.Compile(Compiler) : Compiler->Constant(ConstAlpha);
@@ -54,9 +51,9 @@ int32 UMaterialExpressionBurn::Compile(FMaterialCompiler* Compiler, int32 Output
 	return Compiler->Lerp(IndexB, Burn, IndexAlpha);
 }
 
-void UMaterialExpressionBurn::GetCaption(TArray<FString>& OutCaptions) const
+void UMaterialExpressionMaterialXBurn::GetCaption(TArray<FString>& OutCaptions) const
 {
-	OutCaptions.Add(TEXT("Burn"));
+	OutCaptions.Add(TEXT("MaterialX Burn"));
 }
 #endif
 
diff --git a/Engine/Plugins/Interchange/Runtime/Source/Import/Private/MaterialX/MaterialExpressions/MaterialExpressionBurn.h b/Engine/Plugins/Interchange/Runtime/Source/Import/Private/MaterialX/MaterialExpressions/MaterialExpressionBurn.h
--- a/Engine/Plugins/Interchange/Runtime/Source/Import/Private/MaterialX/MaterialExpressions/MaterialExpressionBurn.h
+++ b/Engine/Plugins/Interchange/Runtime/Source/Import/Private/MaterialX/MaterialExpressions/MaterialExpressionBurn.h
@@ -14,7 +14,7 @@
  * Result: Lerp(B, 1-(1-B)/A, Alpha)
  */
 UCLASS(MinimalAPI, collapsecategories, hidecategories = Object)
-class UMaterialExpressionBurn : public UMaterialExpression
+class UMaterialExpressionMaterialXBurn : public UMaterialExpression
 {
 	GENERATED_UCLASS_BODY()
 
@@ -31,11 +31,11 @@ class UMaterialExpressionBurn : public UMaterialExpression
 	UPROPERTY(EditAnywhere, Category = MaterialExpressionLinearInterpolate, meta = (OverridingInputProperty = "Alpha"))
 	float ConstAlpha = 1.f;
 
-	//~ Begin UMaterialExpression Interface
+	//~ Begin UMaterialExpressionMaterialX Interface
 #if WITH_EDITOR
 	virtual int32 Compile(class FMaterialCompiler* Compiler, int32 OutputIndex) override;
 	virtual void GetCaption(TArray<FString>& OutCaptions) const override;
 #endif
-	//~ End UMaterialExpression Interface
+	//~ End UMaterialExpressionMaterialX Interface
 };
 
diff --git a/Engine/Plugins/Interchange/Runtime/Source/Import/Private/MaterialX/MaterialExpressions/MaterialExpressionDifference.cpp b/Engine/Plugins/Interchange/Runtime/Source/Import/Private/MaterialX/MaterialExpressions/MaterialExpressionDifference.cpp
--- a/Engine/Plugins/Interchange/Runtime/Source/Import/Private/MaterialX/MaterialExpressions/MaterialExpressionDifference.cpp
+++ b/Engine/Plugins/Interchange/Runtime/Source/Import/Private/MaterialX/MaterialExpressions/MaterialExpressionDifference.cpp
@@ -4,40 +4,37 @@
 
 #include UE_INLINE_GENERATED_CPP_BY_NAME(MaterialExpressionDifference)
 
-#define LOCTEXT_NAMESPACE "MaterialExpressionDifference"
+#define LOCTEXT_NAMESPACE "MaterialExpressionMaterialXDifference"
 
-UMaterialExpressionDifference::UMaterialExpressionDifference(const FObjectInitializer& ObjectInitializer)
+UMaterialExpressionMaterialXDifference::UMaterialExpressionMaterialXDifference(const FObjectInitializer& ObjectInitializer)
 	: Super(ObjectInitializer)
 {
 	// Structure to hold one-time initialization
 	struct FConstructorStatics
 	{
 		FText NAME_MaterialX;
-		FText NAME_Compositing;
 		FConstructorStatics()
 			: NAME_MaterialX(LOCTEXT("MaterialX", "MaterialX"))
-			, NAME_Compositing(LOCTEXT("Compositing", "Compositing"))
 		{}
 	};
 	static FConstructorStatics ConstructorStatics;
 
 #if WITH_EDITORONLY_DATA
 	MenuCategories.Add(ConstructorStatics.NAME_MaterialX);
-	MenuCategories.Add(ConstructorStatics.NAME_Compositing);
 #endif
 }
 
 #if WITH_EDITOR
-int32 UMaterialExpressionDifference::Compile(FMaterialCompiler* Compiler, int32 OutputIndex)
+int32 UMaterialExpressionMaterialXDifference::Compile(FMaterialCompiler* Compiler, int32 OutputIndex)
 {
 	if(!A.GetTracedInput().Expression)
 	{
-		return Compiler->Errorf(TEXT("Missing A input"));
+		return Compiler->Errorf(TEXT("Missing MaterialX Difference input A"));
 	}
 
 	if(!B.GetTracedInput().Expression)
 	{
-		return Compiler->Errorf(TEXT("Missing B input"));
+		return Compiler->Errorf(TEXT("Missing MaterialX Difference input B"));
 	}
 
 	int32 IndexAlpha = Alpha.GetTracedInput().Expression ? Alpha.Compile(Compiler) : Compiler->Constant(ConstAlpha);
@@ -47,9 +44,9 @@ int32 UMaterialExpressionDifference::Compile(FMaterialCompiler* Compiler, int32
 	return Compiler->Lerp(IndexB, Diff, IndexAlpha);
 }
 
-void UMaterialExpressionDifference::GetCaption(TArray<FString>& OutCaptions) const
+void UMaterialExpressionMaterialXDifference::GetCaption(TArray<FString>& OutCaptions) const
 {
-	OutCaptions.Add(TEXT("Difference"));
+	OutCaptions.Add(TEXT("MaterialX Difference"));
 }
 #endif
 
diff --git a/Engine/Plugins/Interchange/Runtime/Source/Import/Private/MaterialX/MaterialExpressions/MaterialExpressionDifference.h b/Engine/Plugins/Interchange/Runtime/Source/Import/Private/MaterialX/MaterialExpressions/MaterialExpressionDifference.h
--- a/Engine/Plugins/Interchange/Runtime/Source/Import/Private/MaterialX/MaterialExpressions/MaterialExpressionDifference.h
+++ b/Engine/Plugins/Interchange/Runtime/Source/Import/Private/MaterialX/MaterialExpressions/MaterialExpressionDifference.h
@@ -14,7 +14,7 @@
  * Result: Lerp(B, abs(A - B), Alpha)
  */
 UCLASS(MinimalAPI, collapsecategories, hidecategories = Object)
-class UMaterialExpressionDifference : public UMaterialExpression
+class UMaterialExpressionMaterialXDifference : public UMaterialExpression
 {
 	GENERATED_UCLASS_BODY()
 
@@ -31,11 +31,11 @@ class UMaterialExpressionDifference : public UMaterialExpression
 	UPROPERTY(EditAnywhere, Category = MaterialExpressionLinearInterpolate, meta = (OverridingInputProperty = "Alpha"))
 	float ConstAlpha = 1.f;
 
-	//~ Begin UMaterialExpression Interface
+	//~ Begin UMaterialExpressionMaterialX Interface
 #if WITH_EDITOR
 	virtual int32 Compile(class FMaterialCompiler* Compiler, int32 OutputIndex) override;
 	virtual void GetCaption(TArray<FString>& OutCaptions) const override;
 #endif
-	//~ End UMaterialExpression Interface
+	//~ End UMaterialExpressionMaterialX Interface
 };
 
diff --git a/Engine/Plugins/Interchange/Runtime/Source/Import/Private/MaterialX/MaterialExpressions/MaterialExpressionDisjointOver.cpp b/Engine/Plugins/Interchange/Runtime/Source/Import/Private/MaterialX/MaterialExpressions/MaterialExpressionDisjointOver.cpp
--- a/Engine/Plugins/Interchange/Runtime/Source/Import/Private/MaterialX/MaterialExpressions/MaterialExpressionDisjointOver.cpp
+++ b/Engine/Plugins/Interchange/Runtime/Source/Import/Private/MaterialX/MaterialExpressions/MaterialExpressionDisjointOver.cpp
@@ -4,40 +4,37 @@
 
 #include UE_INLINE_GENERATED_CPP_BY_NAME(MaterialExpressionDisjointOver)
 
-#define LOCTEXT_NAMESPACE "MaterialExpressionDisjointOver"
+#define LOCTEXT_NAMESPACE "MaterialExpressionMaterialXDisjointOver"
 
-UMaterialExpressionDisjointOver::UMaterialExpressionDisjointOver(const FObjectInitializer& ObjectInitializer)
+UMaterialExpressionMaterialXDisjointOver::UMaterialExpressionMaterialXDisjointOver(const FObjectInitializer& ObjectInitializer)
 	: Super(ObjectInitializer)
 {
 	// Structure to hold one-time initialization
 	struct FConstructorStatics
 	{
 		FText NAME_MaterialX;
-		FText NAME_Compositing;
 		FConstructorStatics()
 			: NAME_MaterialX(LOCTEXT("MaterialX", "MaterialX"))
-			, NAME_Compositing(LOCTEXT("Compositing", "Compositing"))
 		{}
 	};
 	static FConstructorStatics ConstructorStatics;
 
 #if WITH_EDITORONLY_DATA
 	MenuCategories.Add(ConstructorStatics.NAME_MaterialX);
-	MenuCategories.Add(ConstructorStatics.NAME_Compositing);
 #endif
 }
 
 #if WITH_EDITOR
-int32 UMaterialExpressionDisjointOver::Compile(FMaterialCompiler* Compiler, int32 OutputIndex)
+int32 UMaterialExpressionMaterialXDisjointOver::Compile(FMaterialCompiler* Compiler, int32 OutputIndex)
 {
 	if(!A.GetTracedInput().Expression)
 	{
-		return Compiler->Errorf(TEXT("Missing A input"));
+		return Compiler->Errorf(TEXT("Missing MaterialX DisjointOver input A"));
 	}
 
 	if(!B.GetTracedInput().Expression)
 	{
-		return Compiler->Errorf(TEXT("Missing B input"));
+		return Compiler->Errorf(TEXT("Missing MaterialX DisjointOver input B"));
 	}
 
 	int32 IndexAlpha = Alpha.GetTracedInput().Expression ? Alpha.Compile(Compiler) : Compiler->Constant(ConstAlpha);
@@ -64,9 +61,9 @@ int32 UMaterialExpressionDisjointOver::Compile(FMaterialCompiler* Compiler, int3
 	return Compiler->Lerp(IndexB, IndexAppend, IndexAlpha);
 }
 
-void UMaterialExpressionDisjointOver::GetCaption(TArray<FString>& OutCaptions) const
+void UMaterialExpressionMaterialXDisjointOver::GetCaption(TArray<FString>& OutCaptions) const
 {
-	OutCaptions.Add(TEXT("DisjointOver"));
+	OutCaptions.Add(TEXT("MaterialX DisjointOver"));
 }
 #endif
 
diff --git a/Engine/Plugins/Interchange/Runtime/Source/Import/Private/MaterialX/MaterialExpressions/MaterialExpressionDisjointOver.h b/Engine/Plugins/Interchange/Runtime/Source/Import/Private/MaterialX/MaterialExpressions/MaterialExpressionDisjointOver.h
--- a/Engine/Plugins/Interchange/Runtime/Source/Import/Private/MaterialX/MaterialExpressions/MaterialExpressionDisjointOver.h
+++ b/Engine/Plugins/Interchange/Runtime/Source/Import/Private/MaterialX/MaterialExpressions/MaterialExpressionDisjointOver.h
@@ -19,7 +19,7 @@
  * Result: Lerp(B, Op, Alpha)
  */
 UCLASS(MinimalAPI, collapsecategories, hidecategories = Object)
-class UMaterialExpressionDisjointOver : public UMaterialExpression
+class UMaterialExpressionMaterialXDisjointOver : public UMaterialExpression
 {
 	GENERATED_UCLASS_BODY()
 
@@ -36,11 +36,11 @@ class UMaterialExpressionDisjointOver : public UMaterialExpression
 	UPROPERTY(EditAnywhere, Category = MaterialExpressionLinearInterpolate, meta = (OverridingInputProperty = "Alpha"))
 	float ConstAlpha = 1.f;
 
-	//~ Begin UMaterialExpression Interface
+	//~ Begin UMaterialExpressionMaterialX Interface
 #if WITH_EDITOR
 	virtual int32 Compile(class FMaterialCompiler* Compiler, int32 OutputIndex) override;
 	virtual void GetCaption(TArray<FString>& OutCaptions) const override;
 #endif
-	//~ End UMaterialExpression Interface
+	//~ End UMaterialExpressionMaterialX Interface
 };
 
diff --git a/Engine/Plugins/Interchange/Runtime/Source/Import/Private/MaterialX/MaterialExpressions/MaterialExpressionDodge.cpp b/Engine/Plugins/Interchange/Runtime/Source/Import/Private/MaterialX/MaterialExpressions/MaterialExpressionDodge.cpp
--- a/Engine/Plugins/Interchange/Runtime/Source/Import/Private/MaterialX/MaterialExpressions/MaterialExpressionDodge.cpp
+++ b/Engine/Plugins/Interchange/Runtime/Source/Import/Private/MaterialX/MaterialExpressions/MaterialExpressionDodge.cpp
@@ -4,40 +4,37 @@
 
 #include UE_INLINE_GENERATED_CPP_BY_NAME(MaterialExpressionDodge)
 
-#define LOCTEXT_NAMESPACE "MaterialExpressionDodge"
+#define LOCTEXT_NAMESPACE "MaterialExpressionMaterialXDodge"
 
-UMaterialExpressionDodge::UMaterialExpressionDodge(const FObjectInitializer& ObjectInitializer)
+UMaterialExpressionMaterialXDodge::UMaterialExpressionMaterialXDodge(const FObjectInitializer& ObjectInitializer)
 	: Super(ObjectInitializer)
 {
 	// Structure to hold one-time initialization
 	struct FConstructorStatics
 	{
 		FText NAME_MaterialX;
-		FText NAME_Compositing;
 		FConstructorStatics()
 			: NAME_MaterialX(LOCTEXT("MaterialX", "MaterialX"))
-			, NAME_Compositing(LOCTEXT("Compositing", "Compositing"))
 		{}
 	};
 	static FConstructorStatics ConstructorStatics;
 
 #if WITH_EDITORONLY_DATA
 	MenuCategories.Add(ConstructorStatics.NAME_MaterialX);
-	MenuCategories.Add(ConstructorStatics.NAME_Compositing);
 #endif
 }
 
 #if WITH_EDITOR
-int32 UMaterialExpressionDodge::Compile(FMaterialCompiler* Compiler, int32 OutputIndex)
+int32 UMaterialExpressionMaterialXDodge::Compile(FMaterialCompiler* Compiler, int32 OutputIndex)
 {
 	if(!A.GetTracedInput().Expression)
 	{
-		return Compiler->Errorf(TEXT("Missing A input"));
+		return Compiler->Errorf(TEXT("Missing MaterialX Dodge input A"));
 	}
 
 	if(!B.GetTracedInput().Expression)
 	{
-		return Compiler->Errorf(TEXT("Missing B input"));
+		return Compiler->Errorf(TEXT("Missing MaterialX Dodge input B"));
 	}
 
 	int32 IndexAlpha = Alpha.GetTracedInput().Expression ? Alpha.Compile(Compiler) : Compiler->Constant(ConstAlpha);
@@ -52,9 +49,9 @@ int32 UMaterialExpressionDodge::Compile(FMaterialCompiler* Compiler, int32 Outpu
 	return Compiler->If(Compiler->Abs(Sub), Zero, Result, Result, Zero, Compiler->Constant(0.00001f));
 }
 
-void UMaterialExpressionDodge::GetCaption(TArray<FString>& OutCaptions) const
+void UMaterialExpressionMaterialXDodge::GetCaption(TArray<FString>& OutCaptions) const
 {
-	OutCaptions.Add(TEXT("Dodge"));
+	OutCaptions.Add(TEXT("MaterialX Dodge"));
 }
 #endif
 
diff --git a/Engine/Plugins/Interchange/Runtime/Source/Import/Private/MaterialX/MaterialExpressions/MaterialExpressionDodge.h b/Engine/Plugins/Interchange/Runtime/Source/Import/Private/MaterialX/MaterialExpressions/MaterialExpressionDodge.h
--- a/Engine/Plugins/Interchange/Runtime/Source/Import/Private/MaterialX/MaterialExpressions/MaterialExpressionDodge.h
+++ b/Engine/Plugins/Interchange/Runtime/Source/Import/Private/MaterialX/MaterialExpressions/MaterialExpressionDodge.h
@@ -14,7 +14,7 @@
  * Result: Lerp(B, B/(1-A), Alpha)
  */
 UCLASS(MinimalAPI, collapsecategories, hidecategories = Object)
-class UMaterialExpressionDodge: public UMaterialExpression
+class UMaterialExpressionMaterialXDodge: public UMaterialExpression
 {
 	GENERATED_UCLASS_BODY()
 
@@ -31,11 +31,11 @@ class UMaterialExpressionDodge: public UMaterialExpression
 	UPROPERTY(EditAnywhere, Category = MaterialExpressionLinearInterpolate, meta = (OverridingInputProperty = "Alpha"))
 	float ConstAlpha = 1.f;
 
-	//~ Begin UMaterialExpression Interface
+	//~ Begin UMaterialExpressionMaterialX Interface
 #if WITH_EDITOR
 	virtual int32 Compile(class FMaterialCompiler* Compiler, int32 OutputIndex) override;
 	virtual void GetCaption(TArray<FString>& OutCaptions) const override;
 #endif
-	//~ End UMaterialExpression Interface
+	//~ End UMaterialExpressionMaterialX Interface
 };
 
diff --git a/Engine/Plugins/Interchange/Runtime/Source/Import/Private/MaterialX/MaterialExpressions/MaterialExpressionExponential.cpp b/Engine/Plugins/Interchange/Runtime/Source/Import/Private/MaterialX/MaterialExpressions/MaterialExpressionExponential.cpp
--- a/Engine/Plugins/Interchange/Runtime/Source/Import/Private/MaterialX/MaterialExpressions/MaterialExpressionExponential.cpp
+++ b/Engine/Plugins/Interchange/Runtime/Source/Import/Private/MaterialX/MaterialExpressions/MaterialExpressionExponential.cpp
@@ -5,35 +5,32 @@
 
 #include UE_INLINE_GENERATED_CPP_BY_NAME(MaterialExpressionExponential)
 
-#define LOCTEXT_NAMESPACE "MaterialExpressionExponential"
+#define LOCTEXT_NAMESPACE "MaterialExpressionMaterialXExponential"
 
-UMaterialExpressionExponential::UMaterialExpressionExponential(const FObjectInitializer& ObjectInitializer)
+UMaterialExpressionMaterialXExponential::UMaterialExpressionMaterialXExponential(const FObjectInitializer& ObjectInitializer)
 	: Super(ObjectInitializer)
 {
 	// Structure to hold one-time initialization
 	struct FConstructorStatics
 	{
 		FText NAME_MaterialX;
-		FText NAME_Math;
 		FConstructorStatics()
 			: NAME_MaterialX(LOCTEXT("MaterialX", "MaterialX"))
-			, NAME_Math(LOCTEXT("Math", "Math"))
 		{}
 	};
 	static FConstructorStatics ConstructorStatics;
 
 #if WITH_EDITORONLY_DATA
 	MenuCategories.Add(ConstructorStatics.NAME_MaterialX);
-	MenuCategories.Add(ConstructorStatics.NAME_Math);
 #endif
 }
 
 #if WITH_EDITOR
-int32 UMaterialExpressionExponential::Compile(FMaterialCompiler* Compiler, int32 OutputIndex)
+int32 UMaterialExpressionMaterialXExponential::Compile(FMaterialCompiler* Compiler, int32 OutputIndex)
 {
 	if(!Input.GetTracedInput().Expression)
 	{
-		return Compiler->Errorf(TEXT("Missing Exponential input"));
+		return Compiler->Errorf(TEXT("Missing MaterialX Exponential input"));
 	}
 
 	int32 Euler = Compiler->Constant(2.7182818);
@@ -41,9 +38,9 @@ int32 UMaterialExpressionExponential::Compile(FMaterialCompiler* Compiler, int32
 	return Compiler->Power(Euler, Input.Compile(Compiler));
 }
 
-void UMaterialExpressionExponential::GetCaption(TArray<FString>& OutCaptions) const
+void UMaterialExpressionMaterialXExponential::GetCaption(TArray<FString>& OutCaptions) const
 {
-	OutCaptions.Add(TEXT("Exponential"));
+	OutCaptions.Add(TEXT("MaterialX Exponential"));
 }
 #endif
 
diff --git a/Engine/Plugins/Interchange/Runtime/Source/Import/Private/MaterialX/MaterialExpressions/MaterialExpressionExponential.h b/Engine/Plugins/Interchange/Runtime/Source/Import/Private/MaterialX/MaterialExpressions/MaterialExpressionExponential.h
--- a/Engine/Plugins/Interchange/Runtime/Source/Import/Private/MaterialX/MaterialExpressions/MaterialExpressionExponential.h
+++ b/Engine/Plugins/Interchange/Runtime/Source/Import/Private/MaterialX/MaterialExpressions/MaterialExpressionExponential.h
@@ -10,18 +10,18 @@
  * A material expression that computes the exponential of an input
  */
 UCLASS(collapsecategories, hidecategories = Object, MinimalAPI)
-class UMaterialExpressionExponential : public UMaterialExpression
+class UMaterialExpressionMaterialXExponential : public UMaterialExpression
 {
 	GENERATED_UCLASS_BODY()
 
 	UPROPERTY()
 	FExpressionInput Input;
 
-	//~ Begin UMaterialExpression Interface
+	//~ Begin UMaterialExpressionMaterialX Interface
 #if WITH_EDITOR
 	virtual int32 Compile(class FMaterialCompiler* Compiler, int32 OutputIndex) override;
 	virtual void GetCaption(TArray<FString>& OutCaptions) const override;
 #endif
-	//~ End UMaterialExpression Interface
+	//~ End UMaterialExpressionMaterialX Interface
 };
 
diff --git a/Engine/Plugins/Interchange/Runtime/Source/Import/Private/MaterialX/MaterialExpressions/MaterialExpressionFractal3D.cpp b/Engine/Plugins/Interchange/Runtime/Source/Import/Private/MaterialX/MaterialExpressions/MaterialExpressionFractal3D.cpp
--- a/Engine/Plugins/Interchange/Runtime/Source/Import/Private/MaterialX/MaterialExpressions/MaterialExpressionFractal3D.cpp
+++ b/Engine/Plugins/Interchange/Runtime/Source/Import/Private/MaterialX/MaterialExpressions/MaterialExpressionFractal3D.cpp
@@ -6,33 +6,29 @@
 
 #include UE_INLINE_GENERATED_CPP_BY_NAME(MaterialExpressionFractal3D)
 
-#define LOCTEXT_NAMESPACE "MaterialExpressionFractal3D"
+#define LOCTEXT_NAMESPACE "MaterialExpressionMaterialXFractal3D"
 
-UMaterialExpressionFractal3D::UMaterialExpressionFractal3D(const FObjectInitializer& ObjectInitializer)
+UMaterialExpressionMaterialXFractal3D::UMaterialExpressionMaterialXFractal3D(const FObjectInitializer& ObjectInitializer)
 	: Super(ObjectInitializer)
 {
 	// Structure to hold one-time initialization
 	struct FConstructorStatics
 	{
 		FText NAME_MaterialX;
-		FText NAME_Procedural;
 		FConstructorStatics()
 			: NAME_MaterialX(LOCTEXT("MaterialX", "MaterialX"))
-			, NAME_Procedural(LOCTEXT("Procedural", "Procedural"))
 		{}
 	};
 	static FConstructorStatics ConstructorStatics;
 
 #if WITH_EDITORONLY_DATA
 	MenuCategories.Add(ConstructorStatics.NAME_MaterialX);
-	MenuCategories.Add(ConstructorStatics.NAME_Procedural);
 #endif
 }
 
 #if WITH_EDITOR
-int32 UMaterialExpressionFractal3D::Compile(FMaterialCompiler* Compiler, int32 OutputIndex)
+int32 UMaterialExpressionMaterialXFractal3D::Compile(FMaterialCompiler* Compiler, int32 OutputIndex)
 {
-
 	UMaterialExpressionCustom* MaterialExpressionCustom = NewObject<UMaterialExpressionCustom>();
 	MaterialExpressionCustom->Inputs[0].InputName = TEXT("Position");
 	MaterialExpressionCustom->Inputs.Add({ TEXT("Octaves") });
@@ -90,9 +86,9 @@ int32 UMaterialExpressionFractal3D::Compile(FMaterialCompiler* Compiler, int32 O
 	return Compiler->Mul(IndexFractal, IndexAmplitude);
 }
 
-void UMaterialExpressionFractal3D::GetCaption(TArray<FString>& OutCaptions) const
+void UMaterialExpressionMaterialXFractal3D::GetCaption(TArray<FString>& OutCaptions) const
 {
-	OutCaptions.Add(TEXT("Fractal3D"));
+	OutCaptions.Add(TEXT("MaterialX Fractal3D"));
 }
 #endif
 
diff --git a/Engine/Plugins/Interchange/Runtime/Source/Import/Private/MaterialX/MaterialExpressions/MaterialExpressionFractal3D.h b/Engine/Plugins/Interchange/Runtime/Source/Import/Private/MaterialX/MaterialExpressions/MaterialExpressionFractal3D.h
--- a/Engine/Plugins/Interchange/Runtime/Source/Import/Private/MaterialX/MaterialExpressions/MaterialExpressionFractal3D.h
+++ b/Engine/Plugins/Interchange/Runtime/Source/Import/Private/MaterialX/MaterialExpressions/MaterialExpressionFractal3D.h
@@ -11,7 +11,7 @@
  */
 
 UCLASS(collapsecategories, hidecategories = Object, MinimalAPI)
-class UMaterialExpressionFractal3D : public UMaterialExpression
+class UMaterialExpressionMaterialXFractal3D : public UMaterialExpression
 {
 	GENERATED_UCLASS_BODY()
 
@@ -69,11 +69,11 @@ class UMaterialExpressionFractal3D : public UMaterialExpression
 	UPROPERTY(EditAnywhere, Category = MaterialExpressionNoise)
 	float OutputMax = 1.f;
 
-	//~ Begin UMaterialExpression Interface
+	//~ Begin UMaterialExpressionMaterialX Interface
 #if WITH_EDITOR
 	virtual int32 Compile(class FMaterialCompiler* Compiler, int32 OutputIndex) override;
 	virtual void GetCaption(TArray<FString>& OutCaptions) const override;
 #endif
-	//~ End UMaterialExpression Interface
+	//~ End UMaterialExpressionMaterialX Interface
 };
 
diff --git a/Engine/Plugins/Interchange/Runtime/Source/Import/Private/MaterialX/MaterialExpressions/MaterialExpressionHsvToRgb.cpp b/Engine/Plugins/Interchange/Runtime/Source/Import/Private/MaterialX/MaterialExpressions/MaterialExpressionHsvToRgb.cpp
--- a/Engine/Plugins/Interchange/Runtime/Source/Import/Private/MaterialX/MaterialExpressions/MaterialExpressionHsvToRgb.cpp
+++ b/Engine/Plugins/Interchange/Runtime/Source/Import/Private/MaterialX/MaterialExpressions/MaterialExpressionHsvToRgb.cpp
@@ -4,35 +4,32 @@
 #include "MaterialCompiler.h"
 #include "Materials/MaterialExpressionCustom.h"
 
-#define LOCTEXT_NAMESPACE "MaterialExpressionHsvToRgb"
+#define LOCTEXT_NAMESPACE "MaterialExpressionMaterialXHsvToRgb"
 
-UMaterialExpressionHsvToRgb::UMaterialExpressionHsvToRgb(const FObjectInitializer& ObjectInitializer)
+UMaterialExpressionMaterialXHsvToRgb::UMaterialExpressionMaterialXHsvToRgb(const FObjectInitializer& ObjectInitializer)
 	: Super(ObjectInitializer)
 {
 	// Structure to hold one-time initialization
 	struct FConstructorStatics
 	{
 		FText NAME_MaterialX;
-		FText NAME_ImageAdjustment;
 		FConstructorStatics()
 			: NAME_MaterialX(LOCTEXT("MaterialX", "MaterialX"))
-			, NAME_ImageAdjustment(LOCTEXT("Image Adjustment", "Image Adjustment"))
 		{}
 	};
 	static FConstructorStatics ConstructorStatics;
 
 #if WITH_EDITORONLY_DATA
 	MenuCategories.Add(ConstructorStatics.NAME_MaterialX);
-	MenuCategories.Add(ConstructorStatics.NAME_ImageAdjustment);
 #endif
 }
 
 #if WITH_EDITOR
-int32 UMaterialExpressionHsvToRgb::Compile(FMaterialCompiler* Compiler, int32 OutputIndex)
+int32 UMaterialExpressionMaterialXHsvToRgb::Compile(FMaterialCompiler* Compiler, int32 OutputIndex)
 {
 	if(!Input.GetTracedInput().Expression)
 	{
-		return Compiler->Errorf(TEXT("Missing RGBtoHSV input"));
+		return Compiler->Errorf(TEXT("Missing MaterialX RGBtoHSV input"));
 	}
 
 	UMaterialExpressionCustom* MaterialExpressionCustom = NewObject<UMaterialExpressionCustom>();
@@ -47,9 +44,9 @@ int32 UMaterialExpressionHsvToRgb::Compile(FMaterialCompiler* Compiler, int32 Ou
 	return MaterialExpressionCustom->Compile(Compiler, OutputIndex);
 }
 
-void UMaterialExpressionHsvToRgb::GetCaption(TArray<FString>& OutCaptions) const
+void UMaterialExpressionMaterialXHsvToRgb::GetCaption(TArray<FString>& OutCaptions) const
 {
-	OutCaptions.Add(TEXT("HSVToRGB"));
+	OutCaptions.Add(TEXT("MaterialX HSVToRGB"));
 }
 #endif
 
diff --git a/Engine/Plugins/Interchange/Runtime/Source/Import/Private/MaterialX/MaterialExpressions/MaterialExpressionHsvToRgb.h b/Engine/Plugins/Interchange/Runtime/Source/Import/Private/MaterialX/MaterialExpressions/MaterialExpressionHsvToRgb.h
--- a/Engine/Plugins/Interchange/Runtime/Source/Import/Private/MaterialX/MaterialExpressions/MaterialExpressionHsvToRgb.h
+++ b/Engine/Plugins/Interchange/Runtime/Source/Import/Private/MaterialX/MaterialExpressions/MaterialExpressionHsvToRgb.h
@@ -6,18 +6,18 @@
 #include "MaterialExpressionHsvToRgb.generated.h"
 
 UCLASS(collapsecategories, hidecategories = Object, MinimalAPI)
-class UMaterialExpressionHsvToRgb : public UMaterialExpression
+class UMaterialExpressionMaterialXHsvToRgb : public UMaterialExpression
 {
 	GENERATED_UCLASS_BODY()
 
 	UPROPERTY()
 	FExpressionInput Input;
 
-	//~ Begin UMaterialExpression Interface
+	//~ Begin UMaterialExpressionMaterialX Interface
 #if WITH_EDITOR
 	virtual int32 Compile(class FMaterialCompiler* Compiler, int32 OutputIndex) override;
 	virtual void GetCaption(TArray<FString>& OutCaptions) const override;
 #endif
-	//~ End UMaterialExpression Interface
+	//~ End UMaterialExpressionMaterialX Interface
 };
 
diff --git a/Engine/Plugins/Interchange/Runtime/Source/Import/Private/MaterialX/MaterialExpressions/MaterialExpressionIn.cpp b/Engine/Plugins/Interchange/Runtime/Source/Import/Private/MaterialX/MaterialExpressions/MaterialExpressionIn.cpp
--- a/Engine/Plugins/Interchange/Runtime/Source/Import/Private/MaterialX/MaterialExpressions/MaterialExpressionIn.cpp
+++ b/Engine/Plugins/Interchange/Runtime/Source/Import/Private/MaterialX/MaterialExpressions/MaterialExpressionIn.cpp
@@ -4,40 +4,37 @@
 
 #include UE_INLINE_GENERATED_CPP_BY_NAME(MaterialExpressionIn)
 
-#define LOCTEXT_NAMESPACE "MaterialExpressionIn"
+#define LOCTEXT_NAMESPACE "MaterialExpressionMaterialXIn"
 
-UMaterialExpressionIn::UMaterialExpressionIn(const FObjectInitializer& ObjectInitializer)
+UMaterialExpressionMaterialXIn::UMaterialExpressionMaterialXIn(const FObjectInitializer& ObjectInitializer)
 	: Super(ObjectInitializer)
 {
 	// Structure to hold one-time initialization
 	struct FConstructorStatics
 	{
 		FText NAME_MaterialX;
-		FText NAME_Compositing;
 		FConstructorStatics()
 			: NAME_MaterialX(LOCTEXT("MaterialX", "MaterialX"))
-			, NAME_Compositing(LOCTEXT("Compositing", "Compositing"))
 		{}
 	};
 	static FConstructorStatics ConstructorStatics;
 
 #if WITH_EDITORONLY_DATA
 	MenuCategories.Add(ConstructorStatics.NAME_MaterialX);
-	MenuCategories.Add(ConstructorStatics.NAME_Compositing);
 #endif
 }
 
 #if WITH_EDITOR
-int32 UMaterialExpressionIn::Compile(FMaterialCompiler* Compiler, int32 OutputIndex)
+int32 UMaterialExpressionMaterialXIn::Compile(FMaterialCompiler* Compiler, int32 OutputIndex)
 {
 	if(!A.GetTracedInput().Expression)
 	{
-		return Compiler->Errorf(TEXT("Missing A input"));
+		return Compiler->Errorf(TEXT("Missing MaterialX In input A"));
 	}
 
 	if(!B.GetTracedInput().Expression)
 	{
-		return Compiler->Errorf(TEXT("Missing B input"));
+		return Compiler->Errorf(TEXT("Missing MaterialX In input B"));
 	}
 
 	int32 IndexAlpha = Alpha.GetTracedInput().Expression ? Alpha.Compile(Compiler) : Compiler->Constant(ConstAlpha);
@@ -50,9 +47,9 @@ int32 UMaterialExpressionIn::Compile(FMaterialCompiler* Compiler, int32 OutputIn
 	return Compiler->Lerp(IndexB, IndexIn, IndexAlpha);
 }
 
-void UMaterialExpressionIn::GetCaption(TArray<FString>& OutCaptions) const
+void UMaterialExpressionMaterialXIn::GetCaption(TArray<FString>& OutCaptions) const
 {
-	OutCaptions.Add(TEXT("In"));
+	OutCaptions.Add(TEXT("MaterialX In"));
 }
 #endif
 
diff --git a/Engine/Plugins/Interchange/Runtime/Source/Import/Private/MaterialX/MaterialExpressions/MaterialExpressionIn.h b/Engine/Plugins/Interchange/Runtime/Source/Import/Private/MaterialX/MaterialExpressions/MaterialExpressionIn.h
--- a/Engine/Plugins/Interchange/Runtime/Source/Import/Private/MaterialX/MaterialExpressions/MaterialExpressionIn.h
+++ b/Engine/Plugins/Interchange/Runtime/Source/Import/Private/MaterialX/MaterialExpressions/MaterialExpressionIn.h
@@ -18,7 +18,7 @@
  * Result: Lerp(B, Ab, Alpha)
  */
 UCLASS(MinimalAPI, collapsecategories, hidecategories = Object)
-class UMaterialExpressionIn: public UMaterialExpression
+class UMaterialExpressionMaterialXIn: public UMaterialExpression
 {
 	GENERATED_UCLASS_BODY()
 
@@ -35,11 +35,11 @@ class UMaterialExpressionIn: public UMaterialExpression
 	UPROPERTY(EditAnywhere, Category = MaterialExpressionLinearInterpolate, meta = (OverridingInputProperty = "Alpha"))
 	float ConstAlpha = 1.f;
 
-	//~ Begin UMaterialExpression Interface
+	//~ Begin UMaterialExpressionMaterialX Interface
 #if WITH_EDITOR
 	virtual int32 Compile(class FMaterialCompiler* Compiler, int32 OutputIndex) override;
 	virtual void GetCaption(TArray<FString>& OutCaptions) const override;
 #endif
-	//~ End UMaterialExpression Interface
+	//~ End UMaterialExpressionMaterialX Interface
 };
 
diff --git a/Engine/Plugins/Interchange/Runtime/Source/Import/Private/MaterialX/MaterialExpressions/MaterialExpressionLength.cpp b/Engine/Plugins/Interchange/Runtime/Source/Import/Private/MaterialX/MaterialExpressions/MaterialExpressionLength.cpp
--- a/Engine/Plugins/Interchange/Runtime/Source/Import/Private/MaterialX/MaterialExpressions/MaterialExpressionLength.cpp
+++ b/Engine/Plugins/Interchange/Runtime/Source/Import/Private/MaterialX/MaterialExpressions/MaterialExpressionLength.cpp
@@ -4,46 +4,47 @@
 
 #include UE_INLINE_GENERATED_CPP_BY_NAME(MaterialExpressionLength)
 
-#define LOCTEXT_NAMESPACE "MaterialExpressionLength"
+#define LOCTEXT_NAMESPACE "MaterialExpressionMaterialXLength"
 
-UMaterialExpressionLength::UMaterialExpressionLength(const FObjectInitializer& ObjectInitializer)
+UMaterialExpressionMaterialXLength::UMaterialExpressionMaterialXLength(const FObjectInitializer& ObjectInitializer)
 	: Super(ObjectInitializer)
 {
 	// Structure to hold one-time initialization
 	struct FConstructorStatics
 	{
 		FText NAME_MaterialX;
-		FText NAME_Math;
-		FText NAME_VectorOps;
 		FConstructorStatics()
 			: NAME_MaterialX(LOCTEXT("MaterialX", "MaterialX"))
-			, NAME_Math(LOCTEXT("Math", "Math"))
-			, NAME_VectorOps(LOCTEXT("VectorOps", "VectorOps"))
 		{}
 	};
 	static FConstructorStatics ConstructorStatics;
 
 #if WITH_EDITORONLY_DATA
 	MenuCategories.Add(ConstructorStatics.NAME_MaterialX);
-	MenuCategories.Add(ConstructorStatics.NAME_Math);
-	MenuCategories.Add(ConstructorStatics.NAME_VectorOps);
 #endif
 }
 
 #if WITH_EDITOR
-int32 UMaterialExpressionLength::Compile(FMaterialCompiler* Compiler, int32 OutputIndex)
+int32 UMaterialExpressionMaterialXLength::Compile(FMaterialCompiler* Compiler, int32 OutputIndex)
 {
 	if(!Input.GetTracedInput().Expression)
 	{
-		return Compiler->Errorf(TEXT("Missing Length input"));
+		return Compiler->Errorf(TEXT("Missing MaterialX Length input"));
 	}
 
-	return Compiler->Length(Input.Compile(Compiler));
+	int32 Index = Input.Compile(Compiler);
+	if(Compiler->GetType(Index) == MCT_Float)
+	{
+		// optimized
+		return Compiler->Abs(Index);
+	}
+
+	return Compiler->Length(Index);
 }
 
-void UMaterialExpressionLength::GetCaption(TArray<FString>& OutCaptions) const
+void UMaterialExpressionMaterialXLength::GetCaption(TArray<FString>& OutCaptions) const
 {
-	OutCaptions.Add(TEXT("Length"));
+	OutCaptions.Add(TEXT("MaterialX Length"));
 }
 #endif
 
diff --git a/Engine/Plugins/Interchange/Runtime/Source/Import/Private/MaterialX/MaterialExpressions/MaterialExpressionLength.h b/Engine/Plugins/Interchange/Runtime/Source/Import/Private/MaterialX/MaterialExpressions/MaterialExpressionLength.h
--- a/Engine/Plugins/Interchange/Runtime/Source/Import/Private/MaterialX/MaterialExpressions/MaterialExpressionLength.h
+++ b/Engine/Plugins/Interchange/Runtime/Source/Import/Private/MaterialX/MaterialExpressions/MaterialExpressionLength.h
@@ -10,18 +10,18 @@
  * A material expression that computes the length (magnitude) of a vector
  */
 UCLASS(collapsecategories, hidecategories = Object, MinimalAPI)
-class UMaterialExpressionLength : public UMaterialExpression
+class UMaterialExpressionMaterialXLength : public UMaterialExpression
 {
 	GENERATED_UCLASS_BODY()
 
 	UPROPERTY()
 	FExpressionInput Input;
 
-	//~ Begin UMaterialExpression Interface
+	//~ Begin UMaterialExpressionMaterialX Interface
 #if WITH_EDITOR
 	virtual int32 Compile(class FMaterialCompiler* Compiler, int32 OutputIndex) override;
 	virtual void GetCaption(TArray<FString>& OutCaptions) const override;
 #endif
-	//~ End UMaterialExpression Interface
+	//~ End UMaterialExpressionMaterialX Interface
 };
 
diff --git a/Engine/Plugins/Interchange/Runtime/Source/Import/Private/MaterialX/MaterialExpressions/MaterialExpressionLogarithm.cpp b/Engine/Plugins/Interchange/Runtime/Source/Import/Private/MaterialX/MaterialExpressions/MaterialExpressionLogarithm.cpp
--- a/Engine/Plugins/Interchange/Runtime/Source/Import/Private/MaterialX/MaterialExpressions/MaterialExpressionLogarithm.cpp
+++ b/Engine/Plugins/Interchange/Runtime/Source/Import/Private/MaterialX/MaterialExpressions/MaterialExpressionLogarithm.cpp
@@ -6,36 +6,32 @@
 
 #include UE_INLINE_GENERATED_CPP_BY_NAME(MaterialExpressionLogarithm)
 
-#define LOCTEXT_NAMESPACE "MaterialExpressionLogarithm"
+#define LOCTEXT_NAMESPACE "MaterialExpressionMaterialXLogarithm"
 
-UMaterialExpressionLogarithm::UMaterialExpressionLogarithm(const FObjectInitializer& ObjectInitializer)
+UMaterialExpressionMaterialXLogarithm::UMaterialExpressionMaterialXLogarithm(const FObjectInitializer& ObjectInitializer)
 	: Super(ObjectInitializer)
 {
 	// Structure to hold one-time initialization
 	struct FConstructorStatics
 	{
-
 		FText NAME_MaterialX;
-		FText NAME_Math;
 		FConstructorStatics()
 			: NAME_MaterialX(LOCTEXT("MaterialX", "MaterialX"))
-			, NAME_Math(LOCTEXT("Math", "Math"))
 		{}
 	};
 	static FConstructorStatics ConstructorStatics;
 
 #if WITH_EDITORONLY_DATA
 	MenuCategories.Add(ConstructorStatics.NAME_MaterialX);
-	MenuCategories.Add(ConstructorStatics.NAME_Math);
 #endif
 }
 
 #if WITH_EDITOR
-int32 UMaterialExpressionLogarithm::Compile(FMaterialCompiler* Compiler, int32 OutputIndex)
+int32 UMaterialExpressionMaterialXLogarithm::Compile(FMaterialCompiler* Compiler, int32 OutputIndex)
 {
 	if(!Input.GetTracedInput().Expression)
 	{
-		return Compiler->Errorf(TEXT("Missing Ln input"));
+		return Compiler->Errorf(TEXT("Missing MaterialX Ln input"));
 	}
 
 	const int32 Log10Euler = Compiler->Constant(0.4342945);
@@ -43,9 +39,9 @@ int32 UMaterialExpressionLogarithm::Compile(FMaterialCompiler* Compiler, int32 O
 	return Compiler->Div(Compiler->Logarithm10(Input.Compile(Compiler)), Log10Euler);
 }
 
-void UMaterialExpressionLogarithm::GetCaption(TArray<FString>& OutCaptions) const
+void UMaterialExpressionMaterialXLogarithm::GetCaption(TArray<FString>& OutCaptions) const
 {
-	OutCaptions.Add(TEXT("Ln"));
+	OutCaptions.Add(TEXT("MaterialX Ln"));
 }
 #endif
 
diff --git a/Engine/Plugins/Interchange/Runtime/Source/Import/Private/MaterialX/MaterialExpressions/MaterialExpressionLogarithm.h b/Engine/Plugins/Interchange/Runtime/Source/Import/Private/MaterialX/MaterialExpressions/MaterialExpressionLogarithm.h
--- a/Engine/Plugins/Interchange/Runtime/Source/Import/Private/MaterialX/MaterialExpressions/MaterialExpressionLogarithm.h
+++ b/Engine/Plugins/Interchange/Runtime/Source/Import/Private/MaterialX/MaterialExpressions/MaterialExpressionLogarithm.h
@@ -10,18 +10,18 @@
  * A material expression that computes the natural logarithm
  */
 UCLASS(collapsecategories, hidecategories = Object, MinimalAPI)
-class UMaterialExpressionLogarithm : public UMaterialExpression
+class UMaterialExpressionMaterialXLogarithm : public UMaterialExpression
 {
 	GENERATED_UCLASS_BODY()
 
 	UPROPERTY()
 	FExpressionInput Input;
 
-	//~ Begin UMaterialExpression Interface
+	//~ Begin UMaterialExpressionMaterialX Interface
 #if WITH_EDITOR
 	virtual int32 Compile(class FMaterialCompiler* Compiler, int32 OutputIndex) override;
 	virtual void GetCaption(TArray<FString>& OutCaptions) const override;
 #endif
-	//~ End UMaterialExpression Interface
+	//~ End UMaterialExpressionMaterialX Interface
 };
 
diff --git a/Engine/Plugins/Interchange/Runtime/Source/Import/Private/MaterialX/MaterialExpressions/MaterialExpressionLuminance.cpp b/Engine/Plugins/Interchange/Runtime/Source/Import/Private/MaterialX/MaterialExpressions/MaterialExpressionLuminance.cpp
--- a/Engine/Plugins/Interchange/Runtime/Source/Import/Private/MaterialX/MaterialExpressions/MaterialExpressionLuminance.cpp
+++ b/Engine/Plugins/Interchange/Runtime/Source/Import/Private/MaterialX/MaterialExpressions/MaterialExpressionLuminance.cpp
@@ -5,37 +5,34 @@
 
 #include UE_INLINE_GENERATED_CPP_BY_NAME(MaterialExpressionLuminance)
 
-#define LOCTEXT_NAMESPACE "MaterialExpressionLuminance"
+#define LOCTEXT_NAMESPACE "MaterialExpressionMaterialXLuminance"
 
-UMaterialExpressionLuminance::UMaterialExpressionLuminance(const FObjectInitializer& ObjectInitializer)
+UMaterialExpressionMaterialXLuminance::UMaterialExpressionMaterialXLuminance(const FObjectInitializer& ObjectInitializer)
 	: Super(ObjectInitializer),
 	LuminanceFactors(0.2722287, 0.6740818, 0.0536895),
-	LuminanceMode(ELuminanceMode::ACEScg)
+	LuminanceMode(EMaterialXLuminanceMode::ACEScg)
 {
 	// Structure to hold one-time initialization
 	struct FConstructorStatics
 	{
 		FText NAME_MaterialX;
-		FText NAME_Adjustment;
 		FConstructorStatics()
 			: NAME_MaterialX(LOCTEXT("MaterialX", "MaterialX"))
-			, NAME_Adjustment(LOCTEXT("Image Adjustment", "Image Adjustment"))
 		{}
 	};
 	static FConstructorStatics ConstructorStatics;
 
 #if WITH_EDITORONLY_DATA
 	MenuCategories.Add(ConstructorStatics.NAME_MaterialX);
-	MenuCategories.Add(ConstructorStatics.NAME_Adjustment);
 #endif
 }
 
 #if WITH_EDITOR
-int32 UMaterialExpressionLuminance::Compile(FMaterialCompiler* Compiler, int32 OutputIndex)
+int32 UMaterialExpressionMaterialXLuminance::Compile(FMaterialCompiler* Compiler, int32 OutputIndex)
 {
 	if(!Input.GetTracedInput().Expression)
 	{
-		return Compiler->Errorf(TEXT("Missing Luminance input"));
+		return Compiler->Errorf(TEXT("Missing MaterialX Luminance input"));
 	}
 
 	const int32 LuminanceFactorsIndex = Compiler->Constant3(LuminanceFactors.R, LuminanceFactors.G, LuminanceFactors.B);
@@ -44,36 +41,36 @@ int32 UMaterialExpressionLuminance::Compile(FMaterialCompiler* Compiler, int32 O
 	return Compiler->Dot(Compiler->ComponentMask(Input.Compile(Compiler), true, true, true, false), LuminanceFactorsIndex);
 }
 
-void UMaterialExpressionLuminance::GetCaption(TArray<FString>& OutCaptions) const
+void UMaterialExpressionMaterialXLuminance::GetCaption(TArray<FString>& OutCaptions) const
 {
-	OutCaptions.Add(TEXT("Luminance"));
+	OutCaptions.Add(TEXT("MaterialX Luminance"));
 }
 
-void UMaterialExpressionLuminance::PostEditChangeProperty(FPropertyChangedEvent& PropertyChangedEvent)
+void UMaterialExpressionMaterialXLuminance::PostEditChangeProperty(FPropertyChangedEvent& PropertyChangedEvent)
 {
 	FProperty* PropertyThatChanged = PropertyChangedEvent.MemberProperty;
 	const FString PropertyName = PropertyThatChanged ? PropertyThatChanged->GetName() : TEXT("");
 
-	if(PropertyName == GET_MEMBER_NAME_STRING_CHECKED(UMaterialExpressionLuminance, LuminanceMode))
+	if(PropertyName == GET_MEMBER_NAME_STRING_CHECKED(UMaterialExpressionMaterialXLuminance, LuminanceMode))
 	{
 		switch(LuminanceMode)
 		{
-		case ELuminanceMode::ACEScg:
+		case EMaterialXLuminanceMode::ACEScg:
 			LuminanceFactors = FLinearColor(0.2722287, 0.6740818, 0.0536895);
 			break;
 
-		case ELuminanceMode::Rec709:
+		case EMaterialXLuminanceMode::Rec709:
 			LuminanceFactors = FLinearColor(0.2126, 0.7152, 0.0722);
 			break;
 
-		case ELuminanceMode::Rec2020:
+		case EMaterialXLuminanceMode::Rec2020:
 			LuminanceFactors = FLinearColor(0.2627, 0.6780, 0.0593);
 			break;
 		}
 	}
-	else if(PropertyName == GET_MEMBER_NAME_STRING_CHECKED(UMaterialExpressionLuminance, LuminanceFactors))
+	else if(PropertyName == GET_MEMBER_NAME_STRING_CHECKED(UMaterialExpressionMaterialXLuminance, LuminanceFactors))
 	{
-		LuminanceMode = ELuminanceMode::Custom;
+		LuminanceMode = EMaterialXLuminanceMode::Custom;
 	}
 		
 	Super::PostEditChangeProperty(PropertyChangedEvent);
diff --git a/Engine/Plugins/Interchange/Runtime/Source/Import/Private/MaterialX/MaterialExpressions/MaterialExpressionLuminance.h b/Engine/Plugins/Interchange/Runtime/Source/Import/Private/MaterialX/MaterialExpressions/MaterialExpressionLuminance.h
--- a/Engine/Plugins/Interchange/Runtime/Source/Import/Private/MaterialX/MaterialExpressions/MaterialExpressionLuminance.h
+++ b/Engine/Plugins/Interchange/Runtime/Source/Import/Private/MaterialX/MaterialExpressions/MaterialExpressionLuminance.h
@@ -12,7 +12,7 @@
  */
 
 UENUM()
-enum class ELuminanceMode : uint8
+enum class EMaterialXLuminanceMode : uint8
 {
 	ACEScg,
 	Rec709,
@@ -22,7 +22,7 @@ enum class ELuminanceMode : uint8
 };
 
 UCLASS(collapsecategories, hidecategories = Object, MinimalAPI)
-class UMaterialExpressionLuminance : public UMaterialExpression
+class UMaterialExpressionMaterialXLuminance : public UMaterialExpression
 {
 	GENERATED_UCLASS_BODY()
 
@@ -33,14 +33,14 @@ class UMaterialExpressionLuminance : public UMaterialExpression
 	FLinearColor LuminanceFactors;    // Color component factors for converting a color to greyscale.
 	
 	UPROPERTY(EditAnywhere, Category = MaterialExpressionLuminance, Meta = (ShowAsInputPin = "Advanced"))
-	ELuminanceMode LuminanceMode;
+	EMaterialXLuminanceMode LuminanceMode;
 	
-	//~ Begin UMaterialExpression Interface
+	//~ Begin UMaterialExpressionMaterialX Interface
 #if WITH_EDITOR
 	virtual int32 Compile(class FMaterialCompiler* Compiler, int32 OutputIndex) override;
 	virtual void GetCaption(TArray<FString>& OutCaptions) const override;
 	virtual void PostEditChangeProperty(FPropertyChangedEvent& PropertyChangedEvent) override;
 #endif
-	//~ End UMaterialExpression Interface
+	//~ End UMaterialExpressionMaterialX Interface
 };
 
diff --git a/Engine/Plugins/Interchange/Runtime/Source/Import/Private/MaterialX/MaterialExpressions/MaterialExpressionMask.cpp b/Engine/Plugins/Interchange/Runtime/Source/Import/Private/MaterialX/MaterialExpressions/MaterialExpressionMask.cpp
--- a/Engine/Plugins/Interchange/Runtime/Source/Import/Private/MaterialX/MaterialExpressions/MaterialExpressionMask.cpp
+++ b/Engine/Plugins/Interchange/Runtime/Source/Import/Private/MaterialX/MaterialExpressions/MaterialExpressionMask.cpp
@@ -4,40 +4,37 @@
 
 #include UE_INLINE_GENERATED_CPP_BY_NAME(MaterialExpressionMask)
 
-#define LOCTEXT_NAMESPACE "MaterialExpressionMask"
+#define LOCTEXT_NAMESPACE "MaterialExpressionMaterialXMask"
 
-UMaterialExpressionMask::UMaterialExpressionMask(const FObjectInitializer& ObjectInitializer)
+UMaterialExpressionMaterialXMask::UMaterialExpressionMaterialXMask(const FObjectInitializer& ObjectInitializer)
 	: Super(ObjectInitializer)
 {
 	// Structure to hold one-time initialization
 	struct FConstructorStatics
 	{
 		FText NAME_MaterialX;
-		FText NAME_Compositing;
 		FConstructorStatics()
 			: NAME_MaterialX(LOCTEXT("MaterialX", "MaterialX"))
-			, NAME_Compositing(LOCTEXT("Compositing", "Compositing"))
 		{}
 	};
 	static FConstructorStatics ConstructorStatics;
 
 #if WITH_EDITORONLY_DATA
 	MenuCategories.Add(ConstructorStatics.NAME_MaterialX);
-	MenuCategories.Add(ConstructorStatics.NAME_Compositing);
 #endif
 }
 
 #if WITH_EDITOR
-int32 UMaterialExpressionMask::Compile(FMaterialCompiler* Compiler, int32 OutputIndex)
+int32 UMaterialExpressionMaterialXMask::Compile(FMaterialCompiler* Compiler, int32 OutputIndex)
 {
 	if(!A.GetTracedInput().Expression)
 	{
-		return Compiler->Errorf(TEXT("Missing A input"));
+		return Compiler->Errorf(TEXT("Missing MaterialX Mask input A"));
 	}
 
 	if(!B.GetTracedInput().Expression)
 	{
-		return Compiler->Errorf(TEXT("Missing B input"));
+		return Compiler->Errorf(TEXT("Missing MaterialX Mask input B"));
 	}
 
 	int32 IndexAlpha = Alpha.GetTracedInput().Expression ? Alpha.Compile(Compiler) : Compiler->Constant(ConstAlpha);
@@ -50,9 +47,9 @@ int32 UMaterialExpressionMask::Compile(FMaterialCompiler* Compiler, int32 Output
 	return Compiler->Lerp(IndexB, Mask, IndexAlpha);
 }
 
-void UMaterialExpressionMask::GetCaption(TArray<FString>& OutCaptions) const
+void UMaterialExpressionMaterialXMask::GetCaption(TArray<FString>& OutCaptions) const
 {
-	OutCaptions.Add(TEXT("Mask"));
+	OutCaptions.Add(TEXT("MaterialX Mask"));
 }
 #endif
 
diff --git a/Engine/Plugins/Interchange/Runtime/Source/Import/Private/MaterialX/MaterialExpressions/MaterialExpressionMask.h b/Engine/Plugins/Interchange/Runtime/Source/Import/Private/MaterialX/MaterialExpressions/MaterialExpressionMask.h
--- a/Engine/Plugins/Interchange/Runtime/Source/Import/Private/MaterialX/MaterialExpressions/MaterialExpressionMask.h
+++ b/Engine/Plugins/Interchange/Runtime/Source/Import/Private/MaterialX/MaterialExpressions/MaterialExpressionMask.h
@@ -18,7 +18,7 @@
  * Result: Lerp(B, B*a, Alpha)
  */
 UCLASS(MinimalAPI, collapsecategories, hidecategories = Object)
-class UMaterialExpressionMask : public UMaterialExpression
+class UMaterialExpressionMaterialXMask : public UMaterialExpression
 {
 	GENERATED_UCLASS_BODY()
 
@@ -35,11 +35,11 @@ class UMaterialExpressionMask : public UMaterialExpression
 	UPROPERTY(EditAnywhere, Category = MaterialExpressionLinearInterpolate, meta = (OverridingInputProperty = "Alpha"))
 	float ConstAlpha = 1.f;
 
-	//~ Begin UMaterialExpression Interface
+	//~ Begin UMaterialExpressionMaterialX Interface
 #if WITH_EDITOR
 	virtual int32 Compile(class FMaterialCompiler* Compiler, int32 OutputIndex) override;
 	virtual void GetCaption(TArray<FString>& OutCaptions) const override;
 #endif
-	//~ End UMaterialExpression Interface
+	//~ End UMaterialExpressionMaterialX Interface
 };
 
diff --git a/Engine/Plugins/Interchange/Runtime/Source/Import/Private/MaterialX/MaterialExpressions/MaterialExpressionMatte.cpp b/Engine/Plugins/Interchange/Runtime/Source/Import/Private/MaterialX/MaterialExpressions/MaterialExpressionMatte.cpp
--- a/Engine/Plugins/Interchange/Runtime/Source/Import/Private/MaterialX/MaterialExpressions/MaterialExpressionMatte.cpp
+++ b/Engine/Plugins/Interchange/Runtime/Source/Import/Private/MaterialX/MaterialExpressions/MaterialExpressionMatte.cpp
@@ -4,40 +4,37 @@
 
 #include UE_INLINE_GENERATED_CPP_BY_NAME(MaterialExpressionMatte)
 
-#define LOCTEXT_NAMESPACE "MaterialExpressionMatte"
+#define LOCTEXT_NAMESPACE "MaterialExpressionMaterialXMatte"
 
-UMaterialExpressionMatte::UMaterialExpressionMatte(const FObjectInitializer& ObjectInitializer)
+UMaterialExpressionMaterialXMatte::UMaterialExpressionMaterialXMatte(const FObjectInitializer& ObjectInitializer)
 	: Super(ObjectInitializer)
 {
 	// Structure to hold one-time initialization
 	struct FConstructorStatics
 	{
 		FText NAME_MaterialX;
-		FText NAME_Compositing;
 		FConstructorStatics()
 			: NAME_MaterialX(LOCTEXT("MaterialX", "MaterialX"))
-			, NAME_Compositing(LOCTEXT("Compositing", "Compositing"))
 		{}
 	};
 	static FConstructorStatics ConstructorStatics;
 
 #if WITH_EDITORONLY_DATA
 	MenuCategories.Add(ConstructorStatics.NAME_MaterialX);
-	MenuCategories.Add(ConstructorStatics.NAME_Compositing);
 #endif
 }
 
 #if WITH_EDITOR
-int32 UMaterialExpressionMatte::Compile(FMaterialCompiler* Compiler, int32 OutputIndex)
+int32 UMaterialExpressionMaterialXMatte::Compile(FMaterialCompiler* Compiler, int32 OutputIndex)
 {
 	if(!A.GetTracedInput().Expression)
 	{
-		return Compiler->Errorf(TEXT("Missing A input"));
+		return Compiler->Errorf(TEXT("Missing MaterialX Matte input A"));
 	}
 
 	if(!B.GetTracedInput().Expression)
 	{
-		return Compiler->Errorf(TEXT("Missing B input"));
+		return Compiler->Errorf(TEXT("Missing MaterialX Matte input B"));
 	}
 
 	int32 IndexAlpha = Alpha.GetTracedInput().Expression ? Alpha.Compile(Compiler) : Compiler->Constant(ConstAlpha);
@@ -59,9 +56,9 @@ int32 UMaterialExpressionMatte::Compile(FMaterialCompiler* Compiler, int32 Outpu
 	return Compiler->Lerp(IndexB, Matte, IndexAlpha);
 }
 
-void UMaterialExpressionMatte::GetCaption(TArray<FString>& OutCaptions) const
+void UMaterialExpressionMaterialXMatte::GetCaption(TArray<FString>& OutCaptions) const
 {
-	OutCaptions.Add(TEXT("Matte"));
+	OutCaptions.Add(TEXT("MaterialX Matte"));
 }
 #endif
 
diff --git a/Engine/Plugins/Interchange/Runtime/Source/Import/Private/MaterialX/MaterialExpressions/MaterialExpressionMatte.h b/Engine/Plugins/Interchange/Runtime/Source/Import/Private/MaterialX/MaterialExpressions/MaterialExpressionMatte.h
--- a/Engine/Plugins/Interchange/Runtime/Source/Import/Private/MaterialX/MaterialExpressions/MaterialExpressionMatte.h
+++ b/Engine/Plugins/Interchange/Runtime/Source/Import/Private/MaterialX/MaterialExpressions/MaterialExpressionMatte.h
@@ -18,7 +18,7 @@
  * Result: Lerp(B, A*a + B*(1-a), Alpha)
  */
 UCLASS(MinimalAPI, collapsecategories, hidecategories = Object)
-class UMaterialExpressionMatte : public UMaterialExpression
+class UMaterialExpressionMaterialXMatte : public UMaterialExpression
 {
 	GENERATED_UCLASS_BODY()
 
@@ -35,11 +35,11 @@ class UMaterialExpressionMatte : public UMaterialExpression
 	UPROPERTY(EditAnywhere, Category = MaterialExpressionLinearInterpolate, meta = (OverridingInputProperty = "Alpha"))
 	float ConstAlpha = 1.f;
 
-	//~ Begin UMaterialExpression Interface
+	//~ Begin UMaterialExpressionMaterialX Interface
 #if WITH_EDITOR
 	virtual int32 Compile(class FMaterialCompiler* Compiler, int32 OutputIndex) override;
 	virtual void GetCaption(TArray<FString>& OutCaptions) const override;
 #endif
-	//~ End UMaterialExpression Interface
+	//~ End UMaterialExpressionMaterialX Interface
 };
 
diff --git a/Engine/Plugins/Interchange/Runtime/Source/Import/Private/MaterialX/MaterialExpressions/MaterialExpressionMinus.cpp b/Engine/Plugins/Interchange/Runtime/Source/Import/Private/MaterialX/MaterialExpressions/MaterialExpressionMinus.cpp
--- a/Engine/Plugins/Interchange/Runtime/Source/Import/Private/MaterialX/MaterialExpressions/MaterialExpressionMinus.cpp
+++ b/Engine/Plugins/Interchange/Runtime/Source/Import/Private/MaterialX/MaterialExpressions/MaterialExpressionMinus.cpp
@@ -4,40 +4,37 @@
 
 #include UE_INLINE_GENERATED_CPP_BY_NAME(MaterialExpressionMinus)
 
-#define LOCTEXT_NAMESPACE "MaterialExpressionMinus"
+#define LOCTEXT_NAMESPACE "MaterialExpressionMaterialXMinus"
 
-UMaterialExpressionMinus::UMaterialExpressionMinus(const FObjectInitializer& ObjectInitializer)
+UMaterialExpressionMaterialXMinus::UMaterialExpressionMaterialXMinus(const FObjectInitializer& ObjectInitializer)
 	: Super(ObjectInitializer)
 {
 	// Structure to hold one-time initialization
 	struct FConstructorStatics
 	{
 		FText NAME_MaterialX;
-		FText NAME_Compositing;
 		FConstructorStatics()
 			: NAME_MaterialX(LOCTEXT("MaterialX", "MaterialX"))
-			, NAME_Compositing(LOCTEXT("Compositing", "Compositing"))
 		{}
 	};
 	static FConstructorStatics ConstructorStatics;
 
 #if WITH_EDITORONLY_DATA
 	MenuCategories.Add(ConstructorStatics.NAME_MaterialX);
-	MenuCategories.Add(ConstructorStatics.NAME_Compositing);
 #endif
 }
 
 #if WITH_EDITOR
-int32 UMaterialExpressionMinus::Compile(FMaterialCompiler* Compiler, int32 OutputIndex)
+int32 UMaterialExpressionMaterialXMinus::Compile(FMaterialCompiler* Compiler, int32 OutputIndex)
 {
 	if(!A.GetTracedInput().Expression)
 	{
-		return Compiler->Errorf(TEXT("Missing A input"));
+		return Compiler->Errorf(TEXT("Missing MaterialX Minus input A"));
 	}
 
 	if(!B.GetTracedInput().Expression)
 	{
-		return Compiler->Errorf(TEXT("Missing B input"));
+		return Compiler->Errorf(TEXT("Missing MaterialX Minus input B"));
 	}
 
 	int32 IndexAlpha = Alpha.GetTracedInput().Expression ? Alpha.Compile(Compiler) : Compiler->Constant(ConstAlpha);
@@ -47,9 +44,9 @@ int32 UMaterialExpressionMinus::Compile(FMaterialCompiler* Compiler, int32 Outpu
 	return Compiler->Lerp(IndexB, Sub, IndexAlpha);
 }
 
-void UMaterialExpressionMinus::GetCaption(TArray<FString>& OutCaptions) const
+void UMaterialExpressionMaterialXMinus::GetCaption(TArray<FString>& OutCaptions) const
 {
-	OutCaptions.Add(TEXT("Minus"));
+	OutCaptions.Add(TEXT("MaterialX Minus"));
 }
 #endif
 
diff --git a/Engine/Plugins/Interchange/Runtime/Source/Import/Private/MaterialX/MaterialExpressions/MaterialExpressionMinus.h b/Engine/Plugins/Interchange/Runtime/Source/Import/Private/MaterialX/MaterialExpressions/MaterialExpressionMinus.h
--- a/Engine/Plugins/Interchange/Runtime/Source/Import/Private/MaterialX/MaterialExpressions/MaterialExpressionMinus.h
+++ b/Engine/Plugins/Interchange/Runtime/Source/Import/Private/MaterialX/MaterialExpressions/MaterialExpressionMinus.h
@@ -14,7 +14,7 @@
  * Result: Lerp(B, A - B, Alpha)
  */
 UCLASS(MinimalAPI, collapsecategories, hidecategories = Object)
-class UMaterialExpressionMinus : public UMaterialExpression
+class UMaterialExpressionMaterialXMinus : public UMaterialExpression
 {
 	GENERATED_UCLASS_BODY()
 
@@ -31,11 +31,11 @@ class UMaterialExpressionMinus : public UMaterialExpression
 	UPROPERTY(EditAnywhere, Category = MaterialExpressionLinearInterpolate, meta = (OverridingInputProperty = "Alpha"))
 	float ConstAlpha = 1.f;
 
-	//~ Begin UMaterialExpression Interface
+	//~ Begin UMaterialExpressionMaterialX Interface
 #if WITH_EDITOR
 	virtual int32 Compile(class FMaterialCompiler* Compiler, int32 OutputIndex) override;
 	virtual void GetCaption(TArray<FString>& OutCaptions) const override;
 #endif
-	//~ End UMaterialExpression Interface
+	//~ End UMaterialExpressionMaterialX Interface
 };
 
diff --git a/Engine/Plugins/Interchange/Runtime/Source/Import/Private/MaterialX/MaterialExpressions/MaterialExpressionOut.cpp b/Engine/Plugins/Interchange/Runtime/Source/Import/Private/MaterialX/MaterialExpressions/MaterialExpressionOut.cpp
--- a/Engine/Plugins/Interchange/Runtime/Source/Import/Private/MaterialX/MaterialExpressions/MaterialExpressionOut.cpp
+++ b/Engine/Plugins/Interchange/Runtime/Source/Import/Private/MaterialX/MaterialExpressions/MaterialExpressionOut.cpp
@@ -4,40 +4,37 @@
 
 #include UE_INLINE_GENERATED_CPP_BY_NAME(MaterialExpressionOut)
 
-#define LOCTEXT_NAMESPACE "MaterialExpressionOut"
+#define LOCTEXT_NAMESPACE "MaterialExpressionMaterialXOut"
 
-UMaterialExpressionOut::UMaterialExpressionOut(const FObjectInitializer& ObjectInitializer)
+UMaterialExpressionMaterialXOut::UMaterialExpressionMaterialXOut(const FObjectInitializer& ObjectInitializer)
 	: Super(ObjectInitializer)
 {
 	// Structure to hold one-time initialization
 	struct FConstructorStatics
 	{
 		FText NAME_MaterialX;
-		FText NAME_Compositing;
 		FConstructorStatics()
 			: NAME_MaterialX(LOCTEXT("MaterialX", "MaterialX"))
-			, NAME_Compositing(LOCTEXT("Compositing", "Compositing"))
 		{}
 	};
 	static FConstructorStatics ConstructorStatics;
 
 #if WITH_EDITORONLY_DATA
 	MenuCategories.Add(ConstructorStatics.NAME_MaterialX);
-	MenuCategories.Add(ConstructorStatics.NAME_Compositing);
 #endif
 }
 
 #if WITH_EDITOR
-int32 UMaterialExpressionOut::Compile(FMaterialCompiler* Compiler, int32 OutputIndex)
+int32 UMaterialExpressionMaterialXOut::Compile(FMaterialCompiler* Compiler, int32 OutputIndex)
 {
 	if(!A.GetTracedInput().Expression)
 	{
-		return Compiler->Errorf(TEXT("Missing A input"));
+		return Compiler->Errorf(TEXT("Missing MaterialX Out input A"));
 	}
 
 	if(!B.GetTracedInput().Expression)
 	{
-		return Compiler->Errorf(TEXT("Missing B input"));
+		return Compiler->Errorf(TEXT("Missing MaterialX Out input B"));
 	}
 
 	int32 IndexAlpha = Alpha.GetTracedInput().Expression ? Alpha.Compile(Compiler) : Compiler->Constant(ConstAlpha);
@@ -53,9 +50,9 @@ int32 UMaterialExpressionOut::Compile(FMaterialCompiler* Compiler, int32 OutputI
 	return Compiler->Lerp(IndexB, Out, IndexAlpha);
 }
 
-void UMaterialExpressionOut::GetCaption(TArray<FString>& OutCaptions) const
+void UMaterialExpressionMaterialXOut::GetCaption(TArray<FString>& OutCaptions) const
 {
-	OutCaptions.Add(TEXT("Out"));
+	OutCaptions.Add(TEXT("MaterialX Out"));
 }
 #endif
 
diff --git a/Engine/Plugins/Interchange/Runtime/Source/Import/Private/MaterialX/MaterialExpressions/MaterialExpressionOut.h b/Engine/Plugins/Interchange/Runtime/Source/Import/Private/MaterialX/MaterialExpressions/MaterialExpressionOut.h
--- a/Engine/Plugins/Interchange/Runtime/Source/Import/Private/MaterialX/MaterialExpressions/MaterialExpressionOut.h
+++ b/Engine/Plugins/Interchange/Runtime/Source/Import/Private/MaterialX/MaterialExpressions/MaterialExpressionOut.h
@@ -18,7 +18,7 @@
  * Result: Lerp(B, A(1-b), Alpha)
  */
 UCLASS(MinimalAPI, collapsecategories, hidecategories = Object)
-class UMaterialExpressionOut : public UMaterialExpression
+class UMaterialExpressionMaterialXOut : public UMaterialExpression
 {
 	GENERATED_UCLASS_BODY()
 
@@ -35,11 +35,11 @@ class UMaterialExpressionOut : public UMaterialExpression
 	UPROPERTY(EditAnywhere, Category = MaterialExpressionLinearInterpolate, meta = (OverridingInputProperty = "Alpha"))
 	float ConstAlpha = 1.f;
 
-	//~ Begin UMaterialExpression Interface
+	//~ Begin UMaterialExpressionMaterialX Interface
 #if WITH_EDITOR
 	virtual int32 Compile(class FMaterialCompiler* Compiler, int32 OutputIndex) override;
 	virtual void GetCaption(TArray<FString>& OutCaptions) const override;
 #endif
-	//~ End UMaterialExpression Interface
+	//~ End UMaterialExpressionMaterialX Interface
 };
 
diff --git a/Engine/Plugins/Interchange/Runtime/Source/Import/Private/MaterialX/MaterialExpressions/MaterialExpressionOver.cpp b/Engine/Plugins/Interchange/Runtime/Source/Import/Private/MaterialX/MaterialExpressions/MaterialExpressionOver.cpp
--- a/Engine/Plugins/Interchange/Runtime/Source/Import/Private/MaterialX/MaterialExpressions/MaterialExpressionOver.cpp
+++ b/Engine/Plugins/Interchange/Runtime/Source/Import/Private/MaterialX/MaterialExpressions/MaterialExpressionOver.cpp
@@ -4,40 +4,37 @@
 
 #include UE_INLINE_GENERATED_CPP_BY_NAME(MaterialExpressionOver)
 
-#define LOCTEXT_NAMESPACE "MaterialExpressionOver"
+#define LOCTEXT_NAMESPACE "MaterialExpressionMaterialXOver"
 
-UMaterialExpressionOver::UMaterialExpressionOver(const FObjectInitializer& ObjectInitializer)
+UMaterialExpressionMaterialXOver::UMaterialExpressionMaterialXOver(const FObjectInitializer& ObjectInitializer)
 	: Super(ObjectInitializer)
 {
 	// Structure to hold one-time initialization
 	struct FConstructorStatics
 	{
 		FText NAME_MaterialX;
-		FText NAME_Compositing;
 		FConstructorStatics()
 			: NAME_MaterialX(LOCTEXT("MaterialX", "MaterialX"))
-			, NAME_Compositing(LOCTEXT("Compositing", "Compositing"))
 		{}
 	};
 	static FConstructorStatics ConstructorStatics;
 
 #if WITH_EDITORONLY_DATA
 	MenuCategories.Add(ConstructorStatics.NAME_MaterialX);
-	MenuCategories.Add(ConstructorStatics.NAME_Compositing);
 #endif
 }
 
 #if WITH_EDITOR
-int32 UMaterialExpressionOver::Compile(FMaterialCompiler* Compiler, int32 OutputIndex)
+int32 UMaterialExpressionMaterialXOver::Compile(FMaterialCompiler* Compiler, int32 OutputIndex)
 {
 	if(!A.GetTracedInput().Expression)
 	{
-		return Compiler->Errorf(TEXT("Missing A input"));
+		return Compiler->Errorf(TEXT("Missing MaterialX Over input A"));
 	}
 
 	if(!B.GetTracedInput().Expression)
 	{
-		return Compiler->Errorf(TEXT("Missing B input"));
+		return Compiler->Errorf(TEXT("Missing MaterialX Over input B"));
 	}
 
 	int32 IndexAlpha = Alpha.GetTracedInput().Expression ? Alpha.Compile(Compiler) : Compiler->Constant(ConstAlpha);
@@ -53,9 +50,9 @@ int32 UMaterialExpressionOver::Compile(FMaterialCompiler* Compiler, int32 Output
 	return Compiler->Lerp(IndexB, Over, IndexAlpha);
 }
 
-void UMaterialExpressionOver::GetCaption(TArray<FString>& OutCaptions) const
+void UMaterialExpressionMaterialXOver::GetCaption(TArray<FString>& OutCaptions) const
 {
-	OutCaptions.Add(TEXT("Over"));
+	OutCaptions.Add(TEXT("MaterialX Over"));
 }
 #endif
 
diff --git a/Engine/Plugins/Interchange/Runtime/Source/Import/Private/MaterialX/MaterialExpressions/MaterialExpressionOver.h b/Engine/Plugins/Interchange/Runtime/Source/Import/Private/MaterialX/MaterialExpressions/MaterialExpressionOver.h
--- a/Engine/Plugins/Interchange/Runtime/Source/Import/Private/MaterialX/MaterialExpressions/MaterialExpressionOver.h
+++ b/Engine/Plugins/Interchange/Runtime/Source/Import/Private/MaterialX/MaterialExpressions/MaterialExpressionOver.h
@@ -18,7 +18,7 @@
  * Result: Lerp(B, A + B(1-a), Alpha)
  */
 UCLASS(MinimalAPI, collapsecategories, hidecategories = Object)
-class UMaterialExpressionOver : public UMaterialExpression
+class UMaterialExpressionMaterialXOver : public UMaterialExpression
 {
 	GENERATED_UCLASS_BODY()
 
@@ -35,11 +35,11 @@ class UMaterialExpressionOver : public UMaterialExpression
 	UPROPERTY(EditAnywhere, Category = MaterialExpressionLinearInterpolate, meta = (OverridingInputProperty = "Alpha"))
 	float ConstAlpha = 1.f;
 
-	//~ Begin UMaterialExpression Interface
+	//~ Begin UMaterialExpressionMaterialX Interface
 #if WITH_EDITOR
 	virtual int32 Compile(class FMaterialCompiler* Compiler, int32 OutputIndex) override;
 	virtual void GetCaption(TArray<FString>& OutCaptions) const override;
 #endif
-	//~ End UMaterialExpression Interface
+	//~ End UMaterialExpressionMaterialX Interface
 };
 
diff --git a/Engine/Plugins/Interchange/Runtime/Source/Import/Private/MaterialX/MaterialExpressions/MaterialExpressionOverlay.cpp b/Engine/Plugins/Interchange/Runtime/Source/Import/Private/MaterialX/MaterialExpressions/MaterialExpressionOverlay.cpp
--- a/Engine/Plugins/Interchange/Runtime/Source/Import/Private/MaterialX/MaterialExpressions/MaterialExpressionOverlay.cpp
+++ b/Engine/Plugins/Interchange/Runtime/Source/Import/Private/MaterialX/MaterialExpressions/MaterialExpressionOverlay.cpp
@@ -4,40 +4,37 @@
 
 #include UE_INLINE_GENERATED_CPP_BY_NAME(MaterialExpressionOverlay)
 
-#define LOCTEXT_NAMESPACE "MaterialExpressionOverlay"
+#define LOCTEXT_NAMESPACE "MaterialExpressionMaterialXOverlay"
 
-UMaterialExpressionOverlay::UMaterialExpressionOverlay(const FObjectInitializer& ObjectInitializer)
+UMaterialExpressionMaterialXOverlay::UMaterialExpressionMaterialXOverlay(const FObjectInitializer& ObjectInitializer)
 	: Super(ObjectInitializer)
 {
 	// Structure to hold one-time initialization
 	struct FConstructorStatics
 	{
 		FText NAME_MaterialX;
-		FText NAME_Compositing;
 		FConstructorStatics()
 			: NAME_MaterialX(LOCTEXT("MaterialX", "MaterialX"))
-			, NAME_Compositing(LOCTEXT("Compositing", "Compositing"))
 		{}
 	};
 	static FConstructorStatics ConstructorStatics;
 
 #if WITH_EDITORONLY_DATA
 	MenuCategories.Add(ConstructorStatics.NAME_MaterialX);
-	MenuCategories.Add(ConstructorStatics.NAME_Compositing);
 #endif
 }
 
 #if WITH_EDITOR
-int32 UMaterialExpressionOverlay::Compile(FMaterialCompiler* Compiler, int32 OutputIndex)
+int32 UMaterialExpressionMaterialXOverlay::Compile(FMaterialCompiler* Compiler, int32 OutputIndex)
 {
 	if(!A.GetTracedInput().Expression)
 	{
-		return Compiler->Errorf(TEXT("Missing A input"));
+		return Compiler->Errorf(TEXT("Missing MaterialX Overlay input A"));
 	}
 
 	if(!B.GetTracedInput().Expression)
 	{
-		return Compiler->Errorf(TEXT("Missing B input"));
+		return Compiler->Errorf(TEXT("Missing MaterialX Overlay input B"));
 	}
 
 	int32 IndexAlpha = Alpha.GetTracedInput().Expression ? Alpha.Compile(Compiler) : Compiler->Constant(ConstAlpha);
@@ -62,9 +59,9 @@ int32 UMaterialExpressionOverlay::Compile(FMaterialCompiler* Compiler, int32 Out
 	return Compiler->Lerp(IndexB, Overlay, IndexAlpha);
 }
 
-void UMaterialExpressionOverlay::GetCaption(TArray<FString>& OutCaptions) const
+void UMaterialExpressionMaterialXOverlay::GetCaption(TArray<FString>& OutCaptions) const
 {
-	OutCaptions.Add(TEXT("Overlay"));
+	OutCaptions.Add(TEXT("MaterialX Overlay"));
 }
 #endif
 
diff --git a/Engine/Plugins/Interchange/Runtime/Source/Import/Private/MaterialX/MaterialExpressions/MaterialExpressionOverlay.h b/Engine/Plugins/Interchange/Runtime/Source/Import/Private/MaterialX/MaterialExpressions/MaterialExpressionOverlay.h
--- a/Engine/Plugins/Interchange/Runtime/Source/Import/Private/MaterialX/MaterialExpressions/MaterialExpressionOverlay.h
+++ b/Engine/Plugins/Interchange/Runtime/Source/Import/Private/MaterialX/MaterialExpressions/MaterialExpressionOverlay.h
@@ -15,7 +15,7 @@
  * Result: Lerp(B, Op, Alpha)
  */
 UCLASS(MinimalAPI, collapsecategories, hidecategories = Object)
-class UMaterialExpressionOverlay : public UMaterialExpression
+class UMaterialExpressionMaterialXOverlay : public UMaterialExpression
 {
 	GENERATED_UCLASS_BODY()
 
@@ -32,11 +32,11 @@ class UMaterialExpressionOverlay : public UMaterialExpression
 	UPROPERTY(EditAnywhere, Category = MaterialExpressionLinearInterpolate, meta = (OverridingInputProperty = "Alpha"))
 	float ConstAlpha = 1.f;
 
-	//~ Begin UMaterialExpression Interface
+	//~ Begin UMaterialExpressionMaterialX Interface
 #if WITH_EDITOR
 	virtual int32 Compile(class FMaterialCompiler* Compiler, int32 OutputIndex) override;
 	virtual void GetCaption(TArray<FString>& OutCaptions) const override;
 #endif
-	//~ End UMaterialExpression Interface
+	//~ End UMaterialExpressionMaterialX Interface
 };
 
diff --git a/Engine/Plugins/Interchange/Runtime/Source/Import/Private/MaterialX/MaterialExpressions/MaterialExpressionPlace2D.cpp b/Engine/Plugins/Interchange/Runtime/Source/Import/Private/MaterialX/MaterialExpressions/MaterialExpressionPlace2D.cpp
--- a/Engine/Plugins/Interchange/Runtime/Source/Import/Private/MaterialX/MaterialExpressions/MaterialExpressionPlace2D.cpp
+++ b/Engine/Plugins/Interchange/Runtime/Source/Import/Private/MaterialX/MaterialExpressions/MaterialExpressionPlace2D.cpp
@@ -16,26 +16,23 @@
 
 #include UE_INLINE_GENERATED_CPP_BY_NAME(MaterialExpressionPlace2D)
 
-#define LOCTEXT_NAMESPACE "MaterialExpressionPlace2D"
+#define LOCTEXT_NAMESPACE "MaterialExpressionMaterialXPlace2D"
 
-UMaterialExpressionPlace2D::UMaterialExpressionPlace2D(const FObjectInitializer& ObjectInitializer)
+UMaterialExpressionMaterialXPlace2D::UMaterialExpressionMaterialXPlace2D(const FObjectInitializer& ObjectInitializer)
 	: Super(ObjectInitializer)
 {
 	// Structure to hold one-time initialization
 	struct FConstructorStatics
 	{
 		FText NAME_MaterialX;
-		FText NAME_Compositing;
 		FConstructorStatics()
 			: NAME_MaterialX(LOCTEXT("MaterialX", "MaterialX"))
-			, NAME_Compositing(LOCTEXT("Math", "Math"))
 		{}
 	};
 	static FConstructorStatics ConstructorStatics;
 
 #if WITH_EDITORONLY_DATA
 	MenuCategories.Add(ConstructorStatics.NAME_MaterialX);
-	MenuCategories.Add(ConstructorStatics.NAME_Compositing);
 #endif
 }
 
@@ -57,7 +54,7 @@ namespace
 	}
 }
 
-int32 UMaterialExpressionPlace2D::Compile(FMaterialCompiler* Compiler, int32 OutputIndex)
+int32 UMaterialExpressionMaterialXPlace2D::Compile(FMaterialCompiler* Compiler, int32 OutputIndex)
 {
 	auto Constant2Expression = [](FExpressionInput& Input, float X, float Y)
 	{
@@ -120,7 +117,7 @@ int32 UMaterialExpressionPlace2D::Compile(FMaterialCompiler* Compiler, int32 Out
 
 	FExpressionInput ApplyRot;
 	{
-		UMaterialExpressionRotate2D* Rotate2D = NewObject<UMaterialExpressionRotate2D>();
+		UMaterialExpressionMaterialXRotate2D* Rotate2D = NewObject<UMaterialExpressionMaterialXRotate2D>();
 		Rotate2D->Input = ApplyScale;
 		Rotate2D->RotationAngle = RotationAngle;
 		ApplyRot.Connect(0, Rotate2D);
@@ -144,12 +141,12 @@ int32 UMaterialExpressionPlace2D::Compile(FMaterialCompiler* Compiler, int32 Out
 	return AppendVectorAddPivot.Compile(Compiler);
 }
 
-void UMaterialExpressionPlace2D::GetCaption(TArray<FString>& OutCaptions) const
+void UMaterialExpressionMaterialXPlace2D::GetCaption(TArray<FString>& OutCaptions) const
 {
-	OutCaptions.Add(TEXT("Place2D"));
+	OutCaptions.Add(TEXT("MaterialX Place2D"));
 }
 
-void UMaterialExpressionPlace2D::GetExpressionToolTip(TArray<FString>& OutToolTip)
+void UMaterialExpressionMaterialXPlace2D::GetExpressionToolTip(TArray<FString>& OutToolTip)
 {
 	ConvertToMultilineToolTip(TEXT("Transform incoming UV texture coordinates from one 2D frame of reference to another."
 								   "The order of operation is: -pivot, scale, rotate, translate, +pivot."),
diff --git a/Engine/Plugins/Interchange/Runtime/Source/Import/Private/MaterialX/MaterialExpressions/MaterialExpressionPlace2D.h b/Engine/Plugins/Interchange/Runtime/Source/Import/Private/MaterialX/MaterialExpressions/MaterialExpressionPlace2D.h
--- a/Engine/Plugins/Interchange/Runtime/Source/Import/Private/MaterialX/MaterialExpressions/MaterialExpressionPlace2D.h
+++ b/Engine/Plugins/Interchange/Runtime/Source/Import/Private/MaterialX/MaterialExpressions/MaterialExpressionPlace2D.h
@@ -14,7 +14,7 @@
  * -pivot, translate, rotate, scale, +pivot which does not introduce texture shear.
  * Default is 0 "SRT" for backward compatibility.*/
 UCLASS(MinimalAPI, collapsecategories, hidecategories = Object)
-class UMaterialExpressionPlace2D : public UMaterialExpression
+class UMaterialExpressionMaterialXPlace2D : public UMaterialExpression
 {
 	GENERATED_UCLASS_BODY()
 
@@ -51,12 +51,12 @@ class UMaterialExpressionPlace2D : public UMaterialExpression
 	UPROPERTY(EditAnywhere, Category = MaterialExpression, meta = (OverridingInputProperty = "Coordinates"))
 	uint8 ConstCoordinate;
 
-	//~ Begin UMaterialExpression Interface
+	//~ Begin UMaterialExpressionMaterialX Interface
 #if WITH_EDITOR
 	virtual int32 Compile(class FMaterialCompiler* Compiler, int32 OutputIndex) override;
 	virtual void GetCaption(TArray<FString>& OutCaptions) const override;
 	virtual void GetExpressionToolTip(TArray<FString>& OutToolTip) override;
 #endif
-	//~ End UMaterialExpression Interface
+	//~ End UMaterialExpressionMaterialX Interface
 };
 
diff --git a/Engine/Plugins/Interchange/Runtime/Source/Import/Private/MaterialX/MaterialExpressions/MaterialExpressionPlus.cpp b/Engine/Plugins/Interchange/Runtime/Source/Import/Private/MaterialX/MaterialExpressions/MaterialExpressionPlus.cpp
--- a/Engine/Plugins/Interchange/Runtime/Source/Import/Private/MaterialX/MaterialExpressions/MaterialExpressionPlus.cpp
+++ b/Engine/Plugins/Interchange/Runtime/Source/Import/Private/MaterialX/MaterialExpressions/MaterialExpressionPlus.cpp
@@ -4,40 +4,37 @@
 
 #include UE_INLINE_GENERATED_CPP_BY_NAME(MaterialExpressionPlus)
 
-#define LOCTEXT_NAMESPACE "MaterialExpressionPlus"
+#define LOCTEXT_NAMESPACE "MaterialExpressionMaterialXPlus"
 
-UMaterialExpressionPlus::UMaterialExpressionPlus(const FObjectInitializer& ObjectInitializer)
+UMaterialExpressionMaterialXPlus::UMaterialExpressionMaterialXPlus(const FObjectInitializer& ObjectInitializer)
 	: Super(ObjectInitializer)
 {
 	// Structure to hold one-time initialization
 	struct FConstructorStatics
 	{
 		FText NAME_MaterialX;
-		FText NAME_Compositing;
 		FConstructorStatics()
 			: NAME_MaterialX(LOCTEXT("MaterialX", "MaterialX"))
-			, NAME_Compositing(LOCTEXT("Compositing", "Compositing"))
 		{}
 	};
 	static FConstructorStatics ConstructorStatics;
 
 #if WITH_EDITORONLY_DATA
 	MenuCategories.Add(ConstructorStatics.NAME_MaterialX);
-	MenuCategories.Add(ConstructorStatics.NAME_Compositing);
 #endif
 }
 
 #if WITH_EDITOR
-int32 UMaterialExpressionPlus::Compile(FMaterialCompiler* Compiler, int32 OutputIndex)
+int32 UMaterialExpressionMaterialXPlus::Compile(FMaterialCompiler* Compiler, int32 OutputIndex)
 {
 	if(!A.GetTracedInput().Expression)
 	{
-		return Compiler->Errorf(TEXT("Missing A input"));
+		return Compiler->Errorf(TEXT("Missing MaterialX Plus input A"));
 	}
 
 	if(!B.GetTracedInput().Expression)
 	{
-		return Compiler->Errorf(TEXT("Missing B input"));
+		return Compiler->Errorf(TEXT("Missing MaterialX Plus input B"));
 	}
 
 	int32 IndexAlpha = Alpha.GetTracedInput().Expression ? Alpha.Compile(Compiler) : Compiler->Constant(ConstAlpha);
@@ -47,9 +44,9 @@ int32 UMaterialExpressionPlus::Compile(FMaterialCompiler* Compiler, int32 Output
 	return Compiler->Lerp(IndexB, Add, IndexAlpha);
 }
 
-void UMaterialExpressionPlus::GetCaption(TArray<FString>& OutCaptions) const
+void UMaterialExpressionMaterialXPlus::GetCaption(TArray<FString>& OutCaptions) const
 {
-	OutCaptions.Add(TEXT("Plus"));
+	OutCaptions.Add(TEXT("MaterialX Plus"));
 }
 #endif
 
diff --git a/Engine/Plugins/Interchange/Runtime/Source/Import/Private/MaterialX/MaterialExpressions/MaterialExpressionPlus.h b/Engine/Plugins/Interchange/Runtime/Source/Import/Private/MaterialX/MaterialExpressions/MaterialExpressionPlus.h
--- a/Engine/Plugins/Interchange/Runtime/Source/Import/Private/MaterialX/MaterialExpressions/MaterialExpressionPlus.h
+++ b/Engine/Plugins/Interchange/Runtime/Source/Import/Private/MaterialX/MaterialExpressions/MaterialExpressionPlus.h
@@ -14,7 +14,7 @@
  * Result: Lerp(B, A + B, Alpha)
  */
 UCLASS(MinimalAPI, collapsecategories, hidecategories = Object)
-class UMaterialExpressionPlus : public UMaterialExpression
+class UMaterialExpressionMaterialXPlus : public UMaterialExpression
 {
 	GENERATED_UCLASS_BODY()
 
@@ -31,11 +31,11 @@ class UMaterialExpressionPlus : public UMaterialExpression
 	UPROPERTY(EditAnywhere, Category = MaterialExpressionLinearInterpolate, meta = (OverridingInputProperty = "Alpha"))
 	float ConstAlpha = 1.f;
 
-	//~ Begin UMaterialExpression Interface
+	//~ Begin UMaterialExpressionMaterialX Interface
 #if WITH_EDITOR
 	virtual int32 Compile(class FMaterialCompiler* Compiler, int32 OutputIndex) override;
 	virtual void GetCaption(TArray<FString>& OutCaptions) const override;
 #endif
-	//~ End UMaterialExpression Interface
+	//~ End UMaterialExpressionMaterialX Interface
 };
 
diff --git a/Engine/Plugins/Interchange/Runtime/Source/Import/Private/MaterialX/MaterialExpressions/MaterialExpressionPremult.cpp b/Engine/Plugins/Interchange/Runtime/Source/Import/Private/MaterialX/MaterialExpressions/MaterialExpressionPremult.cpp
--- a/Engine/Plugins/Interchange/Runtime/Source/Import/Private/MaterialX/MaterialExpressions/MaterialExpressionPremult.cpp
+++ b/Engine/Plugins/Interchange/Runtime/Source/Import/Private/MaterialX/MaterialExpressions/MaterialExpressionPremult.cpp
@@ -4,31 +4,28 @@
 
 #include UE_INLINE_GENERATED_CPP_BY_NAME(MaterialExpressionPremult)
 
-#define LOCTEXT_NAMESPACE "MaterialExpressionPremult"
+#define LOCTEXT_NAMESPACE "MaterialExpressionMaterialXPremult"
 
-UMaterialExpressionPremult::UMaterialExpressionPremult(const FObjectInitializer& ObjectInitializer)
+UMaterialExpressionMaterialXPremult::UMaterialExpressionMaterialXPremult(const FObjectInitializer& ObjectInitializer)
 	: Super(ObjectInitializer)
 {
 	// Structure to hold one-time initialization
 	struct FConstructorStatics
 	{
 		FText NAME_MaterialX;
-		FText NAME_Compositing;
 		FConstructorStatics()
 			: NAME_MaterialX(LOCTEXT("MaterialX", "MaterialX"))
-			, NAME_Compositing(LOCTEXT("Compositing", "Compositing"))
 		{}
 	};
 	static FConstructorStatics ConstructorStatics;
 
 #if WITH_EDITORONLY_DATA
 	MenuCategories.Add(ConstructorStatics.NAME_MaterialX);
-	MenuCategories.Add(ConstructorStatics.NAME_Compositing);
 #endif
 }
 
 #if WITH_EDITOR
-int32 UMaterialExpressionPremult::Compile(FMaterialCompiler* Compiler, int32 OutputIndex)
+int32 UMaterialExpressionMaterialXPremult::Compile(FMaterialCompiler* Compiler, int32 OutputIndex)
 {
 	if(!Input.GetTracedInput().Expression)
 	{
@@ -41,9 +38,9 @@ int32 UMaterialExpressionPremult::Compile(FMaterialCompiler* Compiler, int32 Out
 	return Compiler->AppendVector(Compiler->Mul(RGB, Alpha), Alpha);
 }
 
-void UMaterialExpressionPremult::GetCaption(TArray<FString>& OutCaptions) const
+void UMaterialExpressionMaterialXPremult::GetCaption(TArray<FString>& OutCaptions) const
 {
-	OutCaptions.Add(TEXT("Premult"));
+		OutCaptions.Add(TEXT("MaterialX Premult"));
 }
 #endif
 
diff --git a/Engine/Plugins/Interchange/Runtime/Source/Import/Private/MaterialX/MaterialExpressions/MaterialExpressionPremult.h b/Engine/Plugins/Interchange/Runtime/Source/Import/Private/MaterialX/MaterialExpressions/MaterialExpressionPremult.h
--- a/Engine/Plugins/Interchange/Runtime/Source/Import/Private/MaterialX/MaterialExpressions/MaterialExpressionPremult.h
+++ b/Engine/Plugins/Interchange/Runtime/Source/Import/Private/MaterialX/MaterialExpressions/MaterialExpressionPremult.h
@@ -11,18 +11,18 @@
  * Input must be of type float4
  */
 UCLASS(collapsecategories, hidecategories = Object, MinimalAPI)
-class UMaterialExpressionPremult : public UMaterialExpression
+class UMaterialExpressionMaterialXPremult : public UMaterialExpression
 {
 	GENERATED_UCLASS_BODY()
 
 	UPROPERTY()
 	FExpressionInput Input;
 
-	//~ Begin UMaterialExpression Interface
+	//~ Begin UMaterialExpressionMaterialX Interface
 #if WITH_EDITOR
 	virtual int32 Compile(class FMaterialCompiler* Compiler, int32 OutputIndex) override;
 	virtual void GetCaption(TArray<FString>& OutCaptions) const override;
 #endif
-	//~ End UMaterialExpression Interface
+	//~ End UMaterialExpressionMaterialX Interface
 };
 
diff --git a/Engine/Plugins/Interchange/Runtime/Source/Import/Private/MaterialX/MaterialExpressions/MaterialExpressionRamp4.cpp b/Engine/Plugins/Interchange/Runtime/Source/Import/Private/MaterialX/MaterialExpressions/MaterialExpressionRamp4.cpp
--- a/Engine/Plugins/Interchange/Runtime/Source/Import/Private/MaterialX/MaterialExpressions/MaterialExpressionRamp4.cpp
+++ b/Engine/Plugins/Interchange/Runtime/Source/Import/Private/MaterialX/MaterialExpressions/MaterialExpressionRamp4.cpp
@@ -2,50 +2,47 @@
 #include "MaterialExpressionRamp4.h"
 #include "MaterialCompiler.h"
 
-#define LOCTEXT_NAMESPACE "MaterialExpressionRamp4"
+#define LOCTEXT_NAMESPACE "MaterialExpressionMaterialXRamp4"
 
-UMaterialExpressionRamp4::UMaterialExpressionRamp4(const FObjectInitializer& ObjectInitializer)
+UMaterialExpressionMaterialXRamp4::UMaterialExpressionMaterialXRamp4(const FObjectInitializer& ObjectInitializer)
 	: Super(ObjectInitializer)
 {
 	// Structure to hold one-time initialization
 	struct FConstructorStatics
 	{
 		FText NAME_MaterialX;
-		FText NAME_Utility;
 		FConstructorStatics()
 			: NAME_MaterialX(LOCTEXT("MaterialX", "MaterialX"))
-			, NAME_Utility(LOCTEXT("Procedural", "Procedural"))
 		{}
 	};
 	static FConstructorStatics ConstructorStatics;
 
 #if WITH_EDITORONLY_DATA
 	MenuCategories.Add(ConstructorStatics.NAME_MaterialX);
-	MenuCategories.Add(ConstructorStatics.NAME_Utility);
 #endif
 }
 
 #if WITH_EDITOR
-int32 UMaterialExpressionRamp4::Compile(FMaterialCompiler* Compiler, int32 OutputIndex)
+int32 UMaterialExpressionMaterialXRamp4::Compile(FMaterialCompiler* Compiler, int32 OutputIndex)
 {
 	if(!A.GetTracedInput().Expression)
 	{
-		return Compiler->Errorf(TEXT("Missing A input"));
+		return Compiler->Errorf(TEXT("Missing MaterialX Ramp4 input A"));
 	}
 
 	if(!B.GetTracedInput().Expression)
 	{
-		return Compiler->Errorf(TEXT("Missing B input"));
+		return Compiler->Errorf(TEXT("Missing MaterialX Ramp4 input B"));
 	}
 	
 	if(!C.GetTracedInput().Expression)
 	{
-		return Compiler->Errorf(TEXT("Missing C input"));
+		return Compiler->Errorf(TEXT("Missing MaterialX Ramp4 input C"));
 	}
 
 	if(!D.GetTracedInput().Expression)
 	{
-		return Compiler->Errorf(TEXT("Missing D input"));
+		return Compiler->Errorf(TEXT("Missing MaterialX Ramp4 input D"));
 	}
 
 	int32 IndexCoordinates = Coordinates.GetTracedInput().Expression ? Coordinates.Compile(Compiler) : Compiler->TextureCoordinate(ConstCoordinate, false, false);
@@ -64,12 +61,12 @@ int32 UMaterialExpressionRamp4::Compile(FMaterialCompiler* Compiler, int32 Outpu
 	return Compiler->Lerp(MixBot, MixTop, T);
 }
 
-void UMaterialExpressionRamp4::GetCaption(TArray<FString>& OutCaptions) const
+void UMaterialExpressionMaterialXRamp4::GetCaption(TArray<FString>& OutCaptions) const
 {
-	OutCaptions.Add(TEXT("Ramp4"));
+	OutCaptions.Add(TEXT("MaterialX Ramp4"));
 }
 
-void UMaterialExpressionRamp4::GetExpressionToolTip(TArray<FString>& OutToolTip)
+void UMaterialExpressionMaterialXRamp4::GetExpressionToolTip(TArray<FString>& OutToolTip)
 {
 	ConvertToMultilineToolTip(TEXT("A 4-corner bilinear value ramp."), 40, OutToolTip);
 }
diff --git a/Engine/Plugins/Interchange/Runtime/Source/Import/Private/MaterialX/MaterialExpressions/MaterialExpressionRamp4.h b/Engine/Plugins/Interchange/Runtime/Source/Import/Private/MaterialX/MaterialExpressions/MaterialExpressionRamp4.h
--- a/Engine/Plugins/Interchange/Runtime/Source/Import/Private/MaterialX/MaterialExpressions/MaterialExpressionRamp4.h
+++ b/Engine/Plugins/Interchange/Runtime/Source/Import/Private/MaterialX/MaterialExpressions/MaterialExpressionRamp4.h
@@ -10,7 +10,7 @@
  * A material expression that computes a 4-corner bilinear value ramp..
  */
 UCLASS(collapsecategories, hidecategories = Object, MinimalAPI)
-class UMaterialExpressionRamp4: public UMaterialExpression
+class UMaterialExpressionMaterialXRamp4: public UMaterialExpression
 {
 	GENERATED_UCLASS_BODY()
 
@@ -37,12 +37,12 @@ class UMaterialExpressionRamp4: public UMaterialExpression
 	UPROPERTY(EditAnywhere, Category = MaterialExpression, meta = (OverridingInputProperty = "Coordinates"))
 	uint8 ConstCoordinate;
 
-	//~ Begin UMaterialExpression Interface
+	//~ Begin UMaterialExpressionMaterialX Interface
 #if WITH_EDITOR
 	virtual int32 Compile(class FMaterialCompiler* Compiler, int32 OutputIndex) override;
 	virtual void GetCaption(TArray<FString>& OutCaptions) const override;
 	virtual void GetExpressionToolTip(TArray<FString>& OutToolTip) override;
 #endif
-	//~ End UMaterialExpression Interface
+	//~ End UMaterialExpressionMaterialX Interface
 };
 
diff --git a/Engine/Plugins/Interchange/Runtime/Source/Import/Private/MaterialX/MaterialExpressions/MaterialExpressionRampLeftRight.cpp b/Engine/Plugins/Interchange/Runtime/Source/Import/Private/MaterialX/MaterialExpressions/MaterialExpressionRampLeftRight.cpp
--- a/Engine/Plugins/Interchange/Runtime/Source/Import/Private/MaterialX/MaterialExpressions/MaterialExpressionRampLeftRight.cpp
+++ b/Engine/Plugins/Interchange/Runtime/Source/Import/Private/MaterialX/MaterialExpressions/MaterialExpressionRampLeftRight.cpp
@@ -4,40 +4,37 @@
 
 #include UE_INLINE_GENERATED_CPP_BY_NAME(MaterialExpressionRampLeftRight)
 
-#define LOCTEXT_NAMESPACE "MaterialExpressionRampLeftRight"
+#define LOCTEXT_NAMESPACE "MaterialExpressionMaterialXRampLeftRight"
 
-UMaterialExpressionRampLeftRight::UMaterialExpressionRampLeftRight(const FObjectInitializer& ObjectInitializer)
+UMaterialExpressionMaterialXRampLeftRight::UMaterialExpressionMaterialXRampLeftRight(const FObjectInitializer& ObjectInitializer)
 	: Super(ObjectInitializer)
 {
 	// Structure to hold one-time initialization
 	struct FConstructorStatics
 	{
 		FText NAME_MaterialX;
-		FText NAME_Procedural;
 		FConstructorStatics()
 			: NAME_MaterialX(LOCTEXT("MaterialX", "MaterialX"))
-			, NAME_Procedural(LOCTEXT("Procedural", "Procedural"))
 		{}
 	};
 	static FConstructorStatics ConstructorStatics;
 
 #if WITH_EDITORONLY_DATA
 	MenuCategories.Add(ConstructorStatics.NAME_MaterialX);
-	MenuCategories.Add(ConstructorStatics.NAME_Procedural);
 #endif
 }
 
 #if WITH_EDITOR
-int32 UMaterialExpressionRampLeftRight::Compile(FMaterialCompiler* Compiler, int32 OutputIndex)
+int32 UMaterialExpressionMaterialXRampLeftRight::Compile(FMaterialCompiler* Compiler, int32 OutputIndex)
 {
 	if(!A.GetTracedInput().Expression)
 	{
-		return Compiler->Errorf(TEXT("Missing A input"));
+		return Compiler->Errorf(TEXT("Missing MaterialX RampLR input A"));
 	}
 
 	if(!B.GetTracedInput().Expression)
 	{
-		return Compiler->Errorf(TEXT("Missing B input"));
+		return Compiler->Errorf(TEXT("Missing MaterialX RampLR input B"));
 	}
 
 	int32 CoordinateIndex = Coordinates.GetTracedInput().Expression ? Coordinates.Compile(Compiler) : Compiler->TextureCoordinate(ConstCoordinate, false, false);
@@ -45,9 +42,9 @@ int32 UMaterialExpressionRampLeftRight::Compile(FMaterialCompiler* Compiler, int
 	return Compiler->Lerp(A.Compile(Compiler), B.Compile(Compiler), Compiler->Saturate(Compiler->ComponentMask(CoordinateIndex, true, false, false, false)));
 }
 
-void UMaterialExpressionRampLeftRight::GetCaption(TArray<FString>& OutCaptions) const
+void UMaterialExpressionMaterialXRampLeftRight::GetCaption(TArray<FString>& OutCaptions) const
 {
-	OutCaptions.Add(TEXT("RampLR"));
+	OutCaptions.Add(TEXT("MaterialX RampLR"));
 }
 #endif
 
diff --git a/Engine/Plugins/Interchange/Runtime/Source/Import/Private/MaterialX/MaterialExpressions/MaterialExpressionRampLeftRight.h b/Engine/Plugins/Interchange/Runtime/Source/Import/Private/MaterialX/MaterialExpressions/MaterialExpressionRampLeftRight.h
--- a/Engine/Plugins/Interchange/Runtime/Source/Import/Private/MaterialX/MaterialExpressions/MaterialExpressionRampLeftRight.h
+++ b/Engine/Plugins/Interchange/Runtime/Source/Import/Private/MaterialX/MaterialExpressions/MaterialExpressionRampLeftRight.h
@@ -10,7 +10,7 @@
  * A material expression that computes a left-to-right bilinear value ramp.
  */
 UCLASS(collapsecategories, hidecategories = Object, MinimalAPI)
-class UMaterialExpressionRampLeftRight : public UMaterialExpression
+class UMaterialExpressionMaterialXRampLeftRight : public UMaterialExpression
 {
 	GENERATED_UCLASS_BODY()
 
@@ -27,11 +27,11 @@ class UMaterialExpressionRampLeftRight : public UMaterialExpression
 	UPROPERTY(EditAnywhere, Category = MaterialExpressionTextureSample, meta = (OverridingInputProperty = "Coordinates"))
 	uint8 ConstCoordinate;
 
-	//~ Begin UMaterialExpression Interface
+	//~ Begin UMaterialExpressionMaterialX Interface
 #if WITH_EDITOR
 	virtual int32 Compile(class FMaterialCompiler* Compiler, int32 OutputIndex) override;
 	virtual void GetCaption(TArray<FString>& OutCaptions) const override;
 #endif
-	//~ End UMaterialExpression Interface
+	//~ End UMaterialExpressionMaterialX Interface
 };
 
diff --git a/Engine/Plugins/Interchange/Runtime/Source/Import/Private/MaterialX/MaterialExpressions/MaterialExpressionRampTopBottom.cpp b/Engine/Plugins/Interchange/Runtime/Source/Import/Private/MaterialX/MaterialExpressions/MaterialExpressionRampTopBottom.cpp
--- a/Engine/Plugins/Interchange/Runtime/Source/Import/Private/MaterialX/MaterialExpressions/MaterialExpressionRampTopBottom.cpp
+++ b/Engine/Plugins/Interchange/Runtime/Source/Import/Private/MaterialX/MaterialExpressions/MaterialExpressionRampTopBottom.cpp
@@ -4,40 +4,37 @@
 
 #include UE_INLINE_GENERATED_CPP_BY_NAME(MaterialExpressionRampTopBottom)
 
-#define LOCTEXT_NAMESPACE "MaterialExpressionRampLeftRight"
+#define LOCTEXT_NAMESPACE "MaterialExpressionMaterialXRampLeftRight"
 
-UMaterialExpressionRampTopBottom::UMaterialExpressionRampTopBottom(const FObjectInitializer& ObjectInitializer)
+UMaterialExpressionMaterialXRampTopBottom::UMaterialExpressionMaterialXRampTopBottom(const FObjectInitializer& ObjectInitializer)
 	: Super(ObjectInitializer)
 {
 	// Structure to hold one-time initialization
 	struct FConstructorStatics
 	{
 		FText NAME_MaterialX;
-		FText NAME_Procedural;
 		FConstructorStatics()
 			: NAME_MaterialX(LOCTEXT("MaterialX", "MaterialX"))
-			, NAME_Procedural(LOCTEXT("Procedural", "Procedural"))
 		{}
 	};
 	static FConstructorStatics ConstructorStatics;
 
 #if WITH_EDITORONLY_DATA
 	MenuCategories.Add(ConstructorStatics.NAME_MaterialX);
-	MenuCategories.Add(ConstructorStatics.NAME_Procedural);
 #endif
 }
 
 #if WITH_EDITOR
-int32 UMaterialExpressionRampTopBottom::Compile(FMaterialCompiler* Compiler, int32 OutputIndex)
+int32 UMaterialExpressionMaterialXRampTopBottom::Compile(FMaterialCompiler* Compiler, int32 OutputIndex)
 {
 	if(!A.GetTracedInput().Expression)
 	{
-		return Compiler->Errorf(TEXT("Missing A input"));
+		return Compiler->Errorf(TEXT("Missing MaterialX RampTB input A"));
 	}
 
 	if(!B.GetTracedInput().Expression)
 	{
-		return Compiler->Errorf(TEXT("Missing B input"));
+		return Compiler->Errorf(TEXT("Missing MaterialX RampTB input B"));
 	}
 
 	int32 CoordinateIndex = Coordinates.GetTracedInput().Expression ? Coordinates.Compile(Compiler) : Compiler->TextureCoordinate(ConstCoordinate, false, false);
@@ -45,9 +42,9 @@ int32 UMaterialExpressionRampTopBottom::Compile(FMaterialCompiler* Compiler, int
 	return Compiler->Lerp(B.Compile(Compiler), A.Compile(Compiler), Compiler->Saturate(Compiler->ComponentMask(CoordinateIndex, false, true, false, false)));
 }
 
-void UMaterialExpressionRampTopBottom::GetCaption(TArray<FString>& OutCaptions) const
+void UMaterialExpressionMaterialXRampTopBottom::GetCaption(TArray<FString>& OutCaptions) const
 {
-	OutCaptions.Add(TEXT("RampTB"));
+	OutCaptions.Add(TEXT("MaterialX RampTB"));
 }
 #endif
 
diff --git a/Engine/Plugins/Interchange/Runtime/Source/Import/Private/MaterialX/MaterialExpressions/MaterialExpressionRampTopBottom.h b/Engine/Plugins/Interchange/Runtime/Source/Import/Private/MaterialX/MaterialExpressions/MaterialExpressionRampTopBottom.h
--- a/Engine/Plugins/Interchange/Runtime/Source/Import/Private/MaterialX/MaterialExpressions/MaterialExpressionRampTopBottom.h
+++ b/Engine/Plugins/Interchange/Runtime/Source/Import/Private/MaterialX/MaterialExpressions/MaterialExpressionRampTopBottom.h
@@ -10,7 +10,7 @@
  * A material expression that computes a top-to-bottom bilinear value ramp.
  */
 UCLASS(collapsecategories, hidecategories = Object, MinimalAPI)
-class UMaterialExpressionRampTopBottom : public UMaterialExpression
+class UMaterialExpressionMaterialXRampTopBottom : public UMaterialExpression
 {
 	GENERATED_UCLASS_BODY()
 
@@ -27,11 +27,11 @@ class UMaterialExpressionRampTopBottom : public UMaterialExpression
 	UPROPERTY(EditAnywhere, Category = MaterialExpressionTextureSample, meta = (OverridingInputProperty = "Coordinates"))
 		uint8 ConstCoordinate;
 
-	//~ Begin UMaterialExpression Interface
+	//~ Begin UMaterialExpressionMaterialX Interface
 #if WITH_EDITOR
 	virtual int32 Compile(class FMaterialCompiler* Compiler, int32 OutputIndex) override;
 	virtual void GetCaption(TArray<FString>& OutCaptions) const override;
 #endif
-	//~ End UMaterialExpression Interface
+	//~ End UMaterialExpressionMaterialX Interface
 };
 
diff --git a/Engine/Plugins/Interchange/Runtime/Source/Import/Private/MaterialX/MaterialExpressions/MaterialExpressionRemap.cpp b/Engine/Plugins/Interchange/Runtime/Source/Import/Private/MaterialX/MaterialExpressions/MaterialExpressionRemap.cpp
--- a/Engine/Plugins/Interchange/Runtime/Source/Import/Private/MaterialX/MaterialExpressions/MaterialExpressionRemap.cpp
+++ b/Engine/Plugins/Interchange/Runtime/Source/Import/Private/MaterialX/MaterialExpressions/MaterialExpressionRemap.cpp
@@ -4,19 +4,17 @@
 
 #include UE_INLINE_GENERATED_CPP_BY_NAME(MaterialExpressionRemap)
 
-#define LOCTEXT_NAMESPACE "MaterialExpressionRemap"
+#define LOCTEXT_NAMESPACE "MaterialExpressionMaterialXRemap"
 
-UMaterialExpressionRemap::UMaterialExpressionRemap(const FObjectInitializer& ObjectInitializer)
+UMaterialExpressionMaterialXRemap::UMaterialExpressionMaterialXRemap(const FObjectInitializer& ObjectInitializer)
 	: Super(ObjectInitializer)
 {
 	// Structure to hold one-time initialization
 	struct FConstructorStatics
 	{
 		FText NAME_MaterialX;
-		FText NAME_Adjustment;
 		FConstructorStatics()
 			: NAME_MaterialX(LOCTEXT("MaterialX", "MaterialX"))
-			, NAME_Adjustment(LOCTEXT("Image Adjustment", "Image Adjustment"))
 		{}
 	};
 	static FConstructorStatics ConstructorStatics;
@@ -28,12 +26,11 @@ UMaterialExpressionRemap::UMaterialExpressionRemap(const FObjectInitializer& Obj
 
 #if WITH_EDITORONLY_DATA
 	MenuCategories.Add(ConstructorStatics.NAME_MaterialX);
-	MenuCategories.Add(ConstructorStatics.NAME_Adjustment);
 #endif
 }
 
 #if WITH_EDITOR
-int32 UMaterialExpressionRemap::Compile(FMaterialCompiler* Compiler, int32 OutputIndex)
+int32 UMaterialExpressionMaterialXRemap::Compile(FMaterialCompiler* Compiler, int32 OutputIndex)
 {
 	if(!Input.GetTracedInput().Expression)
 	{
@@ -54,9 +51,9 @@ int32 UMaterialExpressionRemap::Compile(FMaterialCompiler* Compiler, int32 Outpu
 	return Compiler->Add(TargetLowIndex, Compiler->Mul(Compiler->Sub(Input.Compile(Compiler), InputLowIndex), Delta));
 }
 
-void UMaterialExpressionRemap::GetCaption(TArray<FString>& OutCaptions) const
+void UMaterialExpressionMaterialXRemap::GetCaption(TArray<FString>& OutCaptions) const
 {
-	OutCaptions.Add(TEXT("Remap"));
+	OutCaptions.Add(TEXT("MaterialX Remap"));
 }
 #endif
 
diff --git a/Engine/Plugins/Interchange/Runtime/Source/Import/Private/MaterialX/MaterialExpressions/MaterialExpressionRemap.h b/Engine/Plugins/Interchange/Runtime/Source/Import/Private/MaterialX/MaterialExpressions/MaterialExpressionRemap.h
--- a/Engine/Plugins/Interchange/Runtime/Source/Import/Private/MaterialX/MaterialExpressions/MaterialExpressionRemap.h
+++ b/Engine/Plugins/Interchange/Runtime/Source/Import/Private/MaterialX/MaterialExpressions/MaterialExpressionRemap.h
@@ -10,7 +10,7 @@
  * A material expression that Remap a value from one range to another.
  */
 UCLASS(collapsecategories, hidecategories = Object, MinimalAPI)
-class UMaterialExpressionRemap : public UMaterialExpression
+class UMaterialExpressionMaterialXRemap : public UMaterialExpression
 {
 	GENERATED_UCLASS_BODY()
 
@@ -41,11 +41,11 @@ class UMaterialExpressionRemap : public UMaterialExpression
 	UPROPERTY(EditAnywhere, Category = MaterialExpressionClamp, meta = (OverridingInputProperty = "TargetHigh"))
 	float TargetHighDefault;
 
-	//~ Begin UMaterialExpression Interface
+	//~ Begin UMaterialExpressionMaterialX Interface
 #if WITH_EDITOR
 	virtual int32 Compile(class FMaterialCompiler* Compiler, int32 OutputIndex) override;
 	virtual void GetCaption(TArray<FString>& OutCaptions) const override;
 #endif
-	//~ End UMaterialExpression Interface
+	//~ End UMaterialExpressionMaterialX Interface
 };
 
diff --git a/Engine/Plugins/Interchange/Runtime/Source/Import/Private/MaterialX/MaterialExpressions/MaterialExpressionRgbToHsv.cpp b/Engine/Plugins/Interchange/Runtime/Source/Import/Private/MaterialX/MaterialExpressions/MaterialExpressionRgbToHsv.cpp
--- a/Engine/Plugins/Interchange/Runtime/Source/Import/Private/MaterialX/MaterialExpressions/MaterialExpressionRgbToHsv.cpp
+++ b/Engine/Plugins/Interchange/Runtime/Source/Import/Private/MaterialX/MaterialExpressions/MaterialExpressionRgbToHsv.cpp
@@ -4,31 +4,28 @@
 #include "MaterialCompiler.h"
 #include "Materials/MaterialExpressionCustom.h"
 
-#define LOCTEXT_NAMESPACE "MaterialExpressionRgbToHsv"
+#define LOCTEXT_NAMESPACE "MaterialExpressionMaterialXRgbToHsv"
 
-UMaterialExpressionRgbToHsv::UMaterialExpressionRgbToHsv(const FObjectInitializer& ObjectInitializer)
+UMaterialExpressionMaterialXRgbToHsv::UMaterialExpressionMaterialXRgbToHsv(const FObjectInitializer& ObjectInitializer)
 	: Super(ObjectInitializer)
 {
 	// Structure to hold one-time initialization
 	struct FConstructorStatics
 	{
-		FText NAME_ImageAdjustment;
 		FText NAME_MaterialX;
 		FConstructorStatics()
-			: NAME_ImageAdjustment(LOCTEXT("Image Adjustment", "Image Adjustment"))
-			, NAME_MaterialX(LOCTEXT("MaterialX", "MaterialX"))
+			: NAME_MaterialX(LOCTEXT("MaterialX", "MaterialX"))
 		{}
 	};
 	static FConstructorStatics ConstructorStatics;
 
 #if WITH_EDITORONLY_DATA
 	MenuCategories.Add(ConstructorStatics.NAME_MaterialX);
-	MenuCategories.Add(ConstructorStatics.NAME_ImageAdjustment);
 #endif
 }
 
 #if WITH_EDITOR
-int32 UMaterialExpressionRgbToHsv::Compile(FMaterialCompiler* Compiler, int32 OutputIndex)
+int32 UMaterialExpressionMaterialXRgbToHsv::Compile(FMaterialCompiler* Compiler, int32 OutputIndex)
 {
 	if(!Input.GetTracedInput().Expression)
 	{
@@ -51,9 +48,9 @@ int32 UMaterialExpressionRgbToHsv::Compile(FMaterialCompiler* Compiler, int32 Ou
 	return MaterialExpressionCustom->Compile(Compiler, OutputIndex);
 }
 
-void UMaterialExpressionRgbToHsv::GetCaption(TArray<FString>& OutCaptions) const
+void UMaterialExpressionMaterialXRgbToHsv::GetCaption(TArray<FString>& OutCaptions) const
 {
-	OutCaptions.Add(TEXT("RGBToHSV"));
+	OutCaptions.Add(TEXT("MaterialX RGBToHSV"));
 }
 #endif
 
diff --git a/Engine/Plugins/Interchange/Runtime/Source/Import/Private/MaterialX/MaterialExpressions/MaterialExpressionRgbToHsv.h b/Engine/Plugins/Interchange/Runtime/Source/Import/Private/MaterialX/MaterialExpressions/MaterialExpressionRgbToHsv.h
--- a/Engine/Plugins/Interchange/Runtime/Source/Import/Private/MaterialX/MaterialExpressions/MaterialExpressionRgbToHsv.h
+++ b/Engine/Plugins/Interchange/Runtime/Source/Import/Private/MaterialX/MaterialExpressions/MaterialExpressionRgbToHsv.h
@@ -6,18 +6,18 @@
 #include "MaterialExpressionRgbToHsv.generated.h"
 
 UCLASS(collapsecategories, hidecategories = Object, MinimalAPI)
-class UMaterialExpressionRgbToHsv : public UMaterialExpression
+class UMaterialExpressionMaterialXRgbToHsv : public UMaterialExpression
 {
 	GENERATED_UCLASS_BODY()
 
 	UPROPERTY()
 	FExpressionInput Input;
 
-	//~ Begin UMaterialExpression Interface
+	//~ Begin UMaterialExpressionMaterialX Interface
 #if WITH_EDITOR
 	virtual int32 Compile(class FMaterialCompiler* Compiler, int32 OutputIndex) override;
 	virtual void GetCaption(TArray<FString>& OutCaptions) const override;
 #endif
-	//~ End UMaterialExpression Interface
+	//~ End UMaterialExpressionMaterialX Interface
 };
 
diff --git a/Engine/Plugins/Interchange/Runtime/Source/Import/Private/MaterialX/MaterialExpressions/MaterialExpressionRotate2D.cpp b/Engine/Plugins/Interchange/Runtime/Source/Import/Private/MaterialX/MaterialExpressions/MaterialExpressionRotate2D.cpp
--- a/Engine/Plugins/Interchange/Runtime/Source/Import/Private/MaterialX/MaterialExpressions/MaterialExpressionRotate2D.cpp
+++ b/Engine/Plugins/Interchange/Runtime/Source/Import/Private/MaterialX/MaterialExpressions/MaterialExpressionRotate2D.cpp
@@ -6,31 +6,28 @@
 
 #include UE_INLINE_GENERATED_CPP_BY_NAME(MaterialExpressionRotate2D)
 
-#define LOCTEXT_NAMESPACE "MaterialExpressionRotate2D"
+#define LOCTEXT_NAMESPACE "MaterialExpressionMaterialXRotate2D"
 
-UMaterialExpressionRotate2D::UMaterialExpressionRotate2D(const FObjectInitializer& ObjectInitializer)
+UMaterialExpressionMaterialXRotate2D::UMaterialExpressionMaterialXRotate2D(const FObjectInitializer& ObjectInitializer)
 	: Super(ObjectInitializer)
 {
 	// Structure to hold one-time initialization
 	struct FConstructorStatics
 	{
 		FText NAME_MaterialX;
-		FText NAME_Math;
 		FConstructorStatics()
 			: NAME_MaterialX(LOCTEXT("MaterialX", "MaterialX"))
-			, NAME_Math(LOCTEXT("Math", "Math"))
 		{}
 	};
 	static FConstructorStatics ConstructorStatics;
 
 #if WITH_EDITORONLY_DATA
 	MenuCategories.Add(ConstructorStatics.NAME_MaterialX);
-	MenuCategories.Add(ConstructorStatics.NAME_Math);
 #endif
 }
 
 #if WITH_EDITOR
-int32 UMaterialExpressionRotate2D::Compile(FMaterialCompiler* Compiler, int32 OutputIndex)
+int32 UMaterialExpressionMaterialXRotate2D::Compile(FMaterialCompiler* Compiler, int32 OutputIndex)
 {
 	if(!Input.GetTracedInput().Expression)
 	{
@@ -56,12 +53,12 @@ int32 UMaterialExpressionRotate2D::Compile(FMaterialCompiler* Compiler, int32 Ou
 	);
 }
 
-void UMaterialExpressionRotate2D::GetCaption(TArray<FString>& OutCaptions) const
+void UMaterialExpressionMaterialXRotate2D::GetCaption(TArray<FString>& OutCaptions) const
 {
-	OutCaptions.Add(TEXT("Rotate2D"));
+	OutCaptions.Add(TEXT("MaterialX Rotate2D"));
 }
 
-void UMaterialExpressionRotate2D::GetExpressionToolTip(TArray<FString>& OutToolTip)
+void UMaterialExpressionMaterialXRotate2D::GetExpressionToolTip(TArray<FString>& OutToolTip)
 {
 	OutToolTip.Add(TEXT("rotates a vector2 value about the origin in 2D."));
 }
diff --git a/Engine/Plugins/Interchange/Runtime/Source/Import/Private/MaterialX/MaterialExpressions/MaterialExpressionRotate2D.h b/Engine/Plugins/Interchange/Runtime/Source/Import/Private/MaterialX/MaterialExpressions/MaterialExpressionRotate2D.h
--- a/Engine/Plugins/Interchange/Runtime/Source/Import/Private/MaterialX/MaterialExpressions/MaterialExpressionRotate2D.h
+++ b/Engine/Plugins/Interchange/Runtime/Source/Import/Private/MaterialX/MaterialExpressions/MaterialExpressionRotate2D.h
@@ -9,7 +9,7 @@
  * A material expression that rotates a vector2 value about the origin in 2D.
  */
 UCLASS(collapsecategories, hidecategories = Object, MinimalAPI)
-class UMaterialExpressionRotate2D : public UMaterialExpression
+class UMaterialExpressionMaterialXRotate2D : public UMaterialExpression
 {
 	GENERATED_UCLASS_BODY()
 
@@ -25,12 +25,12 @@ class UMaterialExpressionRotate2D : public UMaterialExpression
 	UPROPERTY(EditAnywhere, Category = MaterialExpressionAngle, meta = (OverridingInputProperty = "RotationAngle"))
 	float ConstRotationAngle = 0.f;
 
-	//~ Begin UMaterialExpression Interface
+	//~ Begin UMaterialExpressionMaterialX Interface
 #if WITH_EDITOR
 	virtual int32 Compile(class FMaterialCompiler* Compiler, int32 OutputIndex) override;
 	virtual void GetCaption(TArray<FString>& OutCaptions) const override;
 	virtual void GetExpressionToolTip(TArray<FString>& OutToolTip) override;
 #endif
-	//~ End UMaterialExpression Interface
+	//~ End UMaterialExpressionMaterialX Interface
 };
 
diff --git a/Engine/Plugins/Interchange/Runtime/Source/Import/Private/MaterialX/MaterialExpressions/MaterialExpressionScreen.cpp b/Engine/Plugins/Interchange/Runtime/Source/Import/Private/MaterialX/MaterialExpressions/MaterialExpressionScreen.cpp
--- a/Engine/Plugins/Interchange/Runtime/Source/Import/Private/MaterialX/MaterialExpressions/MaterialExpressionScreen.cpp
+++ b/Engine/Plugins/Interchange/Runtime/Source/Import/Private/MaterialX/MaterialExpressions/MaterialExpressionScreen.cpp
@@ -4,40 +4,37 @@
 
 #include UE_INLINE_GENERATED_CPP_BY_NAME(MaterialExpressionScreen)
 
-#define LOCTEXT_NAMESPACE "MaterialExpressionScreen"
+#define LOCTEXT_NAMESPACE "MaterialExpressionMaterialXScreen"
 
-UMaterialExpressionScreen::UMaterialExpressionScreen(const FObjectInitializer& ObjectInitializer)
+UMaterialExpressionMaterialXScreen::UMaterialExpressionMaterialXScreen(const FObjectInitializer& ObjectInitializer)
 	: Super(ObjectInitializer)
 {
 	// Structure to hold one-time initialization
 	struct FConstructorStatics
 	{
 		FText NAME_MaterialX;
-		FText NAME_Compositing;
 		FConstructorStatics()
 			: NAME_MaterialX(LOCTEXT("MaterialX", "MaterialX"))
-			, NAME_Compositing(LOCTEXT("Compositing", "Compositing"))
 		{}
 	};
 	static FConstructorStatics ConstructorStatics;
 
 #if WITH_EDITORONLY_DATA
 	MenuCategories.Add(ConstructorStatics.NAME_MaterialX);
-	MenuCategories.Add(ConstructorStatics.NAME_Compositing);
 #endif
 }
 
 #if WITH_EDITOR
-int32 UMaterialExpressionScreen::Compile(FMaterialCompiler* Compiler, int32 OutputIndex)
+int32 UMaterialExpressionMaterialXScreen::Compile(FMaterialCompiler* Compiler, int32 OutputIndex)
 {
 	if(!A.GetTracedInput().Expression)
 	{
-		return Compiler->Errorf(TEXT("Missing A input"));
+		return Compiler->Errorf(TEXT("Missing MaterialX Screen input A"));
 	}
 
 	if(!B.GetTracedInput().Expression)
 	{
-		return Compiler->Errorf(TEXT("Missing B input"));
+		return Compiler->Errorf(TEXT("Missing MaterialX Screen input B"));
 	}
 
 	int32 IndexAlpha = Alpha.GetTracedInput().Expression ? Alpha.Compile(Compiler) : Compiler->Constant(ConstAlpha);
@@ -55,9 +52,9 @@ int32 UMaterialExpressionScreen::Compile(FMaterialCompiler* Compiler, int32 Outp
 	return Compiler->Lerp(IndexB, Screen, IndexAlpha);
 }
 
-void UMaterialExpressionScreen::GetCaption(TArray<FString>& OutCaptions) const
+void UMaterialExpressionMaterialXScreen::GetCaption(TArray<FString>& OutCaptions) const
 {
-	OutCaptions.Add(TEXT("Screen"));
+	OutCaptions.Add(TEXT("MaterialX Screen"));
 }
 #endif
 
diff --git a/Engine/Plugins/Interchange/Runtime/Source/Import/Private/MaterialX/MaterialExpressions/MaterialExpressionScreen.h b/Engine/Plugins/Interchange/Runtime/Source/Import/Private/MaterialX/MaterialExpressions/MaterialExpressionScreen.h
--- a/Engine/Plugins/Interchange/Runtime/Source/Import/Private/MaterialX/MaterialExpressions/MaterialExpressionScreen.h
+++ b/Engine/Plugins/Interchange/Runtime/Source/Import/Private/MaterialX/MaterialExpressions/MaterialExpressionScreen.h
@@ -14,7 +14,7 @@
  * Result: Lerp(B, 1-(1-F)(1-B), Alpha)
  */
 UCLASS(MinimalAPI, collapsecategories, hidecategories = Object)
-class UMaterialExpressionScreen : public UMaterialExpression
+class UMaterialExpressionMaterialXScreen : public UMaterialExpression
 {
 	GENERATED_UCLASS_BODY()
 
@@ -31,11 +31,11 @@ class UMaterialExpressionScreen : public UMaterialExpression
 	UPROPERTY(EditAnywhere, Category = MaterialExpressionLinearInterpolate, meta = (OverridingInputProperty = "Alpha"))
 	float ConstAlpha = 1.f;
 
-	//~ Begin UMaterialExpression Interface
+	//~ Begin UMaterialExpressionMaterialX Interface
 #if WITH_EDITOR
 	virtual int32 Compile(class FMaterialCompiler* Compiler, int32 OutputIndex) override;
 	virtual void GetCaption(TArray<FString>& OutCaptions) const override;
 #endif
-	//~ End UMaterialExpression Interface
+	//~ End UMaterialExpressionMaterialX Interface
 };
 
diff --git a/Engine/Plugins/Interchange/Runtime/Source/Import/Private/MaterialX/MaterialExpressions/MaterialExpressionSplitLeftRight.cpp b/Engine/Plugins/Interchange/Runtime/Source/Import/Private/MaterialX/MaterialExpressions/MaterialExpressionSplitLeftRight.cpp
--- a/Engine/Plugins/Interchange/Runtime/Source/Import/Private/MaterialX/MaterialExpressions/MaterialExpressionSplitLeftRight.cpp
+++ b/Engine/Plugins/Interchange/Runtime/Source/Import/Private/MaterialX/MaterialExpressions/MaterialExpressionSplitLeftRight.cpp
@@ -4,9 +4,9 @@
 
 #include UE_INLINE_GENERATED_CPP_BY_NAME(MaterialExpressionSplitLeftRight)
 
-#define LOCTEXT_NAMESPACE "MaterialExpressionSplitLeftRight"
+#define LOCTEXT_NAMESPACE "MaterialExpressionMaterialXSplitLeftRight"
 
-UMaterialExpressionSplitLeftRight::UMaterialExpressionSplitLeftRight(const FObjectInitializer& ObjectInitializer)
+UMaterialExpressionMaterialXSplitLeftRight::UMaterialExpressionMaterialXSplitLeftRight(const FObjectInitializer& ObjectInitializer)
 	: Super(ObjectInitializer),
 	ConstCenter(0.5f),
 	ConstCoordinate(0)
@@ -15,31 +15,28 @@ UMaterialExpressionSplitLeftRight::UMaterialExpressionSplitLeftRight(const FObje
 	struct FConstructorStatics
 	{
 		FText NAME_MaterialX;
-		FText NAME_Procedural;
 		FConstructorStatics()
 			: NAME_MaterialX(LOCTEXT("MaterialX", "MaterialX"))
-			, NAME_Procedural(LOCTEXT("Procedural", "Procedural"))
 		{}
 	};
 	static FConstructorStatics ConstructorStatics;
 
 #if WITH_EDITORONLY_DATA
 	MenuCategories.Add(ConstructorStatics.NAME_MaterialX);
-	MenuCategories.Add(ConstructorStatics.NAME_Procedural);
 #endif
 }
 
 #if WITH_EDITOR
-int32 UMaterialExpressionSplitLeftRight::Compile(FMaterialCompiler* Compiler, int32 OutputIndex)
+int32 UMaterialExpressionMaterialXSplitLeftRight::Compile(FMaterialCompiler* Compiler, int32 OutputIndex)
 {
 	if(!A.GetTracedInput().Expression)
 	{
-		return Compiler->Errorf(TEXT("Missing A input"));
+		return Compiler->Errorf(TEXT("Missing MaterialX SplitLR input A"));
 	}
 
 	if(!B.GetTracedInput().Expression)
 	{
-		return Compiler->Errorf(TEXT("Missing B input"));
+		return Compiler->Errorf(TEXT("Missing MaterialX SplitLR input B"));
 	}
 
 	int32 CoordinateIndex = Coordinates.GetTracedInput().Expression ? Coordinates.Compile(Compiler) : Compiler->TextureCoordinate(ConstCoordinate, false, false);
@@ -53,9 +50,9 @@ int32 UMaterialExpressionSplitLeftRight::Compile(FMaterialCompiler* Compiler, in
 	return Compiler->Lerp(A.Compile(Compiler), B.Compile(Compiler), AAStep);
 }
 
-void UMaterialExpressionSplitLeftRight::GetCaption(TArray<FString>& OutCaptions) const
+void UMaterialExpressionMaterialXSplitLeftRight::GetCaption(TArray<FString>& OutCaptions) const
 {
-	OutCaptions.Add(TEXT("SplitLR"));
+	OutCaptions.Add(TEXT("MaterialX SplitLR"));
 }
 #endif
 
diff --git a/Engine/Plugins/Interchange/Runtime/Source/Import/Private/MaterialX/MaterialExpressions/MaterialExpressionSplitLeftRight.h b/Engine/Plugins/Interchange/Runtime/Source/Import/Private/MaterialX/MaterialExpressions/MaterialExpressionSplitLeftRight.h
--- a/Engine/Plugins/Interchange/Runtime/Source/Import/Private/MaterialX/MaterialExpressions/MaterialExpressionSplitLeftRight.h
+++ b/Engine/Plugins/Interchange/Runtime/Source/Import/Private/MaterialX/MaterialExpressions/MaterialExpressionSplitLeftRight.h
@@ -10,7 +10,7 @@
  * A material expression that computes a left-right split matte, split at a specified u value.
  */
 UCLASS(collapsecategories, hidecategories = Object, MinimalAPI)
-class UMaterialExpressionSplitLeftRight : public UMaterialExpression
+class UMaterialExpressionMaterialXSplitLeftRight : public UMaterialExpression
 {
 	GENERATED_UCLASS_BODY()
 
@@ -34,11 +34,11 @@ class UMaterialExpressionSplitLeftRight : public UMaterialExpression
 	UPROPERTY(EditAnywhere, Category = MaterialExpressionTextureSample, meta = (OverridingInputProperty = "Coordinates"))
 	uint8 ConstCoordinate;
 
-	//~ Begin UMaterialExpression Interface
+	//~ Begin UMaterialExpressionMaterialX Interface
 #if WITH_EDITOR
 	virtual int32 Compile(class FMaterialCompiler* Compiler, int32 OutputIndex) override;
 	virtual void GetCaption(TArray<FString>& OutCaptions) const override;
 #endif
-	//~ End UMaterialExpression Interface
+	//~ End UMaterialExpressionMaterialX Interface
 };
 
diff --git a/Engine/Plugins/Interchange/Runtime/Source/Import/Private/MaterialX/MaterialExpressions/MaterialExpressionSplitTopBottom.cpp b/Engine/Plugins/Interchange/Runtime/Source/Import/Private/MaterialX/MaterialExpressions/MaterialExpressionSplitTopBottom.cpp
--- a/Engine/Plugins/Interchange/Runtime/Source/Import/Private/MaterialX/MaterialExpressions/MaterialExpressionSplitTopBottom.cpp
+++ b/Engine/Plugins/Interchange/Runtime/Source/Import/Private/MaterialX/MaterialExpressions/MaterialExpressionSplitTopBottom.cpp
@@ -4,9 +4,9 @@
 
 #include UE_INLINE_GENERATED_CPP_BY_NAME(MaterialExpressionSplitTopBottom)
 
-#define LOCTEXT_NAMESPACE "MaterialExpressionSplitLeftRight"
+#define LOCTEXT_NAMESPACE "MaterialExpressionMaterialXSplitLeftRight"
 
-UMaterialExpressionSplitTopBottom::UMaterialExpressionSplitTopBottom(const FObjectInitializer& ObjectInitializer)
+UMaterialExpressionMaterialXSplitTopBottom::UMaterialExpressionMaterialXSplitTopBottom(const FObjectInitializer& ObjectInitializer)
 	: Super(ObjectInitializer),
 	ConstCenter(0.5f),
 	ConstCoordinate(0)
@@ -15,31 +15,28 @@ UMaterialExpressionSplitTopBottom::UMaterialExpressionSplitTopBottom(const FObje
 	struct FConstructorStatics
 	{
 		FText NAME_MaterialX;
-		FText NAME_Procedural;
 		FConstructorStatics()
 			: NAME_MaterialX(LOCTEXT("MaterialX", "MaterialX"))
-			, NAME_Procedural(LOCTEXT("Procedural", "Procedural"))
 		{}
 	};
 	static FConstructorStatics ConstructorStatics;
 
 #if WITH_EDITORONLY_DATA
 	MenuCategories.Add(ConstructorStatics.NAME_MaterialX);
-	MenuCategories.Add(ConstructorStatics.NAME_Procedural);
 #endif
 }
 
 #if WITH_EDITOR
-int32 UMaterialExpressionSplitTopBottom::Compile(FMaterialCompiler* Compiler, int32 OutputIndex)
+int32 UMaterialExpressionMaterialXSplitTopBottom::Compile(FMaterialCompiler* Compiler, int32 OutputIndex)
 {
 	if(!A.GetTracedInput().Expression)
 	{
-		return Compiler->Errorf(TEXT("Missing A input"));
+		return Compiler->Errorf(TEXT("Missing MaterialX SplitTB input A"));
 	}
 
 	if(!B.GetTracedInput().Expression)
 	{
-		return Compiler->Errorf(TEXT("Missing B input"));
+		return Compiler->Errorf(TEXT("Missing MaterialX SplitTB input B"));
 	}
 
 	int32 CoordinateIndex = Coordinates.GetTracedInput().Expression ? Coordinates.Compile(Compiler) : Compiler->TextureCoordinate(ConstCoordinate, false, false);
@@ -53,9 +50,9 @@ int32 UMaterialExpressionSplitTopBottom::Compile(FMaterialCompiler* Compiler, in
 	return Compiler->Lerp(B.Compile(Compiler), A.Compile(Compiler), AAStep);
 }
 
-void UMaterialExpressionSplitTopBottom::GetCaption(TArray<FString>& OutCaptions) const
+void UMaterialExpressionMaterialXSplitTopBottom::GetCaption(TArray<FString>& OutCaptions) const
 {
-	OutCaptions.Add(TEXT("SplitTB"));
+	OutCaptions.Add(TEXT("MaterialX SplitTB"));
 }
 #endif
 
diff --git a/Engine/Plugins/Interchange/Runtime/Source/Import/Private/MaterialX/MaterialExpressions/MaterialExpressionSplitTopBottom.h b/Engine/Plugins/Interchange/Runtime/Source/Import/Private/MaterialX/MaterialExpressions/MaterialExpressionSplitTopBottom.h
--- a/Engine/Plugins/Interchange/Runtime/Source/Import/Private/MaterialX/MaterialExpressions/MaterialExpressionSplitTopBottom.h
+++ b/Engine/Plugins/Interchange/Runtime/Source/Import/Private/MaterialX/MaterialExpressions/MaterialExpressionSplitTopBottom.h
@@ -10,7 +10,7 @@
  * A material expression that computes a top-bottom split matte, split at a specified v value.
  */
 UCLASS(collapsecategories, hidecategories = Object, MinimalAPI)
-class UMaterialExpressionSplitTopBottom : public UMaterialExpression
+class UMaterialExpressionMaterialXSplitTopBottom : public UMaterialExpression
 {
 	GENERATED_UCLASS_BODY()
 
@@ -34,11 +34,11 @@ class UMaterialExpressionSplitTopBottom : public UMaterialExpression
 	UPROPERTY(EditAnywhere, Category = MaterialExpressionTextureSample, meta = (OverridingInputProperty = "Coordinates"))
 	uint8 ConstCoordinate;
 
-	//~ Begin UMaterialExpression Interface
+	//~ Begin UMaterialExpressionMaterialX Interface
 #if WITH_EDITOR
 	virtual int32 Compile(class FMaterialCompiler* Compiler, int32 OutputIndex) override;
 	virtual void GetCaption(TArray<FString>& OutCaptions) const override;
 #endif
-	//~ End UMaterialExpression Interface
+	//~ End UMaterialExpressionMaterialX Interface
 };
 
diff --git a/Engine/Plugins/Interchange/Runtime/Source/Import/Private/MaterialX/MaterialExpressions/MaterialExpressionTextureSampleParameterBlur.cpp b/Engine/Plugins/Interchange/Runtime/Source/Import/Private/MaterialX/MaterialExpressions/MaterialExpressionTextureSampleParameterBlur.cpp
--- a/Engine/Plugins/Interchange/Runtime/Source/Import/Private/MaterialX/MaterialExpressions/MaterialExpressionTextureSampleParameterBlur.cpp
+++ b/Engine/Plugins/Interchange/Runtime/Source/Import/Private/MaterialX/MaterialExpressions/MaterialExpressionTextureSampleParameterBlur.cpp
@@ -10,7 +10,7 @@
 
 #include UE_INLINE_GENERATED_CPP_BY_NAME(MaterialExpressionTextureSampleParameterBlur)
 
-#define LOCTEXT_NAMESPACE "MaterialExpressionTextureSampleBlur"
+#define LOCTEXT_NAMESPACE "MaterialExpressionMaterialXTextureSampleBlur"
 
 namespace Gauss
 {
@@ -78,31 +78,29 @@ namespace Box
 	}
 }
 
-UMaterialExpressionTextureSampleParameterBlur::UMaterialExpressionTextureSampleParameterBlur(const FObjectInitializer& ObjectInitializer)
+UMaterialExpressionMaterialXTextureSampleParameterBlur::UMaterialExpressionMaterialXTextureSampleParameterBlur(const FObjectInitializer& ObjectInitializer)
 	: Super(ObjectInitializer)
 {
 	// Structure to hold one-time initialization
 	struct FConstructorStatics
 	{
 		FText NAME_MaterialX;
-		FText NAME_Convolution;
 		FConstructorStatics()
 			: NAME_MaterialX(LOCTEXT("MaterialX", "MaterialX"))
-			, NAME_Convolution(LOCTEXT("Convolution", "Convolution"))
 		{}
 	};
 	static FConstructorStatics ConstructorStatics;
 
 #if WITH_EDITORONLY_DATA
+	MenuCategories.Empty();
 	MenuCategories.Add(ConstructorStatics.NAME_MaterialX);
-	MenuCategories.Add(ConstructorStatics.NAME_Convolution);
 #endif
 }
 
 #if WITH_EDITOR
-int32 UMaterialExpressionTextureSampleParameterBlur::Compile(FMaterialCompiler* Compiler, int32 OutputIndex)
+int32 UMaterialExpressionMaterialXTextureSampleParameterBlur::Compile(FMaterialCompiler* Compiler, int32 OutputIndex)
 {
-	if(KernelSize > 0)
+	if(KernelSize != EMAterialXTextureSampleBlurKernel ::Kernel1)
 	{
 		FExpressionInput TexCoordInput;
 		if(Coordinates.GetTracedInput().Expression)
@@ -149,43 +147,24 @@ int32 UMaterialExpressionTextureSampleParameterBlur::Compile(FMaterialCompiler*
 		);
 
 		const float* Kernel;
-		int32 FilterWidth = 1;
+		int32 FilterWidth;
 
-		if(KernelSize <= (1.f / 3.f))
-		{
-			FilterWidth = 3 / 2;
-			if(Filter == ETextureSampleBlurFilter::Box)
-			{
-				Kernel = Box::KernelWeights3x3;
-			}
-			else
-			{
-				Kernel = Gauss::KernelWeights3x3;
-			}
-		}
-		else if(KernelSize <= (2.f / 3.f))
+		switch(KernelSize)
 		{
+		case EMAterialXTextureSampleBlurKernel::Kernel5:
 			FilterWidth = 5 / 2;
-			if(Filter == ETextureSampleBlurFilter::Box)
-			{
-				Kernel = Box::KernelWeights5x5;
-			}
-			else
-			{
-				Kernel = Gauss::KernelWeights5x5;
-			}
-		}
-		else
-		{
+			Kernel = Filter == EMaterialXTextureSampleBlurFilter::Box ? Box::KernelWeights5x5 : Gauss::KernelWeights5x5;
+			break;
+
+		case EMAterialXTextureSampleBlurKernel::Kernel7:
 			FilterWidth = 7 / 2;
-			if(Filter == ETextureSampleBlurFilter::Box)
-			{
-				Kernel = Box::KernelWeights7x7;
-			}
-			else
-			{
-				Kernel = Gauss::KernelWeights7x7;
-			}
+			Kernel = Filter == EMaterialXTextureSampleBlurFilter::Box ? Box::KernelWeights7x7 : Gauss::KernelWeights7x7;
+			break;
+
+		default:
+			FilterWidth = 3 / 2;
+			Kernel = Filter == EMaterialXTextureSampleBlurFilter::Box ? Box::KernelWeights3x3 : Gauss::KernelWeights3x3;
+			break;
 		}
 
 		UMaterialExpressionCustom* Convolution = NewObject<UMaterialExpressionCustom>();
@@ -236,9 +215,9 @@ int32 UMaterialExpressionTextureSampleParameterBlur::Compile(FMaterialCompiler*
 	}
 }
 
-void UMaterialExpressionTextureSampleParameterBlur::GetCaption(TArray<FString>& OutCaptions) const
+void UMaterialExpressionMaterialXTextureSampleParameterBlur::GetCaption(TArray<FString>& OutCaptions) const
 {
-	OutCaptions.Add(TEXT("ParamBlur"));
+	OutCaptions.Add(TEXT("MaterialX ParamBlur"));
 	OutCaptions.Add(FString::Printf(TEXT("'%s'"), *ParameterName.ToString()));
 }
 #endif // WITH_EDITOR
diff --git a/Engine/Plugins/Interchange/Runtime/Source/Import/Private/MaterialX/MaterialExpressions/MaterialExpressionUnpremult.cpp b/Engine/Plugins/Interchange/Runtime/Source/Import/Private/MaterialX/MaterialExpressions/MaterialExpressionUnpremult.cpp
--- a/Engine/Plugins/Interchange/Runtime/Source/Import/Private/MaterialX/MaterialExpressions/MaterialExpressionUnpremult.cpp
+++ b/Engine/Plugins/Interchange/Runtime/Source/Import/Private/MaterialX/MaterialExpressions/MaterialExpressionUnpremult.cpp
@@ -4,31 +4,28 @@
 
 #include UE_INLINE_GENERATED_CPP_BY_NAME(MaterialExpressionUnpremult)
 
-#define LOCTEXT_NAMESPACE "MaterialExpressionUnpremult"
+#define LOCTEXT_NAMESPACE "MaterialExpressionMaterialXUnpremult"
 
-UMaterialExpressionUnpremult::UMaterialExpressionUnpremult(const FObjectInitializer& ObjectInitializer)
+UMaterialExpressionMaterialXUnpremult::UMaterialExpressionMaterialXUnpremult(const FObjectInitializer& ObjectInitializer)
 	: Super(ObjectInitializer)
 {
 	// Structure to hold one-time initialization
 	struct FConstructorStatics
 	{
 		FText NAME_MaterialX;
-		FText NAME_Compositing;
 		FConstructorStatics()
 			: NAME_MaterialX(LOCTEXT("MaterialX", "MaterialX"))
-			, NAME_Compositing(LOCTEXT("Compositing", "Compositing"))
 		{}
 	};
 	static FConstructorStatics ConstructorStatics;
 
 #if WITH_EDITORONLY_DATA
 	MenuCategories.Add(ConstructorStatics.NAME_MaterialX);
-	MenuCategories.Add(ConstructorStatics.NAME_Compositing);
 #endif
 }
 
 #if WITH_EDITOR
-int32 UMaterialExpressionUnpremult::Compile(FMaterialCompiler* Compiler, int32 OutputIndex)
+int32 UMaterialExpressionMaterialXUnpremult::Compile(FMaterialCompiler* Compiler, int32 OutputIndex)
 {
 	if(!Input.GetTracedInput().Expression)
 	{
@@ -43,9 +40,9 @@ int32 UMaterialExpressionUnpremult::Compile(FMaterialCompiler* Compiler, int32 O
 	return Compiler->If(Alpha, Compiler->Constant(0.f), Unpremult, IndexInput, Unpremult, Compiler->Constant(0.00001f));
 }
 
-void UMaterialExpressionUnpremult::GetCaption(TArray<FString>& OutCaptions) const
+void UMaterialExpressionMaterialXUnpremult::GetCaption(TArray<FString>& OutCaptions) const
 {
-	OutCaptions.Add(TEXT("Unpremult"));
+	OutCaptions.Add(TEXT("MaterialX Unpremult"));
 }
 #endif
 
diff --git a/Engine/Plugins/Interchange/Runtime/Source/Import/Private/MaterialX/MaterialExpressions/MaterialExpressionUnpremult.h b/Engine/Plugins/Interchange/Runtime/Source/Import/Private/MaterialX/MaterialExpressions/MaterialExpressionUnpremult.h
--- a/Engine/Plugins/Interchange/Runtime/Source/Import/Private/MaterialX/MaterialExpressions/MaterialExpressionUnpremult.h
+++ b/Engine/Plugins/Interchange/Runtime/Source/Import/Private/MaterialX/MaterialExpressions/MaterialExpressionUnpremult.h
@@ -12,18 +12,18 @@
  * Input must be of type float4
  */
 UCLASS(collapsecategories, hidecategories = Object, MinimalAPI)
-class UMaterialExpressionUnpremult : public UMaterialExpression
+class UMaterialExpressionMaterialXUnpremult : public UMaterialExpression
 {
 	GENERATED_UCLASS_BODY()
 
 	UPROPERTY()
 	FExpressionInput Input;
 
-	//~ Begin UMaterialExpression Interface
+	//~ Begin UMaterialExpressionMaterialX Interface
 #if WITH_EDITOR
 	virtual int32 Compile(class FMaterialCompiler* Compiler, int32 OutputIndex) override;
 	virtual void GetCaption(TArray<FString>& OutCaptions) const override;
 #endif
-	//~ End UMaterialExpression Interface
+	//~ End UMaterialExpressionMaterialX Interface
 };
 
diff --git a/Engine/Plugins/Interchange/Runtime/Source/Import/Private/MaterialX/MaterialXUtils/MaterialXManager.cpp b/Engine/Plugins/Interchange/Runtime/Source/Import/Private/MaterialX/MaterialXUtils/MaterialXManager.cpp
--- a/Engine/Plugins/Interchange/Runtime/Source/Import/Private/MaterialX/MaterialXUtils/MaterialXManager.cpp
+++ b/Engine/Plugins/Interchange/Runtime/Source/Import/Private/MaterialX/MaterialXUtils/MaterialXManager.cpp
@@ -101,66 +101,66 @@ FMaterialXManager::FMaterialXManager()
 		{MaterialX::Category::CrossProduct, TEXT("Crossproduct")},
 		{MaterialX::Category::Divide,       TEXT("Divide")},
 		{MaterialX::Category::DotProduct,   TEXT("Dotproduct")},
-		{MaterialX::Category::Exp,          TEXT("Exponential")},
+		{MaterialX::Category::Exp,          TEXT("MaterialXExponential")},
 		{MaterialX::Category::Floor,        TEXT("Floor")},
 		{MaterialX::Category::Invert,       TEXT("Subtract")},
-		{MaterialX::Category::Ln,           TEXT("Logarithm")},
-		{MaterialX::Category::Magnitude,    TEXT("Length")},
+		{MaterialX::Category::Ln,           TEXT("MaterialXLogarithm")},
+		{MaterialX::Category::Magnitude,    TEXT("MaterialXLength")},
 		{MaterialX::Category::Max,          TEXT("Max")},
 		{MaterialX::Category::Min,          TEXT("Min")},
 		{MaterialX::Category::Modulo,       TEXT("Fmod")},
 		{MaterialX::Category::Multiply,     TEXT("Multiply")},
 		{MaterialX::Category::Normalize,    TEXT("Normalize")},
-		{MaterialX::Category::Place2D,		TEXT("Place2D")},
+		{MaterialX::Category::Place2D,		TEXT("MaterialXPlace2D")},
 		{MaterialX::Category::Power,        TEXT("Power")},
-		{MaterialX::Category::Rotate2D,		TEXT("Rotate2D")},
-		{MaterialX::Category::RampLR,       TEXT("RampLeftRight")},
-		{MaterialX::Category::RampTB,       TEXT("RampTopBottom")},
+		{MaterialX::Category::Rotate2D,		TEXT("MaterialXRotate2D")},
+		{MaterialX::Category::RampLR,       TEXT("MaterialXRampLeftRight")},
+		{MaterialX::Category::RampTB,       TEXT("MaterialXRampTopBottom")},
 		{MaterialX::Category::Sign,         TEXT("Sign")},
 		{MaterialX::Category::Sin,          TEXT("Sine")},
-		{MaterialX::Category::SplitLR,      TEXT("SplitLeftRight")},
-		{MaterialX::Category::SplitTB,      TEXT("SplitTopBottom")},
+		{MaterialX::Category::SplitLR,      TEXT("MaterialXSplitLeftRight")},
+		{MaterialX::Category::SplitTB,      TEXT("MaterialXSplitTopBottom")},
 		{MaterialX::Category::Sqrt,         TEXT("SquareRoot")},
 		{MaterialX::Category::Sub,          TEXT("Subtract")},
 		{MaterialX::Category::Tan,          TEXT("Tangent")},
 		// Compositing nodes
-		{MaterialX::Category::Burn,         TEXT("Burn")},
-		{MaterialX::Category::Difference,   TEXT("Difference")},
-		{MaterialX::Category::Disjointover, TEXT("Disjointover")},
-		{MaterialX::Category::Dodge,        TEXT("Dodge")},
-		{MaterialX::Category::In,           TEXT("In")},
+		{MaterialX::Category::Burn,         TEXT("MaterialXBurn")},
+		{MaterialX::Category::Difference,   TEXT("MaterialXDifference")},
+		{MaterialX::Category::Disjointover, TEXT("MaterialXDisjointover")},
+		{MaterialX::Category::Dodge,        TEXT("MaterialXDodge")},
+		{MaterialX::Category::In,           TEXT("MaterialXIn")},
 		{MaterialX::Category::Inside,       TEXT("Multiply")},
-		{MaterialX::Category::Mask,         TEXT("Mask")},
-		{MaterialX::Category::Matte,        TEXT("Matte")},
-		{MaterialX::Category::Minus,        TEXT("Minus")},
+		{MaterialX::Category::Mask,         TEXT("MaterialXMask")},
+		{MaterialX::Category::Matte,        TEXT("MaterialXMatte")},
+		{MaterialX::Category::Minus,        TEXT("MaterialXMinus")},
 		{MaterialX::Category::Mix,          TEXT("Lerp")},
-		{MaterialX::Category::Out,          TEXT("Out")},
-		{MaterialX::Category::Over,         TEXT("Over")},
-		{MaterialX::Category::Overlay,      TEXT("Overlay")},
-		{MaterialX::Category::Plus,         TEXT("Plus")},
-		{MaterialX::Category::Premult,      TEXT("Premult")},
-		{MaterialX::Category::Screen,       TEXT("Screen")},
-		{MaterialX::Category::Unpremult,    TEXT("Unpremult")},
+		{MaterialX::Category::Out,          TEXT("MaterialXOut")},
+		{MaterialX::Category::Over,         TEXT("MaterialXOver")},
+		{MaterialX::Category::Overlay,      TEXT("MaterialXOverlay")},
+		{MaterialX::Category::Plus,         TEXT("MaterialXPlus")},
+		{MaterialX::Category::Premult,      TEXT("MaterialXPremult")},
+		{MaterialX::Category::Screen,       TEXT("MaterialXScreen")},
+		{MaterialX::Category::Unpremult,    TEXT("MaterialXUnpremult")},
 		// Channel nodes
 		{MaterialX::Category::Combine2,     TEXT("AppendVector")},
-		{MaterialX::Category::Combine3,     TEXT("Append3Vector")},
-		{MaterialX::Category::Combine4,     TEXT("Append4Vector")},
+		{MaterialX::Category::Combine3,     TEXT("MaterialXAppend3Vector")},
+		{MaterialX::Category::Combine4,     TEXT("MaterialXAppend4Vector")},
 		// Procedural2D nodes
-		{MaterialX::Category::Ramp4,        TEXT("Ramp4")},
-		{MaterialX::Category::RampLR,       TEXT("RampLeftRight")},
-		{MaterialX::Category::RampTB,       TEXT("RampTopBottom")},
-		{MaterialX::Category::SplitLR,      TEXT("SplitLeftRight")},
-		{MaterialX::Category::SplitTB,      TEXT("SplitTopBottom")},
+		{MaterialX::Category::Ramp4,        TEXT("MaterialXRamp4")},
+		{MaterialX::Category::RampLR,       TEXT("MaterialXRampLeftRight")},
+		{MaterialX::Category::RampTB,       TEXT("MaterialXRampTopBottom")},
+		{MaterialX::Category::SplitLR,      TEXT("MaterialXSplitLeftRight")},
+		{MaterialX::Category::SplitTB,      TEXT("MaterialXSplitTopBottom")},
 		// Procedural3D nodes
-		{MaterialX::Category::Fractal3D,    TEXT("Fractal3D")},
+		{MaterialX::Category::Fractal3D,    TEXT("MaterialXFractal3D")},
 		// Geometric nodes 
 		{MaterialX::Category::GeomColor,	TEXT("VertexColor")},
 		{MaterialX::Category::TexCoord,		TEXT("TextureCoordinate")},
 		// Adjustment nodes,
-		{MaterialX::Category::HsvToRgb,		TEXT("HsvToRgb")},
-		{MaterialX::Category::Luminance,    TEXT("Luminance")},
-		{MaterialX::Category::Remap,        TEXT("Remap")},
-		{MaterialX::Category::RgbToHsv,		TEXT("RgbToHsv")},
+		{MaterialX::Category::HsvToRgb,		TEXT("MaterialXHsvToRgb")},
+		{MaterialX::Category::Luminance,    TEXT("MaterialXLuminance")},
+		{MaterialX::Category::Remap,        TEXT("MaterialXRemap")},
+		{MaterialX::Category::RgbToHsv,		TEXT("MaterialXRgbToHsv")},
 		{MaterialX::Category::Saturate,		TEXT("Desaturation")},
 		{MaterialX::Category::Smoothstep,   TEXT("SmoothStep")}
 	}
diff --git a/Engine/Plugins/Interchange/Runtime/Source/Import/Private/MaterialX/MaterialXUtils/MaterialXSurfaceShaderAbstract.cpp b/Engine/Plugins/Interchange/Runtime/Source/Import/Private/MaterialX/MaterialXUtils/MaterialXSurfaceShaderAbstract.cpp
--- a/Engine/Plugins/Interchange/Runtime/Source/Import/Private/MaterialX/MaterialXUtils/MaterialXSurfaceShaderAbstract.cpp
+++ b/Engine/Plugins/Interchange/Runtime/Source/Import/Private/MaterialX/MaterialXUtils/MaterialXSurfaceShaderAbstract.cpp
@@ -13,6 +13,7 @@
 #include "Materials/MaterialExpressionTransformPosition.h"
 #include "Materials/MaterialExpressionVectorNoise.h"
 #include "MaterialX/MaterialXUtils/MaterialXManager.h"
+#include "MaterialX/MaterialExpressions/MaterialExpressionTextureSampleParameterBlur.h"
 
 #define LOCTEXT_NAMESPACE "MaterialXSurfaceShaderAbstract"
 
@@ -957,14 +958,30 @@ void FMaterialXSurfaceShaderAbstract::ConnectBlurInputToOutput(MaterialX::NodePt
 					// By default TextureSampleBox uses gaussian filter
 					if(InputKernel->getValueString() == "box")
 					{
-						TextureShaderNode->AddInt32Attribute(TextureSampleBlur::Attributes::Filter.ToString(), 0);
+						TextureShaderNode->AddInt32Attribute(TextureSampleBlur::Attributes::Filter.ToString(), int32(EMaterialXTextureSampleBlurFilter::Box));
 					}
 				}
 
 				if(mx::InputPtr InputKernel = UpstreamNode->getInput("size"))
 				{
-					float KernelSize = mx::fromValueString<float>(InputKernel->getValueString());
-					TextureShaderNode->AddFloatAttribute(TextureSampleBlur::Attributes::KernelSize.ToString(), KernelSize);
+					if(InputKernel->hasValueString())
+					{
+						float KernelSize = mx::fromValueString<float>(InputKernel->getValueString());
+						constexpr float Kernel1x1 = 0.f / 3.f;
+						constexpr float Kernel3x3 = 1.f / 3.f;
+						constexpr float Kernel5x5 = 2.f / 3.f;
+						constexpr float Kernel7x7 = 3.f / 3.f;
+						TextureShaderNode->AddInt32Attribute(TextureSampleBlur::Attributes::KernelSize.ToString(),
+															 FMath::IsNearlyEqual(KernelSize, Kernel1x1) ? int32(EMAterialXTextureSampleBlurKernel::Kernel1) :
+															 KernelSize <= Kernel3x3 ? int32(EMAterialXTextureSampleBlurKernel::Kernel3) :
+															 KernelSize <= Kernel5x5 ? int32(EMAterialXTextureSampleBlurKernel::Kernel5) :
+															 KernelSize <= Kernel7x7 ? int32(EMAterialXTextureSampleBlurKernel::Kernel7) :
+															 int32(EMAterialXTextureSampleBlurKernel::Kernel1));
+					}
+					else
+					{
+						UE_LOG(LogInterchangeImport, Warning, TEXT("<%s>: input 'size' must have a value"), ANSI_TO_TCHAR(UpstreamNode->getName().c_str()));
+					}
 				}
 			}
 			else
diff --git a/Engine/Plugins/Interchange/Runtime/Source/Import/Public/MaterialX/MaterialExpressions/MaterialExpressionTextureSampleParameterBlur.h b/Engine/Plugins/Interchange/Runtime/Source/Import/Public/MaterialX/MaterialExpressions/MaterialExpressionTextureSampleParameterBlur.h
--- a/Engine/Plugins/Interchange/Runtime/Source/Import/Public/MaterialX/MaterialExpressions/MaterialExpressionTextureSampleParameterBlur.h
+++ b/Engine/Plugins/Interchange/Runtime/Source/Import/Public/MaterialX/MaterialExpressions/MaterialExpressionTextureSampleParameterBlur.h
@@ -5,26 +5,29 @@
 #include "MaterialExpressionTextureSampleParameterBlur.generated.h"
 
 UENUM()
-enum class ETextureSampleBlurFilter : uint8
+enum class EMaterialXTextureSampleBlurFilter : uint8
 {
 	Box,
 	Gaussian
 };
 
+UENUM()
+enum class EMAterialXTextureSampleBlurKernel
+{
+	Kernel1 UMETA(DisplayName = "1x1", ToolTip="No Blur"),
+	Kernel3 UMETA(DisplayName = "3x3"),
+	Kernel5 UMETA(DisplayName = "5x5"),
+	Kernel7 UMETA(DisplayName = "7x7"),
+};
+
 UCLASS(MinimalAPI, collapsecategories, hidecategories = Object)
-class UMaterialExpressionTextureSampleParameterBlur: public UMaterialExpressionTextureSampleParameter2D
+class UMaterialExpressionMaterialXTextureSampleParameterBlur: public UMaterialExpressionTextureSampleParameter2D
 {
 	GENERATED_UCLASS_BODY()
 
-	/** The size of the blur kernel, relative to 0-1 UV space, only 3x3, 5x5 and 7x7 kernel are supported.
-	  * Basically depending on the KernelSize value:
-	  * KernelSize == 0   -> 1x1 Kernel (Meaning no Blur)
-	  * KernelSize <= 1/3 -> 3x3 Kernel
-	  * KernelSize <= 2/3 -> 5x5 Kernel
-	  * Otherwise         -> 7x7 Kernel
-	  */
+	/** The size of the blur kernel, relative to 0-1 UV space. */
 	UPROPERTY(EditAnywhere, Category = Filtering, Meta = (ShowAsInputPin = "Advanced"))
-	float KernelSize = 0.f;
+	EMAterialXTextureSampleBlurKernel KernelSize = EMAterialXTextureSampleBlurKernel::Kernel1;
 
 	/** Size of the filter when we sample a texture coordinate */
 	UPROPERTY(EditAnywhere, Category = Filtering, Meta = (ShowAsInputPin = "Advanced"))
@@ -36,12 +39,12 @@ class UMaterialExpressionTextureSampleParameterBlur: public UMaterialExpressionT
 
 	/** Filter to use when we blur a Texture: Gaussian or Box Linear filter*/
 	UPROPERTY(EditAnywhere, Category = Filtering, Meta = (ShowAsInputPin = "Advanced"))
-	ETextureSampleBlurFilter Filter = ETextureSampleBlurFilter::Gaussian;
+	EMaterialXTextureSampleBlurFilter Filter = EMaterialXTextureSampleBlurFilter::Gaussian;
 
-	//~ Begin UMaterialExpression Interface
+	//~ Begin UMaterialExpressionMaterialX Interface
 #if WITH_EDITOR
 	virtual int32 Compile(class FMaterialCompiler* Compiler, int32 OutputIndex) override;
 	virtual void GetCaption(TArray<FString>& OutCaptions) const override;
 #endif
-	//~ End UMaterialExpression Interface
+	//~ End UMaterialExpressionMaterialX Interface
 };
\ No newline at end of file
diff --git a/Engine/Plugins/Interchange/Runtime/Source/Pipelines/Private/InterchangeGenericMaterialPipeline.cpp b/Engine/Plugins/Interchange/Runtime/Source/Pipelines/Private/InterchangeGenericMaterialPipeline.cpp
--- a/Engine/Plugins/Interchange/Runtime/Source/Pipelines/Private/InterchangeGenericMaterialPipeline.cpp
+++ b/Engine/Plugins/Interchange/Runtime/Source/Pipelines/Private/InterchangeGenericMaterialPipeline.cpp
@@ -1695,7 +1695,7 @@ void UInterchangeGenericMaterialPipeline::HandleTextureSampleNode(const UInterch
 		}
 		else if(TextureNode->IsA<UInterchangeTextureBlurNode>())
 		{
-			ExpressionClassName = UMaterialExpressionTextureSampleParameterBlur::StaticClass()->GetName();
+			ExpressionClassName = UMaterialExpressionMaterialXTextureSampleParameterBlur::StaticClass()->GetName();
 		}
 		else if (TextureNode->IsA<UInterchangeTexture2DNode>())
 		{
@@ -1726,7 +1726,7 @@ void UInterchangeGenericMaterialPipeline::HandleTextureSampleNode(const UInterch
 		}
 	}
 
-	if(ExpressionClassName == UMaterialExpressionTextureSampleParameterBlur::StaticClass()->GetName())
+	if(ExpressionClassName == UMaterialExpressionMaterialXTextureSampleParameterBlur::StaticClass()->GetName())
 	{
 		HandleTextureSampleBlurNode(ShaderNode, MaterialFactoryNode, TextureSampleFactoryNode);
 	}
@@ -1737,35 +1737,35 @@ void UInterchangeGenericMaterialPipeline::HandleTextureSampleBlurNode(const UInt
 	using namespace UE::Interchange::Materials::Standard::Nodes;
 
 	// KernelSize
-	if(float KernelSize; ShaderNode->GetFloatAttribute(TextureSampleBlur::Attributes::KernelSize.ToString(), KernelSize))
+	if(int32 KernelSize; ShaderNode->GetInt32Attribute(TextureSampleBlur::Attributes::KernelSize.ToString(), KernelSize))
 	{
-		const FName KernelSizeMemberName = GET_MEMBER_NAME_CHECKED(UMaterialExpressionTextureSampleParameterBlur, KernelSize);
-		TextureSampleFactoryNode->AddFloatAttribute(KernelSizeMemberName.ToString(), KernelSize);
-		TextureSampleFactoryNode->AddApplyAndFillDelegates<float>(KernelSizeMemberName.ToString(), UMaterialExpressionTextureSampleParameterBlur::StaticClass(), KernelSizeMemberName);
+		const FName KernelSizeMemberName = GET_MEMBER_NAME_CHECKED(UMaterialExpressionMaterialXTextureSampleParameterBlur, KernelSize);
+		TextureSampleFactoryNode->AddInt32Attribute(KernelSizeMemberName.ToString(), KernelSize);
+		TextureSampleFactoryNode->AddApplyAndFillDelegates<int32>(KernelSizeMemberName.ToString(), UMaterialExpressionMaterialXTextureSampleParameterBlur::StaticClass(), KernelSizeMemberName);
 	}
 
 	// FilterSize
 	if(float FilterSize; ShaderNode->GetFloatAttribute(TextureSampleBlur::Attributes::FilterSize.ToString(), FilterSize))
 	{
-		const FName FilterSizeMemberName = GET_MEMBER_NAME_CHECKED(UMaterialExpressionTextureSampleParameterBlur, FilterSize);
+		const FName FilterSizeMemberName = GET_MEMBER_NAME_CHECKED(UMaterialExpressionMaterialXTextureSampleParameterBlur, FilterSize);
 		TextureSampleFactoryNode->AddFloatAttribute(FilterSizeMemberName.ToString(), FilterSize);
-		TextureSampleFactoryNode->AddApplyAndFillDelegates<float>(FilterSizeMemberName.ToString(), UMaterialExpressionTextureSampleParameterBlur::StaticClass(), FilterSizeMemberName);
+		TextureSampleFactoryNode->AddApplyAndFillDelegates<float>(FilterSizeMemberName.ToString(), UMaterialExpressionMaterialXTextureSampleParameterBlur::StaticClass(), FilterSizeMemberName);
 	}
 
 	// FilterOffset
 	if(float FilterOffset; ShaderNode->GetFloatAttribute(TextureSampleBlur::Attributes::FilterOffset.ToString(), FilterOffset))
 	{
-		const FName FilterOffsetMemberName = GET_MEMBER_NAME_CHECKED(UMaterialExpressionTextureSampleParameterBlur, FilterOffset);
+		const FName FilterOffsetMemberName = GET_MEMBER_NAME_CHECKED(UMaterialExpressionMaterialXTextureSampleParameterBlur, FilterOffset);
 		TextureSampleFactoryNode->AddFloatAttribute(FilterOffsetMemberName.ToString(), FilterOffset);
-		TextureSampleFactoryNode->AddApplyAndFillDelegates<float>(FilterOffsetMemberName.ToString(), UMaterialExpressionTextureSampleParameterBlur::StaticClass(), FilterOffsetMemberName);
+		TextureSampleFactoryNode->AddApplyAndFillDelegates<float>(FilterOffsetMemberName.ToString(), UMaterialExpressionMaterialXTextureSampleParameterBlur::StaticClass(), FilterOffsetMemberName);
 	}
 
 	// Filter
 	if(int32 Filter; ShaderNode->GetInt32Attribute(TextureSampleBlur::Attributes::Filter.ToString(), Filter))
 	{
-		const FName FilterMemberName = GET_MEMBER_NAME_CHECKED(UMaterialExpressionTextureSampleParameterBlur, Filter);
+		const FName FilterMemberName = GET_MEMBER_NAME_CHECKED(UMaterialExpressionMaterialXTextureSampleParameterBlur, Filter);
 		TextureSampleFactoryNode->AddInt32Attribute(FilterMemberName.ToString(), Filter);
-		TextureSampleFactoryNode->AddApplyAndFillDelegates<int>(FilterMemberName.ToString(), UMaterialExpressionTextureSampleParameterBlur::StaticClass(), FilterMemberName);
+		TextureSampleFactoryNode->AddApplyAndFillDelegates<int32>(FilterMemberName.ToString(), UMaterialExpressionMaterialXTextureSampleParameterBlur::StaticClass(), FilterMemberName);
 	}
 }
 
