[Interchange - glTF] Move GLTFCore module from GLTFImporter plugin to Interchange.
[Interchange - glTF] Remove deprecated glTF importers and relevant CVar.
Re-instated includes from UE_ENABLE_INCLUDE_ORDER_DEPRECATED_IN_5_2 for GLTFCore

diff --git a/Engine/Plugins/Editor/GLTFImporter/GLTFImporter.uplugin b/Engine/Plugins/Editor/GLTFImporter/GLTFImporter.uplugin
deleted file mode 100644
diff --git a/Engine/Plugins/Editor/GLTFImporter/Source/GLTFCore/GLTFCore.Build.cs b/Engine/Plugins/Editor/GLTFImporter/Source/GLTFCore/GLTFCore.Build.cs
deleted file mode 100644
diff --git a/Engine/Plugins/Editor/GLTFImporter/Source/GLTFCore/Private/GLTF/ConversionUtilities.h b/Engine/Plugins/Editor/GLTFImporter/Source/GLTFCore/Private/GLTF/ConversionUtilities.h
deleted file mode 100644
diff --git a/Engine/Plugins/Editor/GLTFImporter/Source/GLTFCore/Private/GLTF/ExtensionsHandler.cpp b/Engine/Plugins/Editor/GLTFImporter/Source/GLTFCore/Private/GLTF/ExtensionsHandler.cpp
deleted file mode 100644
diff --git a/Engine/Plugins/Editor/GLTFImporter/Source/GLTFCore/Private/GLTF/ExtensionsHandler.h b/Engine/Plugins/Editor/GLTFImporter/Source/GLTFCore/Private/GLTF/ExtensionsHandler.h
deleted file mode 100644
diff --git a/Engine/Plugins/Editor/GLTFImporter/Source/GLTFCore/Private/GLTF/GLTFAccessor.cpp b/Engine/Plugins/Editor/GLTFImporter/Source/GLTFCore/Private/GLTF/GLTFAccessor.cpp
deleted file mode 100644
diff --git a/Engine/Plugins/Editor/GLTFImporter/Source/GLTFCore/Private/GLTF/GLTFAsset.cpp b/Engine/Plugins/Editor/GLTFImporter/Source/GLTFCore/Private/GLTF/GLTFAsset.cpp
deleted file mode 100644
diff --git a/Engine/Plugins/Editor/GLTFImporter/Source/GLTFCore/Private/GLTF/GLTFBinaryReader.cpp b/Engine/Plugins/Editor/GLTFImporter/Source/GLTFCore/Private/GLTF/GLTFBinaryReader.cpp
deleted file mode 100644
diff --git a/Engine/Plugins/Editor/GLTFImporter/Source/GLTFCore/Private/GLTF/GLTFBinaryReader.h b/Engine/Plugins/Editor/GLTFImporter/Source/GLTFCore/Private/GLTF/GLTFBinaryReader.h
deleted file mode 100644
diff --git a/Engine/Plugins/Editor/GLTFImporter/Source/GLTFCore/Private/GLTF/GLTFMaterial.cpp b/Engine/Plugins/Editor/GLTFImporter/Source/GLTFCore/Private/GLTF/GLTFMaterial.cpp
deleted file mode 100644
diff --git a/Engine/Plugins/Editor/GLTFImporter/Source/GLTFCore/Private/GLTF/GLTFMesh.cpp b/Engine/Plugins/Editor/GLTFImporter/Source/GLTFCore/Private/GLTF/GLTFMesh.cpp
deleted file mode 100644
diff --git a/Engine/Plugins/Editor/GLTFImporter/Source/GLTFCore/Private/GLTF/GLTFReader.cpp b/Engine/Plugins/Editor/GLTFImporter/Source/GLTFCore/Private/GLTF/GLTFReader.cpp
deleted file mode 100644
diff --git a/Engine/Plugins/Editor/GLTFImporter/Source/GLTFCore/Private/GLTF/GLTFTexture.cpp b/Engine/Plugins/Editor/GLTFImporter/Source/GLTFCore/Private/GLTF/GLTFTexture.cpp
deleted file mode 100644
diff --git a/Engine/Plugins/Editor/GLTFImporter/Source/GLTFCore/Private/GLTF/JsonUtilities.h b/Engine/Plugins/Editor/GLTFImporter/Source/GLTFCore/Private/GLTF/JsonUtilities.h
deleted file mode 100644
diff --git a/Engine/Plugins/Editor/GLTFImporter/Source/GLTFCore/Private/GLTF/MaterialUtilities.h b/Engine/Plugins/Editor/GLTFImporter/Source/GLTFCore/Private/GLTF/MaterialUtilities.h
deleted file mode 100644
diff --git a/Engine/Plugins/Editor/GLTFImporter/Source/GLTFCore/Private/GLTFCoreModule.cpp b/Engine/Plugins/Editor/GLTFImporter/Source/GLTFCore/Private/GLTFCoreModule.cpp
deleted file mode 100644
diff --git a/Engine/Plugins/Editor/GLTFImporter/Source/GLTFCore/Private/GLTFMapFactory.cpp b/Engine/Plugins/Editor/GLTFImporter/Source/GLTFCore/Private/GLTFMapFactory.cpp
deleted file mode 100644
diff --git a/Engine/Plugins/Editor/GLTFImporter/Source/GLTFCore/Private/GLTFMapFactory.h b/Engine/Plugins/Editor/GLTFImporter/Source/GLTFCore/Private/GLTFMapFactory.h
deleted file mode 100644
diff --git a/Engine/Plugins/Editor/GLTFImporter/Source/GLTFCore/Private/GLTFMaterialExpressions.cpp b/Engine/Plugins/Editor/GLTFImporter/Source/GLTFCore/Private/GLTFMaterialExpressions.cpp
deleted file mode 100644
diff --git a/Engine/Plugins/Editor/GLTFImporter/Source/GLTFCore/Private/GLTFMaterialFactory.cpp b/Engine/Plugins/Editor/GLTFImporter/Source/GLTFCore/Private/GLTFMaterialFactory.cpp
deleted file mode 100644
diff --git a/Engine/Plugins/Editor/GLTFImporter/Source/GLTFCore/Private/GLTFMeshFactory.cpp b/Engine/Plugins/Editor/GLTFImporter/Source/GLTFCore/Private/GLTFMeshFactory.cpp
deleted file mode 100644
diff --git a/Engine/Plugins/Editor/GLTFImporter/Source/GLTFCore/Public/GLTFAccessor.h b/Engine/Plugins/Editor/GLTFImporter/Source/GLTFCore/Public/GLTFAccessor.h
deleted file mode 100644
diff --git a/Engine/Plugins/Editor/GLTFImporter/Source/GLTFCore/Public/GLTFAnimation.h b/Engine/Plugins/Editor/GLTFImporter/Source/GLTFCore/Public/GLTFAnimation.h
deleted file mode 100644
diff --git a/Engine/Plugins/Editor/GLTFImporter/Source/GLTFCore/Public/GLTFAsset.h b/Engine/Plugins/Editor/GLTFImporter/Source/GLTFCore/Public/GLTFAsset.h
deleted file mode 100644
diff --git a/Engine/Plugins/Editor/GLTFImporter/Source/GLTFCore/Public/GLTFCoreModule.h b/Engine/Plugins/Editor/GLTFImporter/Source/GLTFCore/Public/GLTFCoreModule.h
deleted file mode 100644
diff --git a/Engine/Plugins/Editor/GLTFImporter/Source/GLTFCore/Public/GLTFLogger.h b/Engine/Plugins/Editor/GLTFImporter/Source/GLTFCore/Public/GLTFLogger.h
deleted file mode 100644
diff --git a/Engine/Plugins/Editor/GLTFImporter/Source/GLTFCore/Public/GLTFMaterial.h b/Engine/Plugins/Editor/GLTFImporter/Source/GLTFCore/Public/GLTFMaterial.h
deleted file mode 100644
diff --git a/Engine/Plugins/Editor/GLTFImporter/Source/GLTFCore/Public/GLTFMaterialExpressions.h b/Engine/Plugins/Editor/GLTFImporter/Source/GLTFCore/Public/GLTFMaterialExpressions.h
deleted file mode 100644
diff --git a/Engine/Plugins/Editor/GLTFImporter/Source/GLTFCore/Public/GLTFMaterialFactory.h b/Engine/Plugins/Editor/GLTFImporter/Source/GLTFCore/Public/GLTFMaterialFactory.h
deleted file mode 100644
diff --git a/Engine/Plugins/Editor/GLTFImporter/Source/GLTFCore/Public/GLTFMesh.h b/Engine/Plugins/Editor/GLTFImporter/Source/GLTFCore/Public/GLTFMesh.h
deleted file mode 100644
diff --git a/Engine/Plugins/Editor/GLTFImporter/Source/GLTFCore/Public/GLTFMeshFactory.h b/Engine/Plugins/Editor/GLTFImporter/Source/GLTFCore/Public/GLTFMeshFactory.h
deleted file mode 100644
diff --git a/Engine/Plugins/Editor/GLTFImporter/Source/GLTFCore/Public/GLTFNode.h b/Engine/Plugins/Editor/GLTFImporter/Source/GLTFCore/Public/GLTFNode.h
deleted file mode 100644
diff --git a/Engine/Plugins/Editor/GLTFImporter/Source/GLTFCore/Public/GLTFReader.h b/Engine/Plugins/Editor/GLTFImporter/Source/GLTFCore/Public/GLTFReader.h
deleted file mode 100644
diff --git a/Engine/Plugins/Editor/GLTFImporter/Source/GLTFCore/Public/GLTFTexture.h b/Engine/Plugins/Editor/GLTFImporter/Source/GLTFCore/Public/GLTFTexture.h
deleted file mode 100644
diff --git a/Engine/Plugins/Editor/GLTFImporter/Source/GLTFImporter/GLTFImporter.Build.cs b/Engine/Plugins/Editor/GLTFImporter/Source/GLTFImporter/GLTFImporter.Build.cs
deleted file mode 100644
diff --git a/Engine/Plugins/Editor/GLTFImporter/Source/GLTFImporter/Private/GLTFImportFactory.cpp b/Engine/Plugins/Editor/GLTFImporter/Source/GLTFImporter/Private/GLTFImportFactory.cpp
deleted file mode 100644
diff --git a/Engine/Plugins/Editor/GLTFImporter/Source/GLTFImporter/Private/GLTFImportFactory.h b/Engine/Plugins/Editor/GLTFImporter/Source/GLTFImporter/Private/GLTFImportFactory.h
deleted file mode 100644
diff --git a/Engine/Plugins/Editor/GLTFImporter/Source/GLTFImporter/Private/GLTFImportOptions.cpp b/Engine/Plugins/Editor/GLTFImporter/Source/GLTFImporter/Private/GLTFImportOptions.cpp
deleted file mode 100644
diff --git a/Engine/Plugins/Editor/GLTFImporter/Source/GLTFImporter/Private/GLTFImporterContext.cpp b/Engine/Plugins/Editor/GLTFImporter/Source/GLTFImporter/Private/GLTFImporterContext.cpp
deleted file mode 100644
diff --git a/Engine/Plugins/Editor/GLTFImporter/Source/GLTFImporter/Private/GLTFImporterContext.h b/Engine/Plugins/Editor/GLTFImporter/Source/GLTFImporter/Private/GLTFImporterContext.h
deleted file mode 100644
diff --git a/Engine/Plugins/Editor/GLTFImporter/Source/GLTFImporter/Private/GLTFImporterModule.cpp b/Engine/Plugins/Editor/GLTFImporter/Source/GLTFImporter/Private/GLTFImporterModule.cpp
deleted file mode 100644
diff --git a/Engine/Plugins/Editor/GLTFImporter/Source/GLTFImporter/Private/GLTFMaterialElement.cpp b/Engine/Plugins/Editor/GLTFImporter/Source/GLTFImporter/Private/GLTFMaterialElement.cpp
deleted file mode 100644
diff --git a/Engine/Plugins/Editor/GLTFImporter/Source/GLTFImporter/Private/GLTFMaterialElement.h b/Engine/Plugins/Editor/GLTFImporter/Source/GLTFImporter/Private/GLTFMaterialElement.h
deleted file mode 100644
diff --git a/Engine/Plugins/Editor/GLTFImporter/Source/GLTFImporter/Private/GLTFStaticMeshFactory.cpp b/Engine/Plugins/Editor/GLTFImporter/Source/GLTFImporter/Private/GLTFStaticMeshFactory.cpp
deleted file mode 100644
diff --git a/Engine/Plugins/Editor/GLTFImporter/Source/GLTFImporter/Private/GLTFTextureFactory.cpp b/Engine/Plugins/Editor/GLTFImporter/Source/GLTFImporter/Private/GLTFTextureFactory.cpp
deleted file mode 100644
diff --git a/Engine/Plugins/Editor/GLTFImporter/Source/GLTFImporter/Private/GLTFTextureFactory.h b/Engine/Plugins/Editor/GLTFImporter/Source/GLTFImporter/Private/GLTFTextureFactory.h
deleted file mode 100644
diff --git a/Engine/Plugins/Editor/GLTFImporter/Source/GLTFImporter/Private/UI/GLTFOptionsWindow.cpp b/Engine/Plugins/Editor/GLTFImporter/Source/GLTFImporter/Private/UI/GLTFOptionsWindow.cpp
deleted file mode 100644
diff --git a/Engine/Plugins/Editor/GLTFImporter/Source/GLTFImporter/Private/UI/GLTFOptionsWindow.h b/Engine/Plugins/Editor/GLTFImporter/Source/GLTFImporter/Private/UI/GLTFOptionsWindow.h
deleted file mode 100644
diff --git a/Engine/Plugins/Editor/GLTFImporter/Source/GLTFImporter/Public/GLTFImportOptions.h b/Engine/Plugins/Editor/GLTFImporter/Source/GLTFImporter/Public/GLTFImportOptions.h
deleted file mode 100644
diff --git a/Engine/Plugins/Editor/GLTFImporter/Source/GLTFImporter/Public/GLTFImporterModule.h b/Engine/Plugins/Editor/GLTFImporter/Source/GLTFImporter/Public/GLTFImporterModule.h
deleted file mode 100644
diff --git a/Engine/Plugins/Editor/GLTFImporter/Source/GLTFImporter/Public/GLTFStaticMeshFactory.h b/Engine/Plugins/Editor/GLTFImporter/Source/GLTFImporter/Public/GLTFStaticMeshFactory.h
deleted file mode 100644
diff --git a/Engine/Plugins/Enterprise/DatasmithGLTFImporter/DatasmithGLTFImporter.uplugin b/Engine/Plugins/Enterprise/DatasmithGLTFImporter/DatasmithGLTFImporter.uplugin
deleted file mode 100644
diff --git a/Engine/Plugins/Enterprise/DatasmithGLTFImporter/Source/DatasmithGLTFTranslator/DatasmithGLTFTranslator.Build.cs b/Engine/Plugins/Enterprise/DatasmithGLTFImporter/Source/DatasmithGLTFTranslator/DatasmithGLTFTranslator.Build.cs
deleted file mode 100644
diff --git a/Engine/Plugins/Enterprise/DatasmithGLTFImporter/Source/DatasmithGLTFTranslator/Private/DatasmithGLTFAnimationImporter.cpp b/Engine/Plugins/Enterprise/DatasmithGLTFImporter/Source/DatasmithGLTFTranslator/Private/DatasmithGLTFAnimationImporter.cpp
deleted file mode 100644
diff --git a/Engine/Plugins/Enterprise/DatasmithGLTFImporter/Source/DatasmithGLTFTranslator/Private/DatasmithGLTFAnimationImporter.h b/Engine/Plugins/Enterprise/DatasmithGLTFImporter/Source/DatasmithGLTFTranslator/Private/DatasmithGLTFAnimationImporter.h
deleted file mode 100644
diff --git a/Engine/Plugins/Enterprise/DatasmithGLTFImporter/Source/DatasmithGLTFTranslator/Private/DatasmithGLTFImportOptions.h b/Engine/Plugins/Enterprise/DatasmithGLTFImporter/Source/DatasmithGLTFTranslator/Private/DatasmithGLTFImportOptions.h
deleted file mode 100644
diff --git a/Engine/Plugins/Enterprise/DatasmithGLTFImporter/Source/DatasmithGLTFTranslator/Private/DatasmithGLTFImporter.cpp b/Engine/Plugins/Enterprise/DatasmithGLTFImporter/Source/DatasmithGLTFTranslator/Private/DatasmithGLTFImporter.cpp
deleted file mode 100644
diff --git a/Engine/Plugins/Enterprise/DatasmithGLTFImporter/Source/DatasmithGLTFTranslator/Private/DatasmithGLTFImporter.h b/Engine/Plugins/Enterprise/DatasmithGLTFImporter/Source/DatasmithGLTFTranslator/Private/DatasmithGLTFImporter.h
deleted file mode 100644
diff --git a/Engine/Plugins/Enterprise/DatasmithGLTFImporter/Source/DatasmithGLTFTranslator/Private/DatasmithGLTFMaterialElement.cpp b/Engine/Plugins/Enterprise/DatasmithGLTFImporter/Source/DatasmithGLTFTranslator/Private/DatasmithGLTFMaterialElement.cpp
deleted file mode 100644
diff --git a/Engine/Plugins/Enterprise/DatasmithGLTFImporter/Source/DatasmithGLTFTranslator/Private/DatasmithGLTFMaterialElement.h b/Engine/Plugins/Enterprise/DatasmithGLTFImporter/Source/DatasmithGLTFTranslator/Private/DatasmithGLTFMaterialElement.h
deleted file mode 100644
diff --git a/Engine/Plugins/Enterprise/DatasmithGLTFImporter/Source/DatasmithGLTFTranslator/Private/DatasmithGLTFTextureFactory.cpp b/Engine/Plugins/Enterprise/DatasmithGLTFImporter/Source/DatasmithGLTFTranslator/Private/DatasmithGLTFTextureFactory.cpp
deleted file mode 100644
diff --git a/Engine/Plugins/Enterprise/DatasmithGLTFImporter/Source/DatasmithGLTFTranslator/Private/DatasmithGLTFTextureFactory.h b/Engine/Plugins/Enterprise/DatasmithGLTFImporter/Source/DatasmithGLTFTranslator/Private/DatasmithGLTFTextureFactory.h
deleted file mode 100644
diff --git a/Engine/Plugins/Enterprise/DatasmithGLTFImporter/Source/DatasmithGLTFTranslator/Private/DatasmithGLTFTranslator.cpp b/Engine/Plugins/Enterprise/DatasmithGLTFImporter/Source/DatasmithGLTFTranslator/Private/DatasmithGLTFTranslator.cpp
deleted file mode 100644
diff --git a/Engine/Plugins/Enterprise/DatasmithGLTFImporter/Source/DatasmithGLTFTranslator/Private/DatasmithGLTFTranslator.h b/Engine/Plugins/Enterprise/DatasmithGLTFImporter/Source/DatasmithGLTFTranslator/Private/DatasmithGLTFTranslator.h
deleted file mode 100644
diff --git a/Engine/Plugins/Enterprise/DatasmithGLTFImporter/Source/DatasmithGLTFTranslator/Private/DatasmithGLTFTranslatorModule.cpp b/Engine/Plugins/Enterprise/DatasmithGLTFImporter/Source/DatasmithGLTFTranslator/Private/DatasmithGLTFTranslatorModule.cpp
deleted file mode 100644
diff --git a/Engine/Plugins/Enterprise/DatasmithGLTFImporter/Source/DatasmithGLTFTranslator/Public/DatasmithGLTFTranslatorModule.h b/Engine/Plugins/Enterprise/DatasmithGLTFImporter/Source/DatasmithGLTFTranslator/Public/DatasmithGLTFTranslatorModule.h
deleted file mode 100644
diff --git a/Engine/Plugins/Experimental/Enterprise/DatasmithInterchange/Source/DatasmithInterchange/Private/InterchangeDatasmithLevelPipeline.cpp b/Engine/Plugins/Experimental/Enterprise/DatasmithInterchange/Source/DatasmithInterchange/Private/InterchangeDatasmithLevelPipeline.cpp
--- a/Engine/Plugins/Experimental/Enterprise/DatasmithInterchange/Source/DatasmithInterchange/Private/InterchangeDatasmithLevelPipeline.cpp
+++ b/Engine/Plugins/Experimental/Enterprise/DatasmithInterchange/Source/DatasmithInterchange/Private/InterchangeDatasmithLevelPipeline.cpp
@@ -119,7 +119,7 @@ void UInterchangeDatasmithLevelPipeline::ExecutePostImportPipeline(const UInterc
 						FString ActorNodeUid;
 						if (TransformTrackNode->GetCustomActorDependencyUid(ActorNodeUid))
 						{
-							const FString ActorFactoryNodeUid = TEXT("Factory_") + ActorNodeUid;
+							const FString ActorFactoryNodeUid = UInterchangeFactoryBaseNode::BuildFactoryNodeUid(ActorNodeUid);
 							const UInterchangeFactoryBaseNode* ActorFactoryNode = Cast<UInterchangeFactoryBaseNode>(NodeContainer->GetNode(ActorFactoryNodeUid));
 
 							if (ActorFactoryNode)
diff --git a/Engine/Plugins/Experimental/Enterprise/DatasmithInterchange/Source/DatasmithInterchange/Private/InterchangeDatasmithPipeline.cpp b/Engine/Plugins/Experimental/Enterprise/DatasmithInterchange/Source/DatasmithInterchange/Private/InterchangeDatasmithPipeline.cpp
--- a/Engine/Plugins/Experimental/Enterprise/DatasmithInterchange/Source/DatasmithInterchange/Private/InterchangeDatasmithPipeline.cpp
+++ b/Engine/Plugins/Experimental/Enterprise/DatasmithInterchange/Source/DatasmithInterchange/Private/InterchangeDatasmithPipeline.cpp
@@ -133,7 +133,7 @@ void UInterchangeDatasmithPipeline::ExecutePipeline(UInterchangeBaseNodeContaine
 
 	// Datasmith Scene
 	{
-		const FString DatasmithSceneUid = TEXT("Factory_") + DatasmithSceneNode->GetUniqueID();
+		const FString DatasmithSceneUid = UInterchangeFactoryBaseNode::BuildFactoryNodeUid(DatasmithSceneNode->GetUniqueID());
 		const FString DisplayLabel = DatasmithSceneNode->GetDisplayLabel();
 		UInterchangeDatasmithSceneFactoryNode* DatasmithSceneFactoryNode = NewObject<UInterchangeDatasmithSceneFactoryNode>(BaseNodeContainer, NAME_None);
 		if (!ensure(DatasmithSceneFactoryNode))
diff --git a/Engine/Plugins/Experimental/Enterprise/DatasmithRuntime/DatasmithRuntime.uplugin b/Engine/Plugins/Experimental/Enterprise/DatasmithRuntime/DatasmithRuntime.uplugin
--- a/Engine/Plugins/Experimental/Enterprise/DatasmithRuntime/DatasmithRuntime.uplugin
+++ b/Engine/Plugins/Experimental/Enterprise/DatasmithRuntime/DatasmithRuntime.uplugin
@@ -29,10 +29,6 @@
 		{
 			"Name": "UdpMessaging",
 			"Enabled": true
-		},
-		{
-			"Name": "DatasmithGLTFImporter",
-			"Enabled": true
 		}
 	]
 }
\ No newline at end of file
diff --git a/Engine/Plugins/Experimental/Enterprise/DatasmithRuntime/Source/DatasmithRuntime.Build.cs b/Engine/Plugins/Experimental/Enterprise/DatasmithRuntime/Source/DatasmithRuntime.Build.cs
--- a/Engine/Plugins/Experimental/Enterprise/DatasmithRuntime/Source/DatasmithRuntime.Build.cs
+++ b/Engine/Plugins/Experimental/Enterprise/DatasmithRuntime/Source/DatasmithRuntime.Build.cs
@@ -16,7 +16,6 @@ public class DatasmithRuntime : ModuleRules
 				"DatasmithContent",
 				"DatasmithCore",
 				"DatasmithNativeTranslator",
-				"DatasmithGLTFTranslator",
 				"DatasmithTranslator",
 			}
 		);
diff --git a/Engine/Plugins/Experimental/Enterprise/DatasmithRuntime/Source/Private/DatasmithRuntimeModule.cpp b/Engine/Plugins/Experimental/Enterprise/DatasmithRuntime/Source/Private/DatasmithRuntimeModule.cpp
--- a/Engine/Plugins/Experimental/Enterprise/DatasmithRuntime/Source/Private/DatasmithRuntimeModule.cpp
+++ b/Engine/Plugins/Experimental/Enterprise/DatasmithRuntime/Source/Private/DatasmithRuntimeModule.cpp
@@ -9,7 +9,6 @@
 #include "MaterialSelectors/DatasmithRuntimeMaterialSelector.h"
 
 #include "DatasmithTranslatorModule.h"
-#include "DatasmithGLTFTranslatorModule.h"
 #include "ReferenceMaterials/DatasmithReferenceMaterialManager.h"
 
 #if WITH_EDITOR
@@ -26,7 +25,6 @@ public:
 	{
 		// Verify DatasmithTranslatorModule has been loaded
 		check(IDatasmithTranslatorModule::IsAvailable());
-		ensure(IDatasmithGLTFTranslatorModule::IsAvailable());
 
 #if WITH_EDITOR
 		// If don't have any active references to our materials they won't be packaged into monolithic builds, and we wouldn't
diff --git a/Engine/Plugins/Interchange/Runtime/Interchange.uplugin b/Engine/Plugins/Interchange/Runtime/Interchange.uplugin
--- a/Engine/Plugins/Interchange/Runtime/Interchange.uplugin
+++ b/Engine/Plugins/Interchange/Runtime/Interchange.uplugin
@@ -16,6 +16,11 @@
 	"Installed": false,
 	"IsHidden": true,
 	"Modules": [
+		{
+			"Name": "GLTFCore",
+			"Type": "Runtime",
+			"LoadingPhase": "Default"
+		},
 		{
 			"Name": "InterchangeNodes",
 			"Type": "Runtime",
@@ -72,10 +77,6 @@
 		}
 	],
 	"Plugins": [
-		{
-			"Name": "GLTFImporter",
-			"Enabled": true
-		},
 		{
 			"Name": "VariantManagerContent",
 			"Enabled": true
diff --git a/Engine/Plugins/Interchange/Runtime/Source/Import/Private/Gltf/InterchangeGltfTranslator.cpp b/Engine/Plugins/Interchange/Runtime/Source/Import/Private/Gltf/InterchangeGltfTranslator.cpp
--- a/Engine/Plugins/Interchange/Runtime/Source/Import/Private/Gltf/InterchangeGltfTranslator.cpp
+++ b/Engine/Plugins/Interchange/Runtime/Source/Import/Private/Gltf/InterchangeGltfTranslator.cpp
@@ -45,13 +45,6 @@
 
 #define LOCTEXT_NAMESPACE "InterchangeGltfTranslator"
 
-static bool GInterchangeEnableGLTFImport = true;
-static FAutoConsoleVariableRef CCvarInterchangeEnableGLTFImport(
-	TEXT("Interchange.FeatureFlags.Import.GLTF"),
-	GInterchangeEnableGLTFImport,
-	TEXT("Whether glTF support is enabled."),
-	ECVF_Default);
-
 static int GInterchangeGLTFUseMaterialInstancing = 0;
 static FAutoConsoleVariableRef CCvarInterchangeGLTFUseMaterialInstancing(
 	TEXT("Interchange.FeatureFlags.Import.GLTF.UseMaterialInstancing"),
@@ -968,12 +961,9 @@ TArray<FString> UInterchangeGltfTranslator::GetSupportedFormats() const
 {
 	TArray<FString> GltfExtensions;
 
-	if ( GInterchangeEnableGLTFImport || GIsAutomationTesting )
-	{
-		GltfExtensions.Reserve(2);
-		GltfExtensions.Add(TEXT("gltf;GL Transmission Format"));
-		GltfExtensions.Add(TEXT("glb;GL Transmission Format (Binary)"));
-	}
+	GltfExtensions.Reserve(2);
+	GltfExtensions.Add(TEXT("gltf;GL Transmission Format"));
+	GltfExtensions.Add(TEXT("glb;GL Transmission Format (Binary)"));
 
 	return GltfExtensions;
 }
diff --git a/Engine/Plugins/Interchange/Runtime/Source/Parsers/GLTFCore/GLTFCore.Build.cs b/Engine/Plugins/Interchange/Runtime/Source/Parsers/GLTFCore/GLTFCore.Build.cs
new file mode 100644
--- /dev/null
+++ b/Engine/Plugins/Interchange/Runtime/Source/Parsers/GLTFCore/GLTFCore.Build.cs
@@ -0,0 +1,30 @@
+// Copyright Epic Games, Inc. All Rights Reserved.
+
+namespace UnrealBuildTool.Rules
+{
+    public class GLTFCore : ModuleRules
+    {
+        public GLTFCore(ReadOnlyTargetRules Target) : base(Target)
+        {
+            PublicDependencyModuleNames.AddRange(
+                new string[]
+                {
+                    "Core",
+                    "CoreUObject",
+                }
+                );
+
+            PrivateDependencyModuleNames.AddRange(
+                new string[]
+                {
+                    "Engine",
+                    "MeshDescription",
+					"StaticMeshDescription",
+                    "Json",
+                    "RenderCore",
+					"HTTP",
+                }
+                );
+        }
+    }
+}
diff --git a/Engine/Plugins/Interchange/Runtime/Source/Parsers/GLTFCore/Private/GLTF/ConversionUtilities.h b/Engine/Plugins/Interchange/Runtime/Source/Parsers/GLTFCore/Private/GLTF/ConversionUtilities.h
new file mode 100644
--- /dev/null
+++ b/Engine/Plugins/Interchange/Runtime/Source/Parsers/GLTFCore/Private/GLTF/ConversionUtilities.h
@@ -0,0 +1,237 @@
+// Copyright Epic Games, Inc. All Rights Reserved.
+
+#pragma once
+
+
+#include "GLTFAnimation.h"
+#include "GLTFMaterial.h"
+#include "GLTFMesh.h"
+#include "GLTFTexture.h"
+#include "Math/Quat.h"
+
+namespace GLTF
+{
+	// math conversion
+
+	inline FVector ConvertVec3(const FVector& Vec)
+	{
+		// glTF uses a right-handed coordinate system, with Y up.
+		// Unreal uses a left-handed coordinate system, with Z up.
+		return {Vec.X, Vec.Z, Vec.Y};
+	}
+
+	inline FVector4 ConvertTangent(const FVector4& Tangent)
+	{
+		// glTF stores tangent as Vec4, with W component indicating handedness of tangent basis.
+
+		FVector4 Result(ConvertVec3(FVector(Tangent.X, Tangent.Y, Tangent.Z)), Tangent.W);
+		return Result;
+	}
+
+	template<typename T>
+	UE::Math::TQuat<T> ConvertQuat(const UE::Math::TQuat<T>& Quat)
+	{
+		// glTF uses a right-handed coordinate system, with Y up.
+		// Unreal uses a left-handed coordinate system, with Z up.
+		// Quat = (qX, qY, qZ, qW) = (sin(angle/2) * aX, sin(angle/2) * aY, sin(angle/2) * aZ, cons(angle/2))
+		// where (aX, aY, aZ) - rotation axis, angle - rotation angle
+		// Y swapped with Z between these coordinate systems
+		// also, as handedness is changed rotation is inversed - hence negation
+		// therefore QuatUE = (-qX, -qZ, -qY, qw)
+
+		UE::Math::TQuat<T> Result(-Quat.X, -Quat.Z, -Quat.Y, Quat.W);
+		// Not checking if quaternion is normalized
+		// e.g. some sources use non-unit Quats for rotation tangents
+		return Result;
+	}
+
+	inline FMatrix ConvertMat(const FMatrix& Matrix)
+	{
+		// glTF stores matrix elements in column major order
+		// Unreal's FMatrix is row major
+
+		FMatrix Result;
+		for (int32 Row = 0; Row < 4; ++Row)
+		{
+			for (int32 Col = 0; Col < 4; ++Col)
+			{
+				Result.M[Row][Col] = Matrix.M[Col][Row];
+			}
+		}
+		return Result;
+	}
+
+	// enum conversion
+
+	inline FMaterial::EAlphaMode AlphaModeFromString(const FString& Str)
+	{
+		// case sensitive comparison
+		if (FPlatformString::Strcmp(*Str, TEXT("OPAQUE")) == 0)
+		{
+			return FMaterial::EAlphaMode::Opaque;
+		}
+		if (FPlatformString::Strcmp(*Str, TEXT("MASK")) == 0)
+		{
+			return FMaterial::EAlphaMode::Mask;
+		}
+		if (FPlatformString::Strcmp(*Str, TEXT("BLEND")) == 0)
+		{
+			return FMaterial::EAlphaMode::Blend;
+		}
+		ensure(false);
+		return FMaterial::EAlphaMode::Opaque;
+	}
+
+	inline FAnimation::EPath AnimationPathFromString(const FString& Str)
+	{
+		// case sensitive comparison
+		if (FPlatformString::Strcmp(*Str, TEXT("rotation")) == 0)
+		{
+			return FAnimation::EPath::Rotation;
+		}
+		if (FPlatformString::Strcmp(*Str, TEXT("scale")) == 0)
+		{
+			return FAnimation::EPath::Scale;
+		}
+		if (FPlatformString::Strcmp(*Str, TEXT("translation")) == 0)
+		{
+			return FAnimation::EPath::Translation;
+		}
+		if (FPlatformString::Strcmp(*Str, TEXT("weights")) == 0)
+		{
+			return FAnimation::EPath::Weights;
+		}
+		ensure(false);
+		return FAnimation::EPath::Rotation;
+	}
+
+	inline FPrimitive::EMode PrimitiveModeFromNumber(uint32 Num)
+	{
+		static const TArray<FPrimitive::EMode> SafetyCheck({
+		    FPrimitive::EMode::Points,
+		    FPrimitive::EMode::LineLoop,
+		    FPrimitive::EMode::Lines,
+		    FPrimitive::EMode::LineStrip,
+		    FPrimitive::EMode::Triangles,
+		    FPrimitive::EMode::TriangleStrip,
+		    FPrimitive::EMode::TriangleFan,
+		});
+
+		if (!ensure(SafetyCheck.Find(static_cast<FPrimitive::EMode>(Num)) != INDEX_NONE))
+		{
+			return FPrimitive::EMode::Unknown;
+		}
+
+		return static_cast<FPrimitive::EMode>(Num);
+	}
+
+	inline FImage::EFormat ImageFormatFromFilename(const FString& Filename)
+	{
+		// attempt to guess from filename extension
+		if (Filename.EndsWith(TEXT(".png")))
+		{
+			return FImage::EFormat::PNG;
+		}
+		else if (Filename.EndsWith(TEXT(".jpg")) || Filename.EndsWith(TEXT(".jpeg")))
+		{
+			return FImage::EFormat::JPEG;
+		}
+
+		return FImage::EFormat::Unknown;
+	}
+
+	inline FImage::EFormat ImageFormatFromMimeType(const FString& Str)
+	{
+		// case sensitive comparison
+		if (FPlatformString::Strcmp(*Str, TEXT("image/jpeg")) == 0 || FPlatformString::Strcmp(*Str, TEXT("image/jpg")) == 0)
+		{
+			return FImage::EFormat::JPEG;
+		}
+		if (FPlatformString::Strcmp(*Str, TEXT("image/png")) == 0)
+		{
+			return FImage::EFormat::PNG;
+		}
+
+		return FImage::EFormat::Unknown;
+	}
+
+	inline FSampler::EFilter FilterFromNumber(uint16 Num)
+	{
+		static const TArray<FSampler::EFilter> SafetyCheck({FSampler::EFilter::Nearest, FSampler::EFilter::Linear,
+		                                                    FSampler::EFilter::NearestMipmapNearest, FSampler::EFilter::LinearMipmapNearest,
+		                                                    FSampler::EFilter::NearestMipmapLinear, FSampler::EFilter::LinearMipmapLinear});
+		if (!ensure(SafetyCheck.Find(static_cast<FSampler::EFilter>(Num)) != INDEX_NONE))
+		{
+			return FSampler::EFilter::Linear;
+		}
+
+		return static_cast<FSampler::EFilter>(Num);
+	}
+
+	inline FSampler::EWrap WrapModeFromNumber(uint16 Num)
+	{
+		static const TArray<FSampler::EWrap> SafetyCheck({FSampler::EWrap::ClampToEdge, FSampler::EWrap::MirroredRepeat, FSampler::EWrap::Repeat});
+		if (!ensure(SafetyCheck.Find(static_cast<FSampler::EWrap>(Num)) != INDEX_NONE))
+		{
+			return FSampler::EWrap::Repeat;
+		}
+
+		return static_cast<FSampler::EWrap>(Num);
+	}
+
+	inline FAccessor::EType AccessorTypeFromString(const FString& Str)
+	{
+		// case sensitive comparison
+		if (FPlatformString::Strcmp(*Str, TEXT("SCALAR")) == 0)
+		{
+			return FAccessor::EType::Scalar;
+		}
+		if (FPlatformString::Strcmp(*Str, TEXT("VEC2")) == 0)
+		{
+			return FAccessor::EType::Vec2;
+		}
+		if (FPlatformString::Strcmp(*Str, TEXT("VEC3")) == 0)
+		{
+			return FAccessor::EType::Vec3;
+		}
+		if (FPlatformString::Strcmp(*Str, TEXT("VEC4")) == 0)
+		{
+			return FAccessor::EType::Vec4;
+		}
+		if (FPlatformString::Strcmp(*Str, TEXT("MAT2")) == 0)
+		{
+			return FAccessor::EType::Mat2;
+		}
+		if (FPlatformString::Strcmp(*Str, TEXT("MAT3")) == 0)
+		{
+			return FAccessor::EType::Mat3;
+		}
+		if (FPlatformString::Strcmp(*Str, TEXT("MAT4")) == 0)
+		{
+			return FAccessor::EType::Mat4;
+		}
+		return FAccessor::EType::Unknown;
+	}
+
+	inline FAccessor::EComponentType ComponentTypeFromNumber(uint16 Num)
+	{
+		switch (Num)
+		{
+			case 5120:
+				return FAccessor::EComponentType::S8;
+			case 5121:
+				return FAccessor::EComponentType::U8;
+			case 5122:
+				return FAccessor::EComponentType::S16;
+			case 5123:
+				return FAccessor::EComponentType::U16;
+			case 5125:
+				return FAccessor::EComponentType::U32;
+			case 5126:
+				return FAccessor::EComponentType::F32;
+			default:
+				return FAccessor::EComponentType::None;
+		}
+	}
+
+}  // namespace GLTF
diff --git a/Engine/Plugins/Interchange/Runtime/Source/Parsers/GLTFCore/Private/GLTF/ExtensionsHandler.cpp b/Engine/Plugins/Interchange/Runtime/Source/Parsers/GLTFCore/Private/GLTF/ExtensionsHandler.cpp
new file mode 100644
--- /dev/null
+++ b/Engine/Plugins/Interchange/Runtime/Source/Parsers/GLTFCore/Private/GLTF/ExtensionsHandler.cpp
@@ -0,0 +1,476 @@
+// Copyright Epic Games, Inc. All Rights Reserved.
+
+#include "ExtensionsHandler.h"
+
+#include "GLTF/JsonUtilities.h"
+#include "GLTFAsset.h"
+#include "GLTFNode.h"
+#include "MaterialUtilities.h"
+
+namespace GLTF
+{
+	namespace
+	{
+		static const TArray<FString> LightExtensions = { GLTF::ToString(GLTF::EExtension::KHR_LightsPunctual), GLTF::ToString(GLTF::EExtension::KHR_Lights) };
+		static const TArray<FString> MaterialsExtensions = { GLTF::ToString(GLTF::EExtension::KHR_MaterialsVariants) };
+	
+		TSharedPtr<FJsonObject> GetLightExtension(const TSharedPtr<FJsonObject>& Object)
+		{
+			if (!Object)
+				return nullptr;
+
+			TSharedPtr<FJsonObject> LightsObj;
+			if (Object->HasTypedField<EJson::Object>(LightExtensions[0]))
+			{
+				LightsObj = Object->GetObjectField(LightExtensions[0]);
+			}
+			else if (Object->HasTypedField<EJson::Object>(LightExtensions[1]))
+			{
+				LightsObj = Object->GetObjectField(LightExtensions[1]);
+			}
+			return LightsObj;
+		}
+
+		TSharedPtr<FJsonObject> GetExtensions(const FJsonObject& Object)
+		{
+			if (!Object.HasTypedField<EJson::Object>(TEXT("extensions")))
+			{
+				return nullptr;
+			}
+
+			return Object.GetObjectField(TEXT("extensions"));
+		}
+	}
+
+	FExtensionsHandler::FExtensionsHandler(TArray<FLogMessage>& InMessages)
+	    : Messages(InMessages)
+	    , Asset(nullptr)
+	{
+	}
+
+	void FExtensionsHandler::SetupAssetExtensions(const FJsonObject& Object) const
+	{
+		const TSharedPtr<FJsonObject>& ExtensionsObj = GetExtensions(Object);
+
+		if (!ExtensionsObj.IsValid())
+		{
+			return;
+		}
+
+		// lights
+
+		if (TSharedPtr<FJsonObject> LightsObj = GetLightExtension(ExtensionsObj))
+		{
+			Asset->ProcessedExtensions.Add(EExtension::KHR_LightsPunctual);
+
+			uint32 LightCount = ArraySize(*LightsObj, TEXT("lights"));
+			if (LightCount > 0)
+			{
+				Asset->Lights.Reserve(LightCount);
+				for (const TSharedPtr<FJsonValue>& Value : LightsObj->GetArrayField(TEXT("lights")))
+				{
+					SetupLightPunctual(*Value->AsObject());
+				}
+			}
+		}
+
+		// variants
+
+		if (ExtensionsObj->HasTypedField<EJson::Object>(GLTF::ToString(GLTF::EExtension::KHR_MaterialsVariants)))
+		{
+			Asset->ProcessedExtensions.Add(EExtension::KHR_MaterialsVariants);
+
+			TSharedPtr<FJsonObject> VariantsObj = ExtensionsObj->GetObjectField(GLTF::ToString(GLTF::EExtension::KHR_MaterialsVariants));
+			uint32 VariantsCount = ArraySize(*VariantsObj, TEXT("variants"));
+			if (VariantsCount > 0)
+			{
+				Asset->Variants.Reserve(VariantsCount);
+				for (const TSharedPtr<FJsonValue>& Value : VariantsObj->GetArrayField(TEXT("variants")))
+				{
+					const TSharedPtr<FJsonObject>& NameObj = Value->AsObject();
+					const FString Name = NameObj->GetStringField(TEXT("name"));
+					Asset->Variants.Add(Name);
+				}
+			}
+		}
+		
+		TArray<FString> SupportedExtensions;
+		SupportedExtensions.Append(LightExtensions);
+		SupportedExtensions.Append(MaterialsExtensions);
+
+		CheckExtensions(Object, SupportedExtensions);
+	}
+
+	void FExtensionsHandler::SetupMaterialExtensions(const FJsonObject& Object, FMaterial& Material) const
+	{
+		if (!Object.HasTypedField<EJson::Object>(TEXT("extensions")))
+		{
+			return;
+		}
+
+		static const TArray<EExtension> Extensions = { 
+			EExtension::KHR_MaterialsPbrSpecularGlossiness,
+			EExtension::KHR_MaterialsUnlit,
+			EExtension::KHR_MaterialsClearCoat,
+			EExtension::KHR_MaterialsTransmission,
+			EExtension::KHR_MaterialsSheen,
+			EExtension::KHR_MaterialsIOR,
+			EExtension::KHR_MaterialsSpecular,
+			EExtension::KHR_MaterialsEmissiveStrength,
+			EExtension::MSFT_PackingOcclusionRoughnessMetallic,
+			EExtension::MSFT_PackingNormalRoughnessMetallic 
+		};
+		TArray<FString> ExtensionsStringified;
+		for (size_t ExtensionIndex = 0; ExtensionIndex < Extensions.Num(); ExtensionIndex++)
+		{
+			ExtensionsStringified.Add(GLTF::ToString(Extensions[ExtensionIndex]));
+		}
+
+		const FJsonObject& ExtensionsObj = *Object.GetObjectField(TEXT("extensions"));
+		for (int32 Index = 0; Index < Extensions.Num(); ++Index)
+		{
+			const FString ExtensionName = ExtensionsStringified[Index];
+			if (!ExtensionsObj.HasTypedField<EJson::Object>(ExtensionName))
+				continue;
+
+			const FJsonObject& ExtObj = *ExtensionsObj.GetObjectField(ExtensionName);
+
+			const EExtension Extension = Extensions[Index];
+			switch (Extension)
+			{
+				case EExtension::KHR_MaterialsPbrSpecularGlossiness:
+				{
+					const FJsonObject& PBR = ExtObj;
+					GLTF::SetTextureMap(PBR, TEXT("diffuseTexture"), nullptr, Asset->Textures, Material.BaseColor, Messages);
+					Material.BaseColorFactor = FVector4f(GetVec4(PBR, TEXT("diffuseFactor"), FVector4(1.0f, 1.0f, 1.0f, 1.0f)));
+
+					GLTF::SetTextureMap(PBR, TEXT("specularGlossinessTexture"), nullptr, Asset->Textures, Material.SpecularGlossiness.Map, Messages);
+					Material.SpecularGlossiness.SpecularFactor   = GetVec3(PBR, TEXT("specularFactor"), FVector(1.0f));
+					Material.SpecularGlossiness.GlossinessFactor = GetScalar(PBR, TEXT("glossinessFactor"), 1.0f);
+
+					Material.ShadingModel = FMaterial::EShadingModel::SpecularGlossiness;
+
+					Asset->ProcessedExtensions.Add(EExtension::KHR_MaterialsPbrSpecularGlossiness);
+				}
+				break;
+				case EExtension::KHR_MaterialsUnlit:
+				{
+					Material.bIsUnlitShadingModel = true;
+					Asset->ProcessedExtensions.Add(EExtension::KHR_MaterialsUnlit);
+				}
+				break;
+				case EExtension::KHR_MaterialsClearCoat:
+				{
+					const FJsonObject& ClearCoat = ExtObj;
+
+					Material.bHasClearCoat = true;
+
+					Material.ClearCoat.ClearCoatFactor = GetScalar(ClearCoat, TEXT("clearcoatFactor"), 1.0f);
+					GLTF::SetTextureMap(ClearCoat, TEXT("clearcoatTexture"), nullptr, Asset->Textures, Material.ClearCoat.ClearCoatMap, Messages);
+
+					Material.ClearCoat.Roughness = GetScalar(ClearCoat, TEXT("clearcoatRoughnessFactor"), 0.0f);
+					GLTF::SetTextureMap(ClearCoat, TEXT("clearcoatRoughnessTexture"), nullptr, Asset->Textures, Material.ClearCoat.RoughnessMap, Messages);
+
+					Material.ClearCoat.NormalMapUVScale = GLTF::SetTextureMap(ClearCoat, TEXT("clearcoatNormalTexture"), TEXT("scale"), Asset->Textures, Material.ClearCoat.NormalMap, Messages);
+
+					Asset->ProcessedExtensions.Add(EExtension::KHR_MaterialsClearCoat);
+				}
+				break;
+				case EExtension::KHR_MaterialsTransmission:
+				{
+					if (Material.AlphaMode == FMaterial::EAlphaMode::Blend)
+					{
+						break;
+					}
+					const FJsonObject& Transm = ExtObj;
+
+					Material.bHasTransmission = true;
+
+					Material.Transmission.TransmissionFactor = GetScalar(Transm, TEXT("transmissionFactor"), 0.0f);
+					GLTF::SetTextureMap(Transm, TEXT("transmissionTexture"), nullptr, Asset->Textures, Material.Transmission.TransmissionMap, Messages);
+
+					Asset->ProcessedExtensions.Add(EExtension::KHR_MaterialsTransmission);
+				}
+				break;
+				case EExtension::KHR_MaterialsSheen:
+				{
+					const FJsonObject& Sheen = ExtObj;
+
+					Material.bHasSheen = true;
+
+					Material.Sheen.SheenColorFactor = GetVec3(Sheen, TEXT("sheenColorFactor"));
+					GLTF::SetTextureMap(Sheen, TEXT("sheenColorTexture"), nullptr, Asset->Textures, Material.Sheen.SheenColorMap, Messages);
+
+					Material.Sheen.SheenRoughnessFactor = GetScalar(Sheen, TEXT("sheenRoughnessFactor"));
+					GLTF::SetTextureMap(Sheen, TEXT("sheenRoughnessTexture"), nullptr, Asset->Textures, Material.Sheen.SheenRoughnessMap, Messages);
+
+					Asset->ProcessedExtensions.Add(EExtension::KHR_MaterialsSheen);
+				}
+				break;
+				case EExtension::KHR_MaterialsIOR:
+				{
+					const FJsonObject& IOR = ExtObj;
+
+					Material.bHasIOR = true;
+					Material.IOR = GetScalar(IOR, TEXT("ior"), 1.0f);
+
+					Asset->ProcessedExtensions.Add(EExtension::KHR_MaterialsIOR);
+				}
+				break;
+				case EExtension::KHR_MaterialsSpecular:
+				{
+					const FJsonObject& Specular = ExtObj;
+
+					Material.bHasSpecular = true;
+					Material.Specular.SpecularFactor = GetScalar(Specular, TEXT("specularFactor"), 1.0f);
+					Material.Specular.SpecularColorFactor = GetVec3(Specular, TEXT("specularColorFactor"));
+					GLTF::SetTextureMap(Specular, TEXT("specularTexture"), nullptr, Asset->Textures, Material.Specular.SpecularMap, Messages);
+					GLTF::SetTextureMap(Specular, TEXT("specularColorTexture"), nullptr, Asset->Textures, Material.Specular.SpecularColorMap, Messages);
+
+					Asset->ProcessedExtensions.Add(EExtension::KHR_MaterialsSpecular);
+				}
+				break;
+				case EExtension::KHR_MaterialsEmissiveStrength:
+				{
+					const FJsonObject& Emissive = ExtObj;
+
+					Material.bHasEmissiveStrength = true;
+					Material.EmissiveStrength = GetScalar(Emissive, TEXT("emissiveStrength"), 1.0f);
+
+					Asset->ProcessedExtensions.Add(EExtension::KHR_MaterialsEmissiveStrength);
+				}
+				break;
+				case EExtension::MSFT_PackingOcclusionRoughnessMetallic:
+				{
+					const FJsonObject& Packing = ExtObj;
+					if (GLTF::SetTextureMap(Packing, TEXT("occlusionRoughnessMetallicTexture"), nullptr, Asset->Textures, Material.Packing.Map, Messages))
+					{
+						Material.Packing.Flags = (int)GLTF::FMaterial::EPackingFlags::OcclusionRoughnessMetallic;
+					}
+					else if (GLTF::SetTextureMap(Packing, TEXT("roughnessMetallicOcclusionTexture"), nullptr, Asset->Textures, Material.Packing.Map, Messages))
+					{
+						Material.Packing.Flags = (int)GLTF::FMaterial::EPackingFlags::RoughnessMetallicOcclusion;
+					}
+					if (GLTF::SetTextureMap(Packing, TEXT("normalTexture"), nullptr, Asset->Textures, Material.Packing.NormalMap, Messages))
+					{
+						// can have extra packed two channel(RG) normal map
+						Material.Packing.Flags |= (int)GLTF::FMaterial::EPackingFlags::NormalRG;
+					}
+
+					if (Material.Packing.Flags != (int)GLTF::FMaterial::EPackingFlags::None)
+						Asset->ProcessedExtensions.Add(EExtension::MSFT_PackingOcclusionRoughnessMetallic);
+				}
+				break;
+				case EExtension::MSFT_PackingNormalRoughnessMetallic:
+				{
+					const FJsonObject& Packing = ExtObj;
+					if (GLTF::SetTextureMap(Packing, TEXT("normalRoughnessMetallicTexture"), nullptr, Asset->Textures, Material.Packing.Map, Messages))
+					{
+						Material.Packing.NormalMap = Material.Packing.Map;
+						Material.Packing.Flags     = (int)GLTF::FMaterial::EPackingFlags::NormalRoughnessMetallic;
+						Asset->ProcessedExtensions.Add(EExtension::MSFT_PackingNormalRoughnessMetallic);
+					}
+				}
+				break;
+				default:
+					if (!ensure(false))
+					{
+						Messages.Emplace(EMessageSeverity::Warning, FString::Printf(TEXT("Material.Extension not supported: %s"), *ToString(Extension)));
+					}
+					break;
+			}
+		}
+
+		CheckExtensions(Object, ExtensionsStringified);
+	}
+
+	void FExtensionsHandler::SetupBufferExtensions(const FJsonObject& Object, FBuffer& Buffer) const
+	{
+		static const TArray<FString> Extensions;
+		CheckExtensions(Object, Extensions);
+	}
+
+	void FExtensionsHandler::SetupBufferViewExtensions(const FJsonObject& Object, FBufferView& BufferView) const
+	{
+		static const TArray<FString> Extensions;
+		CheckExtensions(Object, Extensions);
+	}
+
+	void FExtensionsHandler::SetupAccessorExtensions(const FJsonObject& Object, FAccessor& Accessor) const
+	{
+		static const TArray<FString> Extensions;
+		CheckExtensions(Object, Extensions);
+	}
+
+	void FExtensionsHandler::SetupPrimitiveExtensions(const FJsonObject& Object, FPrimitive& Primitive) const
+	{
+		if (!Object.HasTypedField<EJson::Object>(TEXT("extensions")))
+		{
+			return;
+		}
+
+		static const TArray<EExtension> Extensions = {
+			EExtension::KHR_MaterialsVariants
+		};
+		TArray<FString> ExtensionsStringified;
+		for (size_t ExtensionIndex = 0; ExtensionIndex < Extensions.Num(); ExtensionIndex++)
+		{
+			ExtensionsStringified.Add(GLTF::ToString(Extensions[ExtensionIndex]));
+		}
+
+		const FJsonObject& ExtensionsObj = *Object.GetObjectField(TEXT("extensions"));
+		for (int32 Index = 0; Index < Extensions.Num(); ++Index)
+		{
+			const FString ExtensionName = ExtensionsStringified[Index];
+			if (!ExtensionsObj.HasTypedField<EJson::Object>(ExtensionName))
+				continue;
+
+			const FJsonObject& ExtObj = *ExtensionsObj.GetObjectField(ExtensionName);
+
+			const EExtension Extension = Extensions[Index];
+			switch (Extension)
+			{
+				case EExtension::KHR_MaterialsVariants:
+				{
+					const TArray<TSharedPtr<FJsonValue>>& Mappings = ExtObj.GetArrayField(TEXT("mappings"));
+
+					for (const TSharedPtr<FJsonValue>& Mapping : Mappings)
+					{
+						FVariantMapping& VariantMapping = Primitive.VariantMappings.Emplace_GetRef();
+
+						const TSharedPtr<FJsonObject> MappingObj = Mapping->AsObject();
+						VariantMapping.MaterialIndex = MappingObj->GetIntegerField(TEXT("material"));
+						const TArray<TSharedPtr<FJsonValue>>& Variants = MappingObj->GetArrayField(TEXT("variants"));
+
+						VariantMapping.VariantIndices.Reserve(Variants.Num());
+						for (const TSharedPtr<FJsonValue>& Variant : Variants)
+						{
+							const int32 VariantIndex = static_cast<int32>(Variant->AsNumber());
+							VariantMapping.VariantIndices.Add(VariantIndex);
+						}
+					}
+
+					Asset->ProcessedExtensions.Add(EExtension::KHR_MaterialsVariants);
+				}
+				break;
+				default:
+					if (!ensure(false))
+					{
+						Messages.Emplace(EMessageSeverity::Warning, FString::Printf(TEXT("Primitive.Extension not supported: %s"), *ToString(Extension)));
+					}
+					break;
+			}
+		}
+
+		CheckExtensions(Object, ExtensionsStringified);
+	}
+
+	void FExtensionsHandler::SetupMeshExtensions(const FJsonObject& Object, FMesh& Mesh) const
+	{
+		static const TArray<FString> Extensions;
+		CheckExtensions(Object, Extensions);
+	}
+
+	void FExtensionsHandler::SetupSceneExtensions(const FJsonObject& Object, FScene& Scene) const
+	{
+		static const TArray<FString> Extensions;
+		CheckExtensions(Object, Extensions);
+	}
+
+	void FExtensionsHandler::SetupNodeExtensions(const FJsonObject& Object, FNode& Node) const
+	{
+		const TSharedPtr<FJsonObject>& ExtensionsObj = GetExtensions(Object);
+		if (TSharedPtr<FJsonObject> LightsObj = GetLightExtension(ExtensionsObj))
+		{
+			Node.LightIndex = GetIndex(*LightsObj, TEXT("light"));
+		}
+
+		static const TArray<FString> Extensions = LightExtensions;
+		CheckExtensions(Object, Extensions);
+	}
+
+	void FExtensionsHandler::SetupCameraExtensions(const FJsonObject& Object, FCamera& Camera) const
+	{
+		static const TArray<FString> Extensions;
+		CheckExtensions(Object, Extensions);
+	}
+
+	void FExtensionsHandler::SetupSkinExtensions(const FJsonObject& Object, FSkinInfo& Skin) const
+	{
+		static const TArray<FString> Extensions;
+		CheckExtensions(Object, Extensions);
+	}
+
+	void FExtensionsHandler::SetupAnimationExtensions(const FJsonObject& Object, FAnimation& Animation) const
+	{
+		static const TArray<FString> Extensions;
+		CheckExtensions(Object, Extensions);
+	}
+
+	void FExtensionsHandler::SetupImageExtensions(const FJsonObject& Object, FImage& Image) const
+	{
+		static const TArray<FString> Extensions;
+		CheckExtensions(Object, Extensions);
+	}
+
+	void FExtensionsHandler::SetupSamplerExtensions(const FJsonObject& Object, FSampler& Sampler) const
+	{
+		static const TArray<FString> Extensions;
+		CheckExtensions(Object, Extensions);
+	}
+
+	void FExtensionsHandler::SetupTextureExtensions(const FJsonObject& Object, FTexture& Texture) const
+	{
+		static const TArray<FString> Extensions;
+		CheckExtensions(Object, Extensions);
+	}
+
+	void FExtensionsHandler::CheckExtensions(const FJsonObject& Object, const TArray<FString>& ExtensionsSupported) const
+	{
+		if (!Object.HasTypedField<EJson::Object>(TEXT("extensions")))
+		{
+			return;
+		}
+
+		const FJsonObject& ExtensionsObj = *Object.GetObjectField(TEXT("extensions"));
+		for (const auto& StrValuePair : ExtensionsObj.Values)
+		{
+			if (ExtensionsSupported.Find(StrValuePair.Get<0>()) == INDEX_NONE)
+			{
+				Messages.Emplace(EMessageSeverity::Warning, FString::Printf(TEXT("Extension is not supported: %s"), *StrValuePair.Get<0>()));
+			}
+		}
+	}
+
+	void FExtensionsHandler::SetupLightPunctual(const FJsonObject& Object) const
+	{
+		const uint32 LightIndex = Asset->Lights.Num();
+		const FNode* Found      = Asset->Nodes.FindByPredicate([LightIndex](const FNode& Node) { return LightIndex == Node.LightIndex; });
+
+		FLight& Light = Asset->Lights.Emplace_GetRef(Found);
+
+		Light.Name      = GetString(Object, TEXT("name"));
+		Light.Color     = GetVec3(Object, TEXT("color"), FVector(1.f));
+		Light.Intensity = GetScalar(Object, TEXT("intensity"), 1.f);
+		Light.Range     = GetScalar(Object, TEXT("range"), Light.Range);
+
+		const FString Type = GetString(Object, TEXT("type"));
+		if (Type == TEXT("spot"))
+		{
+			Light.Type = FLight::EType::Spot;
+			if (Object.HasTypedField<EJson::Object>(TEXT("spot")))
+			{
+				const TSharedPtr<FJsonObject>& SpotObj = Object.GetObjectField(TEXT("spot"));
+				Light.Spot.InnerConeAngle              = GetScalar(*SpotObj, TEXT("innerConeAngle"), 0.f);
+				Light.Spot.OuterConeAngle              = GetScalar(*SpotObj, TEXT("outerConeAngle"), Light.Spot.OuterConeAngle);
+			}
+		}
+		else if (Type == TEXT("point"))
+			Light.Type = FLight::EType::Point;
+		else if (Type == TEXT("directional"))
+			Light.Type = FLight::EType::Directional;
+		else
+			Messages.Emplace(EMessageSeverity::Warning, FString::Printf(TEXT("Light has no type specified: %s"), *Light.Name));
+	}
+
+}  // namespace GLTF
diff --git a/Engine/Plugins/Interchange/Runtime/Source/Parsers/GLTFCore/Private/GLTF/ExtensionsHandler.h b/Engine/Plugins/Interchange/Runtime/Source/Parsers/GLTFCore/Private/GLTF/ExtensionsHandler.h
new file mode 100644
--- /dev/null
+++ b/Engine/Plugins/Interchange/Runtime/Source/Parsers/GLTFCore/Private/GLTF/ExtensionsHandler.h
@@ -0,0 +1,70 @@
+// Copyright Epic Games, Inc. All Rights Reserved.
+
+#pragma once
+
+#include "GLTFLogger.h"
+
+class FJsonObject;
+
+namespace GLTF
+{
+	struct FAsset;
+	struct FMaterial;
+	struct FBuffer;
+	struct FBufferView;
+	struct FAccessor;
+	struct FPrimitive;
+	struct FMesh;
+	struct FScene;
+	struct FNode;
+	struct FSkinInfo;
+	struct FCamera;
+	struct FImage;
+	struct FSampler;
+	struct FTexture;
+	struct FAnimation;
+
+	class FExtensionsHandler
+	{
+	public:
+		explicit FExtensionsHandler(TArray<FLogMessage>& InMessages);
+
+		void SetAsset(GLTF::FAsset& InAsset);
+
+		void SetupAssetExtensions(const FJsonObject& Object) const;
+		void SetupBufferExtensions(const FJsonObject& Object, FBuffer& Buffer) const;
+		void SetupBufferViewExtensions(const FJsonObject& Object, FBufferView& BufferView) const;
+		void SetupAccessorExtensions(const FJsonObject& Object, FAccessor& Accessor) const;
+		void SetupPrimitiveExtensions(const FJsonObject& Object, FPrimitive& Primitive) const;
+		void SetupMeshExtensions(const FJsonObject& Object, FMesh& Mesh) const;
+		void SetupSceneExtensions(const FJsonObject& Object, FScene& Scene) const;
+		void SetupNodeExtensions(const FJsonObject& Object, FNode& Node) const;
+		void SetupCameraExtensions(const FJsonObject& Object, FCamera& Camera) const;
+		void SetupSkinExtensions(const FJsonObject& Object, FSkinInfo& Skin) const;
+		void SetupAnimationExtensions(const FJsonObject& Object, FAnimation& Animation) const;
+		void SetupImageExtensions(const FJsonObject& Object, FImage& Image) const;
+		void SetupSamplerExtensions(const FJsonObject& Object, FSampler& Sampler) const;
+		void SetupTextureExtensions(const FJsonObject& Object, FTexture& Texture) const;
+		void SetupMaterialExtensions(const FJsonObject& Object, FMaterial& Material) const;
+
+		TArray<FLogMessage>& GetMessages()
+		{
+			return Messages;
+		}
+
+	private:
+		void CheckExtensions(const FJsonObject& Object, const TArray<FString>& ExtensionsSupported) const;
+
+		void SetupLightPunctual(const FJsonObject& Object) const;
+
+	private:
+		TArray<FLogMessage>& Messages;
+		GLTF::FAsset*        Asset;
+	};
+
+	inline void FExtensionsHandler::SetAsset(GLTF::FAsset& InAsset)
+	{
+		Asset = &InAsset;
+	}
+
+}  // namespace GLTF
diff --git a/Engine/Plugins/Interchange/Runtime/Source/Parsers/GLTFCore/Private/GLTF/GLTFAccessor.cpp b/Engine/Plugins/Interchange/Runtime/Source/Parsers/GLTFCore/Private/GLTF/GLTFAccessor.cpp
new file mode 100644
--- /dev/null
+++ b/Engine/Plugins/Interchange/Runtime/Source/Parsers/GLTFCore/Private/GLTF/GLTFAccessor.cpp
@@ -0,0 +1,581 @@
+// Copyright Epic Games, Inc. All Rights Reserved.
+
+#include "GLTFAccessor.h"
+
+#include "ConversionUtilities.h"
+
+namespace GLTF
+{
+	namespace
+	{
+		uint32 GetElementSize(FAccessor::EType Type, FAccessor::EComponentType ComponentType)
+		{
+			static const uint8 ComponentSize[]      = {0, 1, 1, 2, 2, 4, 4};      // keep in sync with EComponentType
+			static const uint8 ComponentsPerValue[] = {0, 1, 2, 3, 4, 4, 9, 16};  // keep in sync with EType
+
+			static_assert(
+			    (int)FAccessor::EType::Unknown == 0 && ((int)FAccessor::EType::Count) == (sizeof(ComponentsPerValue) / sizeof(ComponentsPerValue[0])),
+			    "EType doesn't match!");
+			static_assert((int)FAccessor::EComponentType::None == 0 &&
+			                  ((int)FAccessor::EComponentType::Count) == (sizeof(ComponentSize) / sizeof(ComponentSize[0])),
+			              "EComponentType doesn't match!");
+
+			return ComponentsPerValue[(int)Type] * ComponentSize[(int)ComponentType];
+		}
+
+		template <class ReturnType, uint32 Count>
+		ReturnType GetNormalized(FAccessor::EComponentType ComponentType, const void* Pointer)
+		{
+			ReturnType Res;
+
+			// convert to 0..1
+			if (ComponentType == FAccessor::EComponentType::U8)
+			{
+				const uint8*    P = static_cast<const uint8*>(Pointer);
+				constexpr float S = 1.0f / 255.0f;
+				for (uint32 Index = 0; Index < Count; ++Index)
+				{
+					Res[Index] = P[Index] * S;
+				}
+			}
+			else if (ComponentType == FAccessor::EComponentType::U16)
+			{
+				const uint16*   P = static_cast<const uint16*>(Pointer);
+				constexpr float S = 1.0f / 65535.0f;
+				for (uint32 Index = 0; Index < Count; ++Index)
+				{
+					Res[Index] = P[Index] * S;
+				}
+			}
+			else
+				ensure(false);
+
+			return Res;
+		}
+
+		FMatrix GetMatrix(const void* Pointer)
+		{
+			// copy float mat4 directly from buffer
+			const float* P = static_cast<const float*>(Pointer);
+
+			FMatrix Matrix;
+			for (int32 Row = 0; Row < 4; ++Row)
+			{
+				for (int32 Col = 0; Col < 4; ++Col)
+				{
+					// glTF stores matrix elements in column major order
+					// Unreal's FMatrix is row major
+					Matrix.M[Row][Col] = P[Col * 4 + Row];
+				}
+			}
+			return Matrix;
+		}
+
+		template <typename DstT, typename SrcT>
+		void Copy(DstT* Dst, const SrcT* Src, uint32 Count)
+		{
+			for (uint32 Index = 0; Index < Count; ++Index)
+			{
+				*(Dst++) = *(Src++);
+			}
+		}
+
+		template <typename DstT, uint32 ElementCount>
+		void CopyNormalized(DstT* Dst, const void* Src, FAccessor::EComponentType ComponentType, uint32 Count)
+		{
+			// convert to 0..1
+			if (ComponentType == FAccessor::EComponentType::U8)
+			{
+				const uint8*    P = static_cast<const uint8*>(Src);
+				constexpr float S = 1.0f / 255.0f;
+				for (uint32 Index = 0; Index < Count; ++Index)
+				{
+					DstT& VecDst = *Dst++;
+					for (uint32 J = 0; J < ElementCount; ++J)
+					{
+						VecDst[J] = P[J] * S;
+					}
+					P += ElementCount;
+				}
+			}
+			else if (ComponentType == FAccessor::EComponentType::U16)
+			{
+				const uint16*   P = static_cast<const uint16*>(Src);
+				constexpr float S = 1.0f / 65535.0f;
+				for (uint32 Index = 0; Index < Count; ++Index)
+				{
+					DstT& VecDst = *Dst++;
+					for (uint32 J = 0; J < ElementCount; ++J)
+					{
+						VecDst[J] = P[J] * S;
+					}
+					P += ElementCount;
+				}
+			}
+			else
+				ensure(false);
+		}
+
+		// Copy data items that don't need conversion/expansion(i.e. Vec3 to Vec3, uint8 to uint8(not uint16)
+		// but including normalized from fixed-point types
+		template<typename ItemType, uint32 ItemElementCount>
+		void CopyWithoutConversion(const FValidAccessor& Accessor, ItemType* Buffer)
+		{
+			// Stride equals item size => use simpler copy
+			if ((Accessor.ByteStride == 0) || Accessor.ByteStride == sizeof(ItemType))
+			{
+				const void* Src = Accessor.DataAt(0);
+				if (Accessor.ComponentType == FAccessor::EComponentType::F32)
+				{
+					memcpy(Buffer, Src, Accessor.Count * sizeof(ItemType));
+				}
+				else if (Accessor.Normalized)
+				{
+					CopyNormalized<ItemType, ItemElementCount>(Buffer, Src, Accessor.ComponentType, Accessor.Count);
+				}
+			}
+			else
+			{
+				if (Accessor.ComponentType == FAccessor::EComponentType::F32)
+				{
+					for (uint32 Index = 0; Index < Accessor.Count; ++Index)
+					{
+						const void* Pointer = Accessor.DataAt(Index);
+						Buffer[Index] = *static_cast<const ItemType*>(Pointer);
+					}
+				}
+				else if (Accessor.Normalized)
+				{
+					for (uint32 Index = 0; Index < Accessor.Count; ++Index)
+					{
+						const void* Pointer = Accessor.DataAt(Index);
+						Buffer[Index] = GetNormalized<ItemType, ItemElementCount>(Accessor.ComponentType, Pointer);
+					}
+				}
+			}
+		}
+	}
+
+	FAccessor::FAccessor(uint32 InCount, EType InType, EComponentType InCompType, bool InNormalized)
+	    : Count(InCount)
+	    , Type(InType)
+	    , ComponentType(InCompType)
+	    , Normalized(InNormalized)
+	{
+	}
+
+	uint32 FAccessor::GetUnsignedInt(uint32 Index) const
+	{
+		return 0;
+	}
+
+	void FAccessor::GetUnsignedInt16x4(uint32 Index, uint16 Values[4]) const {}
+
+	float FAccessor::GetFloat(uint32 Index) const
+	{
+		return 0.f;
+	}
+
+	FVector2D FAccessor::GetVec2(uint32 Index) const
+	{
+		return FVector2D::ZeroVector;
+	}
+
+	FVector FAccessor::GetVec3(uint32 Index) const
+	{
+		return FVector::ZeroVector;
+	}
+
+	FVector4 FAccessor::GetVec4(uint32 Index) const
+	{
+		return FVector4();
+	}
+
+	FMatrix FAccessor::GetMat4(uint32 Index) const
+	{
+		return FMatrix::Identity;
+	}
+
+	void FAccessor::GetUnsignedIntArray(uint32* Buffer) const {}
+
+	void FAccessor::GetFloatArray(float* Buffer) const {}
+
+	void FAccessor::GetVec2Array(FVector2f* Buffer) const {}
+
+	void FAccessor::GetVec3Array(FVector3f* Buffer) const {}
+
+	void FAccessor::GetVec4Array(FVector4f* Buffer) const {}
+
+	void FAccessor::GetMat4Array(FMatrix44f* Buffer) const {}
+
+	void FAccessor::GetCoordArray(FVector3f* Buffer) const
+	{
+		GetVec3Array(Buffer);
+		for (uint32 Index = 0; Index < Count; ++Index)
+		{
+			Buffer[Index] = (FVector3f)ConvertVec3((FVector)Buffer[Index]);
+		}
+	}
+
+	void FAccessor::GetQuatArray(FVector4f* Buffer) const
+	{
+		GetVec4Array(Buffer);
+		for (uint32 Index = 0; Index < Count; ++Index)
+		{
+			FVector4f& Value = Buffer[Index];
+			FQuat4f Quat(Value[0], Value[1], Value[2], Value[3]);
+			Quat = GLTF::ConvertQuat(Quat);
+			Value.X = Quat.X;
+			Value.Y = Quat.Y;
+			Value.Z = Quat.Z;
+			Value.W = Quat.W;
+		}
+	}
+
+	//
+
+	FValidAccessor::FValidAccessor(FBufferView& InBufferView, uint32 InOffset, uint32 InCount, EType InType, EComponentType InCompType,
+	                               bool InNormalized)
+	    : FAccessor(InCount, InType, InCompType, InNormalized)
+	    , BufferView(InBufferView)
+	    , ByteOffset(InOffset)
+	    , ElementSize(GetElementSize(Type, ComponentType))
+		// BufferView.ByteStride zero means elements are tightly packed
+		, ByteStride(InBufferView.ByteStride ? InBufferView.ByteStride : ElementSize)
+	{
+	}
+
+	bool FValidAccessor::IsValid() const
+	{
+		return BufferView.IsValid();
+	}
+
+	FMD5Hash FValidAccessor::GetHash() const
+	{
+		if (!IsValid())
+		{
+			return FMD5Hash();
+		}
+
+		FMD5 MD5;
+
+		// Contigous view into a buffer: Can hash all at once
+		if (ByteStride == ElementSize)
+		{
+			MD5.Update(DataAt(0), Count * ElementSize);
+		}
+		// Shared buffer: Have to hash in strides
+		else
+		{
+			for (uint32 Index = 0; Index < Count; ++Index)
+			{
+				MD5.Update(DataAt(Index), ElementSize);
+			}
+		}
+
+		uint8 TypeInt = static_cast<uint8>(Type);
+		uint8 ComponentInt = static_cast<uint8>(ComponentType);
+		uint8 NormalizedInt = static_cast<uint8>(Normalized);
+
+		MD5.Update(&TypeInt, sizeof(TypeInt));
+		MD5.Update(&ComponentInt, sizeof(ComponentInt));
+		MD5.Update(&NormalizedInt, sizeof(NormalizedInt));
+
+		FMD5Hash Hash;
+		Hash.Set(MD5);
+		return Hash;
+	}
+
+	uint32 FValidAccessor::GetUnsignedInt(uint32 Index) const
+	{
+		// should be Scalar, not Normalized, unsigned integer (8, 16 or 32 bit)
+
+		if (Index < Count)
+		{
+			if (Type == EType::Scalar && !Normalized)
+			{
+				const uint8* ValuePtr = DataAt(Index);
+				switch (ComponentType)
+				{
+					case EComponentType::U8:
+						return *ValuePtr;
+					case EComponentType::U16:
+						return *reinterpret_cast<const uint16*>(ValuePtr);
+					case EComponentType::U32:
+						return *reinterpret_cast<const uint32*>(ValuePtr);
+					default:
+						break;
+				}
+			}
+		}
+
+		ensure(false);
+		return 0;
+	}
+
+	void FValidAccessor::GetUnsignedInt16x4(uint32 Index, uint16 Values[4]) const
+	{
+		// should be Vec4, not Normalized, unsigned integer (8 or 16 bit)
+
+		if (Index < Count)
+		{
+			if (Type == EType::Vec4 && !Normalized)
+			{
+				const void* ValuePtr = DataAt(Index);
+				switch (ComponentType)
+				{
+					case EComponentType::U8:
+						for (int i = 0; i < 4; ++i)
+						{
+							Values[i] = ((uint8*)ValuePtr)[i];
+						}
+						return;
+					case EComponentType::U16:
+						for (int i = 0; i < 4; ++i)
+						{
+							Values[i] = ((uint16*)ValuePtr)[i];
+						}
+						return;
+					default:
+						break;
+				}
+			}
+		}
+		ensure(false);
+	}
+
+	float FValidAccessor::GetFloat(uint32 Index) const
+	{
+		// should be Scalar float
+
+		if (Index < Count)
+		{
+			if (Type == EType::Scalar && !Normalized)
+			{
+				const uint8* ValuePtr = DataAt(Index);
+				switch (ComponentType)
+				{
+					case EComponentType::F32:
+						return *ValuePtr;
+					default:
+						break;
+				}
+			}
+		}
+
+		ensure(false);
+		return 0.f;
+	}
+
+	FVector2D FValidAccessor::GetVec2(uint32 Index) const
+	{
+		// Spec-defined attributes (TEXCOORD_0, TEXCOORD_1) use only these formats:
+		// - F32
+		// - U8 normalized
+		// - U16 normalized
+		// Custom attributes can use any CompType, so add support for those when needed.
+
+		if (Index < Count)
+		{
+			if (Type == EType::Vec2)  // strict format match, unlike GPU shader fetch
+			{
+				const void* Pointer = DataAt(Index);
+
+				if (ComponentType == EComponentType::F32)
+				{
+					// copy float vec2 directly from buffer
+					return FVector2D(*reinterpret_cast<const FVector2f*>(Pointer));
+				}
+				else if (Normalized)
+				{
+					return FVector2D(GetNormalized<FVector2f, 2>(ComponentType, Pointer));
+				}
+			}
+		}
+
+		ensure(false);
+		return FVector2D::ZeroVector;
+	}
+
+	FVector FValidAccessor::GetVec3(uint32 Index) const
+	{
+		// Spec-defined attributes (POSITION, NORMAL, COLOR_0) use only these formats:
+		// - F32
+		// - U8 normalized
+		// - U16 normalized
+		// Custom attributes can use any CompType, so add support for those when needed.
+
+		if (Index < Count)
+		{
+			if (Type == EType::Vec3)  // strict format match, unlike GPU shader fetch
+			{
+				const void* Pointer = DataAt(Index);
+
+				if (ComponentType == EComponentType::F32)
+				{
+					// copy float vec3 directly from buffer
+					return FVector(*reinterpret_cast<const FVector3f*>(Pointer));
+				}
+				else if (Normalized)
+				{
+					return FVector(GetNormalized<FVector3f, 3>(ComponentType, Pointer));
+				}
+			}
+		}
+
+		ensure(false);
+		return FVector::ZeroVector;
+	}
+
+	FVector4 FValidAccessor::GetVec4(uint32 Index) const
+	{
+		// Spec-defined attributes (TANGENT, COLOR_0) use only these formats:
+		// - F32
+		// - U8 normalized
+		// - U16 normalized
+		// Custom attributes can use any CompType, so add support for those when needed.
+
+		if (Index < Count)
+		{
+			if (Type == EType::Vec4)  // strict format match, unlike GPU shader fetch
+			{
+				const void* Pointer = DataAt(Index);
+
+				if (ComponentType == EComponentType::F32)
+				{
+					// copy float vec4 directly from buffer
+					return FVector4(*reinterpret_cast<const FVector4f*>(Pointer));
+				}
+				else if (Normalized)
+				{
+					return FVector4(GetNormalized<FVector4f, 4>(ComponentType, Pointer));
+				}
+			}
+		}
+
+		ensure(false);
+		return FVector4();
+	}
+
+	FMatrix FValidAccessor::GetMat4(uint32 Index) const
+	{
+		// Focus on F32 for now, add other types as needed.
+
+		if (Index < Count)
+		{
+			if (Type == EType::Mat4)  // strict format match, unlike GPU shader fetch
+			{
+				if (ComponentType == EComponentType::F32)
+				{
+					const void* Pointer = DataAt(Index);
+					return GetMatrix(Pointer);
+				}
+			}
+		}
+
+		ensure(false);
+		return FMatrix();
+	}
+
+	void FValidAccessor::GetUnsignedIntArray(uint32* Buffer) const
+	{
+		if (Type == EType::Scalar && !Normalized)
+		{
+			const uint8* Src = DataAt(0);
+			switch (ComponentType)
+			{
+				case EComponentType::U8:
+					Copy(Buffer, reinterpret_cast<const uint8*>(Src), Count);
+					return;
+				case EComponentType::U16:
+					Copy(Buffer, reinterpret_cast<const uint16*>(Src), Count);
+					return;
+				case EComponentType::U32:
+					memcpy(Buffer, Src, Count * sizeof(uint32));
+					return;
+				default:
+					break;
+			}
+		}
+
+		ensure(false);
+	}
+
+	void FValidAccessor::GetFloatArray(float* Buffer) const
+	{
+		if (Type == EType::Scalar && !Normalized)
+		{
+			const uint8* Src = DataAt(0);
+			switch (ComponentType)
+			{
+				case EComponentType::F32:
+					memcpy(Buffer, Src, Count * sizeof(float));
+					return;
+				default:
+					break;
+			}
+		}
+
+		ensure(false);
+	}
+
+	void FValidAccessor::GetVec2Array(FVector2f* Buffer) const
+	{
+		if (!ensure(Type == EType::Vec2))
+		{
+			return;
+		}
+		CopyWithoutConversion<FVector2f, 2>(*this, Buffer);
+	}
+
+	void FValidAccessor::GetVec3Array(FVector3f* Buffer) const
+	{
+		if (!ensure(Type == EType::Vec3))
+		{
+			return;
+		}
+		CopyWithoutConversion<FVector3f, 3>(*this, Buffer);
+	}
+
+	void FValidAccessor::GetVec4Array(FVector4f* Buffer) const
+	{
+		if (!ensure(Type == EType::Vec4))
+		{
+			return;
+		}
+		CopyWithoutConversion<FVector4f, 4>(*this, Buffer);
+	}
+
+	void FValidAccessor::GetMat4Array(FMatrix44f* Buffer) const
+	{
+		if (Type == EType::Mat4 && ComponentType == EComponentType::F32)  // strict format match, unlike GPU shader fetch
+		{
+			for (uint32 Index = 0; Index < Count; ++Index)
+			{
+				const void* Pointer = DataAt(Index);
+				Buffer[Index]       = FMatrix44f(GetMatrix(Pointer));
+			}
+			return;
+		}
+
+		ensure(false);
+	}
+
+	inline const uint8* FValidAccessor::DataAt(uint32 Index) const
+	{
+		const uint32 Offset = Index * ByteStride;
+		return BufferView.DataAt(Offset + ByteOffset);
+	}
+
+	//
+
+	bool FVoidAccessor::IsValid() const
+	{
+		return false;
+	}
+
+	FMD5Hash FVoidAccessor::GetHash() const
+	{
+		return FMD5Hash();
+	}
+}  // namespace GLTF
diff --git a/Engine/Plugins/Interchange/Runtime/Source/Parsers/GLTFCore/Private/GLTF/GLTFAsset.cpp b/Engine/Plugins/Interchange/Runtime/Source/Parsers/GLTFCore/Private/GLTF/GLTFAsset.cpp
new file mode 100644
--- /dev/null
+++ b/Engine/Plugins/Interchange/Runtime/Source/Parsers/GLTFCore/Private/GLTF/GLTFAsset.cpp
@@ -0,0 +1,410 @@
+// Copyright Epic Games, Inc. All Rights Reserved.
+
+#include "GLTFAsset.h"
+
+#include "GLTFAnimation.h"
+#include "Misc/Paths.h"
+#include "GLTFMaterial.h"
+#include "GLTFMesh.h"
+#include "GLTFNode.h"
+#include "GLTFTexture.h"
+
+namespace GLTF
+{
+	namespace
+	{
+		template <typename T>
+		void GenerateNames(const FString& Prefix, TArray<T>& Objects)
+		{
+			for (int32 Index = 0; Index < Objects.Num(); ++Index)
+			{
+				T& Obj = Objects[Index];
+
+				Obj.UniqueId = Prefix + FString::FromInt(Index);
+
+				if (Obj.Name.IsEmpty())
+				{
+					Obj.Name = Obj.UniqueId;
+				}
+			}
+		}
+
+		//Make the Name FName complient.
+		//For Scenario where Asset.Name is '_x' (where x is a number) making FName to be none:
+		//UE-174253
+		template <typename T>
+		void ValidateFNameCriteria(const FString& Prefix, TArray<T>& Objects)
+		{
+			for (int32 Index = 0; Index < Objects.Num(); ++Index)
+			{
+				T& Obj = Objects[Index];
+
+				if (!Obj.Name.IsEmpty())
+				{
+					FName SanitizationTest(*Obj.Name);
+					if (SanitizationTest.IsNone())
+					{
+						Obj.Name = Prefix + Obj.Name;
+					}
+				}
+			}
+		}
+	}
+
+	const FMetadata::FExtraData* FMetadata::GetExtraData(const FString& Name) const
+	{
+		return Extras.FindByPredicate([&Name](const FMetadata::FExtraData& Data) { return Data.Name == Name; });
+	}
+
+	const TSet<GLTF::EExtension> FAsset::SupportedExtensions = { GLTF::EExtension::KHR_MaterialsPbrSpecularGlossiness,
+																 GLTF::EExtension::KHR_MaterialsUnlit,
+																 GLTF::EExtension::KHR_LightsPunctual,
+																 GLTF::EExtension::KHR_MaterialsClearCoat,
+																 GLTF::EExtension::KHR_MaterialsTransmission,
+																 GLTF::EExtension::KHR_MaterialsSheen,
+																 GLTF::EExtension::KHR_MaterialsVariants,
+																 GLTF::EExtension::KHR_MaterialsSpecular,
+																 GLTF::EExtension::KHR_MaterialsIOR };
+
+	void FAsset::Clear(uint32 BinBufferKBytes, uint32 ExtraBinBufferKBytes)
+	{
+		Buffers.Empty();
+		BufferViews.Empty();
+		Accessors.Empty();
+		Meshes.Empty();
+
+		Scenes.Empty();
+		Nodes.Empty();
+		Cameras.Empty();
+		Lights.Empty();
+		Skins.Empty();
+
+		Images.Empty();
+		Samplers.Empty();
+		Textures.Empty();
+		Materials.Empty();
+
+		ProcessedExtensions.Empty((int)EExtension::Count);
+		ExtensionsRequired.Empty();
+		Metadata.GeneratorName.Empty();
+		Metadata.Extras.Empty();
+
+		if (BinData.Num() > 0)
+			BinData.Empty(BinBufferKBytes * 1024);
+		if (ExtraBinData.Num() > 0)
+			ExtraBinData.Empty(ExtraBinBufferKBytes * 1024);
+	}
+
+	void FAsset::GenerateNames(const FString& Prefix)
+	{
+		ensure(!Prefix.IsEmpty());
+
+		{
+			const FString ScenePrefix = Prefix + TEXT("_scene_");
+
+			for (size_t SceneIndex = 0; SceneIndex < Scenes.Num(); SceneIndex++)
+			{
+				FScene& Scene = Scenes[SceneIndex];
+
+				Scene.UniqueId = ScenePrefix + FString::FromInt(SceneIndex);
+
+				if (Scene.Name.IsEmpty())
+				{
+					Scene.Name = Scene.UniqueId;
+				}
+			}
+		}
+
+		{
+			const FString NodePrefix = Prefix + TEXT("_node_");
+			const FString JointPrefix = Prefix + TEXT("_joint_");
+			for (int32 NodeIndex = 0; NodeIndex < Nodes.Num(); ++NodeIndex)
+			{
+				FNode& Node = Nodes[NodeIndex];
+
+				bool bIsJoint = Node.Type == FNode::EType::Joint;
+				Node.UniqueId = (bIsJoint ? JointPrefix : NodePrefix);
+
+				Node.UniqueId = Node.UniqueId + FString::FromInt(NodeIndex);
+
+				if (Node.Name.IsEmpty())
+				{
+					Node.Name = Node.UniqueId;
+				}
+			}
+		}
+
+		{
+			const FString TexPrefix = Prefix + TEXT("_texture_");
+			for (int32 TextureIndex = 0; TextureIndex < Textures.Num(); ++TextureIndex)
+			{
+				FTexture& Tex = Textures[TextureIndex];
+
+				Tex.UniqueId = TexPrefix + FString::FromInt(TextureIndex);
+
+				if (!Tex.Name.IsEmpty())
+				{
+					// keep base part of the name as it was in gltf
+				}
+				else if (!Tex.Source.Name.IsEmpty())
+				{
+					Tex.Name = Tex.Source.Name;
+				}
+				else if (!Tex.Source.URI.IsEmpty())
+				{
+					//in case texture is embedded:
+					if (Tex.Source.URI.StartsWith(TEXT("data:")))
+					{
+						Tex.Name = Tex.UniqueId;
+					}
+					else
+					{
+						Tex.Name = FPaths::GetBaseFilename(Tex.Source.URI);
+					}
+				}
+				else
+				{
+					Tex.Name = Tex.UniqueId;
+				}
+
+			}
+		}
+
+		{
+			const FString CameraPrefix = Prefix + TEXT("_camera_");
+			for (int32 CameraIndex = 0; CameraIndex < Cameras.Num(); ++CameraIndex)
+			{
+				FCamera& Camera = Cameras[CameraIndex];
+
+				Camera.UniqueId = CameraPrefix + FString::FromInt(CameraIndex);
+
+				if (Camera.Name.IsEmpty())
+				{
+					if (!Camera.Node.Name.IsEmpty())
+					{
+						Camera.Name = Camera.Node.Name; // cant be empty
+					}
+					else
+					{
+						Camera.Name = Camera.UniqueId;
+					}
+				}
+			}
+		}
+
+		{
+			const FString LightPrefix = Prefix + TEXT("_light_");
+			for (int32 LightIndex = 0; LightIndex < Lights.Num(); ++LightIndex)
+			{
+				FLight& Light = Lights[LightIndex];
+
+				Light.UniqueId = LightPrefix + FString::FromInt(LightIndex);
+
+				if (Light.Name.IsEmpty())
+				{
+					if (Light.Node && !Light.Node->Name.IsEmpty())
+					{
+						Light.Name = Light.Node->Name;
+					}
+					else
+					{
+						Light.Name = Light.UniqueId;
+					}
+				}
+			}
+		}
+
+		// General pattern here is that we'll have a Prefix that will be a filename (like MyFile). 
+		// The name (display label) and the UniqueId of an object has been separated.
+		// UniqueId will consist of: [Prefil]+[_ElementType_]+[ElementIndex]
+		// Name will be the provided GLTF object.name, if present,
+		// if Name is not present the UniqueId will be used as the Name.
+		GLTF::GenerateNames(Prefix + TEXT("_material_"), Materials);
+		GLTF::GenerateNames(Prefix + TEXT("_skin_"), Skins);
+		GLTF::GenerateNames(Prefix + TEXT("_animation_"), Animations);
+		GLTF::GenerateNames(Prefix + TEXT("_image_"), Images);
+		GLTF::GenerateNames(Prefix + TEXT("_mesh_"), Meshes);
+
+		// Validate FName criteria for Materials Textures Meshes Animations (assets)
+		GLTF::ValidateFNameCriteria(TEXT("Material_"), Materials);
+		GLTF::ValidateFNameCriteria(TEXT("Texture_"), Textures);
+		GLTF::ValidateFNameCriteria(TEXT("Mesh_"), Meshes);
+		GLTF::ValidateFNameCriteria(TEXT("Animation_"), Animations);
+
+		//Generate MorphTarget Names:
+		//MorphTargetNames have to be unique across meshes as well:
+		TSet<FString> MorphTargetNamesUniquenessCheckAcrossMeshes;
+		int32 MorphTargetCounterAcrossMeshes = 0;
+		for (int32 MeshIndex = 0; MeshIndex < Meshes.Num(); MeshIndex++)
+		{
+			if (Meshes[MeshIndex].Primitives.Num() > 0)
+			{
+				TSet<FString> MorphTargetNamesUniquenessCheck; //per mesh
+
+				for (const FString& MorphTargetName : Meshes[MeshIndex].MorphTargetNames)
+				{
+					MorphTargetNamesUniquenessCheckAcrossMeshes.Add(MorphTargetName);
+					MorphTargetCounterAcrossMeshes++;
+
+					//validate uniqueness of MorphTargetNames
+					MorphTargetNamesUniquenessCheck.Add(MorphTargetName);
+				}
+
+				if (MorphTargetNamesUniquenessCheck.Num() != Meshes[MeshIndex].MorphTargetNames.Num())
+				{
+					Meshes[MeshIndex].MorphTargetNames.Empty();
+				}
+			}
+		}
+		bool bReBuildMorphTargetNames = MorphTargetNamesUniquenessCheckAcrossMeshes.Num() != MorphTargetCounterAcrossMeshes;
+
+		for (int32 MeshIndex = 0; MeshIndex < Meshes.Num(); MeshIndex++)
+		{
+			if (Meshes[MeshIndex].Primitives.Num() > 0)
+			{
+				//at this point the NumberOfMorphTargets across primitives have been validated (they have to be equal)
+				//Note: All primitives MUST have the same number of morph targets in the same order.
+				
+				if (bReBuildMorphTargetNames)
+				{
+					Meshes[MeshIndex].MorphTargetNames.Empty();
+				}
+
+				//check if number of morph targets and the number of morphtarget names:
+				//NumberOfMorphTargetsPerPrimitive will return the first Primitive's MorphTarget's count,
+				//		In case Primitives have varying MorphTarget counts then the ValidationCheck will report false.
+				int32 NumberOfMorphTargets = Meshes[MeshIndex].NumberOfMorphTargetsPerPrimitive();
+				if (NumberOfMorphTargets != Meshes[MeshIndex].MorphTargetNames.Num())
+				{
+					Meshes[MeshIndex].MorphTargetNames.Empty();
+				}
+
+				//if morph target names failed a criteria above, generate unique names:
+				if (Meshes[MeshIndex].MorphTargetNames.IsEmpty())
+				{
+					for (int32 MorphTargetIndex = 0; MorphTargetIndex < NumberOfMorphTargets; MorphTargetIndex++)
+					{
+						Meshes[MeshIndex].MorphTargetNames.Add(Meshes[MeshIndex].UniqueId + TEXT("_") + LexToString(MorphTargetIndex) + TEXT("_MorphTarget"));
+					}
+
+					Meshes[MeshIndex].GenerateIsValidCache(false);
+				}
+			}
+		}
+	}
+
+	void FAsset::GetRootNodes(TArray<int32>& NodeIndices)
+	{
+		TMap<int32, uint32> VisitCount;
+		VisitCount.Reserve(Nodes.Num());
+		for (int32 NodeIndex = 0; NodeIndex < Nodes.Num(); ++NodeIndex)
+		{
+			for (int32 ChildIndex : Nodes[NodeIndex].Children)
+			{
+				uint32* Found = VisitCount.Find(ChildIndex);
+				if (Found)
+				{
+					(*Found)++;
+				}
+				else
+				{
+					VisitCount.Add(ChildIndex) = 1;
+				}
+			}
+		}
+
+		NodeIndices.Empty();
+		for (int32 NodeIndex = 0; NodeIndex < Nodes.Num(); ++NodeIndex)
+		{
+			if (!VisitCount.Contains(NodeIndex))
+				NodeIndices.Add(NodeIndex);
+			else
+				ensure(VisitCount[NodeIndex] == 1);
+		}
+	}
+
+	FAsset::EValidationCheck FAsset::ValidationCheck() const
+	{
+		int32 Res = EValidationCheck::Valid;
+		if (Meshes.FindByPredicate([](const FMesh& Mesh) { return !Mesh.IsValid(); }))
+		{
+			Res |= InvalidMeshPresent;
+		}
+		
+		if (Nodes.FindByPredicate([this](const FNode& Node)
+			{
+				if (Node.MeshIndex != INDEX_NONE)
+				{
+					//check Mesh index validity:
+					if (!Meshes.IsValidIndex(Node.MeshIndex))
+					{
+						return false;
+					}
+
+					//Check Morph Target Weights:
+					//	Number of MorphTargets across primitives are validated in Mesh.IsValid.
+					if (Node.MorphTargetWeights.Num() > 0 && (Node.MorphTargetWeights.Num() != Meshes[Node.MeshIndex].NumberOfMorphTargetsPerPrimitive()))
+					{
+						return false;
+					}
+				}
+
+				return !Node.Transform.IsValid();
+			}))
+		{
+			Res |= InvalidNodePresent;
+		}
+
+		// TODO: lots more validation
+
+		// TODO: verify if image format is PNG or JPEG
+
+		return static_cast<EValidationCheck>(Res);
+	}
+
+	const TCHAR* ToString(EExtension Extension)
+	{
+		switch (Extension)
+		{
+			case GLTF::EExtension::KHR_MaterialsPbrSpecularGlossiness:
+				return TEXT("KHR_materials_pbrSpecularGlossiness");
+			case GLTF::EExtension::KHR_MaterialsUnlit:
+				return TEXT("KHR_materials_unlit");
+			case GLTF::EExtension::KHR_MaterialsClearCoat:
+				return TEXT("KHR_materials_clearcoat");
+			case GLTF::EExtension::KHR_MaterialsTransmission:
+				return TEXT("KHR_materials_transmission");
+			case GLTF::EExtension::KHR_MaterialsSheen:
+				return TEXT("KHR_materials_sheen");
+			case GLTF::EExtension::KHR_MaterialsVariants:
+				return TEXT("KHR_materials_variants");
+			case GLTF::EExtension::KHR_MaterialsIOR:
+				return TEXT("KHR_materials_ior");
+			case GLTF::EExtension::KHR_MaterialsSpecular:
+				return TEXT("KHR_materials_specular");
+			case GLTF::EExtension::KHR_MaterialsEmissiveStrength:
+				return TEXT("KHR_materials_emissive_strength");
+			case GLTF::EExtension::KHR_TextureTransform:
+				return TEXT("KHR_texture_transform");
+			case GLTF::EExtension::KHR_DracoMeshCompression:
+				return TEXT("KHR_draco_mesh_compression");
+			case GLTF::EExtension::KHR_LightsPunctual:
+				return TEXT("KHR_lights_punctual");
+			case GLTF::EExtension::KHR_Lights:
+				return TEXT("KHR_lights");
+			case GLTF::EExtension::KHR_Blend:
+				return TEXT("KHR_blend");
+			case GLTF::EExtension::MSFT_TextureDDS:
+				return TEXT("MSFT_texture_dds");
+			case GLTF::EExtension::MSFT_PackingNormalRoughnessMetallic:
+				return TEXT("MSFT_packing_normalRoughnessMetallic");
+			case GLTF::EExtension::MSFT_PackingOcclusionRoughnessMetallic:
+				return TEXT("MSFT_packing_occlusionRoughnessMetallic");
+			case GLTF::EExtension::Count:
+				ensure(false);
+			default:
+				return TEXT("UnknwonExtension");
+		}
+	}
+
+}  // namespace GLTF
diff --git a/Engine/Plugins/Interchange/Runtime/Source/Parsers/GLTFCore/Private/GLTF/GLTFBinaryReader.cpp b/Engine/Plugins/Interchange/Runtime/Source/Parsers/GLTFCore/Private/GLTF/GLTFBinaryReader.cpp
new file mode 100644
--- /dev/null
+++ b/Engine/Plugins/Interchange/Runtime/Source/Parsers/GLTFCore/Private/GLTF/GLTFBinaryReader.cpp
@@ -0,0 +1,134 @@
+// Copyright Epic Games, Inc. All Rights Reserved.
+
+#include "GLTFBinaryReader.h"
+
+#include "CoreGlobals.h"
+
+namespace GLTF
+{
+	namespace
+	{
+		// round up to the nearest multiple of 4 (used for padding & vertex attribute alignment)
+		inline uint32 Pad4(uint32 X)
+		{
+			return (X + 3) & ~3;
+		}
+
+		bool SignatureMatches(uint32 Signature, const char* ExpectedSignature)
+		{
+			return Signature == *(const uint32*)ExpectedSignature;
+		}
+
+		bool IsHeaderValid(FArchive& Archive)
+		{
+			// Binary glTF files begin with a 12-byte header:
+			// - magic bytes "glTF"
+			// - format version
+			// - size of this file
+
+			const int64 FileSize = Archive.TotalSize();
+			if (FileSize < 12)
+			{
+				return false;
+			}
+
+			uint32 Magic;
+			Archive.SerializeInt(Magic, MAX_uint32);
+			bool MagicOk = SignatureMatches(Magic, "glTF");
+
+			uint32 Version;
+			Archive.SerializeInt(Version, MAX_uint32);
+			bool VersionOk = Version == 2;
+
+			uint32 Size;
+			Archive.SerializeInt(Size, MAX_uint32);
+			bool SizeOk = Size == FileSize;
+
+			return MagicOk && VersionOk && SizeOk;
+		}
+
+		bool ReadChunk(FArchive& FileReader, const char* ExpectedChunkType, bool& OutHasMoreData, TArray<uint8>& OutData)
+		{
+			// Align to next 4-byte boundary before reading anything
+			uint32 Offset        = FileReader.Tell();
+			uint32 AlignedOffset = Pad4(Offset);
+			if (Offset != AlignedOffset)
+			{
+				FileReader.Seek(AlignedOffset);
+			}
+
+			// Each chunk has the form [Size][Type][...Data...]
+			uint32 ChunkType, ChunkDataSize;
+			FileReader.SerializeInt(ChunkDataSize, MAX_uint32);
+			FileReader.SerializeInt(ChunkType, MAX_uint32);
+
+			constexpr uint32 ChunkHeaderSize = 8;
+			const uint32     AvailableData   = FileReader.TotalSize() - (AlignedOffset + ChunkHeaderSize);
+
+			// Is there room for another chunk after this one?
+			OutHasMoreData = AvailableData - Pad4(ChunkDataSize) >= ChunkHeaderSize;
+
+			// Is there room for this chunk's data? (should always be true)
+			if (ChunkDataSize > AvailableData)
+			{
+				return false;
+			}
+
+			if (SignatureMatches(ChunkType, ExpectedChunkType))
+			{
+				// Read this chunk's data
+				OutData.SetNumUninitialized(ChunkDataSize, true);
+				FileReader.Serialize(OutData.GetData(), ChunkDataSize);
+				return true;
+			}
+			else
+			{
+				// Skip past this chunk's data
+				FileReader.Seek(AlignedOffset + ChunkHeaderSize + ChunkDataSize);
+				return false;
+			}
+		}
+	}
+
+	FBinaryFileReader::FBinaryFileReader()
+	    : BinChunk(nullptr)
+	{
+	}
+
+	bool FBinaryFileReader::ReadFile(FArchive& FileReader)
+	{
+		if (!ensure(BinChunk))
+		{
+			return false;
+		}
+
+		Messages.Empty();
+
+		// Binary glTF files begin with a 12-byte header
+		// followed by 1 chunk of JSON and (optionally) 1 chunk of binary data.
+
+		if (!IsHeaderValid(FileReader))
+		{
+			Messages.Emplace(EMessageSeverity::Error, TEXT("Invalid GLTF header!"));
+			return false;
+		}
+
+		JsonChunk.Empty();
+		BinChunk->Empty();
+
+		bool HasMoreData;
+		if (ReadChunk(FileReader, "JSON", HasMoreData, JsonChunk))
+		{
+			// Get BIN chunk if present
+			while (HasMoreData)
+			{
+				if (ReadChunk(FileReader, "BIN", HasMoreData, *BinChunk))
+				{
+					break;
+				}
+			}
+		}
+		return JsonChunk.Num() > 0;
+	}
+
+}  // namespace GLTF
diff --git a/Engine/Plugins/Interchange/Runtime/Source/Parsers/GLTFCore/Private/GLTF/GLTFBinaryReader.h b/Engine/Plugins/Interchange/Runtime/Source/Parsers/GLTFCore/Private/GLTF/GLTFBinaryReader.h
new file mode 100644
--- /dev/null
+++ b/Engine/Plugins/Interchange/Runtime/Source/Parsers/GLTFCore/Private/GLTF/GLTFBinaryReader.h
@@ -0,0 +1,38 @@
+// Copyright Epic Games, Inc. All Rights Reserved.
+
+#pragma once
+
+#include "GLTFLogger.h"
+
+
+struct FScriptContainerElement;
+
+namespace GLTF
+{
+	class FBinaryFileReader : public FBaseLogger
+	{
+	public:
+		FBinaryFileReader();
+
+		bool ReadFile(FArchive& FileReader);
+
+		void SetBuffer(TArray<uint8>& InBuffer);
+
+		const TArray<uint8>& GetJsonBuffer() const;
+
+	private:
+		TArray<uint8> JsonChunk;
+
+		TArray<uint8>* BinChunk;
+	};
+
+	inline void FBinaryFileReader::SetBuffer(TArray<uint8>& InBuffer)
+	{
+		BinChunk = &InBuffer;
+	}
+
+	inline const TArray<uint8>& FBinaryFileReader::GetJsonBuffer() const
+	{
+		return JsonChunk;
+	}
+}
diff --git a/Engine/Plugins/Interchange/Runtime/Source/Parsers/GLTFCore/Private/GLTF/GLTFMaterial.cpp b/Engine/Plugins/Interchange/Runtime/Source/Parsers/GLTFCore/Private/GLTF/GLTFMaterial.cpp
new file mode 100644
--- /dev/null
+++ b/Engine/Plugins/Interchange/Runtime/Source/Parsers/GLTFCore/Private/GLTF/GLTFMaterial.cpp
@@ -0,0 +1,60 @@
+// Copyright Epic Games, Inc. All Rights Reserved.
+
+#include "GLTFMaterial.h"
+
+namespace GLTF
+{
+	FMD5Hash FMaterial::GetHash() const
+	{
+		auto HashTextureMap = [](FMD5& MD5, const FTextureMap& Map)
+		{
+			FMD5Hash MapHash = Map.GetHash();
+			if (MapHash.IsValid())
+			{
+				MD5.Update(MapHash.GetBytes(), MapHash.GetSize());
+			}
+		};
+
+		FMD5 MD5;
+
+		uint8 ShadingModelInt = static_cast<uint8>(ShadingModel);
+		uint8 AlphaModeInt = static_cast<uint8>(AlphaMode);
+		uint8 IsDoubleSized = static_cast<uint8>(bIsDoubleSided);
+		uint8 IsUnlitShadingModel = static_cast<uint8>(bIsUnlitShadingModel);
+
+		HashTextureMap(MD5, BaseColor);
+		MD5.Update(reinterpret_cast<const uint8*>(&BaseColorFactor), sizeof(BaseColorFactor));
+		MD5.Update(&ShadingModelInt, sizeof(ShadingModelInt));
+
+		HashTextureMap(MD5, MetallicRoughness.Map);
+		MD5.Update(reinterpret_cast<const uint8*>(&MetallicRoughness.MetallicFactor), sizeof(MetallicRoughness.MetallicFactor));
+		MD5.Update(reinterpret_cast<const uint8*>(&MetallicRoughness.RoughnessFactor), sizeof(MetallicRoughness.RoughnessFactor));
+
+		HashTextureMap(MD5, SpecularGlossiness.Map);
+		MD5.Update(reinterpret_cast<const uint8*>(&SpecularGlossiness.GlossinessFactor), sizeof(SpecularGlossiness.GlossinessFactor));
+		MD5.Update(reinterpret_cast<const uint8*>(&SpecularGlossiness.SpecularFactor), sizeof(SpecularGlossiness.SpecularFactor));
+
+		HashTextureMap(MD5, Normal);
+		MD5.Update(reinterpret_cast<const uint8*>(&NormalScale), sizeof(NormalScale));
+
+		HashTextureMap(MD5, Occlusion);
+		MD5.Update(reinterpret_cast<const uint8*>(&OcclusionStrength), sizeof(OcclusionStrength));
+
+		HashTextureMap(MD5, Emissive);
+		MD5.Update(reinterpret_cast<const uint8*>(&EmissiveFactor), sizeof(EmissiveFactor));
+
+		MD5.Update(&IsDoubleSized, sizeof(IsDoubleSized));
+		MD5.Update(&AlphaModeInt, sizeof(AlphaModeInt));
+		MD5.Update(reinterpret_cast<const uint8*>(&AlphaCutoff), sizeof(AlphaCutoff));
+
+		MD5.Update(reinterpret_cast<const uint8*>(&Packing.Flags), sizeof(Packing.Flags));
+		HashTextureMap(MD5, Packing.Map);
+		HashTextureMap(MD5, Packing.NormalMap);
+
+		MD5.Update(&IsUnlitShadingModel, sizeof(IsUnlitShadingModel));
+
+		FMD5Hash Hash;
+		Hash.Set(MD5);
+		return Hash;
+	}
+}; // namespace GLTF
diff --git a/Engine/Plugins/Interchange/Runtime/Source/Parsers/GLTFCore/Private/GLTF/GLTFMesh.cpp b/Engine/Plugins/Interchange/Runtime/Source/Parsers/GLTFCore/Private/GLTF/GLTFMesh.cpp
new file mode 100644
--- /dev/null
+++ b/Engine/Plugins/Interchange/Runtime/Source/Parsers/GLTFCore/Private/GLTF/GLTFMesh.cpp
@@ -0,0 +1,634 @@
+// Copyright Epic Games, Inc. All Rights Reserved.
+
+#include "GLTFMesh.h"
+
+#include "ConversionUtilities.h"
+
+namespace GLTF
+{
+	namespace
+	{
+		// convertible to 0.0 .. 1.0 factor
+		// (colors, tex coords, weights, etc.)
+		bool IsConvertibleToNormalizedFloat(const FAccessor& Attrib)
+		{
+			return Attrib.ComponentType == FAccessor::EComponentType::F32 ||
+				(Attrib.Normalized &&
+					(Attrib.ComponentType == FAccessor::EComponentType::U8 || Attrib.ComponentType == FAccessor::EComponentType::U16));
+		}
+	}
+
+	FMorphTarget::FMorphTarget(const FAccessor& InPositionDisplacements, const FAccessor& InNormalDisplacements,
+		const FAccessor& InTangentDisplacements, const FAccessor& InTexCoord0Displacements, const FAccessor& InTexCoord1Displacements, const FAccessor& InColor0Deltas)
+		: PositionDisplacements(InPositionDisplacements)
+		, NormalDisplacements(InNormalDisplacements)
+		, TangentDisplacements(InTangentDisplacements)
+		, TexCoord0Displacements(InTexCoord0Displacements)
+		, TexCoord1Displacements(InTexCoord1Displacements)
+		, Color0Deltas(InColor0Deltas)
+	{
+	}
+
+	bool FMorphTarget::IsValid() const
+	{
+		return HasPositionDisplacements() || HasNormalDisplacements() || HasTangentDisplacements() || HasTexCoordDisplacements(0) || HasTexCoordDisplacements(1) || HasColorDeltas();
+	}
+	FMD5Hash FMorphTarget::GetHash() const
+	{
+		if (!IsValid())
+		{
+			return FMD5Hash();
+		}
+
+		auto HashAccessor = [](FMD5& MD5, const FAccessor& Accessor)
+		{
+			if (!Accessor.IsValid())
+			{
+				return;
+			}
+
+			FMD5Hash MD5Hash = Accessor.GetHash();
+			if (MD5Hash.IsValid())
+			{
+				MD5.Update(MD5Hash.GetBytes(), MD5Hash.GetSize());
+			}
+		};
+
+		FMD5 MD5;
+
+		HashAccessor(MD5, PositionDisplacements);
+		HashAccessor(MD5, NormalDisplacements);
+		HashAccessor(MD5, TangentDisplacements);
+		HashAccessor(MD5, TexCoord0Displacements);
+		HashAccessor(MD5, TexCoord1Displacements);
+		HashAccessor(MD5, Color0Deltas);
+
+		FMD5Hash Hash;
+		Hash.Set(MD5);
+
+		return Hash;
+	}
+
+	bool FMorphTarget::HasPositionDisplacements() const
+	{
+		return PositionDisplacements.IsValid();
+	}
+	void FMorphTarget::GetPositionDisplacements(TArray<FVector3f>& Buffer) const
+	{
+		PositionDisplacements.GetCoordArray(Buffer);
+	}
+	int32 FMorphTarget::GetNumberOfPositionDisplacements() const
+	{
+		return PositionDisplacements.Count;
+	}
+
+	bool FMorphTarget::HasNormalDisplacements() const
+	{
+		return NormalDisplacements.IsValid();
+	}
+	void FMorphTarget::GetNormalDisplacements(TArray<FVector3f>& Buffer) const
+	{
+		NormalDisplacements.GetCoordArray(Buffer);
+	}
+	int32 FMorphTarget::GetNumberOfNormalDisplacements() const
+	{
+		return NormalDisplacements.Count;
+	}
+
+	bool FMorphTarget::HasTangentDisplacements() const
+	{
+		return TangentDisplacements.IsValid();
+	}
+	void FMorphTarget::GetTangentDisplacements(TArray<FVector4f>& Buffer) const
+	{
+		const int32 N = TangentDisplacements.Count;
+		Buffer.Reserve(N);
+		for (int32 Index = 0; Index < N; ++Index)
+		{
+			const FVector TangentVector = TangentDisplacements.GetVec3(Index);
+			Buffer.Emplace((FVector3f)GLTF::ConvertVec3(TangentVector), 0.f);
+		}
+	}
+	int32 FMorphTarget::GetNumberOfTangentDisplacements() const
+	{
+		return TangentDisplacements.Count;
+	}
+
+	bool FMorphTarget::HasTexCoordDisplacements(uint32 Index) const
+	{
+		switch (Index)
+		{
+			case 0:
+				return TexCoord0Displacements.IsValid();
+			case 1:
+				return TexCoord1Displacements.IsValid();
+			default:
+				return false;
+		}
+	}
+	void FMorphTarget::GetTexCoordDisplacements(uint32 Index, TArray<FVector2f>& Buffer) const
+	{
+		switch (Index)
+		{
+		case 0:
+			return TexCoord0Displacements.GetVec2Array(Buffer);
+		case 1:
+			return TexCoord1Displacements.GetVec2Array(Buffer);
+		default:
+			break;
+		}
+	}
+	int32 FMorphTarget::GetNumberOfTexCoordDisplacements(uint32 Index) const
+	{
+		switch (Index)
+		{
+		case 0:
+			return TexCoord0Displacements.Count;
+		case 1:
+			return TexCoord1Displacements.Count;
+		default:
+			return 0;
+		}
+	}
+
+	bool FMorphTarget::HasColorDeltas() const
+	{
+		return Color0Deltas.IsValid() && (Color0Deltas.Type == FAccessor::EType::Vec3 || Color0Deltas.Type == FAccessor::EType::Vec4);
+	}
+	void FMorphTarget::GetColorDeltas(TArray<FVector4f>& Buffer) const
+	{
+		if (Color0Deltas.Type == FAccessor::EType::Vec4)
+		{
+			Color0Deltas.GetVec4Array(Buffer);
+		}
+		else if (Color0Deltas.Type == FAccessor::EType::Vec3)
+		{
+			const int32 N = Color0Deltas.Count;
+			Buffer.Reserve(N);
+			for (int32 Index = 0; Index < N; ++Index)
+			{
+				const FVector Vec = Color0Deltas.GetVec3(Index);
+				Buffer.Emplace((FVector3f)Vec, 0.f);
+			}
+		}
+	}
+	int32 FMorphTarget::GetNumberOfColorDeltas() const
+	{
+		return (Color0Deltas.Type == FAccessor::EType::Vec3 || Color0Deltas.Type == FAccessor::EType::Vec4) ? Color0Deltas.Count : 0;
+	}
+
+
+	FPrimitive::FPrimitive(EMode InMode, int32 InMaterial, const FAccessor& InIndices, const FAccessor& InPosition, const FAccessor& InNormal,
+	                       const FAccessor& InTangent, const FAccessor& InTexCoord0, const FAccessor& InTexCoord1, const FAccessor& InColor0,
+	                       const FAccessor& InJoints0, const FAccessor& InWeights0)
+	    : Mode(InMode)
+	    , MaterialIndex(InMaterial)
+	    , Indices(InIndices)
+	    , Position(InPosition)
+	    , Normal(InNormal)
+	    , Tangent(InTangent)
+	    , TexCoord0(InTexCoord0)
+	    , TexCoord1(InTexCoord1)
+	    , Color0(InColor0)
+	    , Joints0(InJoints0)
+	    , Weights0(InWeights0)
+	{
+	}
+
+	void FPrimitive::GetTangents(TArray<FVector4f>& Buffer) const
+	{
+		Buffer.Reserve(Tangent.Count);
+		for (uint32 Index = 0; Index < Tangent.Count; ++Index)
+		{
+			Buffer.Push((FVector4f)GLTF::ConvertTangent(Tangent.GetVec4(Index)));
+		}
+	}
+
+	void FPrimitive::GetColors(TArray<FVector4f>& Buffer) const
+	{
+		if (Color0.Type == FAccessor::EType::Vec4)
+		{
+			Color0.GetVec4Array(Buffer);
+		}
+		else if (Color0.Type == FAccessor::EType::Vec3)
+		{
+			const int32 N = Color0.Count;
+			Buffer.Reserve(N);
+			for (int32 Index = 0; Index < N; ++Index)
+			{
+				const FVector Vec = Color0.GetVec3(Index);
+				Buffer.Emplace((FVector3f)Vec, 1.f);
+			}
+		}
+	}
+
+	void FPrimitive::GetJointInfluences(TArray<FJointInfluence>& Buffer) const
+	{
+		// return a flat array that corresponds 1-to-1 with vertex positions
+		const int32 N = Joints0.Count;
+		Buffer.Reserve(N);
+		for (int32 Index = 0; Index < N; ++Index)
+		{
+			FJointInfluence& Joint = Buffer.Emplace_GetRef(Weights0.GetVec4(Index));
+			Joints0.GetUnsignedInt16x4(Index, Joint.ID);
+		}
+	}
+
+	FTriangle FPrimitive::TriangleVerts(uint32 T) const
+	{
+		FTriangle Result;
+		if (T >= TriangleCount())
+			return Result;
+
+		const bool Indexed = Indices.IsValid();
+		switch (Mode)
+		{
+			case EMode::Triangles:
+				if (Indexed)
+				{
+					Result.A = Indices.GetUnsignedInt(3 * T);
+					Result.B = Indices.GetUnsignedInt(3 * T + 1);
+					Result.C = Indices.GetUnsignedInt(3 * T + 2);
+				}
+				else
+				{
+					Result.A = 3 * T;
+					Result.B = 3 * T + 1;
+					Result.C = 3 * T + 2;
+				}
+				break;
+			case EMode::TriangleStrip:
+				// are indexed TriangleStrip & TriangleFan valid?
+				// I don't see anything in the spec that says otherwise...
+				if (Indexed)
+				{
+					if (T % 2 == 0)
+					{
+						Result.A = Indices.GetUnsignedInt(T);
+						Result.B = Indices.GetUnsignedInt(T + 1);
+					}
+					else
+					{
+						Result.A = Indices.GetUnsignedInt(T + 1);
+						Result.B = Indices.GetUnsignedInt(T);
+					}
+					Result.C = Indices.GetUnsignedInt(T + 2);
+				}
+				else
+				{
+					if (T % 2 == 0)
+					{
+						Result.A = T;
+						Result.B = T + 1;
+					}
+					else
+					{
+						Result.A = T + 1;
+						Result.B = T;
+					}
+					Result.C = T + 2;
+				}
+				break;
+			case EMode::TriangleFan:
+				if (Indexed)
+				{
+					Result.A = Indices.GetUnsignedInt(0);
+					Result.B = Indices.GetUnsignedInt(T + 1);
+					Result.C = Indices.GetUnsignedInt(T + 2);
+				}
+				else
+				{
+					Result.A = 0;
+					Result.B = T + 1;
+					Result.C = T + 2;
+				}
+				break;
+			default:
+				break;
+		}
+
+		return Result;
+	}
+
+	void FPrimitive::GetTriangleIndices(TArray<uint32>& Buffer) const
+	{
+		if (Mode == EMode::Triangles)
+		{
+			if (Indices.IsValid())
+			{
+				return Indices.GetUnsignedIntArray(Buffer);
+			}
+			else
+			{
+				// generate indices [0 1 2][3 4 5]...
+				const uint32 N = TriangleCount() * 3;
+				Buffer.Reserve(N);
+				for (uint32 Index = 0; Index < N; ++Index)
+				{
+					Buffer.Push(Index);
+				}
+			}
+		}
+		else
+		{
+			const uint32 N = TriangleCount() * 3;
+			Buffer.Reserve(N);
+			for (uint32 Index = 0; Index < TriangleCount(); ++Index)
+			{
+				FTriangle Tri = TriangleVerts(Index);
+				Buffer.Push(Tri.A);
+				Buffer.Push(Tri.B);
+				Buffer.Push(Tri.C);
+			}
+		}
+	}
+
+	uint32 FPrimitive::VertexCount() const
+	{
+		if (Indices.IsValid())
+			return Indices.Count;
+		else
+			return Position.Count;
+	}
+
+	uint32 FPrimitive::TriangleCount() const
+	{
+		switch (Mode)
+		{
+			case EMode::Triangles:
+				return VertexCount() / 3;
+			case EMode::TriangleStrip:
+			case EMode::TriangleFan:
+				return VertexCount() - 2;
+			default:
+				return 0;
+		}
+	}
+
+	void FPrimitive::GenerateIsValidCache()
+	{
+		bIsValidCache = IsValidPrivate();
+	}
+
+	bool FPrimitive::IsValid() const
+	{
+		if (bIsValidCache.IsSet())
+		{
+			return bIsValidCache.GetValue();
+		}
+		else
+		{
+			return IsValidPrivate();
+		}
+	}
+
+	bool FPrimitive::IsValidPrivate() const
+	{
+		// make sure all semantic attributes meet the spec
+
+		if (!Position.IsValid())
+		{
+			return false;
+		}
+
+		const uint32 VertexCount = Position.Count;
+
+		if (Position.Type != FAccessor::EType::Vec3 || Position.ComponentType != FAccessor::EComponentType::F32)
+		{
+			return false;
+		}
+
+		if (Normal.IsValid())
+		{
+			if (Normal.Count != VertexCount)
+			{
+				return false;
+			}
+
+			if (Normal.Type != FAccessor::EType::Vec3 || Normal.ComponentType != FAccessor::EComponentType::F32)
+			{
+				return false;
+			}
+		}
+
+		if (Tangent.IsValid())
+		{
+			if (Tangent.Count != VertexCount)
+			{
+				return false;
+			}
+
+			if (Tangent.Type != FAccessor::EType::Vec4 || Tangent.ComponentType != FAccessor::EComponentType::F32)
+			{
+				return false;
+			}
+		}
+
+		const FAccessor* TexCoords[] = { &TexCoord0, &TexCoord1 };
+		for (const FAccessor* TexCoord : TexCoords)
+		{
+			if (TexCoord->IsValid())
+			{
+				if (TexCoord->Count != VertexCount)
+				{
+					return false;
+				}
+
+				if (TexCoord->Type != FAccessor::EType::Vec2 || !IsConvertibleToNormalizedFloat(*TexCoord))
+				{
+					return false;
+				}
+			}
+		}
+
+		if (Color0.IsValid())
+		{
+			if (Color0.Count != VertexCount)
+			{
+				return false;
+			}
+
+			if (!(Color0.Type == FAccessor::EType::Vec3 || Color0.Type == FAccessor::EType::Vec4) || !IsConvertibleToNormalizedFloat(Color0))
+			{
+				return false;
+			}
+		}
+
+		//Validate Ranges:
+		{
+			TArray<uint32> AttributesSizes;
+			const FAccessor* AttributeAccessors[] = { &Position, &Normal, &Tangent, &TexCoord0, &TexCoord1, &Color0, &Joints0, &Weights0 };
+			for (const FAccessor* AttributeAccessor : AttributeAccessors)
+			{
+				if (AttributeAccessor->IsValid() && (AttributeAccessor->Count > 0))
+				{
+					AttributesSizes.Add(AttributeAccessor->Count);
+				}
+			}
+
+			TArray<uint32> IndicesValues;
+			GetTriangleIndices(IndicesValues);
+
+			for (uint32 Index : IndicesValues)
+			{
+				for (const uint32& AttributeSize : AttributesSizes)
+				{
+					if (AttributeSize <= Index)
+					{
+						return false;
+					}
+				}
+			}
+		}
+
+		//Morph Targets:
+		for (const FMorphTarget& MorphTarget : MorphTargets)
+		{
+			if (!MorphTarget.IsValid())
+			{
+				return false;
+			}
+
+			if (MorphTarget.HasPositionDisplacements() && Position.Count != MorphTarget.GetNumberOfPositionDisplacements())
+			{
+				return false;
+			}
+
+			if (MorphTarget.HasNormalDisplacements() && Normal.Count != MorphTarget.GetNumberOfNormalDisplacements())
+			{
+				return false;
+			}
+
+			if (MorphTarget.HasTangentDisplacements() && Tangent.Count != MorphTarget.GetNumberOfTangentDisplacements())
+			{
+				return false;
+			}
+
+			if (MorphTarget.HasTexCoordDisplacements(0) && TexCoord0.Count != MorphTarget.GetNumberOfTexCoordDisplacements(0))
+			{
+				return false;
+			}
+			if (MorphTarget.HasTexCoordDisplacements(1) && TexCoord0.Count != MorphTarget.GetNumberOfTexCoordDisplacements(1))
+			{
+				return false;
+			}
+
+			if (MorphTarget.HasColorDeltas() && Color0.Count != MorphTarget.GetNumberOfColorDeltas())
+			{
+				return false;
+			}
+		}
+
+		return true;
+	}
+
+	FMD5Hash FPrimitive::GetHash() const
+	{
+		if (!IsValid())
+		{
+			return FMD5Hash();
+		}
+
+		auto HashAccessor = [](FMD5& MD5, const FAccessor& Accessor)
+		{
+			if (!Accessor.IsValid())
+			{
+				return;
+			}
+
+			FMD5Hash MD5Hash = Accessor.GetHash();
+			if (MD5Hash.IsValid())
+			{
+				MD5.Update(MD5Hash.GetBytes(), MD5Hash.GetSize());
+			}
+		};
+
+		FMD5 MD5;
+
+		uint8 ModeInt = static_cast<uint8>(Mode);
+		MD5.Update(&ModeInt, sizeof(ModeInt));
+		MD5.Update(reinterpret_cast<const uint8*>(&MaterialIndex), sizeof(MaterialIndex));
+
+		HashAccessor(MD5, Indices);
+		HashAccessor(MD5, Position);
+
+		if (HasNormals())
+		{
+			HashAccessor(MD5, Normal);
+		}
+
+		if (HasTangents())
+		{
+			HashAccessor(MD5, Tangent);
+		}
+
+		if (HasTexCoords(0))
+		{
+			HashAccessor(MD5, TexCoord0);
+		}
+
+		if (HasTexCoords(1))
+		{
+			HashAccessor(MD5, TexCoord1);
+		}
+
+		if (HasColors())
+		{
+			HashAccessor(MD5, Color0);
+		}
+
+		if (HasJointWeights())
+		{
+			HashAccessor(MD5, Joints0);
+			HashAccessor(MD5, Weights0);
+		}
+
+		//Morph Targets:
+		for (const FMorphTarget& MorphTarget : MorphTargets)
+		{
+			FMD5Hash MorphTargetHash = MorphTarget.GetHash();
+			if (MorphTargetHash.IsValid())
+			{
+				MD5.Update(MorphTargetHash.GetBytes(), MorphTargetHash.GetSize());
+			}
+		}
+
+		FMD5Hash Hash;
+		Hash.Set(MD5);
+		return Hash;
+	}
+
+	FMD5Hash FMesh::GetHash() const
+	{
+		if (!IsValid())
+		{
+			return FMD5Hash();
+		}
+
+		FMD5 MD5;
+
+		for (const GLTF::FPrimitive& Primitive : Primitives)
+		{
+			FMD5Hash PrimitiveHash = Primitive.GetHash();
+			if (PrimitiveHash.IsValid())
+			{
+				MD5.Update(PrimitiveHash.GetBytes(), PrimitiveHash.GetSize());
+			}
+		}
+
+		FMD5Hash Hash;
+		Hash.Set(MD5);
+		return Hash;
+	}
+
+	//This is a helper functions for the validators.
+	//We are allowed to check only the first primitive's morph target count,
+	// because if there is a incosonsitency acorss primitive targets then the mesh IsValid will report false already.
+	int32 FMesh::NumberOfMorphTargetsPerPrimitive() const
+	{
+		int32 MorphTargetCounter = Primitives.Num() > 0 ? Primitives[0].MorphTargets.Num() : 0;
+
+		return MorphTargetCounter;
+	}
+
+}  // namespace GLTF
diff --git a/Engine/Plugins/Interchange/Runtime/Source/Parsers/GLTFCore/Private/GLTF/GLTFReader.cpp b/Engine/Plugins/Interchange/Runtime/Source/Parsers/GLTFCore/Private/GLTF/GLTFReader.cpp
new file mode 100644
--- /dev/null
+++ b/Engine/Plugins/Interchange/Runtime/Source/Parsers/GLTFCore/Private/GLTF/GLTFReader.cpp
@@ -0,0 +1,1090 @@
+// Copyright Epic Games, Inc. All Rights Reserved.
+
+#include "GLTFReader.h"
+
+#include "ExtensionsHandler.h"
+#include "GLTF/ConversionUtilities.h"
+#include "GLTFAsset.h"
+#include "GLTF/JsonUtilities.h"
+#include "GLTFBinaryReader.h"
+#include "GLTFNode.h"
+#include "MaterialUtilities.h"
+
+#include "HAL/PlatformFile.h"
+#include "HAL/PlatformFileManager.h"
+#include "Misc/Base64.h"
+#include "Misc/FileHelper.h"
+#include "Misc/Paths.h"
+#include "Serialization/JsonSerializer.h"
+
+#include "GenericPlatform/GenericPlatformHttp.h"
+
+namespace GLTF
+{
+	namespace
+	{
+		bool DecodeDataURI(const FString& URI, FString& OutMimeType, uint8* OutData, uint32& OutDataSize)
+		{
+			// Data URIs look like "data:[<mime-type>][;encoding],<data>"
+			// glTF always uses base64 encoding for data URIs
+
+			if (!ensure(URI.StartsWith(TEXT("data:"))))
+			{
+				return false;
+			}
+
+			int32      Semicolon, Comma;
+			const bool HasSemicolon = URI.FindChar(TEXT(';'), Semicolon);
+			const bool HasComma     = URI.FindChar(TEXT(','), Comma);
+			if (!(HasSemicolon && HasComma))
+			{
+				return false;
+			}
+
+			const FString Encoding = URI.Mid(Semicolon + 1, Comma - Semicolon - 1);
+			if (Encoding != TEXT("base64"))
+			{
+				return false;
+			}
+
+			OutMimeType = URI.Mid(5, Semicolon - 5);
+
+			const FString EncodedData = URI.RightChop(Comma + 1);
+			OutDataSize               = FBase64::GetDecodedDataSize(EncodedData);
+			return FBase64::Decode(*EncodedData, EncodedData.Len(), OutData);
+		}
+
+		uint32 GetDecodedDataSize(const FString& URI, FString& OutMimeType)
+		{
+			// Data URIs look like "data:[<mime-type>][;encoding],<data>"
+			// glTF always uses base64 encoding for data URIs
+
+			if (!ensure(URI.StartsWith(TEXT("data:"))))
+			{
+				return false;
+			}
+
+			int32      Semicolon, Comma;
+			const bool HasSemicolon = URI.FindChar(TEXT(';'), Semicolon);
+			const bool HasComma     = URI.FindChar(TEXT(','), Comma);
+			if (!(HasSemicolon && HasComma))
+			{
+				return 0;
+			}
+
+			const FString Encoding = URI.Mid(Semicolon + 1, Comma - Semicolon - 1);
+			if (Encoding != TEXT("base64"))
+			{
+				return 0;
+			}
+
+			OutMimeType = URI.Mid(5, Semicolon - 5);
+
+			const FString EncodedData = URI.RightChop(Comma + 1);
+			return FBase64::GetDecodedDataSize(EncodedData);
+		}
+
+		const FAccessor& AccessorAtIndex(const TArray<FValidAccessor>& Accessors, int32 Index)
+		{
+			if (Accessors.IsValidIndex(Index))
+			{
+				return Accessors[Index];
+			}
+			else
+			{
+				static const FVoidAccessor Void;
+				return Void;
+			}
+		}
+	}
+
+	FFileReader::FFileReader()
+	    : BufferCount(0)
+	    , BufferViewCount(0)
+	    , ImageCount(0)
+	    , BinaryReader(new FBinaryFileReader())
+	    , ExtensionsHandler(new FExtensionsHandler(Messages))
+	    , Asset(nullptr)
+	{
+	}
+
+	FFileReader::~FFileReader() {}
+
+	void FFileReader::SetupBuffer(const FJsonObject& Object, const FString& Path)
+	{
+		const uint32 ByteLength = GetUnsignedInt(Object, TEXT("byteLength"), 0);
+		Asset->Buffers.Emplace(ByteLength);
+		FBuffer& Buffer = Asset->Buffers.Last();
+
+		bool bUpdateOffset = false;
+		if (Object.HasTypedField<EJson::String>(TEXT("uri")))
+		{
+			// set Buffer.Data from Object.uri
+
+			FString URI = Object.GetStringField(TEXT("uri"));
+			if (URI.StartsWith(TEXT("data:")))
+			{
+				FString MimeType;
+				uint32  DataSize = 0;
+				bool    bSuccess = DecodeDataURI(URI, MimeType, CurrentBufferOffset, DataSize);
+				if (!bSuccess || (MimeType != TEXT("application/octet-stream") && MimeType != TEXT("application/gltf-buffer")) || !ensure(DataSize == ByteLength))
+				{
+					Messages.Emplace(EMessageSeverity::Error, TEXT("Problem decoding buffer from data URI."));
+				}
+				else
+				{
+					bUpdateOffset = true;
+				}
+			}
+			else
+			{
+				URI = FGenericPlatformHttp::UrlDecode(URI);
+
+				// Load buffer from external file.
+				const FString FullPath = Path / URI;
+				FArchive*     Reader   = IFileManager::Get().CreateFileReader(*FullPath);
+				if (Reader)
+				{
+					const int64 FileSize = Reader->TotalSize();
+					if (ByteLength == FileSize)
+					{
+						Reader->Serialize(CurrentBufferOffset, ByteLength);
+						bUpdateOffset = true;
+					}
+					else
+					{
+						Messages.Emplace(EMessageSeverity::Error, TEXT("Buffer file size does not match."));
+					}
+
+					Reader->Close();
+					delete Reader;
+				}
+				else
+				{
+					Messages.Emplace(EMessageSeverity::Error, FString::Printf(TEXT("Could not load file: '%s'"), *FullPath));
+				}
+			}
+		}
+		else
+		{
+			// Missing URI means use binary chunk of GLB
+			const uint32 BinSize = Asset->BinData.Num();
+			if (BinSize == 0)
+			{
+				Messages.Emplace(EMessageSeverity::Error, TEXT("Buffer from BIN chunk is missing or empty."));
+			}
+			else if (BinSize < ByteLength)
+			{
+				Messages.Emplace(EMessageSeverity::Error, TEXT("Buffer from BIN chunk is too small."));
+			}
+			else
+			{
+				Buffer.Data = Asset->BinData.GetData();
+			}
+		}
+
+		if (bUpdateOffset)
+		{
+			Buffer.Data = CurrentBufferOffset;
+			CurrentBufferOffset += ByteLength;
+		}
+
+		ExtensionsHandler->SetupBufferExtensions(Object, Buffer);
+	}
+
+	void FFileReader::SetupBufferView(const FJsonObject& Object) const
+	{
+		const uint32 BufferIdx = GetUnsignedInt(Object, TEXT("buffer"), BufferCount);
+		if (BufferIdx < BufferCount)  // must be true
+		{
+			const uint32 ByteOffset = GetUnsignedInt(Object, TEXT("byteOffset"), 0);
+			const uint32 ByteLength = GetUnsignedInt(Object, TEXT("byteLength"), 0);
+			const uint32 ByteStride = GetUnsignedInt(Object, TEXT("byteStride"), 0);
+			Asset->BufferViews.Emplace(Asset->Buffers[BufferIdx], ByteOffset, ByteLength, ByteStride);
+			ExtensionsHandler->SetupBufferViewExtensions(Object, Asset->BufferViews.Last());
+		}
+	}
+
+	void FFileReader::SetupAccessor(const FJsonObject& Object) const
+	{
+		const uint32 BufferViewIdx = GetUnsignedInt(Object, TEXT("bufferView"), BufferViewCount);
+		if (BufferViewIdx < BufferViewCount)  // must be true
+		{
+			const uint32                    ByteOffset = GetUnsignedInt(Object, TEXT("byteOffset"), 0);
+			const FAccessor::EComponentType CompType   = ComponentTypeFromNumber(GetUnsignedInt(Object, TEXT("componentType"), 0));
+			const uint32                    Count      = GetUnsignedInt(Object, TEXT("count"), 0);
+			const FAccessor::EType          Type       = AccessorTypeFromString(Object.GetStringField("type"));
+			const bool                      Normalized = GetBool(Object, TEXT("normalized"));
+			Asset->Accessors.Emplace(Asset->BufferViews[BufferViewIdx], ByteOffset, Count, Type, CompType, Normalized);
+			ExtensionsHandler->SetupAccessorExtensions(Object, Asset->Accessors.Last());
+		}
+	}
+
+	void FFileReader::SetupMorphTarget(const FJsonObject& Object, GLTF::FPrimitive& Primitive) const
+	{
+		const TArray<FValidAccessor>& A = Asset->Accessors;
+		const FAccessor& Position = AccessorAtIndex(A, GetIndex(Object, TEXT("POSITION")));
+		const FAccessor& Normal = AccessorAtIndex(A, GetIndex(Object, TEXT("NORMAL")));
+		const FAccessor& Tangent = AccessorAtIndex(A, GetIndex(Object, TEXT("TANGENT")));
+		const FAccessor& TexCoord0 = AccessorAtIndex(A, GetIndex(Object, TEXT("TEXCOORD_0")));
+		const FAccessor& TexCoord1 = AccessorAtIndex(A, GetIndex(Object, TEXT("TEXCOORD_1")));
+		const FAccessor& Color0 = AccessorAtIndex(A, GetIndex(Object, TEXT("COLOR_0")));
+
+		Primitive.MorphTargets.Emplace(Position, Normal, Tangent, TexCoord0, TexCoord1, Color0);
+	}
+
+	void FFileReader::SetupPrimitive(const FJsonObject& Object, FMesh& Mesh) const
+	{
+		const FPrimitive::EMode       Mode = PrimitiveModeFromNumber(GetUnsignedInt(Object, TEXT("mode"), (uint32)FPrimitive::EMode::Triangles));
+		if (Mode == FPrimitive::EMode::Unknown)
+		{
+			return;
+		}
+
+		const int32                   MaterialIndex = GetIndex(Object, TEXT("material"));
+		const TArray<FValidAccessor>& A             = Asset->Accessors;
+
+		const FAccessor& Indices = AccessorAtIndex(A, GetIndex(Object, TEXT("indices")));
+
+		// the only required attribute is POSITION
+		const FJsonObject& Attributes = *Object.GetObjectField("attributes");
+		const FAccessor&   Position   = AccessorAtIndex(A, GetIndex(Attributes, TEXT("POSITION")));
+		const FAccessor&   Normal     = AccessorAtIndex(A, GetIndex(Attributes, TEXT("NORMAL")));
+		const FAccessor&   Tangent    = AccessorAtIndex(A, GetIndex(Attributes, TEXT("TANGENT")));
+		const FAccessor&   TexCoord0  = AccessorAtIndex(A, GetIndex(Attributes, TEXT("TEXCOORD_0")));
+		const FAccessor&   TexCoord1  = AccessorAtIndex(A, GetIndex(Attributes, TEXT("TEXCOORD_1")));
+		const FAccessor&   Color0     = AccessorAtIndex(A, GetIndex(Attributes, TEXT("COLOR_0")));
+		const FAccessor&   Joints0    = AccessorAtIndex(A, GetIndex(Attributes, TEXT("JOINTS_0")));
+		const FAccessor&   Weights0   = AccessorAtIndex(A, GetIndex(Attributes, TEXT("WEIGHTS_0")));
+
+		Mesh.Primitives.Emplace(Mode, MaterialIndex, Indices, Position, Normal, Tangent, TexCoord0, TexCoord1, Color0, Joints0, Weights0);
+
+		//Morph Targets:
+		if (Object.HasField(TEXT("targets")))
+		{
+			const TArray<TSharedPtr<FJsonValue> >& MorphTargets = Object.GetArrayField(TEXT("targets"));
+			for (TSharedPtr<FJsonValue> Value : MorphTargets)
+			{
+				const FJsonObject& MorphTargetObject = *Value->AsObject();
+				SetupMorphTarget(MorphTargetObject, Mesh.Primitives.Last());
+			}
+		}
+
+		ExtensionsHandler->SetupPrimitiveExtensions(Object, Mesh.Primitives.Last());
+	}
+
+	void FFileReader::SetupMesh(const FJsonObject& Object) const
+	{
+		Asset->Meshes.Emplace();
+		FMesh& Mesh = Asset->Meshes.Last();
+
+		const TArray<TSharedPtr<FJsonValue> >& PrimArray = Object.GetArrayField(TEXT("primitives"));
+
+		Mesh.Name = GetString(Object, TEXT("name"));
+		Mesh.Primitives.Reserve(PrimArray.Num());
+
+		int32 NumberOfMorphTargets = -1;
+		for (TSharedPtr<FJsonValue> Value : PrimArray)
+		{
+			const FJsonObject& PrimObject = *Value->AsObject();
+			SetupPrimitive(PrimObject, Mesh);
+
+			if (NumberOfMorphTargets == -1)
+			{
+				NumberOfMorphTargets = Mesh.Primitives.Last().MorphTargets.Num();
+			}
+			else
+			{
+				if (NumberOfMorphTargets != Mesh.Primitives.Last().MorphTargets.Num())
+				{
+					//All primitives MUST have the same number of morph targets in the same order.
+					Messages.Emplace(EMessageSeverity::Error, TEXT("Number of Primitive.Targets is not consistent across the Mesh."));
+				}
+			}
+		}
+
+		// Morph Target Weights:
+		if (Object.HasField(TEXT("weights")))
+
+		{
+			const TArray<TSharedPtr<FJsonValue> >& MorphTargetWeightsArray = Object.GetArrayField(TEXT("weights"));
+			for (TSharedPtr<FJsonValue> Value : MorphTargetWeightsArray)
+			{
+				Mesh.MorphTargetWeights.Add(Value->AsNumber());
+			}
+		}
+
+		// Morph Target Names:
+		if (Object.HasField(TEXT("extras")))
+		{
+			const TSharedPtr<FJsonObject>& Extras = Object.GetObjectField(TEXT("extras"));
+			if (Extras->HasField(TEXT("targetNames")))
+			{
+				const TArray<TSharedPtr<FJsonValue>>& TargetNamesArray = Extras->GetArrayField(TEXT("targetNames"));
+				for (TSharedPtr<FJsonValue> Value : TargetNamesArray)
+				{
+					Mesh.MorphTargetNames.Add(Value->AsString());
+				}
+			}
+		}
+
+		Mesh.GenerateIsValidCache();
+
+		ExtensionsHandler->SetupMeshExtensions(Object, Mesh);
+	}
+
+	void FFileReader::SetupScene(const FJsonObject& Object) const
+	{
+		FScene& Scene = Asset->Scenes.Emplace_GetRef();
+
+		Scene.Name = GetString(Object, TEXT("name"));
+		if (Object.HasField(TEXT("nodes")))
+		{
+			const TArray<TSharedPtr<FJsonValue> >& NodesArray = Object.GetArrayField(TEXT("nodes"));
+			Scene.Nodes.Reserve(NodesArray.Num());
+			for (TSharedPtr<FJsonValue> Value : NodesArray)
+			{
+				const int32 NodeIndex = Value->AsNumber();
+				Scene.Nodes.Add(NodeIndex);
+
+				BuildParentIndices(INDEX_NONE, NodeIndex);
+			}
+		}
+
+		ExtensionsHandler->SetupSceneExtensions(Object, Scene);
+	}
+
+	void FFileReader::SetupNode(const FJsonObject& Object) const
+	{
+		FNode& Node = Asset->Nodes.Emplace_GetRef();
+
+		Node.Name = GetString(Object, TEXT("name"));
+
+		if (Object.HasField(TEXT("matrix")))
+		{
+			Node.Transform.SetFromMatrix(GetMat4(Object, TEXT("matrix")));
+			Node.Transform.SetRotation(GLTF::ConvertQuat(Node.Transform.GetRotation()));
+		}
+		else
+		{
+			Node.Transform.SetTranslation(GetVec3(Object, TEXT("translation")));
+			Node.Transform.SetRotation(GetQuat(Object, TEXT("rotation")));
+			Node.Transform.SetScale3D(GetVec3(Object, TEXT("scale"), FVector::OneVector));
+		}
+		Node.Transform.SetTranslation(GLTF::ConvertVec3(Node.Transform.GetTranslation()));
+		Node.Transform.SetScale3D(GLTF::ConvertVec3(Node.Transform.GetScale3D()));
+
+		if (Object.HasField(TEXT("children")))
+		{
+			const TArray<TSharedPtr<FJsonValue> >& ChildArray = Object.GetArrayField(TEXT("children"));
+			Node.Children.Reserve(ChildArray.Num());
+			for (TSharedPtr<FJsonValue> Value : ChildArray)
+			{
+				const int32 ChildIndex = Value->AsNumber();
+				Node.Children.Add(ChildIndex);
+			}
+		}
+
+		Node.MeshIndex   = GetIndex(Object, TEXT("mesh"));
+		Node.Skindex     = GetIndex(Object, TEXT("skin"));
+		Node.CameraIndex = GetIndex(Object, TEXT("camera"));
+
+		if (Object.HasField(TEXT("weights")))
+		{
+			const TArray<TSharedPtr<FJsonValue> >& ChildArray = Object.GetArrayField(TEXT("weights"));
+			Node.MorphTargetWeights.Reserve(ChildArray.Num());
+			for (TSharedPtr<FJsonValue> Value : ChildArray)
+			{
+				Node.MorphTargetWeights.Add(Value->AsNumber());
+			}
+		}
+
+		ExtensionsHandler->SetupNodeExtensions(Object, Node);
+	}
+
+	void FFileReader::SetupCamera(const FJsonObject& Object) const
+	{
+		const uint32 CameraIndex = Asset->Cameras.Num();
+		const FNode* Found       = Asset->Nodes.FindByPredicate([CameraIndex](const FNode& Node) { return CameraIndex == Node.CameraIndex; });
+		FString Name = GetString(Object, TEXT("name"));
+		if (!Found)
+		{
+			Messages.Emplace(EMessageSeverity::Warning, FString::Printf(TEXT("No camera node found for camera %d('%s')"), CameraIndex, *Name));
+			return;
+		}
+
+		FCamera& Camera = Asset->Cameras.Emplace_GetRef(*Found);
+		Camera.Name     = GetString(Object, TEXT("name"));
+
+		const FString Type = GetString(Object, TEXT("type"));
+		if (Type == TEXT("perspective"))
+		{
+			const FJsonObject& Perspective = *Object.GetObjectField(Type);
+
+			Camera.ZNear                   = GetScalar(Perspective, TEXT("znear"), 0.f);
+			Camera.ZFar                    = GetScalar(Perspective, TEXT("zfar"), Camera.ZNear + 10.f);
+			Camera.Perspective.AspectRatio = GetScalar(Perspective, TEXT("aspectRatio"), 1.f);
+			Camera.Perspective.Fov         = GetScalar(Perspective, TEXT("yfov"), 0.f);
+			Camera.bIsPerspective          = true;
+		}
+		else if (Type == TEXT("orthographic"))
+		{
+			const FJsonObject& Orthographic = *Object.GetObjectField(Type);
+
+			Camera.ZNear                       = GetScalar(Orthographic, TEXT("znear"), 0.f);
+			Camera.ZFar                        = GetScalar(Orthographic, TEXT("zfar"), Camera.ZNear + 10.f);
+			Camera.Orthographic.XMagnification = GetScalar(Orthographic, TEXT("xmag"), 0.f);
+			Camera.Orthographic.YMagnification = GetScalar(Orthographic, TEXT("ymag"), 0.f);
+			Camera.bIsPerspective              = false;
+		}
+		else
+			Messages.Emplace(EMessageSeverity::Error, TEXT("Invalid camera type: ") + Type);
+
+		ExtensionsHandler->SetupCameraExtensions(Object, Camera);
+	}
+
+	void FFileReader::SetupAnimation(const FJsonObject& Object) const
+	{
+		FAnimation& Animation = Asset->Animations.Emplace_GetRef();
+		Animation.Name        = GetString(Object, TEXT("name"));
+
+		// create samplers
+		{
+			const TArray<TSharedPtr<FJsonValue> >& SampplerArray = Object.GetArrayField(TEXT("samplers"));
+			Animation.Samplers.Reserve(SampplerArray.Num());
+			for (const TSharedPtr<FJsonValue>& Value : SampplerArray)
+			{
+				const FJsonObject& SamplerObject = *Value->AsObject();
+				const int32        Input         = GetIndex(SamplerObject, TEXT("input"));
+				const int32        Output        = GetIndex(SamplerObject, TEXT("output"));
+				if (ensure((Input != INDEX_NONE) && (Output != INDEX_NONE)))
+				{
+					FAnimation::FSampler Sampler(Asset->Accessors[Input], Asset->Accessors[Output]);
+					FString InterpolationStr = GetString(SamplerObject, TEXT("interpolation"), TEXT("LINEAR"));
+					if (InterpolationStr == TEXT("LINEAR"))
+					{
+						Sampler.Interpolation = FAnimation::EInterpolation::Linear;
+					}
+					else if (InterpolationStr == TEXT("STEP"))
+					{
+						Sampler.Interpolation = FAnimation::EInterpolation::Step;
+					}
+					else if (InterpolationStr == TEXT("CUBICSPLINE"))
+					{
+						Sampler.Interpolation = FAnimation::EInterpolation::CubicSpline;
+					}
+
+					Animation.Samplers.Add(Sampler);
+				}
+			}
+		}
+
+		// create channels
+		{
+			const TArray<TSharedPtr<FJsonValue> >& ChannelsArray = Object.GetArrayField(TEXT("channels"));
+			Animation.Channels.Reserve(ChannelsArray.Num());
+			for (const TSharedPtr<FJsonValue>& Value : ChannelsArray)
+			{
+				const FJsonObject& ChannelObject = *Value->AsObject();
+				const int32        Index         = GetIndex(ChannelObject, TEXT("sampler"));
+				if (!ensure(Index != INDEX_NONE))
+				{
+					continue;
+				}
+
+				const FJsonObject& TargetObject = *ChannelObject.GetObjectField(TEXT("target"));
+				const int32        NodeIndex    = GetIndex(TargetObject, TEXT("node"));
+				if (!ensure(NodeIndex != INDEX_NONE))
+				{
+					continue;
+				}
+
+				FAnimation::FChannel Channel(Asset->Nodes[NodeIndex]);
+				Channel.Sampler     = Index;
+				Channel.Target.Path = AnimationPathFromString(GetString(TargetObject, TEXT("path")));
+				Animation.Channels.Add(Channel);
+			}
+		}
+
+		ExtensionsHandler->SetupAnimationExtensions(Object, Animation);
+	}
+
+	void FFileReader::SetupSkin(const FJsonObject& Object) const
+	{
+		const FAccessor& InverseBindMatrices = AccessorAtIndex(Asset->Accessors, GetIndex(Object, TEXT("inverseBindMatrices")));
+
+		FSkinInfo& Skin = Asset->Skins.Emplace_GetRef(InverseBindMatrices);
+		Skin.Name       = GetString(Object, TEXT("name"));
+
+		const TArray<TSharedPtr<FJsonValue> >& JointArray = Object.GetArrayField(TEXT("joints"));
+		Skin.Joints.Reserve(JointArray.Num());
+		for (TSharedPtr<FJsonValue> Value : JointArray)
+		{
+			const int32 JointIndex = Value->AsNumber();
+			Skin.Joints.Add(JointIndex);
+		}
+
+		Skin.Skeleton = GetIndex(Object, TEXT("skeleton"));
+		if (Skin.Skeleton != INDEX_NONE)
+		{
+			//From gltf specification:
+			//If Skeleton is present:
+			//	"The index of the node used as a skeleton root."
+			//If it is not part of the Joints list yet, add it:
+			bool bSkeletonIsAlreadyJoint = false;
+			for (const int32& JointIndex : Skin.Joints)
+			{
+				if (JointIndex == Skin.Skeleton)
+				{
+					bSkeletonIsAlreadyJoint = true;
+					break;
+				}
+			}
+			if (!bSkeletonIsAlreadyJoint)
+			{
+				Skin.Joints.Add(Skin.Skeleton);
+			}
+		}
+
+		ExtensionsHandler->SetupSkinExtensions(Object, Skin);
+	}
+
+	void FFileReader::SetupImage(const FJsonObject& Object, const FString& Path, bool bInLoadImageData)
+	{
+		FImage& Image = Asset->Images.Emplace_GetRef();
+		Image.Name    = GetString(Object, TEXT("name"));
+
+		bool bUpdateOffset = false;
+		if (Object.HasTypedField<EJson::String>(TEXT("uri")))
+		{
+			// Get data now, so Unreal doesn't need to care about where the data came from.
+			// Unreal *is* responsible for decoding Data based on Format.
+
+			Image.URI = Object.GetStringField(TEXT("uri"));
+			if (Image.URI.StartsWith(TEXT("data:")))
+			{
+				uint32  ImageSize = 0;
+				FString MimeType;
+				bool    bSuccess = DecodeDataURI(Image.URI, MimeType, CurrentBufferOffset, ImageSize);
+				Image.Format     = ImageFormatFromMimeType(MimeType);
+				if (!bSuccess || Image.Format == FImage::EFormat::Unknown)
+				{
+					Messages.Emplace(EMessageSeverity::Error, TEXT("Problem decoding image from data URI."));
+				}
+				else
+				{
+					Image.DataByteLength = ImageSize;
+					bUpdateOffset        = true;
+				}
+			}
+			else  // Load buffer from external file.
+			{
+				Image.URI = FGenericPlatformHttp::UrlDecode(Image.URI);
+				Image.Format = ImageFormatFromFilename(Image.URI);
+
+				Image.FilePath = Path / Image.URI;
+				if (bInLoadImageData)
+				{
+					FArchive* Reader = IFileManager::Get().CreateFileReader(*Image.FilePath);
+					if (Reader)
+					{
+						const int64 FileSize = Reader->TotalSize();
+						Reader->Serialize(CurrentBufferOffset, FileSize);
+						Reader->Close();
+						delete Reader;
+
+						Image.DataByteLength = FileSize;
+						bUpdateOffset        = true;
+					}
+					else
+					{
+						Messages.Emplace(EMessageSeverity::Error, TEXT("Could not load image file."));
+					}
+				}
+			}
+		}
+		else
+		{
+			// Missing URI means use a BufferView
+			const int32 Index = GetIndex(Object, TEXT("bufferView"));
+			if (Asset->BufferViews.IsValidIndex(Index))
+			{
+				Image.Format = ImageFormatFromMimeType(GetString(Object, TEXT("mimeType")));
+
+				const FBufferView& BufferView = Asset->BufferViews[Index];
+				// We just created Image, so Image.Data is empty. Fill it with encoded bytes!
+				Image.DataByteLength = BufferView.ByteLength;
+				Image.Data           = static_cast<const uint8*>(BufferView.DataAt(0));
+			}
+		}
+
+		if (bUpdateOffset)
+		{
+			Image.Data = CurrentBufferOffset;
+			CurrentBufferOffset += Image.DataByteLength;
+		}
+
+		ExtensionsHandler->SetupImageExtensions(Object, Image);
+	}
+
+	void FFileReader::SetupSampler(const FJsonObject& Object) const
+	{
+		FSampler& Sampler = Asset->Samplers.Emplace_GetRef();
+
+		// spec doesn't specify default value, use linear
+		Sampler.MinFilter = FilterFromNumber(GetUnsignedInt(Object, TEXT("minFilter"), (uint32)FSampler::EFilter::Linear));
+		Sampler.MagFilter = FilterFromNumber(GetUnsignedInt(Object, TEXT("magFilter"), (uint32)FSampler::EFilter::Linear));
+		// default mode is Repeat according to spec
+		Sampler.WrapS = WrapModeFromNumber(GetUnsignedInt(Object, TEXT("wrapS"), (uint32)FSampler::EWrap::Repeat));
+		Sampler.WrapT = WrapModeFromNumber(GetUnsignedInt(Object, TEXT("wrapT"), (uint32)FSampler::EWrap::Repeat));
+
+		ExtensionsHandler->SetupSamplerExtensions(Object, Sampler);
+	}
+
+	void FFileReader::SetupTexture(const FJsonObject& Object) const
+	{
+		int32 SourceIndex  = GetIndex(Object, TEXT("source"));
+		int32 SamplerIndex = GetIndex(Object, TEXT("sampler"));
+
+		// According to the spec it's possible to have a Texture with no Image source.
+		// In that case use a default image (checkerboard?).
+
+		if (Asset->Images.IsValidIndex(SourceIndex))
+		{
+			const bool HasSampler = Asset->Samplers.IsValidIndex(SamplerIndex);
+
+			const FString TexName = GetString(Object, TEXT("name"));
+
+			const FImage&   Source  = Asset->Images[SourceIndex];
+			const FSampler& Sampler = HasSampler ? Asset->Samplers[SamplerIndex] : FSampler::DefaultSampler;
+
+			Asset->Textures.Emplace(TexName, Source, Sampler);
+			ExtensionsHandler->SetupTextureExtensions(Object, Asset->Textures.Last());
+		}
+		else
+		{
+			Messages.Emplace(EMessageSeverity::Warning, TEXT("Invalid texture source index: ") + FString::FromInt(SourceIndex));
+		}
+	}
+
+	void FFileReader::SetupMaterial(const FJsonObject& Object) const
+	{
+		Asset->Materials.Emplace(GetString(Object, TEXT("name")));
+		FMaterial& Material = Asset->Materials.Last();
+
+		GLTF::SetTextureMap(Object, TEXT("emissiveTexture"), nullptr, Asset->Textures, Material.Emissive, ExtensionsHandler->GetMessages());
+		Material.EmissiveFactor = (FVector3f)GetVec3(Object, TEXT("emissiveFactor"));
+
+		Material.NormalScale       = GLTF::SetTextureMap(Object, TEXT("normalTexture"), TEXT("scale"), Asset->Textures, Material.Normal, ExtensionsHandler->GetMessages());
+		Material.OcclusionStrength = GLTF::SetTextureMap(Object, TEXT("occlusionTexture"), TEXT("strength"), Asset->Textures, Material.Occlusion, ExtensionsHandler->GetMessages());
+
+		if (Object.HasTypedField<EJson::Object>(TEXT("pbrMetallicRoughness")))
+		{
+			const FJsonObject& PBR = *Object.GetObjectField(TEXT("pbrMetallicRoughness"));
+
+			GLTF::SetTextureMap(PBR, TEXT("baseColorTexture"), nullptr, Asset->Textures, Material.BaseColor, ExtensionsHandler->GetMessages());
+			Material.BaseColorFactor = (FVector4f)GetVec4(PBR, TEXT("baseColorFactor"), FVector4(1.0f, 1.0f, 1.0f, 1.0f));
+
+			GLTF::SetTextureMap(PBR, TEXT("metallicRoughnessTexture"), nullptr, Asset->Textures, Material.MetallicRoughness.Map, ExtensionsHandler->GetMessages());
+			Material.MetallicRoughness.MetallicFactor  = GetScalar(PBR, TEXT("metallicFactor"), 1.0f);
+			Material.MetallicRoughness.RoughnessFactor = GetScalar(PBR, TEXT("roughnessFactor"), 1.0f);
+		}
+
+		if (Object.HasTypedField<EJson::String>(TEXT("alphaMode")))
+		{
+			Material.AlphaMode = AlphaModeFromString(Object.GetStringField(TEXT("alphaMode")));
+			if (Material.AlphaMode == FMaterial::EAlphaMode::Mask)
+			{
+				Material.AlphaCutoff = GetScalar(Object, TEXT("alphaCutoff"), 0.5f);
+			}
+		}
+
+		Material.bIsDoubleSided = GetBool(Object, TEXT("doubleSided"));
+
+		ExtensionsHandler->SetupMaterialExtensions(Object, Material);
+	}
+
+	void FFileReader::ReadFile(const FString& InFilePath, bool bInLoadImageData, bool bInLoadMetadata, GLTF::FAsset& OutAsset)
+	{
+		Messages.Empty();
+
+		TUniquePtr<FArchive> FileReader(IFileManager::Get().CreateFileReader(*InFilePath));
+		TUniquePtr<FArchive> JsonFileReader;
+		if (!FileReader)
+		{
+			Messages.Emplace(EMessageSeverity::Error, TEXT("Can't load file: ") + InFilePath);
+			return;
+		}
+
+		const FString Extension = FPaths::GetExtension(InFilePath);
+		if (Extension == TEXT("gltf"))
+		{
+			// Convert to UTF8
+			FFileHelper::LoadFileToString(JsonBuffer, *InFilePath);
+		}
+		else if (Extension == TEXT("glb"))
+		{
+			BinaryReader->SetBuffer(OutAsset.BinData);
+			if (!BinaryReader->ReadFile(*FileReader))
+			{
+				Messages.Append(BinaryReader->GetLogMessages());
+				return;
+			}
+
+			// Convert to UTF8
+			const TArray<uint8>& Buffer = BinaryReader->GetJsonBuffer();
+			FFileHelper::BufferToString(JsonBuffer, Buffer.GetData(), Buffer.Num());
+		}
+		else
+		{
+			Messages.Emplace(EMessageSeverity::Error, TEXT("Invalid extension."));
+			return;
+		}
+		JsonFileReader.Reset(new FBufferReader(JsonBuffer.GetCharArray().GetData(), sizeof(FString::ElementType) * JsonBuffer.Len(), false));
+
+		JsonRoot                                                  = MakeShareable(new FJsonObject);
+		TSharedRef<TJsonReader<FString::ElementType> > JsonReader = TJsonReader<FString::ElementType>::Create(JsonFileReader.Get());
+		if (!FJsonSerializer::Deserialize(JsonReader, JsonRoot))
+		{
+			JsonRoot.Reset();
+			Messages.Emplace(EMessageSeverity::Error, TEXT("Problem loading JSON."));
+			return;
+		}
+
+		// Check file format version to make sure we can read it.
+		const TSharedPtr<FJsonObject>& AssetInfo = JsonRoot->GetObjectField(TEXT("asset"));
+		if (AssetInfo->HasTypedField<EJson::Number>(TEXT("minVersion")))
+		{
+			const double MinVersion = AssetInfo->GetNumberField(TEXT("minVersion"));
+			if (MinVersion > 2.0)
+			{
+				Messages.Emplace(EMessageSeverity::Error, TEXT("This importer supports glTF version 2.0 (or compatible) assets."));
+				return;
+			}
+			OutAsset.Metadata.Version = MinVersion;
+		}
+		else
+		{
+			const double Version = AssetInfo->GetNumberField(TEXT("version"));
+			if (Version < 2.0)
+			{
+				Messages.Emplace(EMessageSeverity::Error, TEXT("This importer supports glTF asset version 2.0 or later."));
+				return;
+			}
+			OutAsset.Metadata.Version = Version;
+		}
+		if (bInLoadMetadata)
+			LoadMetadata(OutAsset);
+
+		const FString ResourcesPath = FPaths::GetPath(InFilePath);
+		ImportAsset(ResourcesPath, bInLoadImageData, OutAsset);
+
+		// generate asset name
+		{
+			OutAsset.Name = FPaths::GetBaseFilename(InFilePath);
+			if (OutAsset.Name.ToLower() == TEXT("scene"))
+			{
+				// change name, try to see if asset title is given
+				if (const FMetadata::FExtraData* Extra = OutAsset.Metadata.GetExtraData(TEXT("title")))
+				{
+					OutAsset.Name = Extra->Value;
+				}
+				else
+				{
+					OutAsset.Name = FPaths::GetBaseFilename(FPaths::GetPath(InFilePath));
+				}
+			}
+		}
+		OutAsset.GenerateNames(OutAsset.Name);
+
+		if (OutAsset.ValidationCheck() != FAsset::Valid)
+		{
+			Messages.Emplace(EMessageSeverity::Warning, FString::Printf(TEXT("For GLTF Asset [%s] not all imported objects are valid."), *OutAsset.Name));
+		}
+
+		JsonRoot.Reset();
+	}
+
+	void FFileReader::LoadMetadata(GLTF::FAsset& OutAsset)
+	{
+		const TSharedPtr<FJsonObject>& AssetInfo = JsonRoot->GetObjectField(TEXT("asset"));
+		if (AssetInfo->HasField(TEXT("generator")))
+			OutAsset.Metadata.GeneratorName = AssetInfo->GetStringField(TEXT("generator"));
+
+		if (!AssetInfo->HasField(TEXT("extras")))
+			return;
+
+		const TSharedPtr<FJsonObject>& Extras = AssetInfo->GetObjectField(TEXT("extras"));
+		for (const auto& ValuePair : Extras->Values)
+		{
+			const TSharedPtr<FJsonValue>& JsonValue = ValuePair.Value;
+			OutAsset.Metadata.Extras.Emplace(FMetadata::FExtraData {ValuePair.Key, JsonValue->AsString()});
+		}
+	}
+
+	void FFileReader::AllocateExtraData(const FString& InResourcesPath, bool bInLoadImageData, TArray<uint8>& OutExtraData)
+	{
+		uint32 ExtraBufferSize = 0;
+		if (BufferCount > 0)
+		{
+			for (const TSharedPtr<FJsonValue>& Value : JsonRoot->GetArrayField(TEXT("buffers")))
+			{
+				const FJsonObject& Object     = *Value->AsObject();
+				const uint32       ByteLength = GetUnsignedInt(Object, TEXT("byteLength"), 0);
+				if (!Object.HasTypedField<EJson::String>(TEXT("uri")))
+					continue;
+
+				FString URI = Object.GetStringField(TEXT("uri"));
+				if (URI.StartsWith(TEXT("data:")))
+				{
+					FString      MimeType;
+					const uint32 DataSize = GetDecodedDataSize(URI, MimeType);
+					if (DataSize > 0 && (MimeType == TEXT("application/octet-stream") || MimeType == (TEXT("application/gltf-buffer"))))
+					{
+						ensure(DataSize == ByteLength);
+						ExtraBufferSize += ByteLength;
+					}
+				}
+				else
+				{
+					URI = FGenericPlatformHttp::UrlDecode(URI);
+					const FString FullPath = InResourcesPath / URI;
+					const int64   FileSize = FPlatformFileManager::Get().GetPlatformFile().FileSize(*FullPath);
+					if (ByteLength == FileSize)
+					{
+						ExtraBufferSize += ByteLength;
+					}
+				}
+			}
+		}
+
+		if (ImageCount)
+		{
+			for (const TSharedPtr<FJsonValue>& Value : JsonRoot->GetArrayField(TEXT("images")))
+			{
+				const FJsonObject& Object = *Value->AsObject();
+
+				if (!Object.HasTypedField<EJson::String>(TEXT("uri")))
+					continue;
+
+				FString URI = Object.GetStringField(TEXT("uri"));
+				if (URI.StartsWith(TEXT("data:")))
+				{
+					FString         MimeType;
+					const uint32    DataSize = GetDecodedDataSize(URI, MimeType);
+					FImage::EFormat Format   = ImageFormatFromMimeType(MimeType);
+					if (DataSize > 0 && Format != FImage::EFormat::Unknown)
+					{
+						ExtraBufferSize += DataSize;
+					}
+				}
+				else if (bInLoadImageData)
+				{
+					URI = FGenericPlatformHttp::UrlDecode(URI);
+
+					FImage::EFormat Format = ImageFormatFromFilename(URI);
+					if (Format != FImage::EFormat::Unknown)
+					{
+						const FString FullPath = InResourcesPath / URI;
+						const int64   FileSize = FPlatformFileManager::Get().GetPlatformFile().FileSize(*FullPath);
+						ExtraBufferSize += FileSize;
+					}
+				}
+			}
+		}
+
+		OutExtraData.Reserve(ExtraBufferSize + 16);
+		OutExtraData.SetNumUninitialized(ExtraBufferSize);
+		CurrentBufferOffset = ExtraBufferSize ? OutExtraData.GetData() : nullptr;
+	}
+
+	void FFileReader::ImportAsset(const FString& InResourcesPath, bool bInLoadImageData, GLTF::FAsset& OutAsset)
+	{
+		BufferCount                = ArraySize(*JsonRoot, TEXT("buffers"));
+		BufferViewCount            = ArraySize(*JsonRoot, TEXT("bufferViews"));
+		const uint32 AccessorCount = ArraySize(*JsonRoot, TEXT("accessors"));
+		const uint32 MeshCount     = ArraySize(*JsonRoot, TEXT("meshes"));
+
+		const uint32 SceneCount      = ArraySize(*JsonRoot, TEXT("scenes"));
+		const uint32 NodeCount       = ArraySize(*JsonRoot, TEXT("nodes"));
+		const uint32 CameraCount     = ArraySize(*JsonRoot, TEXT("cameras"));
+		const uint32 SkinCount       = ArraySize(*JsonRoot, TEXT("skins"));
+		const uint32 AnimationsCount = ArraySize(*JsonRoot, TEXT("animations"));
+
+		ImageCount                 = ArraySize(*JsonRoot, TEXT("images"));
+		const uint32 SamplerCount  = ArraySize(*JsonRoot, TEXT("samplers"));
+		const uint32 TextureCount  = ArraySize(*JsonRoot, TEXT("textures"));
+		const uint32 MaterialCount = ArraySize(*JsonRoot, TEXT("materials"));
+
+		const uint32 ExtensionsRequiredCount = ArraySize(*JsonRoot, TEXT("extensionsRequired"));
+
+		{
+			// cleanup and reserve
+			OutAsset.Buffers.Empty(BufferCount);
+			OutAsset.BufferViews.Empty(BufferViewCount);
+			OutAsset.Accessors.Empty(AccessorCount);
+			OutAsset.Meshes.Empty(MeshCount);
+			OutAsset.Scenes.Empty(SceneCount);
+			OutAsset.Nodes.Empty(NodeCount);
+			OutAsset.Cameras.Empty(CameraCount);
+			OutAsset.Lights.Empty(10);
+			OutAsset.Skins.Empty(SkinCount);
+			OutAsset.Animations.Empty(AnimationsCount);
+			OutAsset.Images.Empty(ImageCount);
+			OutAsset.Samplers.Empty(SamplerCount);
+			OutAsset.Textures.Empty(TextureCount);
+			OutAsset.Materials.Empty(MaterialCount);
+			OutAsset.ProcessedExtensions.Empty((int)EExtension::Count);
+			OutAsset.ExtensionsUsed.Empty();
+			OutAsset.ExtensionsRequired.Empty();
+		}
+
+		// allocate asset mapped data for images and buffers
+		AllocateExtraData(InResourcesPath, bInLoadImageData, OutAsset.ExtraBinData);
+
+		Asset = &OutAsset;
+		ExtensionsHandler->SetAsset(OutAsset);
+
+		SetupObjects(BufferCount, TEXT("buffers"), [this, InResourcesPath](const FJsonObject& Object) { SetupBuffer(Object, InResourcesPath); });
+		SetupObjects(BufferViewCount, TEXT("bufferViews"), [this](const FJsonObject& Object) { SetupBufferView(Object); });
+		SetupObjects(AccessorCount, TEXT("accessors"), [this](const FJsonObject& Object) { SetupAccessor(Object); });
+
+		SetupObjects(MeshCount, TEXT("meshes"), [this](const FJsonObject& Object) { SetupMesh(Object); });
+		SetupObjects(NodeCount, TEXT("nodes"), [this](const FJsonObject& Object) { SetupNode(Object); });
+		SetupObjects(SceneCount, TEXT("scenes"), [this](const FJsonObject& Object) { SetupScene(Object); });
+		SetupObjects(CameraCount, TEXT("cameras"), [this](const FJsonObject& Object) { SetupCamera(Object); });
+		SetupObjects(SkinCount, TEXT("skins"), [this](const FJsonObject& Object) { SetupSkin(Object); });
+		SetupObjects(AnimationsCount, TEXT("animations"), [this](const FJsonObject& Object) { SetupAnimation(Object); });
+
+		SetupObjects(ImageCount, TEXT("images"), [this, InResourcesPath, bInLoadImageData](const FJsonObject& Object) {
+			SetupImage(Object, InResourcesPath, bInLoadImageData);
+		});
+		SetupObjects(SamplerCount, TEXT("samplers"), [this](const FJsonObject& Object) { SetupSampler(Object); });
+		SetupObjects(TextureCount, TEXT("textures"), [this](const FJsonObject& Object) { SetupTexture(Object); });
+		SetupObjects(MaterialCount, TEXT("materials"), [this](const FJsonObject& Object) { SetupMaterial(Object); });
+
+		/**
+		 * Returns true if all required extensions are supported.
+		 * According to gltf specification checking only the top-level extensionsRequired property in order to decide
+		 * if the model import is supported should be sufficient as the documentation states:
+		 * "All glTF extensions required to load and/or render an asset MUST be listed in the top-level extensionsRequired array"
+		 */
+		const TArray<TSharedPtr<FJsonValue>>* ExtensionsRequired;
+		if (JsonRoot->TryGetArrayField(TEXT("extensionsRequired"), ExtensionsRequired))
+		{
+			if (ensure(ExtensionsRequired))
+			{
+				OutAsset.ExtensionsRequired.Reserve(OutAsset.ExtensionsRequired.Num() + ExtensionsRequired->Num());
+				for (const TSharedPtr<FJsonValue>& Extension : *ExtensionsRequired)
+				{
+					OutAsset.ExtensionsRequired.Add(Extension->AsString());
+				}
+			}
+		}
+
+		const TArray<TSharedPtr<FJsonValue>>* ExtensionsUsed;
+		if (JsonRoot->TryGetArrayField(TEXT("extensionsUsed"), ExtensionsUsed))
+		{
+			OutAsset.ExtensionsUsed.Reserve(ExtensionsUsed->Num());
+			for (const TSharedPtr<FJsonValue>& Extension : *ExtensionsUsed)
+			{
+				OutAsset.ExtensionsUsed.Add(Extension->AsString());
+			}
+		}
+
+		SetupNodesType();
+		ExtensionsHandler->SetupAssetExtensions(*JsonRoot);
+		BuildRootJoints();
+	}
+
+	template <typename SetupFunc>
+	void FFileReader::SetupObjects(uint32 ObjectCount, const TCHAR* FieldName, SetupFunc Func) const
+	{
+		if (ObjectCount > 0)
+		{
+			for (const TSharedPtr<FJsonValue>& Value : JsonRoot->GetArrayField(FieldName))
+			{
+				const FJsonObject& Object = *Value->AsObject();
+				Func(Object);
+			}
+		}
+	}
+
+	void FFileReader::SetupNodesType() const
+	{
+		// setup node types
+		for (FNode& Node : Asset->Nodes)
+		{
+			if (Node.MeshIndex != INDEX_NONE)
+			{
+				Node.Type = Node.Skindex != INDEX_NONE ? FNode::EType::MeshSkinned : FNode::EType::Mesh;
+			}
+			else if (Node.CameraIndex != INDEX_NONE)
+			{
+				Node.Type = FNode::EType::Camera;
+			}
+			else if (Node.LightIndex != INDEX_NONE)
+			{
+				Node.Type = FNode::EType::Light;
+			}
+			else
+			{
+				ensure(Node.Transform.IsValid());
+				if (!Node.Transform.GetRotation().IsIdentity() || !Node.Transform.GetTranslation().IsZero() ||
+				    !Node.Transform.GetScale3D().Equals(FVector(1.f)))
+				{
+					Node.Type = FNode::EType::Transform;
+				}
+			}
+		}
+		for (const FSkinInfo& Skin : Asset->Skins)
+		{
+			for (int32 JointIndex : Skin.Joints)
+			{
+				ensure(Asset->Nodes[JointIndex].Type == FNode::EType::None 
+					|| Asset->Nodes[JointIndex].Type == FNode::EType::Transform
+					|| Asset->Nodes[JointIndex].Type == FNode::EType::Joint);
+				Asset->Nodes[JointIndex].Type = FNode::EType::Joint;
+			}
+		}
+	}
+
+	void FFileReader::BuildParentIndices(int32 ParentNodeIndex, int32 CurrentNodeIndex) const
+	{
+		if (!Asset->Nodes.IsValidIndex(CurrentNodeIndex))
+		{
+			return;
+		}
+		GLTF::FNode& Node = Asset->Nodes[CurrentNodeIndex];
+		Node.ParentIndex = ParentNodeIndex;
+
+		for (const int32 ChildNodeIndex : Node.Children)
+		{
+			BuildParentIndices(CurrentNodeIndex, ChildNodeIndex);
+		}
+	}
+
+	int32 FFileReader::FindRootJointIndex(int32 CurrentIndex) const
+	{
+		if (!ensure(Asset->Nodes.IsValidIndex(CurrentIndex)))
+		{
+			return INDEX_NONE;
+		}
+		while (Asset->Nodes.IsValidIndex(Asset->Nodes[CurrentIndex].ParentIndex) && Asset->Nodes[Asset->Nodes[CurrentIndex].ParentIndex].Type == GLTF::FNode::EType::Joint)
+		{
+			CurrentIndex = Asset->Nodes[CurrentIndex].ParentIndex;
+		}
+		return CurrentIndex;
+	}
+	void FFileReader::BuildRootJoints() const
+	{
+		for (size_t Index = 0; Index < Asset->Nodes.Num(); Index++)
+		{
+			if (Asset->Nodes[Index].Type == GLTF::FNode::EType::Joint)
+			{
+				Asset->Nodes[Index].RootJointIndex = FindRootJointIndex(Index);
+			}
+		}
+	}
+
+}  // namespace GLTF
diff --git a/Engine/Plugins/Interchange/Runtime/Source/Parsers/GLTFCore/Private/GLTF/GLTFTexture.cpp b/Engine/Plugins/Interchange/Runtime/Source/Parsers/GLTFCore/Private/GLTF/GLTFTexture.cpp
new file mode 100644
--- /dev/null
+++ b/Engine/Plugins/Interchange/Runtime/Source/Parsers/GLTFCore/Private/GLTF/GLTFTexture.cpp
@@ -0,0 +1,9 @@
+// Copyright Epic Games, Inc. All Rights Reserved.
+
+#include "GLTFTexture.h"
+
+namespace GLTF
+{
+	const FSampler FSampler::DefaultSampler;
+
+}  // namespace GLTF
diff --git a/Engine/Plugins/Interchange/Runtime/Source/Parsers/GLTFCore/Private/GLTF/JsonUtilities.h b/Engine/Plugins/Interchange/Runtime/Source/Parsers/GLTFCore/Private/GLTF/JsonUtilities.h
new file mode 100644
--- /dev/null
+++ b/Engine/Plugins/Interchange/Runtime/Source/Parsers/GLTFCore/Private/GLTF/JsonUtilities.h
@@ -0,0 +1,147 @@
+// Copyright Epic Games, Inc. All Rights Reserved.
+
+#pragma once
+
+#include "ConversionUtilities.h"
+
+#include "Dom/JsonObject.h"
+
+namespace GLTF
+{
+	inline uint32 ArraySize(const FJsonObject& Object, const FString& Name)
+	{
+		if (Object.HasTypedField<EJson::Array>(Name))
+		{
+			return Object.GetArrayField(Name).Num();
+		}
+
+		return 0;  // for empty arrays & if array does not exist
+	}
+
+	inline FString GetString(const FJsonObject& Object, const TCHAR* Name, const FString& DefaultValue = FString())
+	{
+		if (Object.HasTypedField<EJson::String>(Name))
+		{
+			return Object.GetStringField(Name);
+		}
+
+		return DefaultValue;
+	}
+
+	inline bool GetBool(const FJsonObject& Object, const TCHAR* Name, bool DefaultValue = false)
+	{
+		if (Object.HasTypedField<EJson::Boolean>(Name))
+		{
+			return Object.GetBoolField(Name);
+		}
+
+		return DefaultValue;
+	}
+
+	inline uint32 GetUnsignedInt(const FJsonObject& Object, const TCHAR* Name, uint32 DefaultValue)
+	{
+		if (Object.HasTypedField<EJson::Number>(Name))
+		{
+			int32 SignedValue = Object.GetIntegerField(Name);
+			if (SignedValue >= 0)
+			{
+				return (uint32)SignedValue;
+			}
+			// complain if negative? if fractional?
+		}
+
+		return DefaultValue;
+	}
+
+	inline uint32 GetIndex(const FJsonObject& Object, const TCHAR* Name)
+	{
+		return GetUnsignedInt(Object, Name, INDEX_NONE);
+	}
+
+	inline float GetScalar(const FJsonObject& Object, const TCHAR* Name, float DefaultValue = 0.0f)
+	{
+		if (Object.HasTypedField<EJson::Number>(Name))
+		{
+			return Object.GetNumberField(Name);
+		}
+
+		return DefaultValue;
+	}
+
+	inline FVector GetVec3(const FJsonObject& Object, const TCHAR* Name, const FVector& DefaultValue = FVector::ZeroVector)
+	{
+		if (Object.HasTypedField<EJson::Array>(Name))
+		{
+			const TArray<TSharedPtr<FJsonValue> >& Array = Object.GetArrayField(Name);
+			if (Array.Num() == 3)
+			{
+				float X = Array[0]->AsNumber();
+				float Y = Array[1]->AsNumber();
+				float Z = Array[2]->AsNumber();
+				return FVector(X, Y, Z);
+			}
+		}
+
+		return DefaultValue;
+	}
+
+	inline FVector4 GetVec4(const FJsonObject& Object, const TCHAR* Name, const FVector4& DefaultValue = FVector4())
+	{
+		if (Object.HasTypedField<EJson::Array>(Name))
+		{
+			const TArray<TSharedPtr<FJsonValue> >& Array = Object.GetArrayField(Name);
+			if (Array.Num() == 4)
+			{
+				float X = Array[0]->AsNumber();
+				float Y = Array[1]->AsNumber();
+				float Z = Array[2]->AsNumber();
+				float W = Array[3]->AsNumber();
+				return FVector4(X, Y, Z, W);
+			}
+		}
+
+		return DefaultValue;
+	}
+
+	inline FQuat GetQuat(const FJsonObject& Object, const TCHAR* Name, const FQuat& DefaultValue = FQuat(0, 0, 0, 1))
+	{
+		if (Object.HasTypedField<EJson::Array>(Name))
+		{
+			const TArray<TSharedPtr<FJsonValue> >& Array = Object.GetArrayField(Name);
+			if (Array.Num() == 4)
+			{
+				float X = Array[0]->AsNumber();
+				float Y = Array[1]->AsNumber();
+				float Z = Array[2]->AsNumber();
+				float W = Array[3]->AsNumber();
+				return GLTF::ConvertQuat(FQuat(X, Y, Z, W));
+			}
+		}
+
+		return DefaultValue;
+	}
+
+	inline FMatrix GetMat4(const FJsonObject& Object, const TCHAR* Name, const FMatrix& DefaultValue = FMatrix::Identity)
+	{
+		if (Object.HasTypedField<EJson::Array>(Name))
+		{
+			const TArray<TSharedPtr<FJsonValue> >& Array = Object.GetArrayField(Name);
+			if (Array.Num() == 16)
+			{
+				FMatrix Matrix;
+				for (int32 Row = 0; Row < 4; ++Row)
+				{
+					for (int32 Col = 0; Col < 4; ++Col)
+					{
+						Matrix.M[Row][Col] = Array[Col * 4 + Row]->AsNumber();
+					}
+				}
+
+				return GLTF::ConvertMat(Matrix);
+			}
+		}
+
+		return DefaultValue;
+	}
+
+}  // namespace GLTF
diff --git a/Engine/Plugins/Interchange/Runtime/Source/Parsers/GLTFCore/Private/GLTF/MaterialUtilities.cpp b/Engine/Plugins/Interchange/Runtime/Source/Parsers/GLTFCore/Private/GLTF/MaterialUtilities.cpp
new file mode 100644
--- /dev/null
+++ b/Engine/Plugins/Interchange/Runtime/Source/Parsers/GLTFCore/Private/GLTF/MaterialUtilities.cpp
@@ -0,0 +1,100 @@
+// Copyright Epic Games, Inc. All Rights Reserved.
+
+#include "MaterialUtilities.h"
+#include "GLTFAsset.h"
+
+namespace GLTF
+{
+	// Returns scale factor if JSON has it, 1.0 by default.
+	float SetTextureMap(const FJsonObject& InObject, const TCHAR* InTexName, const TCHAR* InScaleName, const TArray<FTexture>& Textures,
+	                           GLTF::FTextureMap& OutMap, TArray<FLogMessage>& OutMessages)
+	{
+		float Scale = 1.0f;
+
+		if (InObject.HasTypedField<EJson::Object>(InTexName))
+		{
+			const FJsonObject& TexObj   = *InObject.GetObjectField(InTexName);
+			int32              TexIndex = GetIndex(TexObj, TEXT("index"));
+			if (Textures.IsValidIndex(TexIndex))
+			{
+				OutMap.TextureIndex = TexIndex;
+				uint32 TexCoord     = GetUnsignedInt(TexObj, TEXT("texCoord"), 0);
+
+				OutMap.TexCoord = TexCoord;
+				if (InScaleName && TexObj.HasTypedField<EJson::Number>(InScaleName))
+				{
+					Scale = TexObj.GetNumberField(InScaleName);
+				}
+
+				// Handle textureInfo extensions
+				if (TexObj.HasTypedField<EJson::Object>(TEXT("extensions")))
+				{
+					static const TArray<EExtension> SupportedTextureInfoExtensions = {
+						EExtension::KHR_TextureTransform
+					};
+					TArray<FString> SupportedTextureInfoExtensionsStringified;
+					for (size_t ExtensionIndex = 0; ExtensionIndex < SupportedTextureInfoExtensions.Num(); ExtensionIndex++)
+					{
+						SupportedTextureInfoExtensionsStringified.Add(GLTF::ToString(SupportedTextureInfoExtensions[ExtensionIndex]));
+					}
+
+					const FJsonObject& ExtensionsObj = *TexObj.GetObjectField(TEXT("extensions"));
+
+					for (int32 Index = 0; Index < SupportedTextureInfoExtensions.Num(); ++Index)
+					{
+						const FString ExtensionName = SupportedTextureInfoExtensionsStringified[Index];
+						if (!ExtensionsObj.HasTypedField<EJson::Object>(ExtensionName))
+						{
+							OutMessages.Emplace(EMessageSeverity::Warning, FString::Printf(TEXT("Extension is not supported: %s"), *ExtensionName));
+							continue;
+						}
+
+						const FJsonObject& ExtObj = *ExtensionsObj.GetObjectField(ExtensionName);
+
+						const EExtension Extension = SupportedTextureInfoExtensions[Index];
+						switch (Extension)
+						{
+							case EExtension::KHR_TextureTransform:
+							{
+								OutMap.bHasTextureTransform = true;
+								OutMap.TextureTransform.Offset[0] = 0.0f;
+								OutMap.TextureTransform.Offset[1] = 0.0f;
+								OutMap.TextureTransform.Scale[0] = 1.0f;
+								OutMap.TextureTransform.Scale[1] = 1.0f;
+								OutMap.TextureTransform.Rotation = GetScalar(ExtObj, TEXT("rotation"), 0.f);
+
+								if (ExtObj.HasTypedField<EJson::Array>(TEXT("offset")))
+								{
+									const TArray<TSharedPtr<FJsonValue>>& Offset = ExtObj.GetArrayField(TEXT("offset"));
+									if (Offset.Num() >= 2)
+									{
+										OutMap.TextureTransform.Offset[0] = static_cast<float>(Offset[0]->AsNumber());
+										OutMap.TextureTransform.Offset[1] = static_cast<float>(Offset[1]->AsNumber());
+									}
+								}
+
+								if (ExtObj.HasTypedField<EJson::Array>(TEXT("scale")))
+								{
+									const TArray<TSharedPtr<FJsonValue>>& TransformScale = ExtObj.GetArrayField(TEXT("scale"));
+									if(TransformScale.Num() >= 2)
+									{
+										OutMap.TextureTransform.Scale[0] = static_cast<float>(TransformScale[0]->AsNumber());
+										OutMap.TextureTransform.Scale[1] = static_cast<float>(TransformScale[1]->AsNumber());
+									}
+								}
+							}
+							break;
+						}
+					}
+				}
+			}
+		}
+		else
+		{
+			Scale = 0.f;
+		}
+
+		return Scale;
+	}
+
+}  // namespace GLTF
diff --git a/Engine/Plugins/Interchange/Runtime/Source/Parsers/GLTFCore/Private/GLTF/MaterialUtilities.h b/Engine/Plugins/Interchange/Runtime/Source/Parsers/GLTFCore/Private/GLTF/MaterialUtilities.h
new file mode 100644
--- /dev/null
+++ b/Engine/Plugins/Interchange/Runtime/Source/Parsers/GLTFCore/Private/GLTF/MaterialUtilities.h
@@ -0,0 +1,17 @@
+// Copyright Epic Games, Inc. All Rights Reserved.
+
+#pragma once
+
+#include "GLTFLogger.h"
+#include "JsonUtilities.h"
+
+namespace GLTF { struct FTexture; }
+namespace GLTF { struct FTextureMap; }
+
+namespace GLTF
+{
+	// Returns scale factor if JSON has it, 1.0 by default.
+	float SetTextureMap(const FJsonObject& InObject, const TCHAR* InTexName, const TCHAR* InScaleName, const TArray<FTexture>& Textures,
+		GLTF::FTextureMap& OutMap, TArray<FLogMessage>& OutMessages);
+
+}  // namespace GLTF
diff --git a/Engine/Plugins/Interchange/Runtime/Source/Parsers/GLTFCore/Private/GLTFCoreModule.cpp b/Engine/Plugins/Interchange/Runtime/Source/Parsers/GLTFCore/Private/GLTFCoreModule.cpp
new file mode 100644
--- /dev/null
+++ b/Engine/Plugins/Interchange/Runtime/Source/Parsers/GLTFCore/Private/GLTFCoreModule.cpp
@@ -0,0 +1,18 @@
+// Copyright Epic Games, Inc. All Rights Reserved.
+
+#include "GLTFCoreModule.h"
+
+
+/**
+ * glTF Core module implementation (private)
+ */
+class FGLTFCoreModule : public IGLTFCoreModule
+{
+
+public:
+	virtual void StartupModule() override {}
+
+	virtual void ShutdownModule() override {}
+};
+
+IMPLEMENT_MODULE(FGLTFCoreModule, GLTFCore);
diff --git a/Engine/Plugins/Interchange/Runtime/Source/Parsers/GLTFCore/Private/GLTFMapFactory.cpp b/Engine/Plugins/Interchange/Runtime/Source/Parsers/GLTFCore/Private/GLTFMapFactory.cpp
new file mode 100644
--- /dev/null
+++ b/Engine/Plugins/Interchange/Runtime/Source/Parsers/GLTFCore/Private/GLTFMapFactory.cpp
@@ -0,0 +1,354 @@
+// Copyright Epic Games, Inc. All Rights Reserved.
+
+#include "GLTFMapFactory.h"
+
+#include "GLTFMaterial.h"
+#include "GLTFMaterialExpressions.h"
+
+#include "GLTFMaterialFactory.h"
+
+namespace GLTF
+{
+	namespace
+	{
+		void CreateTextureCoordinate(int32 TexCoord, FMaterialExpressionTexture& TexExpression, FMaterialElement& MaterialElement)
+		{
+			if (TexCoord)
+			{
+				FMaterialExpressionTextureCoordinate* CoordExpression = MaterialElement.AddMaterialExpression<FMaterialExpressionTextureCoordinate>();
+				CoordExpression->SetCoordinateIndex(TexCoord);
+				CoordExpression->ConnectExpression(TexExpression.GetInputCoordinate(), 0);
+			}
+		}
+
+		void CreateTextureTransform(const FTextureTransform& InTextureTransform, FMaterialExpressionTexture& TexExpression, FMaterialElement& MaterialElement)
+		{
+			FMaterialExpression* CoordExpression = TexExpression.GetInputCoordinate().GetExpression();
+
+			if (CoordExpression == nullptr)
+			{
+				CoordExpression = MaterialElement.AddMaterialExpression<FMaterialExpressionTextureCoordinate>();
+			}
+
+			FMaterialExpression* LastExprInChain = nullptr;
+
+			FMaterialExpressionGeneric* UVMultiplyExpr = MaterialElement.AddMaterialExpression<FMaterialExpressionGeneric>();
+			UVMultiplyExpr->SetExpressionName(TEXT("Multiply"));
+
+			FMaterialExpressionGeneric* UVScale = MaterialElement.AddMaterialExpression<FMaterialExpressionGeneric>();
+			UVScale->SetExpressionName(TEXT("Constant2Vector"));
+			UVScale->SetFloatProperty(TEXT("R"), InTextureTransform.Scale[0]);
+			UVScale->SetFloatProperty(TEXT("G"), InTextureTransform.Scale[1]);
+
+			CoordExpression->ConnectExpression(*UVMultiplyExpr->GetInput(0), 0);
+			UVScale->ConnectExpression(*UVMultiplyExpr->GetInput(1), 0);
+
+			LastExprInChain = UVMultiplyExpr;
+
+			if (!FMath::IsNearlyZero(InTextureTransform.Rotation))
+			{
+				FMaterialExpressionFunctionCall* UVRotate = MaterialElement.AddMaterialExpression<FMaterialExpressionFunctionCall>();
+				UVRotate->SetFunctionPathName(TEXT("/Engine/Functions/Engine_MaterialFunctions02/Texturing/CustomRotator.CustomRotator"));
+
+				float AngleRadians = InTextureTransform.Rotation;
+
+				if (AngleRadians < 0.0f)
+				{
+					AngleRadians = TWO_PI - AngleRadians;
+				}
+
+				FMaterialExpressionScalar* RotationAngle = MaterialElement.AddMaterialExpression<FMaterialExpressionScalar>();
+				RotationAngle->GetScalar() = 1.0f - (AngleRadians / TWO_PI); // Normalize angle value
+
+				FMaterialExpressionGeneric* RotationPivot = MaterialElement.AddMaterialExpression<FMaterialExpressionGeneric>();
+				RotationPivot->SetExpressionName(TEXT("Constant2Vector"));
+				RotationPivot->SetFloatProperty(TEXT("R"), 0.0f);
+				RotationPivot->SetFloatProperty(TEXT("G"), 0.0f);
+
+				LastExprInChain->ConnectExpression(*UVRotate->GetInput(0), 0);
+				RotationPivot->ConnectExpression(*UVRotate->GetInput(1), 0);
+				RotationAngle->ConnectExpression(*UVRotate->GetInput(2), 0);
+
+				LastExprInChain = UVRotate;
+			}
+
+			if (!FMath::IsNearlyZero(InTextureTransform.Offset[0]) || 
+				!FMath::IsNearlyZero(InTextureTransform.Offset[1]))
+			{
+				FMaterialExpressionGeneric* UVOffset = MaterialElement.AddMaterialExpression<FMaterialExpressionGeneric>();
+				UVOffset->SetExpressionName(TEXT("Constant2Vector"));
+				UVOffset->SetFloatProperty(TEXT("R"), InTextureTransform.Offset[0]);
+				UVOffset->SetFloatProperty(TEXT("G"), InTextureTransform.Offset[1]);
+
+				FMaterialExpressionGeneric* AddExpr = MaterialElement.AddMaterialExpression<FMaterialExpressionGeneric>();
+				AddExpr->SetExpressionName(TEXT("Add"));
+
+				LastExprInChain->ConnectExpression(*AddExpr->GetInput(0), 0);
+				UVOffset->ConnectExpression(*AddExpr->GetInput(1), 0);
+
+				LastExprInChain = AddExpr;
+			}
+
+			LastExprInChain->ConnectExpression(TexExpression.GetInputCoordinate(), 0);
+		}
+
+		FMaterialExpressionInput& GetFirstInput(FMaterialExpression* Expression)
+		{
+			FMaterialExpressionInput* Input = Expression->GetInput(0);
+			if (Input->GetExpression() == nullptr)
+				return *Input;
+			else
+				return GetFirstInput(Input->GetExpression());
+		}
+
+		inline void SetExpresisonValue(float Value, FMaterialExpressionScalar* Expression)
+		{
+			Expression->GetScalar() = Value;
+		}
+
+		inline void SetExpresisonValue(const FVector3f& Color, FMaterialExpressionColor* Expression)
+		{
+			Expression->GetColor() = FLinearColor(Color);
+		}
+
+		inline void SetExpresisonValue(const FVector4f& Color, FMaterialExpressionColor* Expression)
+		{
+			Expression->GetColor() = FLinearColor(Color);
+		}
+	}
+
+	FPBRMapFactory::FPBRMapFactory(ITextureFactory& TextureFactory)
+	    : CurrentMaterialElement(nullptr)
+	    , TextureFactory(TextureFactory)
+	    , ParentPackage(nullptr)
+	    , Flags(RF_NoFlags)
+	{
+	}
+
+	void FPBRMapFactory::CreateNormalMap(const GLTF::FTexture& Map, int CoordinateIndex, float NormalScale, const GLTF::FTextureTransform* TextureTransform)
+	{
+		if (!ensure(CurrentMaterialElement))
+		{
+			return;
+		}
+
+		FMaterialExpressionTexture* TexExpression = CreateTextureMap(Map, CoordinateIndex, TEXT("Normal Map"), ETextureMode::Normal, TextureTransform);
+		if (!TexExpression)
+		{
+			return;
+		}
+
+		FMaterialExpressionScalar* ScalarExpression = CurrentMaterialElement->AddMaterialExpression<FMaterialExpressionScalar>();
+		ScalarExpression->SetName(TEXT("Normal Scale"));
+		ScalarExpression->SetGroupName(*GroupName);
+		ScalarExpression->GetScalar() = NormalScale;
+
+		// GLTF specifies that the following formula is used:
+		// scaledNormal = normalize((<sampled normal texture value> * 2.0 - 1.0) * vec3(<normal scale>, <normal scale>, 1.0)).
+		FMaterialExpressionGeneric* NormalXY = CurrentMaterialElement->AddMaterialExpression<FMaterialExpressionGeneric>();
+		NormalXY->SetExpressionName(TEXT("ComponentMask"));
+		NormalXY->SetBoolProperty(TEXT("R"), true);
+		NormalXY->SetBoolProperty(TEXT("G"), true);
+		NormalXY->SetBoolProperty(TEXT("B"), false);
+		TexExpression->ConnectExpression(*NormalXY->GetInput(0), 0);
+
+		FMaterialExpressionGeneric* NormalZ = CurrentMaterialElement->AddMaterialExpression<FMaterialExpressionGeneric>();
+		NormalZ->SetExpressionName(TEXT("ComponentMask"));
+		NormalZ->SetBoolProperty(TEXT("R"), false);
+		NormalZ->SetBoolProperty(TEXT("G"), false);
+		NormalZ->SetBoolProperty(TEXT("B"), true);
+		TexExpression->ConnectExpression(*NormalZ->GetInput(0), 0);
+
+		FMaterialExpressionGeneric* MultiplyXYExpression = CurrentMaterialElement->AddMaterialExpression<FMaterialExpressionGeneric>();
+		MultiplyXYExpression->SetExpressionName(TEXT("Multiply"));
+		NormalXY->ConnectExpression(*MultiplyXYExpression->GetInput(0), 0);
+		ScalarExpression->ConnectExpression(*MultiplyXYExpression->GetInput(1), 0);
+
+		FMaterialExpressionGeneric* ReconstructNormalVector = CurrentMaterialElement->AddMaterialExpression<FMaterialExpressionGeneric>();
+		ReconstructNormalVector->SetExpressionName(TEXT("AppendVector"));
+		MultiplyXYExpression->ConnectExpression(*ReconstructNormalVector->GetInput(0), 0);
+		NormalZ->ConnectExpression(*ReconstructNormalVector->GetInput(1), 0);
+
+		FMaterialExpressionGeneric* Normalize = CurrentMaterialElement->AddMaterialExpression<FMaterialExpressionGeneric>();
+		Normalize->SetExpressionName(TEXT("Normalize"));
+
+		ReconstructNormalVector->ConnectExpression(*Normalize->GetInput(0), 0);
+
+		Normalize->ConnectExpression(CurrentMaterialElement->GetNormal(), 0);
+	}
+
+	FMaterialExpression* FPBRMapFactory::CreateColorMap(const GLTF::FTexture& Map, int CoordinateIndex, const FVector3f& Color, const TCHAR* MapName,
+	                                                    const TCHAR* ValueName, ETextureMode TextureMode, FMaterialExpressionInput& MaterialInput,
+														const GLTF::FTextureTransform* TextureTransform)
+	{
+		return CreateMap<FMaterialExpressionColor>(Map, CoordinateIndex, Color, MapName, ValueName, TextureMode, MaterialInput, TextureTransform);
+	}
+
+	FMaterialExpression* FPBRMapFactory::CreateColorMap(const GLTF::FTexture& Map, int CoordinateIndex, const FVector4f& Color, const TCHAR* MapName,
+	                                                    const TCHAR* ValueName, ETextureMode TextureMode, FMaterialExpressionInput& MaterialInput, 
+														const GLTF::FTextureTransform* TextureTransform)
+	{
+		return CreateMap<FMaterialExpressionColor>(Map, CoordinateIndex, Color, MapName, ValueName, TextureMode, MaterialInput, TextureTransform);
+	}
+
+	void FPBRMapFactory::CreateMultiMap(const GLTF::FTexture& Map, int CoordinateIndex, const TCHAR* MapName, const FMapChannel* MapChannels,
+	                                    uint32 MapChannelsCount, ETextureMode TextureMode, const GLTF::FTextureTransform* TextureTransform)
+	{
+		FExpressionList ValueExpressions;
+		for (uint32 Index = 0; Index < MapChannelsCount; ++Index)
+		{
+			const FMapChannel& MapChannel = MapChannels[Index];
+
+			FMaterialExpressionParameter* ValueExpression = nullptr;
+			switch (MapChannel.Channel)
+			{
+				case EChannel::RG:
+				case EChannel::RGB:
+				case EChannel::All:
+				{
+					ValueExpression = CurrentMaterialElement->AddMaterialExpression<FMaterialExpressionColor>();
+					FMaterialExpressionColor* ColorExpression = static_cast<FMaterialExpressionColor*>(ValueExpression);
+					SetExpresisonValue(*reinterpret_cast<const FVector3f*>(MapChannel.VecValue), ColorExpression);
+					ColorExpression->SetGroupName(*GroupName);
+					break;
+				}
+				default:
+					ValueExpression = CurrentMaterialElement->AddMaterialExpression<FMaterialExpressionScalar>();
+					FMaterialExpressionScalar* ScalarExpression = static_cast<FMaterialExpressionScalar*>(ValueExpression);
+					SetExpresisonValue(MapChannel.Value, ScalarExpression);
+					ScalarExpression->SetGroupName(*GroupName);
+					break;
+			}
+
+			ValueExpression->SetName(MapChannel.ValueName);
+
+			ValueExpressions.Add(ValueExpression);
+		}
+
+		FMaterialExpressionTexture* TexExpression = CreateTextureMap(Map, CoordinateIndex, MapName, TextureMode, TextureTransform);
+		if (TexExpression)
+		{
+			for (uint32 Index = 0; Index < MapChannelsCount; ++Index)
+			{
+				const FMapChannel& MapChannel = MapChannels[Index];
+
+				FMaterialExpressionGeneric* MultiplyExpression = CurrentMaterialElement->AddMaterialExpression<FMaterialExpressionGeneric>();
+				MultiplyExpression->SetExpressionName(TEXT("Multiply"));
+
+				switch (MapChannel.Channel)
+				{
+				case EChannel::RG:
+				{
+					FMaterialExpressionFunctionCall* MakeFloat2 = CurrentMaterialElement->AddMaterialExpression<FMaterialExpressionFunctionCall>();
+					MakeFloat2->SetFunctionPathName(TEXT("/Engine/Functions/Engine_MaterialFunctions02/Utility/MakeFloat2.MakeFloat2"));
+
+					TexExpression->ConnectExpression(*MakeFloat2->GetInput(0), (int32)EChannel::Red);
+					TexExpression->ConnectExpression(*MakeFloat2->GetInput(1), (int32)EChannel::Green);
+					MakeFloat2->ConnectExpression(*MultiplyExpression->GetInput(0), 0);
+					break;
+				}
+				// RGB is the top channel now
+				case EChannel::RGB:
+					TexExpression->ConnectExpression(*MultiplyExpression->GetInput(0), 0);
+					break;
+				case EChannel::All:
+					TexExpression->ConnectExpression(*MultiplyExpression->GetInput(0), 5);
+					break;
+				default:
+					// single channel connection
+					TexExpression->ConnectExpression(*MultiplyExpression->GetInput(0), (int32)MapChannel.Channel);
+					break;
+				}
+
+				ValueExpressions[Index]->ConnectExpression(*MultiplyExpression->GetInput(1), 0);
+				if (MapChannel.OutputExpression)
+				{
+					MultiplyExpression->ConnectExpression(GetFirstInput(MapChannel.OutputExpression), 0);
+					MapChannel.OutputExpression->ConnectExpression(*MapChannel.MaterialInput, 0);
+				}
+				else
+				{
+					MultiplyExpression->ConnectExpression(*MapChannel.MaterialInput, 0);
+				}
+			}
+		}
+		else
+		{
+			// no texture present, make value connections
+			for (uint32 Index = 0; Index < MapChannelsCount; ++Index)
+			{
+				const FMapChannel& MapChannel = MapChannels[Index];
+
+				if (MapChannel.OutputExpression)
+				{
+					ValueExpressions[Index]->ConnectExpression(GetFirstInput(MapChannel.OutputExpression), 0);
+					MapChannel.OutputExpression->ConnectExpression(*MapChannel.MaterialInput, 0);
+				}
+				else
+				{
+					ValueExpressions[Index]->ConnectExpression(*MapChannel.MaterialInput, 0);
+				}
+			}
+		}
+	}
+
+	FMaterialExpressionTexture* FPBRMapFactory::CreateTextureMap(const GLTF::FTexture& Map, int CoordinateIndex, const TCHAR* MapName,
+	                                                             ETextureMode TextureMode, const GLTF::FTextureTransform* TextureTransform)
+	{
+		ITextureElement* Texture = TextureFactory.CreateTexture(Map, ParentPackage, Flags, TextureMode);
+
+		FMaterialExpressionTexture* TexExpression = nullptr;
+		if (Texture)
+		{
+			TexExpression = CurrentMaterialElement->AddMaterialExpression<FMaterialExpressionTexture>();
+			TexExpression->SetTexture(Texture);
+			TexExpression->SetName(*(FString(MapName) + TEXT(" Map")));
+			TexExpression->SetGroupName(*GroupName);
+
+			CreateTextureCoordinate(CoordinateIndex, *TexExpression, *CurrentMaterialElement);
+
+			if (nullptr != TextureTransform)
+			{
+				CreateTextureTransform(*TextureTransform, *TexExpression, *CurrentMaterialElement);
+			}
+		}
+		return TexExpression;
+	}
+
+	template <class ValueExpressionClass, class ValueClass>
+	FMaterialExpression* FPBRMapFactory::CreateMap(const GLTF::FTexture& Map, int CoordinateIndex, const ValueClass& Value, const TCHAR* MapName,
+	                                               const TCHAR* ValueName, ETextureMode TextureMode, FMaterialExpressionInput& MaterialInput, 
+												   const GLTF::FTextureTransform* TextureTransform)
+	{
+		if (!ensure(MapName) || !ensure(CurrentMaterialElement))
+		{
+			return nullptr;
+		}
+
+		ValueExpressionClass* ValueExpression = CurrentMaterialElement->AddMaterialExpression<ValueExpressionClass>();
+		if (ValueName)
+			ValueExpression->SetName(*(FString(MapName) + TEXT(" ") + ValueName));
+		else
+			ValueExpression->SetName(MapName);
+		ValueExpression->SetGroupName(*GroupName);
+		SetExpresisonValue(Value, ValueExpression);
+
+		FMaterialExpressionTexture* TexExpression = CreateTextureMap(Map, CoordinateIndex, MapName, TextureMode, TextureTransform);
+		if (TexExpression)
+		{
+			FMaterialExpressionGeneric* MultiplyExpression = CurrentMaterialElement->AddMaterialExpression<FMaterialExpressionGeneric>();
+			MultiplyExpression->SetExpressionName(TEXT("Multiply"));
+
+			TexExpression->ConnectExpression(*MultiplyExpression->GetInput(0), 0);
+			ValueExpression->ConnectExpression(*MultiplyExpression->GetInput(1), 0);
+			MultiplyExpression->ConnectExpression(MaterialInput, 0);
+			return MultiplyExpression;
+		}
+		else
+		{
+			ValueExpression->ConnectExpression(MaterialInput, 0);
+			return ValueExpression;
+		}
+	}
+
+}  // namespace GLTF
diff --git a/Engine/Plugins/Interchange/Runtime/Source/Parsers/GLTFCore/Private/GLTFMapFactory.h b/Engine/Plugins/Interchange/Runtime/Source/Parsers/GLTFCore/Private/GLTFMapFactory.h
new file mode 100644
--- /dev/null
+++ b/Engine/Plugins/Interchange/Runtime/Source/Parsers/GLTFCore/Private/GLTFMapFactory.h
@@ -0,0 +1,121 @@
+// Copyright Epic Games, Inc. All Rights Reserved.
+
+#pragma once
+
+
+#include "Math/Vector4.h"
+#include "UObject/ObjectMacros.h"
+
+namespace GLTF { class FMaterialElement; }
+namespace GLTF { class FMaterialExpression; }
+namespace GLTF { class FMaterialExpressionInput; }
+namespace GLTF { class FMaterialExpressionTexture; }
+namespace GLTF { class ITextureFactory; }
+namespace GLTF { enum class ETextureMode; }
+template <uint32 NumInlineElements> class TFixedAllocator;
+
+namespace GLTF
+{
+	struct FTexture;
+	struct FTextureTransform;
+	class FTextureFactory;
+
+	struct FPBRMapFactory
+	{
+		enum class EChannel
+		{
+			All,
+			Red,
+			Green,
+			Blue,
+			Alpha,
+			RG,
+			RGB
+		};
+
+		struct FMapChannel
+		{
+			union {
+				float Value;
+				float VecValue[3];
+			};
+
+			const TCHAR*              ValueName;
+			EChannel                  Channel;
+			FMaterialExpressionInput* MaterialInput;
+			FMaterialExpression*      OutputExpression;
+
+			FMapChannel(float InValue, const TCHAR* InValueName, EChannel InChannel, FMaterialExpressionInput* InMaterialInput, FMaterialExpression* InOutputExpression)
+				: FMapChannel(InValueName, InChannel, InMaterialInput, InOutputExpression)
+			{
+				Value = InValue;
+			}
+
+			FMapChannel(const FVector& InVecValue, const TCHAR* InValueName, EChannel InChannel, FMaterialExpressionInput* InMaterialInput, FMaterialExpression* InOutputExpression)
+				: FMapChannel(InValueName, InChannel, InMaterialInput, InOutputExpression)
+			{
+				SetValue(InVecValue);
+			}
+
+			void SetValue(const FVector& Vec)
+			{
+				VecValue[0] = (float)Vec.X;
+				VecValue[1] = (float)Vec.Y;
+				VecValue[2] = (float)Vec.Z;
+			}
+
+		private:
+			FMapChannel(const TCHAR* InValueName, EChannel InChannel, FMaterialExpressionInput* InMaterialInput, FMaterialExpression* InOutputExpression)
+				: ValueName(InValueName)
+				, Channel(InChannel)
+				, MaterialInput(InMaterialInput)
+				, OutputExpression(InOutputExpression)
+			{
+			}
+		};
+
+		FMaterialElement* CurrentMaterialElement;
+		FString           GroupName;
+
+	public:
+		FPBRMapFactory(ITextureFactory& TextureFactory);
+
+		void SetParentPackage(UObject* ParentPackage, EObjectFlags Flags);
+
+		void CreateNormalMap(const GLTF::FTexture& Map, int CoordinateIndex, float NormalScale, const GLTF::FTextureTransform* TextureTransform);
+
+		FMaterialExpression* CreateColorMap(const GLTF::FTexture& Map, int CoordinateIndex, const FVector3f& Color, const TCHAR* MapName,
+		                                    const TCHAR* ValueName, ETextureMode TextureMode, FMaterialExpressionInput& MaterialInput,
+											const GLTF::FTextureTransform* TextureTransform);
+
+		FMaterialExpression* CreateColorMap(const GLTF::FTexture& Map, int CoordinateIndex, const FVector4f& Color, const TCHAR* MapName,
+		                                    const TCHAR* ValueName, ETextureMode TextureMode, FMaterialExpressionInput& MaterialInput,
+											const GLTF::FTextureTransform* TextureTransform);
+
+		FMaterialExpressionTexture* CreateTextureMap(const GLTF::FTexture& Map, int CoordinateIndex, const TCHAR* MapName, ETextureMode TextureMode, 
+													 const GLTF::FTextureTransform* TextureTransform);
+
+		void CreateMultiMap(const GLTF::FTexture& Map, int CoordinateIndex, const TCHAR* MapName, const FMapChannel* MapChannels,
+		                    uint32 MapChannelsCount, ETextureMode TextureMode, const GLTF::FTextureTransform* TextureTransform);
+
+	private:
+		using FExpressionList = TArray<FMaterialExpression*, TFixedAllocator<4> >;
+
+		template <class ValueExpressionClass, class ValueClass>
+		FMaterialExpression* CreateMap(const GLTF::FTexture& Map, int CoordinateIndex, const ValueClass& Value, const TCHAR* MapName,
+		                               const TCHAR* ValueName, ETextureMode TextureMode, FMaterialExpressionInput& MaterialInput, 
+									   const GLTF::FTextureTransform* TextureTransform);
+
+	private:
+		ITextureFactory& TextureFactory;
+		UObject*         ParentPackage;
+		EObjectFlags     Flags;
+	};
+
+	inline void FPBRMapFactory::SetParentPackage(UObject* InParentPackage, EObjectFlags InFlags)
+	{
+		ParentPackage = InParentPackage;
+		Flags         = InFlags;
+	}
+
+}  // namespace GLTF
diff --git a/Engine/Plugins/Interchange/Runtime/Source/Parsers/GLTFCore/Private/GLTFMaterialExpressions.cpp b/Engine/Plugins/Interchange/Runtime/Source/Parsers/GLTFCore/Private/GLTFMaterialExpressions.cpp
new file mode 100644
--- /dev/null
+++ b/Engine/Plugins/Interchange/Runtime/Source/Parsers/GLTFCore/Private/GLTFMaterialExpressions.cpp
@@ -0,0 +1,226 @@
+// Copyright Epic Games, Inc. All Rights Reserved.
+
+#include "GLTFMaterialExpressions.h"
+
+
+namespace GLTF
+{
+	FMaterialExpressionInput::FMaterialExpressionInput(const FString& Name)
+	    : Name(Name)
+	    , Expression(nullptr)
+	    , OutputIndex(INDEX_NONE)
+	{
+	}
+
+	FMaterialExpression::FMaterialExpression() {}
+
+	FMaterialExpression::FMaterialExpression(const FMaterialExpression& Expr)
+	    : Type(Expr.Type)
+	{
+	}
+	FMaterialExpression::FMaterialExpression(int Type)
+	    : Type((EMaterialExpressionType)Type)
+	{
+	}
+
+	void FMaterialExpression::ConnectExpression(FMaterialExpressionInput& ExpressionInput, int32 OutputIndex)
+	{
+		if (OutputIndex != INDEX_NONE)
+		{
+			ExpressionInput.Expression  = this;
+			ExpressionInput.OutputIndex = OutputIndex;
+		}
+	}
+
+	FMaterialExpressionColor::FMaterialExpressionColor()
+	    : FMaterialExpressionParameter(Type)
+	{
+	}
+
+	FMaterialExpressionTexture::FMaterialExpressionTexture()
+	    : FMaterialExpressionParameter(Type)
+	    , InputCoordinate(TEXT("InputCoordinate"))
+	    , Texture(nullptr)
+	{
+	}
+
+	FMaterialExpressionTextureCoordinate::FMaterialExpressionTextureCoordinate()
+	    : FMaterialExpression(Type)
+	    , CoordinateIndex(0)
+	{
+	}
+
+	FMaterialExpressionGeneric::FMaterialExpressionGeneric()
+	    : FMaterialExpression(Type)
+	{
+	}
+
+	FMaterialExpressionFunctionCall::FMaterialExpressionFunctionCall()
+	    : FMaterialExpression(Type)
+	{
+	}
+
+	//
+
+	FMaterialExpressionParameter::FMaterialExpressionParameter(int Type)
+	    : FMaterialExpression(Type)
+	{
+	}
+
+	void FMaterialExpressionParameter::SetName(const TCHAR* InName)
+	{
+		ensure(Type == EMaterialExpressionType::ConstantColor || Type == EMaterialExpressionType::ConstantScalar ||
+		      Type == EMaterialExpressionType::Texture);
+		Name = InName;
+	}
+
+	const TCHAR* FMaterialExpressionParameter::GetName() const
+	{
+		ensure(Type == EMaterialExpressionType::ConstantColor || Type == EMaterialExpressionType::ConstantScalar ||
+		      Type == EMaterialExpressionType::Texture);
+		return *Name;
+	}
+
+	void FMaterialExpressionParameter::SetGroupName(const TCHAR* InGroupName)
+	{
+		ensure(Type == EMaterialExpressionType::ConstantColor || Type == EMaterialExpressionType::ConstantScalar ||
+		      Type == EMaterialExpressionType::Texture);
+		GroupName = InGroupName;
+	}
+
+	const TCHAR* FMaterialExpressionParameter::GetGroupName() const
+	{
+		ensure(Type == EMaterialExpressionType::ConstantColor || Type == EMaterialExpressionType::ConstantScalar ||
+		      Type == EMaterialExpressionType::Texture);
+		return *GroupName;
+	}
+
+	//
+
+	FMaterialExpressionInput* FMaterialExpressionTexture::GetInput(int32 Index)
+	{
+		ensure(Index == 0);
+		return &InputCoordinate;
+	}
+
+	int32 FMaterialExpressionTexture::GetInputCount() const
+	{
+		return 1;
+	}
+
+	FMaterialExpressionInput* FMaterialExpressionGeneric::GetInput(int32 Index)
+	{
+		while (!Inputs.IsValidIndex(Index))
+		{
+			Inputs.Emplace(*FString::FromInt(Inputs.Num()));
+		}
+
+		return &Inputs[Index];
+	}
+
+	int32 FMaterialExpressionGeneric::GetInputCount() const
+	{
+		return Inputs.Num();
+	}
+
+	FMaterialExpressionInput* FMaterialExpressionFunctionCall::GetInput(int32 Index)
+	{
+		while (!Inputs.IsValidIndex(Index))
+		{
+			Inputs.Emplace(*FString::FromInt(Inputs.Num()));
+		}
+
+		return &Inputs[Index];
+	}
+
+	int32 FMaterialExpressionFunctionCall::GetInputCount() const
+	{
+		return Inputs.Num();
+	}
+
+	//
+
+	FMaterialElement::FMaterialElement(const FString& Name)
+	    : Name(Name)
+	    , BaseColor(TEXT("BaseColor"))
+	    , Metallic(TEXT("Metallic"))
+	    , Specular(TEXT("Specular"))
+	    , Roughness(TEXT("Roughness"))
+	    , EmissiveColor(TEXT("EmissiveColor"))
+	    , Opacity(TEXT("Opacity"))
+	    , Normal(TEXT("Normal"))
+	    , WorldDisplacement(TEXT("WorldDisplacement"))
+	    , Refraction(TEXT("Refraction"))
+	    , AmbientOcclusion(TEXT("AmbientOcclusion"))
+		, ClearCoat(TEXT("ClearCoat"))
+		, ClearCoatRoughness(TEXT("ClearCoatRoughness"))
+		, ThinTranslucentMaterialOutput(nullptr)
+		, ClearCoatBottomNormalOutput(nullptr)
+	    , bIsFinal(false)
+	{
+	}
+
+	FMaterialElement::~FMaterialElement()
+	{
+		for (GLTF::FMaterialExpression* Expr : Expressions)
+		{
+			delete Expr;
+		}
+	}
+
+    FMD5Hash FMaterialElement::GetGLTFMaterialHash() const
+	{
+		return GLTFMaterialHash;
+	}
+
+	void FMaterialElement::SetGLTFMaterialHash(FMD5Hash Hash)
+	{
+		GLTFMaterialHash = Hash;
+	}
+
+	int32 FMaterialElement::GetExpressionsCount() const
+	{
+		return Expressions.Num();
+	}
+
+	FMaterialExpression* FMaterialElement::GetExpression(int32 Index)
+	{
+		if (!ensure(Index < Expressions.Num()))
+		{
+			return nullptr;
+		}
+		return Expressions[Index];
+	}
+
+	FMaterialExpression* FMaterialElement::AddMaterialExpression(EMaterialExpressionType ExpressionType)
+	{
+		FMaterialExpression* Expression = nullptr;
+		switch (ExpressionType)
+		{
+			case GLTF::EMaterialExpressionType::ConstantColor:
+				Expression = new FMaterialExpressionColor();
+				break;
+			case GLTF::EMaterialExpressionType::ConstantScalar:
+				Expression = new FMaterialExpressionScalar();
+				break;
+			case GLTF::EMaterialExpressionType::FunctionCall:
+				Expression = new FMaterialExpressionFunctionCall();
+				break;
+			case GLTF::EMaterialExpressionType::Generic:
+				Expression = new FMaterialExpressionGeneric();
+				break;
+			case GLTF::EMaterialExpressionType::Texture:
+				Expression = new FMaterialExpressionTexture();
+				break;
+			case GLTF::EMaterialExpressionType::TextureCoordinate:
+				Expression = new FMaterialExpressionTextureCoordinate();
+				break;
+			default:
+				ensure(false);
+		}
+		if (Expression)
+			Expressions.Add(Expression);
+		return Expression;
+	}
+
+}  // namespace GLTF
diff --git a/Engine/Plugins/Interchange/Runtime/Source/Parsers/GLTFCore/Private/GLTFMaterialFactory.cpp b/Engine/Plugins/Interchange/Runtime/Source/Parsers/GLTFCore/Private/GLTFMaterialFactory.cpp
new file mode 100644
--- /dev/null
+++ b/Engine/Plugins/Interchange/Runtime/Source/Parsers/GLTFCore/Private/GLTFMaterialFactory.cpp
@@ -0,0 +1,836 @@
+// Copyright Epic Games, Inc. All Rights Reserved.
+
+#include "GLTFMaterialFactory.h"
+
+// IWYU pragma: begin_keep
+#include "GLTFAnimation.h"
+#include "GLTFMapFactory.h"
+#include "GLTFMaterial.h"
+#include "GLTFMaterialExpressions.h"
+#include "GLTFAsset.h"
+#include "Engine/EngineTypes.h"
+#include "GLTFMesh.h"
+#include "GLTFNode.h"
+#include "GLTFTexture.h"
+// IWYU pragma: end_keep
+
+namespace GLTF
+{
+	namespace
+	{
+		const GLTF::FTexture& GetTexture(const GLTF::FTextureMap& Map, const TArray<GLTF::FTexture>& Textures)
+		{
+			static const GLTF::FImage   Image;
+			static const GLTF::FTexture None(FString(), Image, GLTF::FSampler::DefaultSampler);
+			return Map.TextureIndex != INDEX_NONE ? Textures[Map.TextureIndex] : None;
+		}
+	}
+
+	class FMaterialFactoryImpl : public GLTF::FBaseLogger
+	{
+	public:
+		FMaterialFactoryImpl(IMaterialElementFactory* MaterialElementFactory, ITextureFactory* TextureFactory)
+		    : MaterialElementFactory(MaterialElementFactory)
+		    , TextureFactory(TextureFactory)
+		{
+			if (!ensure(MaterialElementFactory))
+			{
+				Messages.Emplace(EMessageSeverity::Warning, FString::Printf(TEXT("MaterialElementFactory is unexpected nullptr.")));
+			}
+		}
+
+		const TArray<FMaterialElement*>& CreateMaterials(const GLTF::FAsset& Asset, UObject* ParentPackage, EObjectFlags Flags);
+
+	private:
+		void HandleOpacity(const TArray<GLTF::FTexture>& Texture, const GLTF::FMaterial& GLTFMaterial, FMaterialElement& MaterialElement);
+		void HandleShadingModel(const TArray<GLTF::FTexture>& Texture, const GLTF::FMaterial& GLTFMaterial, FPBRMapFactory& MapFactory, FMaterialElement& MaterialElement);
+		void HandleOcclusion(const TArray<GLTF::FTexture>& Texture, const GLTF::FMaterial& GLTFMaterial, FPBRMapFactory& MapFactory, FMaterialElement& MaterialElement);
+		void HandleEmissive(const TArray<GLTF::FTexture>& Textures, const GLTF::FMaterial& GLTFMaterial, FPBRMapFactory& MapFactory, FMaterialElement& MaterialElement);
+		void HandleNormal(const TArray<GLTF::FTexture>& Textures, const GLTF::FMaterial& GLTFMaterial, FPBRMapFactory& MapFactory, FMaterialElement& MaterialElement);
+		void HandleClearCoat(const TArray<GLTF::FTexture>& Textures, const GLTF::FMaterial& GLTFMaterial, FPBRMapFactory& MapFactory, FMaterialElement& MaterialElement);
+		void HandleTransmission(const TArray<GLTF::FTexture>& Textures, const GLTF::FMaterial& GLTFMaterial, FPBRMapFactory& MapFactory, FMaterialElement& MaterialElement);
+		void HandleSheen(const TArray<GLTF::FTexture>& Textures, const GLTF::FMaterial& GLTFMaterial, FPBRMapFactory& MapFactory, FMaterialElement& MaterialElement);
+		void HandleIOR(const TArray<GLTF::FTexture>& Textures, const GLTF::FMaterial& GLTFMaterial, FPBRMapFactory& MapFactory, FMaterialElement& MaterialElement);
+		void HandleSpecular(const TArray<GLTF::FTexture>& Textures, const GLTF::FMaterial& GLTFMaterial, FPBRMapFactory& MapFactory, FMaterialElement& MaterialElement);
+
+
+	private:
+		TUniquePtr<IMaterialElementFactory> MaterialElementFactory;
+		TUniquePtr<ITextureFactory>         TextureFactory;
+		TArray<FMaterialElement*>           Materials;
+
+		friend class FMaterialFactory;
+	};
+
+	namespace
+	{
+		EBlendMode ConvertAlphaMode(FMaterial::EAlphaMode Mode)
+		{
+			switch (Mode)
+			{
+				case FMaterial::EAlphaMode::Opaque:
+					return EBlendMode::BLEND_Opaque;
+				case FMaterial::EAlphaMode::Blend:
+					return EBlendMode::BLEND_Translucent;
+				case FMaterial::EAlphaMode::Mask:
+					return EBlendMode::BLEND_Masked;
+				default:
+					return EBlendMode::BLEND_Opaque;
+			}
+		}
+
+		template <class ReturnClass>
+		ReturnClass* FindExpression(const FString& Name, FMaterialElement& MaterialElement)
+		{
+			ReturnClass* Result = nullptr;
+			for (int32 Index = 0; Index < MaterialElement.GetExpressionsCount(); ++Index)
+			{
+				FMaterialExpression* Expression = MaterialElement.GetExpression(Index);
+
+				if (!Expression)
+				{
+					continue;
+				}
+
+				if (Expression->GetType() != EMaterialExpressionType::ConstantColor &&
+				    Expression->GetType() != EMaterialExpressionType::ConstantScalar && Expression->GetType() != EMaterialExpressionType::Texture)
+					continue;
+
+				FMaterialExpressionParameter* ExpressionParameter = static_cast<FMaterialExpressionParameter*>(Expression);
+				if (ExpressionParameter->GetName() == Name)
+				{
+					if (ensure(Expression->GetType() == (EMaterialExpressionType)ReturnClass::Type))
+					{
+						Result = static_cast<ReturnClass*>(ExpressionParameter);
+					}
+					break;
+				}
+			}
+			return Result;
+		}
+
+	}
+
+	const TArray<FMaterialElement*>& FMaterialFactoryImpl::CreateMaterials(const GLTF::FAsset& Asset, UObject* ParentPackage, EObjectFlags Flags)
+	{
+		TextureFactory->CleanUp();
+		Materials.Empty();
+
+		if (!MaterialElementFactory)
+		{
+			return Materials;
+		}
+
+		Materials.Reserve(Asset.Materials.Num() + 1);
+
+		Messages.Empty();
+
+		FPBRMapFactory MapFactory(*TextureFactory);
+		MapFactory.SetParentPackage(ParentPackage, Flags);
+
+		for (const GLTF::FMaterial& GLTFMaterial : Asset.Materials)
+		{
+			if (!ensure(!GLTFMaterial.Name.IsEmpty()))
+			{
+				continue;
+			}
+
+			FMaterialElement* MaterialElement = MaterialElementFactory->CreateMaterial(*GLTFMaterial.Name, ParentPackage, Flags);
+			MaterialElement->SetTwoSided(GLTFMaterial.bIsDoubleSided);
+			MaterialElement->SetBlendMode(ConvertAlphaMode(GLTFMaterial.AlphaMode));
+
+			MapFactory.CurrentMaterialElement = MaterialElement;
+
+			HandleShadingModel(Asset.Textures, GLTFMaterial, MapFactory, *MaterialElement);
+			HandleOpacity(Asset.Textures, GLTFMaterial, *MaterialElement);
+			HandleClearCoat(Asset.Textures, GLTFMaterial, MapFactory, *MaterialElement);
+			HandleTransmission(Asset.Textures, GLTFMaterial, MapFactory, *MaterialElement);
+			HandleSheen(Asset.Textures, GLTFMaterial, MapFactory, *MaterialElement);
+			HandleIOR(Asset.Textures, GLTFMaterial, MapFactory, *MaterialElement);
+			HandleSpecular(Asset.Textures, GLTFMaterial, MapFactory, *MaterialElement);
+
+			// Additional maps
+			HandleOcclusion(Asset.Textures, GLTFMaterial, MapFactory, *MaterialElement);
+			HandleEmissive(Asset.Textures, GLTFMaterial, MapFactory, *MaterialElement);
+			HandleNormal(Asset.Textures, GLTFMaterial, MapFactory, *MaterialElement);
+
+			MaterialElement->SetGLTFMaterialHash(GLTFMaterial.GetHash());
+            MaterialElement->Finalize();
+			Materials.Add(MaterialElement);
+		}
+
+		return Materials;
+	}
+
+	void FMaterialFactoryImpl::HandleOpacity(const TArray<GLTF::FTexture>& Textures, const GLTF::FMaterial& GLTFMaterial, FMaterialElement& MaterialElement)
+	{
+		if (GLTFMaterial.IsOpaque() || GLTFMaterial.bHasTransmission) // With transmission, we handle opacity different
+		{
+			return;
+		}
+
+		const TCHAR* GroupName = TEXT("Opacity");
+
+		FMaterialExpressionTexture* BaseColorMap = FindExpression<FMaterialExpressionTexture>(TEXT("BaseColor Map"), MaterialElement);
+		switch (GLTFMaterial.AlphaMode)
+		{
+			case FMaterial::EAlphaMode::Mask:
+			{
+				FMaterialExpressionColor* BaseColorFactor = FindExpression<FMaterialExpressionColor>(TEXT("BaseColor"), MaterialElement);
+
+				FMaterialExpressionGeneric* MultiplyExpression = MaterialElement.AddMaterialExpression<FMaterialExpressionGeneric>();
+				MultiplyExpression->SetExpressionName(TEXT("Multiply"));
+
+				BaseColorFactor->ConnectExpression(*MultiplyExpression->GetInput(1), (int)FPBRMapFactory::EChannel::Alpha);
+				BaseColorMap->ConnectExpression(*MultiplyExpression->GetInput(0), (int)FPBRMapFactory::EChannel::Alpha);
+
+				FMaterialExpressionFunctionCall* CuttofExpression = MaterialElement.AddMaterialExpression<FMaterialExpressionFunctionCall>();
+				CuttofExpression->SetFunctionPathName(TEXT("/Engine/Functions/Engine_MaterialFunctions02/SmoothStep.SmoothStep"));
+
+				FMaterialExpressionScalar* ValueExpression = MaterialElement.AddMaterialExpression<FMaterialExpressionScalar>();
+				ValueExpression->SetName(TEXT("Alpha Cuttof"));
+				ValueExpression->SetGroupName(GroupName);
+				ValueExpression->GetScalar() = GLTFMaterial.AlphaCutoff;
+
+				MultiplyExpression->ConnectExpression(*CuttofExpression->GetInput(0), 0);
+				ValueExpression->ConnectExpression(*CuttofExpression->GetInput(1), 0);
+				ValueExpression->ConnectExpression(*CuttofExpression->GetInput(2), 0);
+
+				CuttofExpression->ConnectExpression(MaterialElement.GetOpacity(), 0);
+				break;
+			}
+			case FMaterial::EAlphaMode::Blend:
+			{
+				FMaterialExpressionScalar* ValueExpression = MaterialElement.AddMaterialExpression<FMaterialExpressionScalar>();
+				ValueExpression->SetName(TEXT("IOR"));
+				ValueExpression->SetGroupName(GroupName);
+				ValueExpression->GetScalar() = 1.f;
+				ValueExpression->ConnectExpression(MaterialElement.GetRefraction(), 0);
+
+				FMaterialExpressionColor* BaseColorFactor = FindExpression<FMaterialExpressionColor>(TEXT("BaseColor"), MaterialElement);
+				if (BaseColorMap)
+				{
+					FMaterialExpressionGeneric* MultiplyExpression = MaterialElement.AddMaterialExpression<FMaterialExpressionGeneric>();
+					MultiplyExpression->SetExpressionName(TEXT("Multiply"));
+
+					BaseColorFactor->ConnectExpression(*MultiplyExpression->GetInput(1), (int)FPBRMapFactory::EChannel::Alpha);
+					BaseColorMap->ConnectExpression(*MultiplyExpression->GetInput(0), (int)FPBRMapFactory::EChannel::Alpha);
+					MultiplyExpression->ConnectExpression(MaterialElement.GetOpacity(), 0);
+				}
+				else
+				{
+					BaseColorFactor->ConnectExpression(MaterialElement.GetOpacity(), (int)FPBRMapFactory::EChannel::Alpha);
+				}
+				break;
+			}
+			default:
+				if (!ensure(false))
+				{
+					Messages.Emplace(EMessageSeverity::Warning, FString::Printf(TEXT("AlphaMode Opaque not supported for Opacity")));
+				}
+				break;
+		}
+	}
+
+	void FMaterialFactoryImpl::HandleShadingModel(const TArray<GLTF::FTexture>& Textures, const GLTF::FMaterial& GLTFMaterial, FPBRMapFactory& MapFactory, FMaterialElement& MaterialElement)
+	{
+		MapFactory.GroupName = TEXT("GGX");
+
+		FMaterialExpressionInput* BaseColorInput = &MaterialElement.GetBaseColor();
+
+		if (GLTFMaterial.bIsUnlitShadingModel)
+		{
+			BaseColorInput = &MaterialElement.GetEmissiveColor();
+			MaterialElement.SetShadingModel(EGLTFMaterialShadingModel::Unlit);
+		}
+
+		TArray<FPBRMapFactory::FMapChannel, TFixedAllocator<4>> Maps;
+		if (GLTFMaterial.ShadingModel == FMaterial::EShadingModel::MetallicRoughness)
+		{
+			// Base Color
+			MapFactory.GroupName = TEXT("Base Color");
+			MapFactory.CreateColorMap(GetTexture(GLTFMaterial.BaseColor, Textures),
+									  GLTFMaterial.BaseColor.TexCoord,
+									  GLTFMaterial.BaseColorFactor,
+									  TEXT("BaseColor"),
+									  nullptr,
+									  ETextureMode::Color,
+									  *BaseColorInput,
+									  GLTFMaterial.BaseColor.bHasTextureTransform ? &GLTFMaterial.BaseColor.TextureTransform : nullptr);
+
+			if (!GLTFMaterial.bIsUnlitShadingModel)
+			{
+				// Metallic
+				Maps.Emplace(GLTFMaterial.MetallicRoughness.MetallicFactor,
+							 TEXT("Metallic Factor"),
+							 FPBRMapFactory::EChannel::Blue,
+							 &MaterialElement.GetMetallic(),
+							 nullptr);
+
+				// Roughness
+				Maps.Emplace(GLTFMaterial.MetallicRoughness.RoughnessFactor,
+							 TEXT("Roughness Factor"),
+							 FPBRMapFactory::EChannel::Green,
+							 &MaterialElement.GetRoughness(),
+							 nullptr);
+
+				MapFactory.CreateMultiMap(GetTexture(GLTFMaterial.MetallicRoughness.Map, Textures),
+										  GLTFMaterial.MetallicRoughness.Map.TexCoord,
+										  TEXT("MetallicRoughness"),
+										  Maps.GetData(),
+										  Maps.Num(),
+										  ETextureMode::Grayscale,
+										  GLTFMaterial.MetallicRoughness.Map.bHasTextureTransform ? &GLTFMaterial.MetallicRoughness.Map.TextureTransform : nullptr);
+			}
+		}
+		else if (GLTFMaterial.ShadingModel == FMaterial::EShadingModel::SpecularGlossiness)
+		{
+			// We'll actually just convert it into MetalRoughness in the material graph
+			FMaterialExpressionFunctionCall* SpecGlossToMetalRough = MaterialElement.AddMaterialExpression<FMaterialExpressionFunctionCall>();
+			SpecGlossToMetalRough->SetFunctionPathName(TEXT("/Interchange/gltf/SpecGlossToMetalRoughness.SpecGlossToMetalRoughness"));
+			SpecGlossToMetalRough->ConnectExpression(*BaseColorInput, 0);
+
+			if (!GLTFMaterial.bIsUnlitShadingModel)
+			{
+				SpecGlossToMetalRough->ConnectExpression(MaterialElement.GetMetallic(), 1);
+
+				FMaterialExpressionGeneric* GlossToRoughness = MaterialElement.AddMaterialExpression<FMaterialExpressionGeneric>();
+				GlossToRoughness->SetExpressionName(TEXT("OneMinus"));
+				GlossToRoughness->ConnectExpression(MaterialElement.GetRoughness(), 0);
+
+				// Diffuse Color (BaseColor/BaseColorFactor are used to store the Diffuse alternatives for Spec/Gloss)
+				MapFactory.GroupName = TEXT("Diffuse Color");
+				FMaterialExpression* Diffuse = MapFactory.CreateColorMap(GetTexture(GLTFMaterial.BaseColor, Textures),
+																		 GLTFMaterial.BaseColor.TexCoord,
+																		 GLTFMaterial.BaseColorFactor,
+																		 TEXT("Diffuse"),
+																		 TEXT("Color"),
+																		 ETextureMode::Color,
+																		 *SpecGlossToMetalRough->GetInput(1),
+																		 GLTFMaterial.BaseColor.bHasTextureTransform ? &GLTFMaterial.BaseColor.TextureTransform : nullptr);
+
+				// Specular (goes into SpecGlossToMetalRough conversion)
+				Maps.Emplace(GLTFMaterial.SpecularGlossiness.SpecularFactor,
+							 TEXT("Specular Factor"),
+							 FPBRMapFactory::EChannel::RGB,
+							 SpecGlossToMetalRough->GetInput(0),
+							 nullptr);
+
+				// Glossiness (converted to Roughness)
+				Maps.Emplace(GLTFMaterial.SpecularGlossiness.GlossinessFactor,
+							 TEXT("Glossiness Factor"),
+							 FPBRMapFactory::EChannel::Alpha,
+							 GlossToRoughness->GetInput(0),
+							 nullptr);
+
+				// Creates the multimap for Specular and Glossiness
+				MapFactory.CreateMultiMap(GetTexture(GLTFMaterial.SpecularGlossiness.Map, Textures),
+										  GLTFMaterial.SpecularGlossiness.Map.TexCoord,
+										  TEXT("SpecularGlossiness"),
+										  Maps.GetData(),
+										  Maps.Num(),
+										  ETextureMode::Color,
+										  GLTFMaterial.SpecularGlossiness.Map.bHasTextureTransform ? &GLTFMaterial.SpecularGlossiness.Map.TextureTransform : nullptr);
+			}
+		}
+	}
+
+	void FMaterialFactoryImpl::HandleOcclusion(const TArray<GLTF::FTexture>& Textures, const GLTF::FMaterial& GLTFMaterial, FPBRMapFactory& MapFactory, FMaterialElement& MaterialElement)
+	{
+		MapFactory.GroupName = TEXT("Occlusion");
+
+		FMaterialExpressionTexture* TexExpression = MapFactory.CreateTextureMap(
+			GetTexture(GLTFMaterial.Occlusion, Textures), GLTFMaterial.Occlusion.TexCoord, TEXT("Occlusion"), ETextureMode::Grayscale, 
+			GLTFMaterial.Occlusion.bHasTextureTransform ? &GLTFMaterial.Occlusion.TextureTransform : nullptr);
+
+		if (!TexExpression)
+			return;
+
+		FMaterialExpressionScalar* ConstantExpression = MaterialElement.AddMaterialExpression<FMaterialExpressionScalar>();
+		ConstantExpression->GetScalar()               = 1.f;
+
+		FMaterialExpressionGeneric* LerpExpression = MaterialElement.AddMaterialExpression<FMaterialExpressionGeneric>();
+		LerpExpression->SetExpressionName(TEXT("LinearInterpolate"));
+
+		FMaterialExpressionScalar* ValueExpression = MaterialElement.AddMaterialExpression<FMaterialExpressionScalar>();
+		ValueExpression->SetName(TEXT("Occlusion Strength"));
+		ValueExpression->SetGroupName(*MapFactory.GroupName);
+		ValueExpression->GetScalar() = GLTFMaterial.OcclusionStrength;
+
+		ConstantExpression->ConnectExpression(*LerpExpression->GetInput(0), 0);
+		TexExpression->ConnectExpression(*LerpExpression->GetInput(1), (int)FPBRMapFactory::EChannel::Red);  // ignore other channels
+		ValueExpression->ConnectExpression(*LerpExpression->GetInput(2), 0);
+
+		LerpExpression->ConnectExpression(MaterialElement.GetAmbientOcclusion(), 0);
+	}
+
+	void FMaterialFactoryImpl::HandleEmissive(const TArray<GLTF::FTexture>& Textures, const GLTF::FMaterial& GLTFMaterial, FPBRMapFactory& MapFactory, FMaterialElement& MaterialElement)
+	{
+		if (GLTFMaterial.Emissive.TextureIndex == INDEX_NONE || GLTFMaterial.EmissiveFactor.IsNearlyZero() || GLTFMaterial.bIsUnlitShadingModel)
+		{
+			return;
+		}
+
+		MapFactory.GroupName = TEXT("Emission");
+		MapFactory.CreateColorMap(GetTexture(GLTFMaterial.Emissive, Textures),
+								  GLTFMaterial.Emissive.TexCoord,
+								  GLTFMaterial.EmissiveFactor,
+								  TEXT("Emissive"),
+								  TEXT("Color"),
+								  ETextureMode::Color, // emissive map is in sRGB space
+								  MaterialElement.GetEmissiveColor(),
+								  GLTFMaterial.Emissive.bHasTextureTransform ? &GLTFMaterial.Emissive.TextureTransform : nullptr);
+	}
+
+	void FMaterialFactoryImpl::HandleNormal(const TArray<GLTF::FTexture>& Textures, const GLTF::FMaterial& GLTFMaterial, FPBRMapFactory& MapFactory, FMaterialElement& MaterialElement)
+	{
+		MapFactory.GroupName = TEXT("Normal");
+		MapFactory.CreateNormalMap(GetTexture(GLTFMaterial.Normal, Textures),
+								   GLTFMaterial.Normal.TexCoord,
+								   GLTFMaterial.NormalScale,
+								   GLTFMaterial.Normal.bHasTextureTransform ? &GLTFMaterial.Normal.TextureTransform : nullptr);
+	}
+
+	void FMaterialFactoryImpl::HandleClearCoat(const TArray<GLTF::FTexture>& Textures, const GLTF::FMaterial& GLTFMaterial, FPBRMapFactory& MapFactory, FMaterialElement& MaterialElement)
+	{
+		if (!GLTFMaterial.bHasClearCoat || GLTFMaterial.bIsUnlitShadingModel || FMath::IsNearlyEqual(GLTFMaterial.ClearCoat.ClearCoatFactor, 0.0f))
+		{
+			return;
+		}
+
+		MaterialElement.SetShadingModel(EGLTFMaterialShadingModel::ClearCoat);
+
+		FMaterialExpressionScalar* ClearCoatFactor = MaterialElement.AddMaterialExpression<FMaterialExpressionScalar>();
+		ClearCoatFactor->GetScalar() = GLTFMaterial.ClearCoat.ClearCoatFactor;
+		ClearCoatFactor->SetName(TEXT("ClearCoatFactor"));
+
+		FMaterialExpressionScalar* ClearCoatRoughnessFactor = MaterialElement.AddMaterialExpression<FMaterialExpressionScalar>();
+		ClearCoatRoughnessFactor->GetScalar() = GLTFMaterial.ClearCoat.Roughness;
+		ClearCoatRoughnessFactor->SetName(TEXT("ClearCoatRoughnessFactor"));
+
+		FMaterialExpressionTexture* ClearCoatTexture = MapFactory.CreateTextureMap(
+			GetTexture(GLTFMaterial.ClearCoat.ClearCoatMap, Textures),
+			GLTFMaterial.ClearCoat.ClearCoatMap.TexCoord,
+			TEXT("ClearCoat"),
+			ETextureMode::Color,
+			GLTFMaterial.ClearCoat.ClearCoatMap.bHasTextureTransform ? &GLTFMaterial.ClearCoat.ClearCoatMap.TextureTransform : nullptr);
+
+		FMaterialExpressionTexture* ClearCoatRoughnessTexture = MapFactory.CreateTextureMap(
+			GetTexture(GLTFMaterial.ClearCoat.RoughnessMap, Textures),
+			GLTFMaterial.ClearCoat.RoughnessMap.TexCoord,
+			TEXT("ClearCoatRoughness"),
+			ETextureMode::Color,
+			GLTFMaterial.ClearCoat.RoughnessMap.bHasTextureTransform ? &GLTFMaterial.ClearCoat.RoughnessMap.TextureTransform : nullptr);
+
+		FMaterialExpression* ClearCoatExpr = ClearCoatFactor;
+		if (ClearCoatTexture)
+		{
+			FMaterialExpressionGeneric* MultiplyExpr = MaterialElement.AddMaterialExpression<FMaterialExpressionGeneric>();
+			MultiplyExpr->SetExpressionName(TEXT("Multiply"));
+
+			ClearCoatFactor->ConnectExpression(*MultiplyExpr->GetInput(0), 0);
+			ClearCoatTexture->ConnectExpression(*MultiplyExpr->GetInput(1), (int)FPBRMapFactory::EChannel::Red);
+			ClearCoatExpr = MultiplyExpr;
+		}
+
+		FMaterialExpression* ClearCoatRougnessExpr = ClearCoatRoughnessFactor;
+		if (ClearCoatRoughnessTexture)
+		{
+			FMaterialExpressionGeneric* MultiplyExpr = MaterialElement.AddMaterialExpression<FMaterialExpressionGeneric>();
+			MultiplyExpr->SetExpressionName(TEXT("Multiply"));
+
+			ClearCoatRoughnessFactor->ConnectExpression(*MultiplyExpr->GetInput(0), 0);
+			ClearCoatRoughnessTexture->ConnectExpression(*MultiplyExpr->GetInput(1), (int)FPBRMapFactory::EChannel::Green);
+			ClearCoatRougnessExpr = MultiplyExpr;
+		}
+
+		ClearCoatExpr->ConnectExpression(MaterialElement.GetClearCoat(), 0);
+		ClearCoatRougnessExpr->ConnectExpression(MaterialElement.GetClearCoatRoughness(), 0);
+
+		FMaterialExpressionTexture* ClearCoatNormalTexture = MapFactory.CreateTextureMap(
+			GetTexture(GLTFMaterial.ClearCoat.NormalMap, Textures),
+			GLTFMaterial.ClearCoat.NormalMap.TexCoord,
+			TEXT("ClearCoatNormal"),
+			ETextureMode::Normal,
+			GLTFMaterial.ClearCoat.NormalMap.bHasTextureTransform ? &GLTFMaterial.ClearCoat.NormalMap.TextureTransform : nullptr);
+
+		if (ClearCoatNormalTexture)
+		{
+			FMaterialExpressionScalar* UVScale = MaterialElement.AddMaterialExpression<FMaterialExpressionScalar>();
+			UVScale->GetScalar() = GLTFMaterial.ClearCoat.NormalMapUVScale;
+			UVScale->SetName(TEXT("NormalMapUVScale"));
+
+			FMaterialExpressionGeneric* Mult = MaterialElement.AddMaterialExpression<FMaterialExpressionGeneric>();
+			Mult->SetExpressionName(TEXT("Multiply"));
+
+			FMaterialExpressionTextureCoordinate* TexCoord = MaterialElement.AddMaterialExpression<FMaterialExpressionTextureCoordinate>();
+			TexCoord->SetCoordinateIndex(GLTFMaterial.ClearCoat.NormalMap.TexCoord);
+			
+			FMaterialExpressionGeneric* ClearCoatNormalCustomOutput = MaterialElement.AddMaterialExpression<FMaterialExpressionGeneric>();
+			ClearCoatNormalCustomOutput->SetExpressionName(TEXT("ClearCoatNormalCustomOutput"));
+			
+			UVScale->ConnectExpression(*Mult->GetInput(0), 0);
+			TexCoord->ConnectExpression(*Mult->GetInput(1), 0);
+			Mult->ConnectExpression(*ClearCoatNormalTexture->GetInput(0), 0);
+			ClearCoatNormalTexture->ConnectExpression(*ClearCoatNormalCustomOutput->GetInput(0), 0);
+		}
+	}
+
+	void FMaterialFactoryImpl::HandleTransmission(const TArray<GLTF::FTexture>& Textures, const GLTF::FMaterial& GLTFMaterial, FPBRMapFactory& MapFactory, FMaterialElement& MaterialElement)
+	{
+		if (!GLTFMaterial.bHasTransmission)
+		{
+			return;
+		}
+
+		MaterialElement.SetBlendMode(ConvertAlphaMode(GLTF::FMaterial::EAlphaMode::Blend));
+		MaterialElement.SetShadingModel(GLTF::EGLTFMaterialShadingModel::ThinTranslucent);
+		MaterialElement.SetTranslucencyLightingMode(ETranslucencyLightingMode::TLM_SurfacePerPixelLighting);
+		MaterialElement.SetTwoSided(GLTFMaterial.bIsDoubleSided);
+
+		FMaterialExpressionGeneric* ThinTranslucentOutput = MaterialElement.AddMaterialExpression<FMaterialExpressionGeneric>();
+		ThinTranslucentOutput->SetExpressionName(TEXT("ThinTranslucentMaterialOutput"));
+
+		FMaterialExpressionScalar* TransmissionFactorExpr = MaterialElement.AddMaterialExpression<FMaterialExpressionScalar>();
+		TransmissionFactorExpr->GetScalar() = GLTFMaterial.Transmission.TransmissionFactor;
+		TransmissionFactorExpr->SetName(TEXT("TransmissionFactor"));
+
+		FMaterialExpressionTexture* TransmissionTexture = MapFactory.CreateTextureMap(
+			GetTexture(GLTFMaterial.Transmission.TransmissionMap, Textures),
+			GLTFMaterial.Transmission.TransmissionMap.TexCoord,
+			TEXT("Transmission"),
+			ETextureMode::Color,
+			GLTFMaterial.Transmission.TransmissionMap.bHasTextureTransform ? &GLTFMaterial.Transmission.TransmissionMap.TextureTransform : nullptr);
+
+		FMaterialExpression* MetallicFactor = MaterialElement.GetMetallic().GetExpression();
+		FMaterialExpression* RoughnessFactor = MaterialElement.GetRoughness().GetExpression();
+
+		FMaterialExpression* MetallicExpr = nullptr;
+		FMaterialExpression* RoughnessExpr = nullptr;
+
+		FMaterialExpressionTexture* MetallicRoughnessTexture = FindExpression<FMaterialExpressionTexture>(TEXT("MetallicRoughness"), MaterialElement);
+		
+		if (MetallicRoughnessTexture)
+		{
+			FMaterialExpressionGeneric* MultMetallic = MaterialElement.AddMaterialExpression<FMaterialExpressionGeneric>();
+			MultMetallic->SetExpressionName(TEXT("Multiply"));
+			FMaterialExpressionGeneric* MultRoughness = MaterialElement.AddMaterialExpression<FMaterialExpressionGeneric>();
+			MultRoughness->SetExpressionName(TEXT("Multiply"));
+
+			MetallicRoughnessTexture->ConnectExpression(*MultMetallic->GetInput(0), (int)FPBRMapFactory::EChannel::Blue);
+			MetallicFactor->ConnectExpression(*MultMetallic->GetInput(1), 0);
+
+			MetallicRoughnessTexture->ConnectExpression(*MultRoughness->GetInput(0), (int)FPBRMapFactory::EChannel::Green);
+			RoughnessFactor->ConnectExpression(*MultMetallic->GetInput(1), 0);
+
+			MetallicExpr = MultMetallic;
+			RoughnessExpr = MultRoughness;
+		}
+		else
+		{
+			MetallicExpr = MetallicFactor;
+			RoughnessExpr = RoughnessFactor;
+		}
+
+		FMaterialExpressionTexture* BaseColorMap = FindExpression<FMaterialExpressionTexture>(TEXT("BaseColor Map"), MaterialElement);
+
+		if (GLTFMaterial.AlphaMode == FMaterial::EAlphaMode::Mask)
+		{
+			// Build opacity mask graph
+
+			FMaterialExpressionColor* BaseColorFactor = FindExpression<FMaterialExpressionColor>(TEXT("BaseColor"), MaterialElement);
+
+			FMaterialExpression* BaseColor1 = nullptr;
+			FMaterialExpression* BaseColor2 = nullptr;
+
+			if (BaseColorMap)
+			{
+				BaseColor1 = BaseColorMap;
+			}
+			else
+			{
+				FMaterialExpressionColor* DefaultColor = MaterialElement.AddMaterialExpression<FMaterialExpressionColor>();
+				DefaultColor->GetColor() = FLinearColor::White;
+				BaseColor1 = DefaultColor;
+			}
+
+			if (BaseColorFactor)
+			{
+				BaseColor2 = BaseColorFactor;
+			}
+			else
+			{
+				FMaterialExpressionColor* DefaultColor = MaterialElement.AddMaterialExpression<FMaterialExpressionColor>();
+				DefaultColor->GetColor() = FLinearColor::White;
+				BaseColor2 = DefaultColor;
+			}
+
+			FMaterialExpressionGeneric* Mult1 = MaterialElement.AddMaterialExpression<FMaterialExpressionGeneric>();
+			Mult1->SetExpressionName(TEXT("Multiply"));
+
+			BaseColor1->ConnectExpression(*Mult1->GetInput(0), 0);
+			BaseColor2->ConnectExpression(*Mult1->GetInput(1), 0);
+
+			FMaterialExpressionGeneric* Step1 = MaterialElement.AddMaterialExpression<FMaterialExpressionGeneric>();
+			Step1->SetExpressionName(TEXT("Step"));
+			Step1->SetFloatProperty(TEXT("ConstX"), 0.333f);
+
+			FMaterialExpressionGeneric* OneMinus1 = MaterialElement.AddMaterialExpression<FMaterialExpressionGeneric>();
+			OneMinus1->SetExpressionName(TEXT("OneMinus"));
+			
+			BaseColor1->ConnectExpression(*Step1->GetInput(0), (int)FPBRMapFactory::EChannel::Alpha);
+			Step1->ConnectExpression(*OneMinus1->GetInput(0), 0);
+			Step1->ConnectExpression(MaterialElement.GetSpecular(), 0);
+
+			FMaterialExpressionGeneric* OpacityMult = MaterialElement.AddMaterialExpression<FMaterialExpressionGeneric>();
+			OpacityMult->SetExpressionName(TEXT("Multiply"));
+
+			OneMinus1->ConnectExpression(*OpacityMult->GetInput(0), 0);
+			OneMinus1->ConnectExpression(MaterialElement.GetSpecular(), 0);
+
+			OpacityMult->ConnectExpression(MaterialElement.GetOpacity(), 0);
+
+			FMaterialExpressionGeneric* Max1 = MaterialElement.AddMaterialExpression<FMaterialExpressionGeneric>();
+			Max1->SetExpressionName(TEXT("Max"));
+
+			FMaterialExpressionGeneric* Max2 = MaterialElement.AddMaterialExpression<FMaterialExpressionGeneric>();
+			Max2->SetExpressionName(TEXT("Max"));
+
+			MetallicExpr->ConnectExpression(*Max1->GetInput(0), 0);
+			RoughnessExpr->ConnectExpression(*Max1->GetInput(1), 0);
+	
+			Max1->ConnectExpression(*Max2->GetInput(1), 0);
+			Max2->ConnectExpression(*OpacityMult->GetInput(1), 0);
+
+			FMaterialExpression* TransmissionExpr = nullptr;
+
+			if (TransmissionTexture)
+			{
+				FMaterialExpressionGeneric* MultTransm = MaterialElement.AddMaterialExpression<FMaterialExpressionGeneric>();
+				MultTransm->SetExpressionName(TEXT("Multiply"));
+
+				TransmissionTexture->ConnectExpression(*MultTransm->GetInput(0), (int)FPBRMapFactory::EChannel::Red);
+				TransmissionFactorExpr->ConnectExpression(*MultTransm->GetInput(1), 0);
+
+				TransmissionExpr = MultTransm;
+			}
+			else
+			{
+				TransmissionExpr = TransmissionFactorExpr;
+			}
+
+			// Transmission network
+			FMaterialExpressionGeneric* Mult2 = MaterialElement.AddMaterialExpression<FMaterialExpressionGeneric>();
+			Mult2->SetExpressionName(TEXT("Multiply"));
+
+			FMaterialExpressionGeneric* OneMinus2 = MaterialElement.AddMaterialExpression<FMaterialExpressionGeneric>();
+			OneMinus2->SetExpressionName(TEXT("OneMinus"));
+
+			FMaterialExpressionGeneric* LerpThinTransl = MaterialElement.AddMaterialExpression<FMaterialExpressionGeneric>();
+			LerpThinTransl->SetExpressionName(TEXT("LinearInterpolate"));
+
+			TransmissionExpr->ConnectExpression(*Mult2->GetInput(0), 0);
+			TransmissionExpr->ConnectExpression(*OneMinus2->GetInput(0), 0);
+			Mult1->ConnectExpression(*Mult2->GetInput(1), 0);
+
+			Mult2->ConnectExpression(*LerpThinTransl->GetInput(0), 0);
+			Step1->ConnectExpression(*LerpThinTransl->GetInput(2), 0);
+
+			OneMinus2->ConnectExpression(*Max2->GetInput(0), 0);
+
+			LerpThinTransl->ConnectExpression(*ThinTranslucentOutput->GetInput(0), 0);
+		
+			Mult1->ConnectExpression(MaterialElement.GetBaseColor(), 0);
+		}
+		else
+		{
+			FMaterialExpressionGeneric* BaseColorMult = MaterialElement.AddMaterialExpression<FMaterialExpressionGeneric>();
+			BaseColorMult->SetExpressionName(TEXT("Multiply"));
+
+			BaseColorMult->ConnectExpression(*ThinTranslucentOutput->GetInput(0), 0);
+
+			FMaterialExpressionGeneric* OneMinus = MaterialElement.AddMaterialExpression<FMaterialExpressionGeneric>();
+			OneMinus->SetExpressionName(TEXT("OneMinus"));
+
+			FMaterialExpressionGeneric* Max1 = MaterialElement.AddMaterialExpression<FMaterialExpressionGeneric>();
+			Max1->SetExpressionName(TEXT("Max"));
+
+			FMaterialExpressionGeneric* Max2 = MaterialElement.AddMaterialExpression<FMaterialExpressionGeneric>();
+			Max2->SetExpressionName(TEXT("Max"));
+
+			MetallicExpr->ConnectExpression(*Max1->GetInput(0), 0);
+			RoughnessExpr->ConnectExpression(*Max1->GetInput(1), 0);
+			OneMinus->ConnectExpression(*Max2->GetInput(0), 0);
+			Max1->ConnectExpression(*Max2->GetInput(1), 0);
+
+			FMaterialExpression* BaseColorExpr = MaterialElement.GetBaseColor().GetExpression();
+			if (BaseColorExpr)
+			{
+				BaseColorExpr->ConnectExpression(*BaseColorMult->GetInput(0), MaterialElement.GetBaseColor().GetOutputIndex());
+			}
+
+			if (TransmissionTexture)
+			{
+				FMaterialExpressionGeneric* TransmissionMult = MaterialElement.AddMaterialExpression<FMaterialExpressionGeneric>();
+				TransmissionMult->SetExpressionName(TEXT("Multiply"));
+				TransmissionFactorExpr->ConnectExpression(*TransmissionMult->GetInput(0), 0);
+				TransmissionTexture->ConnectExpression(*TransmissionMult->GetInput(1), (int)FPBRMapFactory::EChannel::Red);
+				TransmissionMult->ConnectExpression(*BaseColorMult->GetInput(1), 0);
+				TransmissionMult->ConnectExpression(*OneMinus->GetInput(0), 0);
+			}
+			else
+			{
+				TransmissionFactorExpr->ConnectExpression(*BaseColorMult->GetInput(1), 0);
+				TransmissionFactorExpr->ConnectExpression(*OneMinus->GetInput(0), 0);
+			}
+
+			Max2->ConnectExpression(MaterialElement.GetOpacity(), 0);
+		}
+	}
+	
+	void FMaterialFactoryImpl::HandleSheen(const TArray<GLTF::FTexture>& Textures, const GLTF::FMaterial& GLTFMaterial, FPBRMapFactory& MapFactory, FMaterialElement& MaterialElement)
+	{
+		if (!GLTFMaterial.bHasSheen || GLTFMaterial.bIsUnlitShadingModel)
+		{
+			return;
+		}
+
+		FMaterialExpressionFunctionCall* FuzzyShadingCall = MaterialElement.AddMaterialExpression<FMaterialExpressionFunctionCall>();
+		FuzzyShadingCall->SetFunctionPathName(TEXT("/Engine/Functions/Engine_MaterialFunctions01/Shading/FuzzyShading.FuzzyShading"));
+
+		FMaterialExpressionColor* Fuzzyness = MaterialElement.AddMaterialExpression<FMaterialExpressionColor>();
+		Fuzzyness->SetName(TEXT("Fuzzyness"));
+		Fuzzyness->GetColor() = FLinearColor(0.9f, 0.8f, 2.0f, 1.0f);
+		Fuzzyness->ConnectExpression(*FuzzyShadingCall->GetInput(2), (int)FPBRMapFactory::EChannel::Red);	// CoreDarkness
+		Fuzzyness->ConnectExpression(*FuzzyShadingCall->GetInput(4), (int)FPBRMapFactory::EChannel::Green);	// EdgeBrightness
+		Fuzzyness->ConnectExpression(*FuzzyShadingCall->GetInput(3), (int)FPBRMapFactory::EChannel::Blue);	// Power
+
+		if (FMaterialExpression* BaseColorExpr = MaterialElement.GetBaseColor().GetExpression())
+		{
+			BaseColorExpr->ConnectExpression(*FuzzyShadingCall->GetInput(0), MaterialElement.GetBaseColor().GetOutputIndex());
+
+			FMaterialExpressionGeneric* LerpExpression = MaterialElement.AddMaterialExpression<FMaterialExpressionGeneric>();
+			LerpExpression->SetExpressionName(TEXT("LinearInterpolate"));
+
+			FuzzyShadingCall->ConnectExpression(*LerpExpression->GetInput(0), 0);
+			BaseColorExpr->ConnectExpression(*LerpExpression->GetInput(1), MaterialElement.GetBaseColor().GetOutputIndex());
+			Fuzzyness->ConnectExpression(*LerpExpression->GetInput(2), (int)FPBRMapFactory::EChannel::Alpha);
+
+			LerpExpression->ConnectExpression(MaterialElement.GetBaseColor(), 0);
+		}
+
+		if (FMaterialExpression* NormalExpr = MaterialElement.GetNormal().GetExpression())
+		{
+			NormalExpr->ConnectExpression(*FuzzyShadingCall->GetInput(1), MaterialElement.GetNormal().GetOutputIndex());
+		}
+	}
+
+	void FMaterialFactoryImpl::HandleIOR(const TArray<GLTF::FTexture>& Textures, const GLTF::FMaterial& GLTFMaterial, FPBRMapFactory& MapFactory, FMaterialElement& MaterialElement)
+	{
+		if (!GLTFMaterial.bHasIOR || GLTFMaterial.ShadingModel == FMaterial::EShadingModel::SpecularGlossiness || GLTFMaterial.bIsUnlitShadingModel)
+		{
+			// As per glTF 2.0 spec, IOR extension is not compatible with specular-glossiness workflow or unlit shading model
+			return;
+		}
+
+		if (!GLTFMaterial.bHasTransmission)
+		{
+			MaterialElement.SetBlendMode(ConvertAlphaMode(GLTF::FMaterial::EAlphaMode::Blend));
+			MaterialElement.SetTranslucencyLightingMode(ETranslucencyLightingMode::TLM_Surface);
+			MaterialElement.SetTwoSided(GLTFMaterial.bIsDoubleSided);
+		}
+
+		FMaterialExpressionGeneric* LerpExpr = MaterialElement.AddMaterialExpression<FMaterialExpressionGeneric>();
+		LerpExpr->SetExpressionName(TEXT("LinearInterpolate"));
+
+		FMaterialExpressionScalar* ConstExpr = MaterialElement.AddMaterialExpression<FMaterialExpressionScalar>();
+		ConstExpr->GetScalar() = 1.f;
+
+		FMaterialExpressionScalar* IORExpr = MaterialElement.AddMaterialExpression<FMaterialExpressionScalar>();
+		IORExpr->SetName(TEXT("IOR"));
+		IORExpr->GetScalar() = GLTFMaterial.IOR;
+
+		FMaterialExpressionGeneric* FresnelExpr = MaterialElement.AddMaterialExpression<FMaterialExpressionGeneric>();
+		FresnelExpr->SetExpressionName(TEXT("Fresnel"));
+
+		ConstExpr->ConnectExpression(*LerpExpr->GetInput(0), 0);
+		IORExpr->ConnectExpression(*LerpExpr->GetInput(1), 0);
+		FresnelExpr->ConnectExpression(*LerpExpr->GetInput(2), 0);
+
+		LerpExpr->ConnectExpression(MaterialElement.GetRefraction(), 0);
+	}
+
+	void FMaterialFactoryImpl::HandleSpecular(const TArray<GLTF::FTexture>& Textures, const GLTF::FMaterial& GLTFMaterial, FPBRMapFactory& MapFactory, FMaterialElement& MaterialElement)
+	{
+		if (!GLTFMaterial.bHasSpecular || GLTFMaterial.ShadingModel == FMaterial::EShadingModel::SpecularGlossiness || GLTFMaterial.bIsUnlitShadingModel)
+		{
+			// As per glTF 2.0 spec, specular extension is not compatible with specular-glossiness workflow or unlit shading model
+			return;
+		}
+
+		FMaterialExpressionScalar* SpecularFactorExpr = MaterialElement.AddMaterialExpression<FMaterialExpressionScalar>();
+		SpecularFactorExpr->SetName(TEXT("SpecularFactor"));
+		SpecularFactorExpr->GetScalar() = GLTFMaterial.Specular.SpecularFactor;
+
+		FMaterialExpressionTexture* SpecularTextureExpr = MapFactory.CreateTextureMap(
+			GetTexture(GLTFMaterial.Specular.SpecularMap, Textures), GLTFMaterial.Specular.SpecularMap.TexCoord, TEXT("SpecularMap"), ETextureMode::Color,
+			GLTFMaterial.Specular.SpecularMap.bHasTextureTransform ? &GLTFMaterial.Specular.SpecularMap.TextureTransform : nullptr);
+
+		FMaterialExpression* SpecularExpr = SpecularFactorExpr;
+
+		if (SpecularTextureExpr)
+		{
+			FMaterialExpressionGeneric* MultiplyExpr = MaterialElement.AddMaterialExpression<FMaterialExpressionGeneric>();
+			MultiplyExpr->SetExpressionName(TEXT("Multiply"));
+			
+			SpecularFactorExpr->ConnectExpression(*MultiplyExpr->GetInput(0), 0);
+			SpecularTextureExpr->ConnectExpression(*MultiplyExpr->GetInput(1), (int)FPBRMapFactory::EChannel::Alpha);
+
+			SpecularExpr = MultiplyExpr;
+		}
+
+		SpecularExpr->ConnectExpression(MaterialElement.GetSpecular(), 0);
+	}
+
+	FMaterialFactory::FMaterialFactory(IMaterialElementFactory* MaterialElementFactory, ITextureFactory* TextureFactory)
+	    : Impl(new FMaterialFactoryImpl(MaterialElementFactory, TextureFactory))
+	{
+	}
+
+	FMaterialFactory::~FMaterialFactory() {}
+
+	const TArray<FMaterialElement*>& FMaterialFactory::CreateMaterials(const GLTF::FAsset& Asset, UObject* ParentPackage, EObjectFlags Flags)
+	{
+		return Impl->CreateMaterials(Asset, ParentPackage, Flags);
+	}
+
+	const TArray<FLogMessage>& FMaterialFactory::GetLogMessages() const
+	{
+		return Impl->GetLogMessages();
+	}
+
+	const TArray<FMaterialElement*>& FMaterialFactory::GetMaterials() const
+	{
+		return Impl->Materials;
+	}
+
+	IMaterialElementFactory& FMaterialFactory::GetMaterialElementFactory()
+	{
+		return *Impl->MaterialElementFactory;
+	}
+
+	ITextureFactory& FMaterialFactory::GetTextureFactory()
+	{
+		return *Impl->TextureFactory;
+	}
+
+	void FMaterialFactory::CleanUp()
+	{
+		for (FMaterialElement* MaterialElement : Impl->Materials)
+		{
+			delete MaterialElement;
+		}
+		Impl->Materials.Empty();
+	}
+
+}  // namespace GLTF
diff --git a/Engine/Plugins/Interchange/Runtime/Source/Parsers/GLTFCore/Private/GLTFMeshFactory.cpp b/Engine/Plugins/Interchange/Runtime/Source/Parsers/GLTFCore/Private/GLTFMeshFactory.cpp
new file mode 100644
--- /dev/null
+++ b/Engine/Plugins/Interchange/Runtime/Source/Parsers/GLTFCore/Private/GLTFMeshFactory.cpp
@@ -0,0 +1,651 @@
+// Copyright Epic Games, Inc. All Rights Reserved.
+
+#include "GLTFMeshFactory.h"
+
+
+#include "GLTFMesh.h"
+#include "StaticMeshAttributes.h"
+#include "Util/ColorConstants.h"
+
+namespace GLTF
+{
+	class FMeshFactoryImpl : public GLTF::FBaseLogger
+	{
+	public:
+		FMeshFactoryImpl();
+
+		void FillMeshDescription(const GLTF::FMesh &Mesh, FMeshDescription* MeshDescription, const TArray<float>& MorphTargetWeights);
+
+		void CleanUp();
+
+	private:
+
+		bool ImportPrimitive(const GLTF::FPrimitive&                        Primitive,  //
+			int32                                          PrimitiveIndex,
+			int32                                          NumUVs,
+			bool                                           bMeshHasTagents,
+			bool                                           bMeshHasColors,
+			const TVertexInstanceAttributesRef<FVector3f>&   VertexInstanceNormals,
+			const TVertexInstanceAttributesRef<FVector3f>&   VertexInstanceTangents,
+			const TVertexInstanceAttributesRef<float>&     VertexInstanceBinormalSigns,
+			const TVertexInstanceAttributesRef<FVector2f>& VertexInstanceUVs,
+			const TVertexInstanceAttributesRef<FVector4f>&  VertexInstanceColors,
+			const TEdgeAttributesRef<bool>&                EdgeHardnesses,
+			FMeshDescription* MeshDescription,
+			const TArray<float>& MorphTargetWeights);
+
+		inline TArray<FVector4f>& GetVector4dBuffer(int32 Index)
+		{
+			check(Index < (sizeof(Vector4dBuffers) / sizeof(Vector4dBuffers[0])));
+			uint32 ReserveSize = Vector4dBuffers[Index].Num() + Vector4dBuffers[Index].GetSlack();
+			Vector4dBuffers[Index].Empty(ReserveSize);
+			return Vector4dBuffers[Index];
+		}
+
+		inline TArray<FVector3f>& GetVectorBuffer(int32 Index)
+		{
+			check(Index < (sizeof(VectorBuffers) / sizeof(VectorBuffers[0])));
+			uint32 ReserveSize = VectorBuffers[Index].Num() + VectorBuffers[Index].GetSlack();
+			VectorBuffers[Index].Empty(ReserveSize);
+			return VectorBuffers[Index];
+		}
+
+		inline TArray<FVector2f>& GetVector2dBuffer(int32 Index)
+		{
+			check(Index < (sizeof(Vector2dBuffers) / sizeof(Vector2dBuffers[0])));
+			uint32 ReserveSize = Vector2dBuffers[Index].Num() + Vector2dBuffers[Index].GetSlack();
+			Vector2dBuffers[Index].Empty(ReserveSize);
+			return Vector2dBuffers[Index];
+		}
+
+		inline TArray<uint32>& GetIntBuffer()
+		{
+			uint32 ReserveSize = IntBuffer.Num() + IntBuffer.GetSlack();
+			IntBuffer.Empty(ReserveSize);
+			return IntBuffer;
+		}
+
+	private:
+
+		enum Vector3fMorphTargetAttributes
+		{
+			POSITIONS,
+			NORMALS
+		};
+		void UpdateVector3fFromMorphTargets(TArray<FVector3f>& AttrbiuteValues, const FPrimitive& Primitive, const TArray<float>& MorphTargetWeights, const Vector3fMorphTargetAttributes& AttributeToProcess);
+
+		enum Vector4dMorphTargetAttributes
+		{
+			TANGENTS,
+			COLORS
+		};
+		void UpdateVector4dFromMorphTargets(TArray<FVector4f>& AttrbiuteValues, const FPrimitive& Primitive, const TArray<float>& MorphTargetWeights, const Vector4dMorphTargetAttributes& AttributeToProcess);
+
+		void UpdateTexCoordsFromMorphTargets(TArray<FVector2f>& AttrbiuteValues, const FPrimitive& Primitive, const TArray<float>& MorphTargetWeights, int32 Index);
+
+		enum
+		{
+			NormalBufferIndex = 0,
+			PositionBufferIndex = 1,
+			ReindexBufferIndex = 2,
+			MorphTargetBufferIndex = 3,
+			VectorBufferCount = 4,
+			UvReindexBufferIndex = MAX_MESH_TEXTURE_COORDS_MD,
+			MorphTargetUvBufferIndex = UvReindexBufferIndex+1,
+			ColorBufferIndex = 0,
+			TangentBufferIndex = 1,
+			Reindex4dBufferIndex = 2,
+			MorphTarget4dBufferIndex = 3,
+			Vector4dBufferCount = 4,
+		};
+
+		float                ImportUniformScale;
+
+		TSet<int32>                  MaterialIndicesUsed;
+		TMap<int32, FPolygonGroupID> MaterialIndexToPolygonGroupID;
+		TArray<FMeshFactory::FIndexVertexIdMap>    PositionIndexToVertexIdPerPrim;
+
+		TArray<FVector2f>                       Vector2dBuffers[MAX_MESH_TEXTURE_COORDS_MD + 2];
+		TArray<FVector3f>                       VectorBuffers[VectorBufferCount];
+		TArray<FVector4f>                       Vector4dBuffers[Vector4dBufferCount];
+		TArray<uint32>                          IntBuffer;
+		TArray<FVertexInstanceID>				CornerVertexInstanceIDs;
+		uint32                                  MaxReserveSize;
+
+		friend class FMeshFactory;
+	};
+
+	namespace
+	{
+		template <typename T>
+		void ReIndex(const TArray<T>& Source, const TArray<uint32>& Indices, TArray<T>& Dst)
+		{
+			if (!ensure(&Source != &Dst))
+			{
+				return;
+			}
+
+			Dst.Reserve(Indices.Num());
+			for (uint32 Index : Indices)
+			{
+				Dst.Add(Source[Index]);
+			}
+		}
+
+		void GenerateFlatNormals(const TArray<FVector3f>& Positions, const TArray<uint32>& Indices, TArray<FVector3f>& Normals)
+		{
+			Normals.Empty();
+
+			uint32 N = Indices.Num();
+			if (!ensure(N % 3 == 0))
+			{
+				N -= N % 3;
+			}
+			Normals.AddUninitialized(N);
+
+			for (uint32 i = 0; i < N; i += 3)
+			{
+				const FVector3f& A = Positions[Indices[i]];
+				const FVector3f& B = Positions[Indices[i + 1]];
+				const FVector3f& C = Positions[Indices[i + 2]];
+
+				const FVector3f Normal = FVector3f::CrossProduct(A - B, A - C).GetSafeNormal();
+				// Same for each corner of the triangle.
+				Normals[i] = Normal;
+				Normals[i + 1] = Normal;
+				Normals[i + 2] = Normal;
+			}
+		}
+
+		int32 GetNumUVs(const GLTF::FMesh& Mesh)
+		{
+			int32 NumUVs = 0;
+			for (int32 UVIndex = 0; UVIndex < MAX_MESH_TEXTURE_COORDS_MD; ++UVIndex)
+			{
+				if (Mesh.HasTexCoords(UVIndex))
+				{
+					NumUVs++;
+				}
+				else
+				{
+					break;
+				}
+			}
+			return NumUVs;
+		}
+	}
+
+	FMeshFactoryImpl::FMeshFactoryImpl()
+		: ImportUniformScale(1.f)
+	{
+		CornerVertexInstanceIDs.SetNum(3);
+	}
+
+	void FMeshFactoryImpl::UpdateVector3fFromMorphTargets(TArray<FVector3f>& AttributeValues, const FPrimitive& Primitive, const TArray<float>& MorphTargetWeights, const Vector3fMorphTargetAttributes& AttributeToProcess)
+	{
+		if (MorphTargetWeights.Num() != Primitive.MorphTargets.Num())
+		{
+			return;
+		}
+
+		for (int32 MorphTargetIndex = 0; MorphTargetIndex < Primitive.MorphTargets.Num(); MorphTargetIndex++)
+		{
+			float MorphTargetWeight = MorphTargetWeights.Num() > MorphTargetIndex ? MorphTargetWeights[MorphTargetIndex] : 0;
+
+			if (MorphTargetWeight > 0 && MorphTargetWeight <= 1.0)
+			{
+				TArray<FVector3f>& Displacements = GetVectorBuffer(MorphTargetBufferIndex);
+
+				switch (AttributeToProcess)
+				{
+				case POSITIONS:
+					if (!Primitive.MorphTargets[MorphTargetIndex].HasPositionDisplacements()) continue;
+					Primitive.MorphTargets[MorphTargetIndex].GetPositionDisplacements(Displacements);
+					break;
+				case NORMALS:
+					if (!Primitive.MorphTargets[MorphTargetIndex].HasNormalDisplacements()) continue;
+					Primitive.MorphTargets[MorphTargetIndex].GetNormalDisplacements(Displacements);
+					break;
+				default:
+					return;
+					break;
+				}
+
+				if (AttributeValues.Num() != Displacements.Num())
+				{
+					continue;
+				}
+				
+				for (int32 Index = 0; Index < Displacements.Num(); Index++)
+				{
+					AttributeValues[Index] += MorphTargetWeight * Displacements[Index];
+				}
+			}
+		}
+	}
+	void FMeshFactoryImpl::UpdateVector4dFromMorphTargets(TArray<FVector4f>& AttributeValues, const FPrimitive& Primitive, const TArray<float>& MorphTargetWeights, const Vector4dMorphTargetAttributes& AttributeToProcess)
+	{
+		if (MorphTargetWeights.Num() != Primitive.MorphTargets.Num())
+		{
+			return;
+		}
+
+		for (int32 MorphTargetIndex = 0; MorphTargetIndex < Primitive.MorphTargets.Num(); MorphTargetIndex++)
+		{
+			float MorphTargetWeight = MorphTargetWeights.Num() > MorphTargetIndex ? MorphTargetWeights[MorphTargetIndex] : 0;
+
+			if (MorphTargetWeight > 0 && MorphTargetWeight <= 1.0)
+			{
+				TArray<FVector4f>& Displacements = GetVector4dBuffer(MorphTargetBufferIndex);
+
+				switch (AttributeToProcess)
+				{
+				case TANGENTS:
+					if (!Primitive.MorphTargets[MorphTargetIndex].HasTangentDisplacements()) continue;
+					Primitive.MorphTargets[MorphTargetIndex].GetTangentDisplacements(Displacements);
+					break;
+				case COLORS:
+					if (!Primitive.MorphTargets[MorphTargetIndex].HasColorDeltas()) continue;
+					Primitive.MorphTargets[MorphTargetIndex].GetColorDeltas(Displacements);
+					break;
+				default:
+					return;
+					break;
+				}
+
+				if (AttributeValues.Num() != Displacements.Num())
+				{
+					continue;
+				}
+
+				for (int32 Index = 0; Index < Displacements.Num(); Index++)
+				{
+					AttributeValues[Index] += MorphTargetWeight * Displacements[Index];
+
+					if (AttributeToProcess == Vector4dMorphTargetAttributes::COLORS)
+					{
+						//After applying color deltas, all components of each COLOR_0 morphed accessor element MUST be clamped to [0.0, 1.0] range.
+						AttributeValues[Index].X = FMath::Clamp(AttributeValues[Index].X, .0f, 1.f);
+						AttributeValues[Index].Y = FMath::Clamp(AttributeValues[Index].Y, .0f, 1.f);
+						AttributeValues[Index].Z = FMath::Clamp(AttributeValues[Index].W, .0f, 1.f);
+						AttributeValues[Index].W = FMath::Clamp(AttributeValues[Index].Z, .0f, 1.f);
+					}
+				}
+			}
+		}
+	}
+	void FMeshFactoryImpl::UpdateTexCoordsFromMorphTargets(TArray<FVector2f>& AttributeValues, const FPrimitive& Primitive, const TArray<float>& MorphTargetWeights, int32 UvIndex)
+	{
+		if (MorphTargetWeights.Num() != Primitive.MorphTargets.Num())
+		{
+			return;
+		}
+
+		for (int32 MorphTargetIndex = 0; MorphTargetIndex < Primitive.MorphTargets.Num(); MorphTargetIndex++)
+		{
+			float MorphTargetWeight = MorphTargetWeights.Num() > MorphTargetIndex ? MorphTargetWeights[MorphTargetIndex] : 0;
+
+			if (MorphTargetWeight > 0 && MorphTargetWeight <= 1.0)
+			{
+				TArray<FVector2f>& Displacements = GetVector2dBuffer(MorphTargetUvBufferIndex);
+
+				if (!Primitive.MorphTargets[MorphTargetIndex].HasTexCoordDisplacements(UvIndex)) continue;
+
+				Primitive.MorphTargets[MorphTargetIndex].GetTexCoordDisplacements(UvIndex, Displacements);
+
+				if (AttributeValues.Num() != Displacements.Num())
+				{
+					continue;
+				}
+
+				for (int32 Index = 0; Index < Displacements.Num(); Index++)
+				{
+					AttributeValues[Index] += MorphTargetWeight * Displacements[Index];
+				}
+			}
+		}
+	}
+
+	void FMeshFactoryImpl::FillMeshDescription(const FMesh &Mesh, FMeshDescription* MeshDescription, const TArray<float>& MorphTargetWeights)
+	{
+		const int32 NumUVs = FMath::Max(1, GetNumUVs(Mesh));
+
+		FStaticMeshAttributes StaticMeshAttributes(*MeshDescription);
+		StaticMeshAttributes.Register();
+
+		TVertexAttributesRef<FVector3f> VertexPositions = StaticMeshAttributes.GetVertexPositions();
+		TEdgeAttributesRef<bool>  EdgeHardnesses = StaticMeshAttributes.GetEdgeHardnesses();
+		TPolygonGroupAttributesRef<FName> PolygonGroupImportedMaterialSlotNames = StaticMeshAttributes.GetPolygonGroupMaterialSlotNames();
+		TVertexInstanceAttributesRef<FVector3f> VertexInstanceNormals = StaticMeshAttributes.GetVertexInstanceNormals();
+		TVertexInstanceAttributesRef<FVector3f> VertexInstanceTangents = StaticMeshAttributes.GetVertexInstanceTangents();
+		TVertexInstanceAttributesRef<float> VertexInstanceBinormalSigns = StaticMeshAttributes.GetVertexInstanceBinormalSigns();
+		TVertexInstanceAttributesRef<FVector2f> VertexInstanceUVs = StaticMeshAttributes.GetVertexInstanceUVs();
+		TVertexInstanceAttributesRef<FVector4f> VertexInstanceColors = StaticMeshAttributes.GetVertexInstanceColors();
+		VertexInstanceUVs.SetNumChannels(NumUVs);
+
+		MaterialIndicesUsed.Empty(10);
+		// Add the vertex position
+		PositionIndexToVertexIdPerPrim.SetNum(FMath::Max(Mesh.Primitives.Num(), PositionIndexToVertexIdPerPrim.Num()));
+		for (int32 Index = 0; Index < Mesh.Primitives.Num(); ++Index)
+		{
+			const FPrimitive& Primitive = Mesh.Primitives[Index];
+
+			if (!Primitive.IsValid())
+			{
+				Messages.Emplace(EMessageSeverity::Warning, TEXT("Mesh has an invalid primitive: ") + Mesh.Name);
+				continue;
+			}
+
+			// Remember which primitives use which materials.
+			MaterialIndicesUsed.Add(Primitive.MaterialIndex);
+
+			TArray<FVector3f>& Positions = GetVectorBuffer(PositionBufferIndex);
+			Primitive.GetPositions(Positions);
+
+			UpdateVector3fFromMorphTargets(Positions, Primitive, MorphTargetWeights, Vector3fMorphTargetAttributes::POSITIONS);
+
+			FMeshFactory::FIndexVertexIdMap& PositionIndexToVertexId = PositionIndexToVertexIdPerPrim[Index];
+			PositionIndexToVertexId.Empty(Positions.Num());
+			for (int32 PositionIndex = 0; PositionIndex < Positions.Num(); ++PositionIndex)
+			{
+				const FVertexID& VertexID = MeshDescription->CreateVertex();
+				VertexPositions[VertexID] = Positions[PositionIndex] * ImportUniformScale;
+				PositionIndexToVertexId.Add(PositionIndex, VertexID);
+			}
+		}
+
+		// Add the PolygonGroup
+		MaterialIndexToPolygonGroupID.Empty(10);
+		for (int32 MaterialIndex : MaterialIndicesUsed)
+		{
+			const FPolygonGroupID& PolygonGroupID = MeshDescription->CreatePolygonGroup();
+			MaterialIndexToPolygonGroupID.Add(MaterialIndex, PolygonGroupID);
+
+			const FName ImportedSlotName(*FString::FromInt(MaterialIndex));
+			PolygonGroupImportedMaterialSlotNames[PolygonGroupID] = ImportedSlotName;
+		}
+
+		// Add the VertexInstance
+		bool bMeshUsesEmptyMaterial = false;
+		bool bDidGenerateTexCoords = false;
+		for (int32 Index = 0; Index < Mesh.Primitives.Num(); ++Index)
+		{
+			const FPrimitive& Primitive = Mesh.Primitives[Index];
+
+			if (!Primitive.IsValid())
+			{
+				continue;
+			}
+
+			const bool        bHasDegenerateTriangles =
+				ImportPrimitive(Primitive, Index, NumUVs, Mesh.HasTangents(), Mesh.HasColors(),  //
+					VertexInstanceNormals, VertexInstanceTangents, VertexInstanceBinormalSigns, VertexInstanceUVs,
+					VertexInstanceColors,  //
+					EdgeHardnesses, MeshDescription,
+					MorphTargetWeights);
+
+			bMeshUsesEmptyMaterial |= Primitive.MaterialIndex == INDEX_NONE;
+			for (int32 UVIndex = 0; UVIndex < NumUVs; ++UVIndex)
+			{
+				if (!Primitive.HasTexCoords(UVIndex))
+				{
+					bDidGenerateTexCoords = true;
+					break;
+				}
+			}
+
+			if (bHasDegenerateTriangles)
+			{
+				Messages.Emplace(EMessageSeverity::Warning,
+					FString::Printf(TEXT("Mesh %s has primitive with degenerate triangles: %d"), *Mesh.Name, Index));
+			}
+		}
+		if (bMeshUsesEmptyMaterial)
+		{
+			Messages.Emplace(EMessageSeverity::Warning, TEXT("Mesh has primitives with no materials assigned: ") + Mesh.Name);
+		}
+	}
+
+	bool FMeshFactoryImpl::ImportPrimitive(const GLTF::FPrimitive&                        Primitive,  //
+		int32                                          PrimitiveIndex,
+		int32                                          NumUVs,
+		bool                                           bMeshHasTangents,
+		bool                                           bMeshHasColors,
+		const TVertexInstanceAttributesRef<FVector3f>&   VertexInstanceNormals,
+		const TVertexInstanceAttributesRef<FVector3f>&   VertexInstanceTangents,
+		const TVertexInstanceAttributesRef<float>&     VertexInstanceBinormalSigns,
+		const TVertexInstanceAttributesRef<FVector2f>& VertexInstanceUVs,
+		const TVertexInstanceAttributesRef<FVector4f>&  VertexInstanceColors,
+		const TEdgeAttributesRef<bool>&                EdgeHardnesses,
+		FMeshDescription* MeshDescription,
+		const TArray<float>& MorphTargetWeights)
+	{
+		const FPolygonGroupID CurrentPolygonGroupID = MaterialIndexToPolygonGroupID[Primitive.MaterialIndex];
+		const uint32          TriCount = Primitive.TriangleCount();
+
+		TArray<uint32>& Indices = GetIntBuffer();
+		Primitive.GetTriangleIndices(Indices);
+
+		//Note from gltf 2.0 specification:
+		//	When the base mesh primitive does not specify normals, client implementations MUST calculate flat normals for each morph target; the provided tangents and their displacements (if present) MUST be ignored.
+		bool bIgnoreTangents = false;
+
+		TArray<FVector3f>& Normals = GetVectorBuffer(NormalBufferIndex);
+		// glTF does not guarantee each primitive within a mesh has the same attributes.
+		// Fill in gaps as needed:
+		// - missing normals will be flat, based on triangle orientation
+		// - missing UVs will be (0,0)
+		// - missing tangents will be (0,0,1)
+		if (Primitive.HasNormals())
+		{
+			TArray<FVector3f>& ReindexBuffer = GetVectorBuffer(ReindexBufferIndex);
+			Primitive.GetNormals(Normals);
+
+			UpdateVector3fFromMorphTargets(Normals, Primitive, MorphTargetWeights, Vector3fMorphTargetAttributes::NORMALS);
+
+			ReIndex(Normals, Indices, ReindexBuffer);
+			Swap(Normals, ReindexBuffer);
+		}
+		else
+		{
+			TArray<FVector3f>& Positions = GetVectorBuffer(PositionBufferIndex);
+			Primitive.GetPositions(Positions);
+
+			//update positions with morph targets before generating flat normals:
+			UpdateVector3fFromMorphTargets(Positions, Primitive, MorphTargetWeights, Vector3fMorphTargetAttributes::POSITIONS);
+
+			GenerateFlatNormals(Positions, Indices, Normals);
+			bIgnoreTangents = true;
+		}
+
+		TArray<FVector4f>& Tangents = GetVector4dBuffer(TangentBufferIndex);
+		if (!bIgnoreTangents && Primitive.HasTangents())
+		{
+			TArray<FVector4f>& ReindexBuffer = GetVector4dBuffer(Reindex4dBufferIndex);
+			Primitive.GetTangents(Tangents);
+
+			UpdateVector4dFromMorphTargets(Tangents, Primitive, MorphTargetWeights, Vector4dMorphTargetAttributes::TANGENTS);
+
+			ReIndex(Tangents, Indices, ReindexBuffer);
+			Swap(Tangents, ReindexBuffer);
+		}
+		else if (bMeshHasTangents)
+		{
+			// If other primitives in this mesh have tangents, generate filler ones for this primitive, to avoid gaps.
+			Tangents.Init(FVector4f(1.0f, 0.0f, 0.0f, 1.0f), Primitive.VertexCount());
+		}
+
+		TArray<FVector4f>& Colors = GetVector4dBuffer(ColorBufferIndex);
+		if (Primitive.HasColors())
+		{
+			TArray<FVector4f>& ReindexBuffer = GetVector4dBuffer(Reindex4dBufferIndex);
+			Primitive.GetColors(Colors);
+
+			UpdateVector4dFromMorphTargets(Colors, Primitive, MorphTargetWeights, Vector4dMorphTargetAttributes::COLORS);
+
+			ReIndex(Colors, Indices, ReindexBuffer);
+			Swap(Colors, ReindexBuffer);
+		}
+		else if (bMeshHasColors)
+		{
+			// If other primitives in this mesh have colors, generate filler ones for this primitive, to avoid gaps.
+			Colors.Init(FVector4f(1.0f), Primitive.VertexCount());
+		}
+
+		int32_t            AvailableBufferIndex = 0;
+		TArray<FVector2f>* UVs[MAX_MESH_TEXTURE_COORDS_MD];
+		for (int32 UVIndex = 0; UVIndex < NumUVs; ++UVIndex)
+		{
+			UVs[UVIndex] = &GetVector2dBuffer(AvailableBufferIndex++);
+			if (Primitive.HasTexCoords(UVIndex))
+			{
+				TArray<FVector2f>& ReindexBuffer = GetVector2dBuffer(UvReindexBufferIndex);
+				Primitive.GetTexCoords(UVIndex, *UVs[UVIndex]);
+
+				UpdateTexCoordsFromMorphTargets(*UVs[UVIndex], Primitive, MorphTargetWeights, UVIndex);
+
+				ReIndex(*UVs[UVIndex], Indices, ReindexBuffer);
+				Swap(*UVs[UVIndex], ReindexBuffer);
+			}
+			else
+			{
+				// Unreal StaticMesh must have UV channel 0.
+				// glTF doesn't require this since not all materials need texture coordinates.
+				// We also fill UV channel > 1 for this primitive if other primitives have it, to avoid gaps.
+				(*UVs[UVIndex]).AddZeroed(Primitive.VertexCount());
+			}
+		}
+
+		bool bHasDegenerateTriangles = false;
+		// Now add all vertexInstances
+		FVertexID         CornerVertexIDs[3];
+		for (uint32 TriangleIndex = 0; TriangleIndex < TriCount; ++TriangleIndex)
+		{
+			for (int32 Corner = 0; Corner < 3; ++Corner)
+			{
+				const int32 VertexIndex = Indices[TriangleIndex * 3 + Corner];
+				const FVertexID VertexID = PositionIndexToVertexIdPerPrim[PrimitiveIndex][VertexIndex];
+				CornerVertexIDs[Corner] = VertexID;
+			}
+
+			// Check for degenerate triangles
+			const FVertexID& Vertex1 = CornerVertexIDs[0];
+			const FVertexID& Vertex2 = CornerVertexIDs[1];
+			const FVertexID& Vertex3 = CornerVertexIDs[2];
+
+			if (Vertex1 == Vertex2 || Vertex2 == Vertex3 || Vertex1 == Vertex3)
+			{
+				bHasDegenerateTriangles = true;
+				continue; // Triangle is degenerate, skip it
+			}
+
+			for (int32 Corner = 0; Corner < 3; ++Corner)
+			{
+				const FVertexInstanceID VertexInstanceID = MeshDescription->CreateVertexInstance(CornerVertexIDs[Corner]);
+				CornerVertexInstanceIDs[Corner] = VertexInstanceID;
+
+				const uint32 IndiceIndex = TriangleIndex * 3 + Corner;
+
+				VertexInstanceNormals[VertexInstanceID] = Normals[IndiceIndex];
+
+				if (!bIgnoreTangents && Tangents.Num() > 0)
+				{
+					VertexInstanceTangents[VertexInstanceID] = Tangents[IndiceIndex];
+					VertexInstanceBinormalSigns[VertexInstanceID] = Tangents[IndiceIndex].W;
+				}
+
+				for (int32 UVIndex = 0; UVIndex < NumUVs; ++UVIndex)
+				{
+					VertexInstanceUVs.Set(VertexInstanceID, UVIndex, (*UVs[UVIndex])[IndiceIndex]);
+				}
+
+				if (Colors.Num() > 0)
+				{
+					FVector4f Color = Colors[IndiceIndex];
+					// Because the engine will auto convert our vertex colors to sRGB, we need to counter the conversion here to ensure the end result is still linear space
+					UE::Geometry::LinearColors::SRGBToLinear(Color);
+					VertexInstanceColors[VertexInstanceID] = Color;
+				}
+			}
+
+			// Insert a polygon into the mesh
+			TArray<FEdgeID> NewEdgeIDs;
+			const FPolygonID NewPolygonID = MeshDescription->CreatePolygon(CurrentPolygonGroupID, CornerVertexInstanceIDs, &NewEdgeIDs);
+
+			for (const FEdgeID& NewEdgeID : NewEdgeIDs)
+			{
+				// Make all faces part of the same smoothing group, so Unreal will combine identical adjacent verts.
+				// (Is there a way to set auto-gen smoothing threshold? glTF spec says to generate flat normals if they're not specified.
+				//   We want to combine identical verts whether they're smooth neighbors or triangles belonging to the same flat polygon.)
+				EdgeHardnesses[NewEdgeID] = false;
+			}
+		}
+		return bHasDegenerateTriangles;
+	}
+
+	inline void FMeshFactoryImpl::CleanUp()
+	{
+		const uint32 ReserveSize = FMath::Min(uint32(IntBuffer.Num() + IntBuffer.GetSlack()), MaxReserveSize);  // cap reserved size
+		IntBuffer.Empty(ReserveSize);
+		Vector2dBuffers[0].Empty(ReserveSize);
+		for (TArray<FVector3f>& Array : VectorBuffers)
+		{
+			Array.Empty(ReserveSize);
+		}
+		for (int32 Index = 1; Index < MAX_MESH_TEXTURE_COORDS_MD + 1; ++Index)
+		{
+			TArray<FVector2f>& Array = Vector2dBuffers[Index];
+			Array.Empty();
+		}
+		for (TArray<FVector4f>& Array : Vector4dBuffers)
+		{
+			Array.Empty();
+		}
+
+		Messages.Empty();
+	}
+
+	//
+
+	FMeshFactory::FMeshFactory()
+		: Impl(new FMeshFactoryImpl())
+	{
+	}
+
+	FMeshFactory::~FMeshFactory() {}
+
+	void FMeshFactory::FillMeshDescription(const GLTF::FMesh &Mesh, FMeshDescription* MeshDescription, const TArray<float>& MorphTargetWeights)
+	{
+		Impl->FillMeshDescription(Mesh, MeshDescription, MorphTargetWeights);
+	}
+
+	const TArray<FLogMessage>& FMeshFactory::GetLogMessages() const
+	{
+		return Impl->GetLogMessages();
+	}
+
+	float FMeshFactory::GetUniformScale() const
+	{
+		return Impl->ImportUniformScale;
+	}
+
+	void FMeshFactory::SetUniformScale(float Scale)
+	{
+		Impl->ImportUniformScale = Scale;
+	}
+
+	void FMeshFactory::SetReserveSize(uint32 Size)
+	{
+		Impl->MaxReserveSize = Size;
+	}
+
+	void FMeshFactory::CleanUp()
+	{
+		Impl->CleanUp();
+	}
+
+	TArray<FMeshFactory::FIndexVertexIdMap>& FMeshFactory::GetPositionIndexToVertexIdPerPrim() const
+	{
+		return Impl->PositionIndexToVertexIdPerPrim;
+	}
+} //namespace GLTF
diff --git a/Engine/Plugins/Interchange/Runtime/Source/Parsers/GLTFCore/Public/GLTFAccessor.h b/Engine/Plugins/Interchange/Runtime/Source/Parsers/GLTFCore/Public/GLTFAccessor.h
new file mode 100644
--- /dev/null
+++ b/Engine/Plugins/Interchange/Runtime/Source/Parsers/GLTFCore/Public/GLTFAccessor.h
@@ -0,0 +1,231 @@
+// Copyright Epic Games, Inc. All Rights Reserved.
+
+#pragma once
+
+#include "CoreMinimal.h"
+
+struct FMD5Hash;
+
+namespace GLTF
+{
+	struct GLTFCORE_API FBuffer
+	{
+		const uint32 ByteLength;
+		const uint8* Data;
+
+		explicit FBuffer(uint32 InByteLength)
+		    : ByteLength(InByteLength)
+		    , Data(nullptr)
+		{
+		}
+
+		bool IsValid() const
+		{
+			return Data != nullptr;
+		}
+
+		const uint8* DataAt(uint32 Offset) const
+		{
+			checkSlow(Data);
+			return Data + Offset;
+		}
+	};
+
+	struct GLTFCORE_API FBufferView
+	{
+		const FBuffer& Buffer;
+		const uint32 ByteOffset;
+		const uint32 ByteLength;
+		// if zero then accessor elements are tightly packed, i.e., effective stride equals the size of the element
+		const uint32 ByteStride;  // range 4..252
+
+		explicit FBufferView(const FBuffer& InBuffer, uint32 InOffset, uint32 InLength, uint32 InStride)
+		    : Buffer(InBuffer)
+		    , ByteOffset(InOffset)
+		    , ByteLength(InLength)
+		    , ByteStride(InStride)
+		{
+			// check that view fits completely inside the buffer
+		}
+
+		bool IsValid() const
+		{
+			return Buffer.IsValid();
+		}
+
+		const uint8* DataAt(uint32 Offset) const
+		{
+			return Buffer.DataAt(Offset + ByteOffset);
+		}
+	};
+
+	struct GLTFCORE_API FAccessor
+	{
+		// accessor stores the data but has no usage semantics
+
+		enum class EType
+		{
+			Unknown,
+			Scalar,
+			Vec2,
+			Vec3,
+			Vec4,
+			Mat2,
+			Mat3,
+			Mat4,
+			Count
+		};
+
+		enum class EComponentType
+		{
+			None,
+			S8,   // signed byte
+			U8,   // unsigned byte
+			S16,  // signed short
+			U16,  // unsigned short
+			U32,  // unsigned int -- only valid for indices, not attributes
+			F32,  // float
+			Count
+		};
+
+		const uint32         Count;
+		const EType          Type;
+		const EComponentType ComponentType;
+		const bool           Normalized;
+
+		FAccessor(uint32 InCount, EType InType, EComponentType InComponentType, bool InNormalized);
+
+		virtual bool IsValid() const = 0;
+		virtual FMD5Hash GetHash() const = 0;
+
+		virtual uint32 GetUnsignedInt(uint32 Index) const;
+		virtual void   GetUnsignedInt16x4(uint32 Index, uint16 Values[4]) const;
+
+		virtual float     GetFloat(uint32 Index) const;
+		virtual FVector2D GetVec2(uint32 Index) const;
+		virtual FVector   GetVec3(uint32 Index) const;
+		virtual FVector4  GetVec4(uint32 Index) const;
+		virtual FMatrix   GetMat4(uint32 Index) const;
+
+		void         GetUnsignedIntArray(TArray<uint32>& Buffer) const;
+		virtual void GetUnsignedIntArray(uint32* Buffer) const;
+		void         GetFloatArray(TArray<float>& Buffer) const;
+		virtual void GetFloatArray(float* Buffer) const;
+		void         GetVec2Array(TArray<FVector2f>& Buffer) const;
+		virtual void GetVec2Array(FVector2f* Buffer) const;
+		void         GetVec3Array(TArray<FVector3f>& Buffer) const;
+		virtual void GetVec3Array(FVector3f* Buffer) const;
+		///@note Performs axis conversion for vec3s(i.e. from glTF right-handed and Y-up to left-handed and Z-up).
+		void         GetCoordArray(TArray<FVector3f>& Buffer) const;
+		void         GetCoordArray(FVector3f* Buffer) const;
+		void         GetVec4Array(TArray<FVector4f>& Buffer) const;
+		virtual void GetVec4Array(FVector4f* Buffer) const;
+		///@note Performs axis conversion for quaternion(i.e. from glTF right-handed and Y-up to left-handed and Z-up).
+		void         GetQuatArray(TArray<FVector4f>& Buffer) const;
+		void         GetQuatArray(FVector4f* Buffer) const;
+		void         GetMat4Array(TArray<FMatrix44f>& Buffer) const;
+		virtual void GetMat4Array(FMatrix44f* Buffer) const;
+	};
+
+	struct GLTFCORE_API FValidAccessor final : FAccessor
+	{
+		FValidAccessor(FBufferView& InBufferView, uint32 InOffset, uint32 InCount, EType InType, EComponentType InCompType, bool InNormalized);
+
+		bool IsValid() const override;
+		FMD5Hash GetHash() const override;
+
+		uint32 GetUnsignedInt(uint32 Index) const override;
+		void   GetUnsignedInt16x4(uint32 Index, uint16 Values[4]) const override;
+
+		float     GetFloat(uint32 Index) const override;
+		FVector2D GetVec2(uint32 Index) const override;
+		FVector   GetVec3(uint32 Index) const override;
+		FVector4  GetVec4(uint32 Index) const override;
+
+		FMatrix GetMat4(uint32 Index) const override;
+
+		void GetUnsignedIntArray(uint32* Buffer) const override;
+		void GetFloatArray(float* Buffer) const override;
+		void GetVec2Array(FVector2f* Buffer) const override;
+		void GetVec3Array(FVector3f* Buffer) const override;
+		void GetVec4Array(FVector4f* Buffer) const override;
+		void GetMat4Array(FMatrix44f* Buffer) const override;
+
+		const uint8* DataAt(uint32 Index) const;
+
+		const FBufferView& BufferView;
+		const uint32       ByteOffset;
+		const uint32       ElementSize;
+		const uint32	   ByteStride;
+	};
+
+	struct GLTFCORE_API FVoidAccessor final : FAccessor
+	{
+		FVoidAccessor()
+		    : FAccessor(0, EType::Scalar, EComponentType::S8, false)
+		{
+		}
+
+		bool IsValid() const override;
+		FMD5Hash GetHash() const override;
+	};
+
+	//
+
+	inline void FAccessor::GetUnsignedIntArray(TArray<uint32>& Buffer) const
+	{
+		if (IsValid())
+			Buffer.SetNumUninitialized(Count, false);
+		GetUnsignedIntArray(Buffer.GetData());
+	}
+
+	inline void FAccessor::GetFloatArray(TArray<float>& Buffer) const
+	{
+		if (IsValid())
+			Buffer.SetNumUninitialized(Count, false);
+		GetFloatArray(Buffer.GetData());
+	}
+
+	inline void FAccessor::GetVec2Array(TArray<FVector2f>& Buffer) const
+	{
+		if (IsValid())
+			Buffer.SetNumUninitialized(Count, false);
+		GetVec2Array(Buffer.GetData());
+	}
+
+	inline void FAccessor::GetVec3Array(TArray<FVector3f>& Buffer) const
+	{
+		if (IsValid())
+			Buffer.SetNumUninitialized(Count, false);
+		GetVec3Array(Buffer.GetData());
+	}
+
+	inline void FAccessor::GetCoordArray(TArray<FVector3f>& Buffer) const
+	{
+		if (IsValid())
+			Buffer.SetNumUninitialized(Count, false);
+		GetCoordArray(Buffer.GetData());
+	}
+
+	inline void FAccessor::GetVec4Array(TArray<FVector4f>& Buffer) const
+	{
+		if (IsValid())
+			Buffer.SetNumUninitialized(Count, false);
+		GetVec4Array(Buffer.GetData());
+	}
+
+	inline void FAccessor::GetQuatArray(TArray<FVector4f>& Buffer) const
+	{
+		if (IsValid())
+			Buffer.SetNumUninitialized(Count, false);
+		GetQuatArray(Buffer.GetData());
+	}
+
+	inline void FAccessor::GetMat4Array(TArray<FMatrix44f>& Buffer) const
+	{
+		if (IsValid())
+			Buffer.SetNumUninitialized(Count, false);
+		GetMat4Array(Buffer.GetData());
+	}
+
+}  // namespace GLTF
diff --git a/Engine/Plugins/Interchange/Runtime/Source/Parsers/GLTFCore/Public/GLTFAnimation.h b/Engine/Plugins/Interchange/Runtime/Source/Parsers/GLTFCore/Public/GLTFAnimation.h
new file mode 100644
--- /dev/null
+++ b/Engine/Plugins/Interchange/Runtime/Source/Parsers/GLTFCore/Public/GLTFAnimation.h
@@ -0,0 +1,78 @@
+// Copyright Epic Games, Inc. All Rights Reserved.
+
+#pragma once
+
+#include "CoreMinimal.h"
+
+struct FScriptContainerElement;
+
+namespace GLTF
+{
+	struct FNode;
+	struct FAccessor;
+
+	struct GLTFCORE_API FAnimation
+	{
+		enum class EInterpolation
+		{
+			Linear,
+			Step,
+			CubicSpline,
+		};
+
+		enum class EPath
+		{
+			Translation,
+			Rotation,
+			Scale,
+			Weights
+		};
+
+		struct FTarget
+		{
+			const FNode& Node;
+			EPath        Path;
+
+			FTarget(const FNode& Node)
+			    : Node(Node)
+			{
+			}
+		};
+
+		struct FChannel
+		{
+			// The index of a sampler in this animation used to compute the value for the target.
+			int32 Sampler;
+			// The index of the node and TRS property to target.
+			FTarget Target;
+
+			FChannel(const FNode& Node)
+			    : Target(Node)
+			{
+			}
+		};
+
+		struct FSampler
+		{
+			EInterpolation Interpolation;
+			// The accessor containing keyframe input values, e.g., time. Always float and seconds.
+			const FAccessor& Input;
+			// The accessor containing keyframe output values.
+			const FAccessor& Output;
+
+			FSampler(const FAccessor& Input, const FAccessor& Output)
+			    : Interpolation(EInterpolation::Linear)
+			    , Input(Input)
+			    , Output(Output)
+			{
+			}
+		};
+
+		FString          Name;
+		TArray<FSampler> Samplers;
+		TArray<FChannel> Channels;
+
+		FString          UniqueId; //will be generated in FAsset::GenerateNames
+	};
+
+}  // namespace GLTF
diff --git a/Engine/Plugins/Interchange/Runtime/Source/Parsers/GLTFCore/Public/GLTFAsset.h b/Engine/Plugins/Interchange/Runtime/Source/Parsers/GLTFCore/Public/GLTFAsset.h
new file mode 100644
--- /dev/null
+++ b/Engine/Plugins/Interchange/Runtime/Source/Parsers/GLTFCore/Public/GLTFAsset.h
@@ -0,0 +1,152 @@
+// Copyright Epic Games, Inc. All Rights Reserved.
+
+#pragma once
+
+#include "CoreMinimal.h"
+#include "GLTFAnimation.h"
+#include "GLTFMaterial.h"
+#include "GLTFMesh.h"
+#include "GLTFNode.h"
+#include "GLTFTexture.h"
+
+namespace GLTF { struct FAnimation; }
+namespace GLTF { struct FBuffer; }
+namespace GLTF { struct FBufferView; }
+namespace GLTF { struct FCamera; }
+namespace GLTF { struct FImage; }
+namespace GLTF { struct FLight; }
+namespace GLTF { struct FMaterial; }
+namespace GLTF { struct FMesh; }
+namespace GLTF { struct FNode; }
+namespace GLTF { struct FSampler; }
+namespace GLTF { struct FSkinInfo; }
+namespace GLTF { struct FTexture; }
+namespace GLTF { struct FValidAccessor; }
+struct FScriptContainerElement;
+
+namespace GLTF
+{
+	enum class EExtension
+	{
+		KHR_MaterialsPbrSpecularGlossiness,
+		KHR_MaterialsUnlit,
+		KHR_MaterialsClearCoat,
+		KHR_MaterialsTransmission,
+		KHR_MaterialsSheen,
+		KHR_MaterialsVariants,
+		KHR_MaterialsIOR,
+		KHR_MaterialsSpecular,
+		KHR_MaterialsEmissiveStrength,
+		KHR_TextureTransform,
+		KHR_DracoMeshCompression,
+		KHR_LightsPunctual,
+		KHR_Lights,
+		KHR_Blend,
+		MSFT_TextureDDS,
+		MSFT_PackingNormalRoughnessMetallic,
+		MSFT_PackingOcclusionRoughnessMetallic,
+		Count
+	};
+
+	struct GLTFCORE_API FScene
+	{
+		FString       Name;
+		TArray<int32> Nodes;
+		FString	      UniqueId; //will be generated in FAsset::GenerateNames
+	};
+
+	struct GLTFCORE_API FMetadata
+	{
+		struct FExtraData
+		{
+			FString Name;
+			FString Value;
+		};
+		FString            GeneratorName;
+		float              Version;
+		TArray<FExtraData> Extras;
+
+		const FExtraData* GetExtraData(const FString& Name) const;
+	};
+
+	struct GLTFCORE_API FAsset : FNoncopyable
+	{
+		static const TSet<EExtension> SupportedExtensions;
+
+		enum EValidationCheck
+		{
+			Valid,
+			InvalidMeshPresent   = 0x1,
+			InvalidNodePresent = 0x2,
+		};
+
+		FString Name;
+
+		TArray<FBuffer>        Buffers;
+		TArray<FBufferView>    BufferViews;
+		TArray<FValidAccessor> Accessors;
+		TArray<FMesh>          Meshes;
+
+		TArray<FScene>     Scenes;
+		TArray<FNode>      Nodes;
+		TArray<FCamera>    Cameras;
+		TArray<FLight>     Lights;
+		TArray<FSkinInfo>  Skins;
+		TArray<FAnimation> Animations;
+
+		TArray<FImage>    Images;
+		TArray<FSampler>  Samplers;
+		TArray<FTexture>  Textures;
+		TArray<FMaterial> Materials;
+
+		TArray<FString>   Variants;
+
+		TArray<FString>	 ExtensionsUsed;      //marked in the gltf
+		TArray<FString>	 ExtensionsRequired;  //marked in the gltf
+
+		TSet<EExtension> ProcessedExtensions;
+		FMetadata        Metadata;
+
+
+		/**
+		 * Will clear the asset's buffers.
+		 *
+		 * @param BinBufferKBytes - bytes to reserve for the bin chunk buffer.
+		 * @param ExtraBinBufferKBytes -  bytes to reserve for the extra binary buffer(eg. image, mime data, etc.)
+		 * @note Only reserves buffers if they had any existing data.
+		 */
+		void Clear(uint32 BinBufferKBytes, uint32 ExtraBinBufferKBytes);
+
+		/**
+		 * Will generate names for any entities(nodes, meshes, etc.) that have the name field missing.
+		 * Is called from GLTF::FFileReader::ReadFile
+		 *
+		 * @param Prefix - prefix to add to the entities name.
+		 */
+		void GenerateNames(const FString& Prefix);
+
+		/**
+		 * Finds the indices for the nodes which are root nodes.
+		 *
+		 * @param NodeIndices - array with the results
+		 */
+		void GetRootNodes(TArray<int32>& NodeIndices);
+
+		/**
+		 * Returns EValidationCheck::Valid if the asset passes the post-import validation checks.
+		 */
+		EValidationCheck ValidationCheck() const;
+
+	private:
+		// Binary glTF files can have embedded data after JSON.
+		// This will be empty when reading from a text glTF (common) or a binary glTF with no BIN chunk (rare).
+		TArray<uint8> BinData;
+		// Extra binary data used for images from disk, mime data and so on.
+		TArray<uint8> ExtraBinData;
+
+		friend class FFileReader;
+	};
+
+	GLTFCORE_API const TCHAR* ToString(EExtension Extension);
+
+}  // namespace GLTF
diff --git a/Engine/Plugins/Interchange/Runtime/Source/Parsers/GLTFCore/Public/GLTFCoreModule.h b/Engine/Plugins/Interchange/Runtime/Source/Parsers/GLTFCore/Public/GLTFCoreModule.h
new file mode 100644
--- /dev/null
+++ b/Engine/Plugins/Interchange/Runtime/Source/Parsers/GLTFCore/Public/GLTFCoreModule.h
@@ -0,0 +1,33 @@
+// Copyright Epic Games, Inc. All Rights Reserved.
+
+#pragma once
+
+#include "CoreMinimal.h"
+#include "Modules/ModuleManager.h"
+
+/**
+ * The public interface of the GLTFCore module
+ */
+class IGLTFCoreModule : public IModuleInterface
+{
+public:
+	/**
+	 * Singleton-like access to IGLTFCore
+	 *
+	 * @return Returns IGLTFCore singleton instance, loading the module on demand if needed
+	 */
+	static IGLTFCoreModule& Get()
+	{
+		return FModuleManager::LoadModuleChecked<IGLTFCoreModule>("GLTFCore");
+	}
+
+	/**
+	 * Checks to see if this module is loaded and ready.  It is only valid to call Get() if IsAvailable() returns true.
+	 *
+	 * @return True if the module is loaded and ready to use
+	 */
+	static bool IsAvailable()
+	{
+		return FModuleManager::Get().IsModuleLoaded("GLTFCore");
+	}
+};
diff --git a/Engine/Plugins/Interchange/Runtime/Source/Parsers/GLTFCore/Public/GLTFLogger.h b/Engine/Plugins/Interchange/Runtime/Source/Parsers/GLTFCore/Public/GLTFLogger.h
new file mode 100644
--- /dev/null
+++ b/Engine/Plugins/Interchange/Runtime/Source/Parsers/GLTFCore/Public/GLTFLogger.h
@@ -0,0 +1,31 @@
+// Copyright Epic Games, Inc. All Rights Reserved.
+
+#pragma once
+
+#include "Templates/Tuple.h"
+#include "Containers/UnrealString.h"
+
+namespace GLTF
+{
+	enum class EMessageSeverity
+	{
+		Warning,
+		Error,
+	};
+	using FLogMessage = TTuple<EMessageSeverity, FString>;
+
+	class GLTFCORE_API FBaseLogger
+	{
+	public:
+		const TArray<FLogMessage>& GetLogMessages() const;
+
+	protected:
+		mutable TArray<FLogMessage> Messages;
+	};
+
+	inline const TArray<FLogMessage>& FBaseLogger::GetLogMessages() const
+	{
+		return Messages;
+	}
+}  // namespace GLTF
+
diff --git a/Engine/Plugins/Interchange/Runtime/Source/Parsers/GLTFCore/Public/GLTFMaterial.h b/Engine/Plugins/Interchange/Runtime/Source/Parsers/GLTFCore/Public/GLTFMaterial.h
new file mode 100644
--- /dev/null
+++ b/Engine/Plugins/Interchange/Runtime/Source/Parsers/GLTFCore/Public/GLTFMaterial.h
@@ -0,0 +1,200 @@
+// Copyright Epic Games, Inc. All Rights Reserved.
+
+#pragma once
+
+#include "CoreMinimal.h"
+#include "Misc/SecureHash.h"
+
+namespace GLTF
+{
+	struct GLTFCORE_API FTextureTransform
+	{
+		float Offset[2];
+		float Rotation;
+		float Scale[2];
+
+		FTextureTransform()
+		{
+			Offset[0] = Offset[1] = 0.0f;
+			Scale[0] = Scale[1] = 1.0f;
+			Rotation = 0.0f;
+		}
+	};
+
+	struct GLTFCORE_API FTextureMap
+	{
+		int32 TextureIndex;
+		uint8 TexCoord;
+
+		bool bHasTextureTransform;
+		FTextureTransform TextureTransform;
+
+		FTextureMap()
+		    : TextureIndex(INDEX_NONE)
+		    , TexCoord(0)
+			, bHasTextureTransform(false)
+		{
+		}
+
+		FMD5Hash GetHash() const
+		{
+			FMD5 MD5;
+			MD5.Update(reinterpret_cast<const uint8*>(&TextureIndex), sizeof(TextureIndex));
+			MD5.Update(&TexCoord, sizeof(TexCoord));
+			FMD5Hash Hash;
+			Hash.Set(MD5);
+			return Hash;
+		}
+	};
+
+	struct GLTFCORE_API FMaterial
+	{
+		enum class EAlphaMode
+		{
+			Opaque,
+			Mask,
+			Blend
+		};
+
+		enum class EShadingModel
+		{
+			MetallicRoughness,
+			SpecularGlossiness,
+		};
+
+		enum class EPackingFlags
+		{
+			// no packing, i.e. default: Unused (R) Roughness (G), Metallic (B) map
+			None = 0x0,
+			// packing two channel (RG) normal map
+			NormalRG = 0x1,
+			// packing Occlusion (R), Roughness (G), Metallic (B) map
+			OcclusionRoughnessMetallic = 0x2,
+			// packing Roughness (R), Metallic (G), Occlusion (B) map
+			RoughnessMetallicOcclusion = 0x4,
+			// packing Normal (RG), Roughness (B), Metallic (A) map
+			NormalRoughnessMetallic = NormalRG | 0x8,
+		};
+
+		struct FMetallicRoughness
+		{
+			FTextureMap Map;
+			float       MetallicFactor = 1.0f;
+			float       RoughnessFactor = 1.0f;
+		};
+		// "pbrSpecularGlossiness" extension
+		// Implements a specular glossiness shading model and is mutually exclusive with "specular" extension, which uses metallic roughness model
+		struct FSpecularGlossiness 
+		{
+			FTextureMap Map;
+			FVector     SpecularFactor = FVector(1.0f);
+			float       GlossinessFactor = 1.0f;
+		};
+		struct FClearCoat
+		{
+			float ClearCoatFactor = 1.0f;
+			FTextureMap ClearCoatMap;
+
+			float Roughness = 0.0f;
+			FTextureMap RoughnessMap;
+
+			float NormalMapUVScale = 1.0f;
+			FTextureMap NormalMap;
+		};
+		struct FTransmission
+		{
+			float TransmissionFactor = 0.0f;
+			FTextureMap TransmissionMap;
+		};
+		struct FSheen
+		{
+			FVector SheenColorFactor = FVector::Zero();
+			FTextureMap SheenColorMap;
+			float SheenRoughnessFactor = 0.f;
+			FTextureMap SheenRoughnessMap;
+		};
+		struct FSpecular // "specular" extension
+		{
+			float SpecularFactor = 1.0f;
+			FTextureMap SpecularMap;
+			FVector SpecularColorFactor = FVector(1.0f, 1.0f, 1.0f);
+			FTextureMap SpecularColorMap;
+		};
+
+		struct FPacking
+		{
+			int         Flags = (int)EPackingFlags::None;  // see EPackingFlags
+			FTextureMap Map;
+			FTextureMap NormalMap;
+		};
+
+		FString Name;
+
+		// PBR properties
+		FTextureMap         BaseColor;		 // Used for DiffuseColor on Specular-Glossiness mode
+		FVector4f           BaseColorFactor; // Used for DiffuseFactor on Specular-Glossiness mode
+		EShadingModel       ShadingModel;
+		FMetallicRoughness  MetallicRoughness;
+		FSpecularGlossiness SpecularGlossiness;
+		FClearCoat          ClearCoat;
+		FTransmission       Transmission;
+		FSheen              Sheen;
+		FSpecular           Specular;
+		float               IOR;
+		float               EmissiveStrength;
+
+		// base properties
+		FTextureMap Normal;
+		FTextureMap Occlusion;
+		FTextureMap Emissive;
+		float       NormalScale;
+		float       OcclusionStrength;
+		FVector3f   EmissiveFactor;
+
+		// material properties
+		bool       bIsDoubleSided;
+		EAlphaMode AlphaMode;
+		float      AlphaCutoff;  // only used when AlphaMode == Mask
+
+		// extension properties
+		FPacking Packing;
+		bool     bIsUnlitShadingModel;
+		bool     bHasClearCoat;
+		bool     bHasSheen;
+		bool     bHasTransmission;
+		bool     bHasIOR;
+		bool     bHasSpecular;
+		bool     bHasEmissiveStrength;
+
+		FString UniqueId; //will be generated in FAsset::GenerateNames
+
+		FMaterial(const FString& Name)
+		    : Name(Name)
+		    , BaseColorFactor {1.0f, 1.0f, 1.0f, 1.0f}
+		    , ShadingModel(EShadingModel::MetallicRoughness)
+			, IOR(1.0f)
+			, EmissiveStrength(1.0f)
+		    , NormalScale(1.f)
+		    , OcclusionStrength(1.f)
+		    , EmissiveFactor(FVector::ZeroVector)
+		    , bIsDoubleSided(false)
+		    , AlphaMode(EAlphaMode::Opaque)
+		    , AlphaCutoff(0.5f)
+		    , bIsUnlitShadingModel(false)
+			, bHasClearCoat(false)
+			, bHasSheen(false)
+			, bHasTransmission(false)
+			, bHasIOR(false)
+			, bHasSpecular(false)
+			, bHasEmissiveStrength(false)
+		{
+		}
+
+		bool IsOpaque() const
+		{
+			return AlphaMode == EAlphaMode::Opaque;
+		}
+
+		FMD5Hash GetHash() const;
+	};
+}  // namespace GLTF
diff --git a/Engine/Plugins/Interchange/Runtime/Source/Parsers/GLTFCore/Public/GLTFMaterialExpressions.h b/Engine/Plugins/Interchange/Runtime/Source/Parsers/GLTFCore/Public/GLTFMaterialExpressions.h
new file mode 100644
--- /dev/null
+++ b/Engine/Plugins/Interchange/Runtime/Source/Parsers/GLTFCore/Public/GLTFMaterialExpressions.h
@@ -0,0 +1,475 @@
+// Copyright Epic Games, Inc. All Rights Reserved.
+
+#pragma once
+
+#include "UObject/StrongObjectPtr.h"
+#include "Misc/SecureHash.h"
+
+namespace GLTF
+{
+	class FMaterialExpressionInput;
+	class FMaterialExpressionTexture;
+	class FMaterialExpressionFunctionCall;
+	class FMaterialElement;
+
+	enum class EMaterialExpressionType
+	{
+		ConstantColor,
+		ConstantScalar,
+		FunctionCall,
+		Generic,
+		Texture,
+		TextureCoordinate,
+	};
+
+	enum class EGLTFMaterialShadingModel
+	{
+		DefaultLit,
+		Subsurface,
+		ClearCoat,
+		ThinTranslucent,
+		Unlit
+	};
+
+	class GLTFCORE_API ITextureElement
+	{
+	public:
+		virtual ~ITextureElement() = default;
+	};
+
+	class GLTFCORE_API FMaterialExpression
+	{
+	public:
+		virtual ~FMaterialExpression() {}
+
+		virtual FMaterialExpressionInput* GetInput(int32 Index)
+		{
+			return nullptr;
+		}
+
+		virtual int32 GetInputCount() const
+		{
+			return 0;
+		}
+
+		EMaterialExpressionType GetType() const
+		{
+			return Type;
+		}
+
+		template <EMaterialExpressionType T>
+		bool IsA() const
+		{
+			return T == Type;
+		}
+
+		void ConnectExpression(FMaterialExpressionInput& ExpressionInput, int32 OutputIndex);
+
+	protected:
+		FMaterialExpression();
+		FMaterialExpression(const FMaterialExpression& Expr);
+		explicit FMaterialExpression(int Type);
+
+	protected:
+		EMaterialExpressionType Type;
+
+		friend class FMaterialElement;
+	};
+
+	class GLTFCORE_API FMaterialExpressionInput
+	{
+	public:
+		FMaterialExpression* GetExpression()
+		{
+			return Expression;
+		}
+
+		const FMaterialExpression* GetExpression() const
+		{
+			return Expression;
+		}
+
+		int32 GetOutputIndex() const
+		{
+			return OutputIndex;
+		}
+
+	private:
+		FMaterialExpressionInput(const FString& Name);
+
+	private:
+		FString              Name;
+		FMaterialExpression* Expression;
+		int32                OutputIndex;
+
+		friend class FMaterialElement;
+		friend class FMaterialExpression;
+		friend class FMaterialExpressionTexture;
+		friend class TArray<FMaterialExpressionInput>;
+	};
+
+	class GLTFCORE_API FMaterialExpressionParameter : public FMaterialExpression
+	{
+	public:
+		virtual ~FMaterialExpressionParameter() {}
+
+		void         SetName(const TCHAR* Name);
+		const TCHAR* GetName() const;
+		void         SetGroupName(const TCHAR* GroupName);
+		const TCHAR* GetGroupName() const;
+
+	protected:
+		explicit FMaterialExpressionParameter(int Type);
+
+	private:
+		FString Name;
+		FString GroupName;
+	};
+
+	class GLTFCORE_API FMaterialExpressionScalar : public FMaterialExpressionParameter
+	{
+	public:
+		enum
+		{
+			Type = (int)EMaterialExpressionType::ConstantScalar
+		};
+
+		FMaterialExpressionScalar()
+		    : FMaterialExpressionParameter(Type)
+		{
+		}
+
+		float& GetScalar()
+		{
+			return Scalar;
+		}
+
+		float GetScalar() const
+		{
+			return Scalar;
+		}
+
+	private:
+		float Scalar;
+	};
+
+	class GLTFCORE_API FMaterialExpressionColor : public FMaterialExpressionParameter
+	{
+	public:
+		enum
+		{
+			Type = (int)EMaterialExpressionType::ConstantColor
+		};
+
+		FMaterialExpressionColor();
+
+		FLinearColor& GetColor()
+		{
+			return Color;
+		}
+
+		const FLinearColor& GetColor() const
+		{
+			return Color;
+		}
+
+	private:
+		FLinearColor Color;
+	};
+
+	class GLTFCORE_API FMaterialExpressionTexture : public FMaterialExpressionParameter
+	{
+	public:
+		enum
+		{
+			Type = (int)EMaterialExpressionType::Texture
+		};
+
+		FMaterialExpressionTexture();
+
+		void SetTexture(ITextureElement* InTexture)
+		{
+			Texture = InTexture;
+		}
+
+		ITextureElement* GetTexture() const
+		{
+			return Texture;
+		}
+
+		FMaterialExpressionInput& GetInputCoordinate()
+		{
+			return InputCoordinate;
+		}
+
+		virtual FMaterialExpressionInput* GetInput(int32 Index) override;
+		virtual int32                     GetInputCount() const override;
+
+	private:
+		FMaterialExpressionInput InputCoordinate;
+		ITextureElement*         Texture;
+	};
+
+	class GLTFCORE_API FMaterialExpressionTextureCoordinate : public FMaterialExpression
+	{
+	public:
+		enum
+		{
+			Type = (int)EMaterialExpressionType::TextureCoordinate
+		};
+
+		FMaterialExpressionTextureCoordinate();
+
+		void SetCoordinateIndex(int32 InCoordinateIndex)
+		{
+			CoordinateIndex = InCoordinateIndex;
+		}
+
+		int32 GetCoordinateIndex() const
+		{
+			return CoordinateIndex;
+		}
+
+	private:
+		int32 CoordinateIndex;
+	};
+
+	class GLTFCORE_API FMaterialExpressionGeneric : public FMaterialExpression
+	{
+	public:
+		enum
+		{
+			Type = (int)EMaterialExpressionType::Generic
+		};
+
+		FMaterialExpressionGeneric();
+
+		void SetExpressionName(const TCHAR* InExpressionName)
+		{
+			ExpressionName = InExpressionName;
+		}
+
+		const TCHAR* GetExpressionName() const
+		{
+			return *ExpressionName;
+		}
+
+		virtual FMaterialExpressionInput* GetInput(int32 Index) override;
+		virtual int32                     GetInputCount() const override;
+
+		void SetBoolProperty(const TCHAR* Name, bool Value)
+		{
+			BoolProperties.Add(Name, Value);
+		}
+
+		void SetFloatProperty(const TCHAR* Name, float Value)
+		{
+			FloatProperties.Add(Name, Value);
+		}
+
+		const TMap<FString, bool>& GetBoolProperties() const
+		{
+			return BoolProperties;
+		}
+
+		const TMap<FString, float>& GetFloatProperties() const
+		{
+			return FloatProperties;
+		}
+
+	private:
+		FString                          ExpressionName;
+		TArray<FMaterialExpressionInput> Inputs;
+
+		TMap<FString, bool> BoolProperties;
+		TMap<FString, float> FloatProperties;
+	};
+
+	class GLTFCORE_API FMaterialExpressionFunctionCall : public FMaterialExpression
+	{
+	public:
+		enum
+		{
+			Type = (int)EMaterialExpressionType::FunctionCall
+		};
+
+		FMaterialExpressionFunctionCall();
+
+		void SetFunctionPathName(const TCHAR* InFunctionPathName)
+		{
+			FunctionPathName = InFunctionPathName;
+		}
+
+		const TCHAR* GetFunctionPathName() const
+		{
+			return *FunctionPathName;
+		}
+
+		virtual FMaterialExpressionInput* GetInput(int32 Index) override;
+		virtual int32                     GetInputCount() const override;
+
+	private:
+		FString                          FunctionPathName;
+		TArray<FMaterialExpressionInput> Inputs;
+	};
+
+	class GLTFCORE_API FMaterialElement
+	{
+	public:
+		FMaterialElement(const FString& Name);
+		virtual ~FMaterialElement();
+
+		virtual int  GetBlendMode() const          = 0;
+		virtual void SetBlendMode(int InBlendMode) = 0;
+		virtual bool GetTwoSided() const           = 0;
+		virtual void SetTwoSided(bool bTwoSided)   = 0;
+		virtual bool GetIsThinSurface() const = 0;
+		virtual void SetIsThinSurface(bool bIsThinSurface) = 0;
+		virtual void SetShadingModel(EGLTFMaterialShadingModel InShadingModel) = 0;
+		virtual void SetTranslucencyLightingMode(int InLightingMode) = 0;
+
+		virtual void Finalize() = 0;
+
+		FMD5Hash GetGLTFMaterialHash() const;
+		void SetGLTFMaterialHash(FMD5Hash Hash);
+
+		const FString& GetName() const;
+
+		FMaterialExpressionInput& GetBaseColor();
+		FMaterialExpressionInput& GetMetallic();
+		FMaterialExpressionInput& GetSpecular();
+		FMaterialExpressionInput& GetRoughness();
+		FMaterialExpressionInput& GetEmissiveColor();
+		FMaterialExpressionInput& GetOpacity();
+		FMaterialExpressionInput& GetNormal();
+		FMaterialExpressionInput& GetWorldDisplacement();
+		FMaterialExpressionInput& GetRefraction();
+		FMaterialExpressionInput& GetAmbientOcclusion();
+		FMaterialExpressionInput& GetClearCoat();
+		FMaterialExpressionInput& GetClearCoatRoughness();
+
+		int32                GetExpressionsCount() const;
+		FMaterialExpression* GetExpression(int32 Index);
+
+		FMaterialExpression* AddMaterialExpression(EMaterialExpressionType ExpressionType);
+		template <typename T>
+		T* AddMaterialExpression();
+
+	protected:
+		FString                  Name;
+		FMaterialExpressionInput BaseColor;
+		FMaterialExpressionInput Metallic;
+		FMaterialExpressionInput Specular;
+		FMaterialExpressionInput Roughness;
+		FMaterialExpressionInput EmissiveColor;
+		FMaterialExpressionInput Opacity;
+		FMaterialExpressionInput Normal;
+		FMaterialExpressionInput WorldDisplacement;
+		FMaterialExpressionInput Refraction;
+		FMaterialExpressionInput AmbientOcclusion;
+		FMaterialExpressionInput ClearCoat;
+		FMaterialExpressionInput ClearCoatRoughness;
+
+		TArray<FMaterialExpression*> Expressions;
+
+		GLTF::FMaterialExpression* ThinTranslucentMaterialOutput;
+		GLTF::FMaterialExpression* ClearCoatBottomNormalOutput;
+
+		bool bIsFinal;
+
+        FMD5Hash GLTFMaterialHash;
+	};
+
+	inline const FString& FMaterialElement::GetName() const
+	{
+		return Name;
+	}
+
+	inline FMaterialExpressionInput& FMaterialElement::GetBaseColor()
+	{
+		return BaseColor;
+	}
+	inline FMaterialExpressionInput& FMaterialElement::GetMetallic()
+	{
+		return Metallic;
+	}
+	inline FMaterialExpressionInput& FMaterialElement::GetSpecular()
+	{
+		return Specular;
+	}
+	inline FMaterialExpressionInput& FMaterialElement::GetRoughness()
+	{
+		return Roughness;
+	}
+	inline FMaterialExpressionInput& FMaterialElement::GetEmissiveColor()
+	{
+		return EmissiveColor;
+	}
+	inline FMaterialExpressionInput& FMaterialElement::GetOpacity()
+	{
+		return Opacity;
+	}
+	inline FMaterialExpressionInput& FMaterialElement::GetNormal()
+	{
+		return Normal;
+	}
+	inline FMaterialExpressionInput& FMaterialElement::GetWorldDisplacement()
+	{
+		return WorldDisplacement;
+	}
+	inline FMaterialExpressionInput& FMaterialElement::GetRefraction()
+	{
+		return Refraction;
+	}
+	inline FMaterialExpressionInput& FMaterialElement::GetAmbientOcclusion()
+	{
+		return AmbientOcclusion;
+	}
+	inline FMaterialExpressionInput& FMaterialElement::GetClearCoat()
+	{
+		return ClearCoat;
+	}
+	inline FMaterialExpressionInput& FMaterialElement::GetClearCoatRoughness()
+	{
+		return ClearCoatRoughness;
+	}
+
+	template <>
+	inline FMaterialExpressionScalar* FMaterialElement::AddMaterialExpression<FMaterialExpressionScalar>()
+	{
+		return static_cast<FMaterialExpressionScalar*>(AddMaterialExpression(EMaterialExpressionType(FMaterialExpressionScalar::Type)));
+	}
+
+	template <>
+	inline FMaterialExpressionColor* FMaterialElement::AddMaterialExpression<FMaterialExpressionColor>()
+	{
+		return static_cast<FMaterialExpressionColor*>(AddMaterialExpression(EMaterialExpressionType(FMaterialExpressionColor::Type)));
+	}
+
+	template <>
+	inline FMaterialExpressionTexture* FMaterialElement::AddMaterialExpression<FMaterialExpressionTexture>()
+	{
+		return static_cast<FMaterialExpressionTexture*>(AddMaterialExpression(EMaterialExpressionType(FMaterialExpressionTexture::Type)));
+	}
+
+	template <>
+	inline FMaterialExpressionTextureCoordinate* FMaterialElement::AddMaterialExpression<FMaterialExpressionTextureCoordinate>()
+	{
+		return static_cast<FMaterialExpressionTextureCoordinate*>(
+		    AddMaterialExpression(EMaterialExpressionType(FMaterialExpressionTextureCoordinate::Type)));
+	}
+
+	template <>
+	inline FMaterialExpressionFunctionCall* FMaterialElement::AddMaterialExpression<FMaterialExpressionFunctionCall>()
+	{
+		return static_cast<FMaterialExpressionFunctionCall*>(AddMaterialExpression(EMaterialExpressionType(FMaterialExpressionFunctionCall::Type)));
+	}
+
+	template <>
+	inline FMaterialExpressionGeneric* FMaterialElement::AddMaterialExpression<FMaterialExpressionGeneric>()
+	{
+		return static_cast<FMaterialExpressionGeneric*>(AddMaterialExpression(EMaterialExpressionType(FMaterialExpressionGeneric::Type)));
+	}
+
+}  // namespace GLTF
+
diff --git a/Engine/Plugins/Interchange/Runtime/Source/Parsers/GLTFCore/Public/GLTFMaterialFactory.h b/Engine/Plugins/Interchange/Runtime/Source/Parsers/GLTFCore/Public/GLTFMaterialFactory.h
new file mode 100644
--- /dev/null
+++ b/Engine/Plugins/Interchange/Runtime/Source/Parsers/GLTFCore/Public/GLTFMaterialFactory.h
@@ -0,0 +1,64 @@
+// Copyright Epic Games, Inc. All Rights Reserved.
+
+#pragma once
+
+#include "GLTFLogger.h"
+
+#include "CoreMinimal.h"
+#include "UObject/ObjectMacros.h"
+
+namespace GLTF
+{
+	struct FAsset;
+	struct FTexture;
+	class FMaterialElement;
+	class ITextureElement;
+	class FMaterialFactoryImpl;
+
+	enum class ETextureMode
+	{
+		Color,
+		Grayscale,
+		Normal
+	};
+
+	class GLTFCORE_API ITextureFactory
+	{
+	public:
+		virtual ~ITextureFactory() = default;
+
+		virtual ITextureElement* CreateTexture(const GLTF::FTexture& Texture, UObject* ParentPackage, EObjectFlags Flags,
+		                                       GLTF::ETextureMode TextureMode) = 0;
+		virtual void             CleanUp()                                     = 0;
+	};
+
+	class GLTFCORE_API IMaterialElementFactory
+	{
+	public:
+		virtual ~IMaterialElementFactory() = default;
+
+		virtual FMaterialElement* CreateMaterial(const TCHAR* Name, UObject* ParentPackage, EObjectFlags Flags) = 0;
+	};
+
+	class GLTFCORE_API FMaterialFactory
+	{
+	public:
+		FMaterialFactory(IMaterialElementFactory* MaterialElementFactory, ITextureFactory* TextureFactory);
+		~FMaterialFactory();
+
+		const TArray<FMaterialElement*>& CreateMaterials(const GLTF::FAsset& Asset, UObject* ParentPackage, EObjectFlags Flags);
+
+		const TArray<FLogMessage>&       GetLogMessages() const;
+		const TArray<FMaterialElement*>& GetMaterials() const;
+
+		IMaterialElementFactory& GetMaterialElementFactory();
+		ITextureFactory&         GetTextureFactory();
+
+		void CleanUp();
+
+	private:
+		TUniquePtr<FMaterialFactoryImpl> Impl;
+	};
+
+}  // namespace GLTF
+
diff --git a/Engine/Plugins/Interchange/Runtime/Source/Parsers/GLTFCore/Public/GLTFMesh.h b/Engine/Plugins/Interchange/Runtime/Source/Parsers/GLTFCore/Public/GLTFMesh.h
new file mode 100644
--- /dev/null
+++ b/Engine/Plugins/Interchange/Runtime/Source/Parsers/GLTFCore/Public/GLTFMesh.h
@@ -0,0 +1,337 @@
+// Copyright Epic Games, Inc. All Rights Reserved.
+
+#pragma once
+
+#include "GLTFAccessor.h"
+#include "CoreMinimal.h"
+
+struct FMD5Hash;
+
+namespace GLTF
+{
+	// for storing triangle indices as a unit
+	struct GLTFCORE_API FTriangle
+	{
+		uint32 A, B, C;
+
+		FTriangle()
+		    : A(0)
+		    , B(0)
+		    , C(0)
+		{
+		}
+	};
+
+	struct GLTFCORE_API FJointInfluence
+	{
+		FVector4 Weight;
+		uint16   ID[4];
+
+		FJointInfluence(const FVector4& InWeight)
+		    : Weight(InWeight)
+		    , ID {0, 0, 0, 0}
+		{
+		}
+	};
+
+	struct GLTFCORE_API FVariantMapping
+	{
+		int32 MaterialIndex;
+		TArray<int32> VariantIndices;
+	};
+
+	struct GLTFCORE_API FMorphTarget
+	{
+		FMorphTarget(const FAccessor& InPositionDisplacements, const FAccessor& InNormalDisplacements,
+			const FAccessor& InTangentDisplacements, const FAccessor& InTexCoord0Displacements, const FAccessor& InTexCoord1Displacements, const FAccessor& InColor0Deltas);
+
+		bool IsValid() const;
+		FMD5Hash GetHash() const;
+
+		bool HasPositionDisplacements() const;
+		void GetPositionDisplacements(TArray<FVector3f>& Buffer) const;
+		int32 GetNumberOfPositionDisplacements() const;
+
+		bool HasNormalDisplacements() const;
+		void GetNormalDisplacements(TArray<FVector3f>& Buffer) const;
+		int32 GetNumberOfNormalDisplacements() const;
+
+		bool HasTangentDisplacements() const;
+		void GetTangentDisplacements(TArray<FVector4f>& Buffer) const;
+		int32 GetNumberOfTangentDisplacements() const;
+
+		bool HasTexCoordDisplacements(uint32 Index) const;
+		void GetTexCoordDisplacements(uint32 Index, TArray<FVector2f>& Buffer) const;
+		int32 GetNumberOfTexCoordDisplacements(uint32 Index) const;
+
+		bool HasColorDeltas() const;
+		void GetColorDeltas(TArray<FVector4f>& Buffer) const;
+		int32 GetNumberOfColorDeltas() const;
+
+	private:
+		const FAccessor& PositionDisplacements;
+		const FAccessor& NormalDisplacements;
+		const FAccessor& TangentDisplacements;
+		const FAccessor& TexCoord0Displacements;
+		const FAccessor& TexCoord1Displacements;
+		const FAccessor& Color0Deltas;
+	};
+
+	struct GLTFCORE_API FPrimitive
+	{
+		enum class EMode
+		{
+			// valid but unsupported
+			Points    = 0,
+			Lines     = 1,
+			LineLoop  = 2,
+			LineStrip = 3,
+			// initially supported
+			Triangles = 4,
+			// will be supported prior to release
+			TriangleStrip = 5,
+			TriangleFan   = 6,
+
+			//
+			Unknown = 7
+		};
+
+		const EMode Mode;
+		const int32 MaterialIndex;
+		TArray<FVariantMapping>	VariantMappings;
+
+		FPrimitive(EMode InMode, int32 InMaterial, const FAccessor& InIndices, const FAccessor& InPosition, const FAccessor& InNormal,
+		           const FAccessor& InTangent, const FAccessor& InTexCoord0, const FAccessor& InTexCoord1, const FAccessor& InColor0,
+		           const FAccessor& InJoints0, const FAccessor& InWeights0);
+
+		void GenerateIsValidCache();
+		bool IsValid() const;
+		FMD5Hash GetHash() const;
+
+		void GetPositions(TArray<FVector3f>& Buffer) const;
+		bool HasNormals() const;
+		void GetNormals(TArray<FVector3f>& Buffer) const;
+		bool HasTangents() const;
+		void GetTangents(TArray<FVector4f>& Buffer) const;
+		bool HasTexCoords(uint32 Index) const;
+		void GetTexCoords(uint32 Index, TArray<FVector2f>& Buffer) const;
+		void GetColors(TArray<FVector4f>& Buffer) const;
+		bool HasColors() const;
+		bool HasJointWeights() const;
+		void GetJointInfluences(TArray<FJointInfluence>& Buffer) const;
+
+		FTriangle TriangleVerts(uint32 T) const;
+		void      GetTriangleIndices(TArray<uint32>& Buffer) const;
+
+		uint32 VertexCount() const;
+		uint32 TriangleCount() const;
+
+		//
+		TArray<FMorphTarget> MorphTargets;
+
+	private:
+		bool IsValidPrivate() const;
+
+		// index buffer
+		const FAccessor& Indices;
+
+		// common attributes
+		const FAccessor& Position;  // always required
+		const FAccessor& Normal;
+		const FAccessor& Tangent;
+		const FAccessor& TexCoord0;
+		const FAccessor& TexCoord1;
+		const FAccessor& Color0;
+		// skeletal mesh attributes
+		const FAccessor& Joints0;
+		const FAccessor& Weights0;
+
+		//Validity cache:
+		TOptional<bool> bIsValidCache;
+	};
+
+
+	struct GLTFCORE_API FMesh
+	{
+		FString				Name;
+		TArray<FPrimitive>	Primitives;
+
+		TArray<float>		MorphTargetWeights;
+		TArray<FString>		MorphTargetNames;
+
+		FString				UniqueId; //will be generated in FAsset::GenerateNames
+	
+		bool HasNormals() const;
+		bool HasTangents() const;
+		bool HasTexCoords(uint32 Index) const;
+		bool HasColors() const;
+		bool HasJointWeights() const;
+
+		void GenerateIsValidCache(bool GenerateIsValidCacheForPrimitives = true);
+		bool IsValid() const;
+		FMD5Hash GetHash() const;
+
+		int32 NumberOfMorphTargetsPerPrimitive() const;
+
+	private:
+		bool IsValidPrivate() const;
+		//Validity cache:
+		TOptional<bool> bIsValidCache;
+	};
+
+	//
+
+	inline bool FPrimitive::HasNormals() const
+	{
+		return Normal.IsValid();
+	}
+
+	inline bool FPrimitive::HasTangents() const
+	{
+		return Tangent.IsValid();
+	}
+
+	inline bool FPrimitive::HasTexCoords(uint32 Index) const
+	{
+		switch (Index)
+		{
+			case 0:
+				return TexCoord0.IsValid();
+			case 1:
+				return TexCoord1.IsValid();
+			default:
+				return false;
+		}
+	}
+
+	inline bool FPrimitive::HasColors() const
+	{
+		return Color0.IsValid() && (Color0.Type == FAccessor::EType::Vec3 || Color0.Type == FAccessor::EType::Vec4);
+	}
+
+	inline void FPrimitive::GetPositions(TArray<FVector3f>& Buffer) const
+	{
+		Position.GetCoordArray(Buffer);
+	}
+
+	inline void FPrimitive::GetNormals(TArray<FVector3f>& Buffer) const
+	{
+		Normal.GetCoordArray(Buffer);
+	}
+
+	inline void FPrimitive::GetTexCoords(uint32 Index, TArray<FVector2f>& Buffer) const
+	{
+		switch (Index)
+		{
+			case 0:
+				return TexCoord0.GetVec2Array(Buffer);
+			case 1:
+				return TexCoord1.GetVec2Array(Buffer);
+			default:
+				ensure(false);
+				break;
+		}
+	}
+
+	inline bool FPrimitive::HasJointWeights() const
+	{
+		return Joints0.IsValid() && Weights0.IsValid();
+	}
+
+	//
+
+	inline bool FMesh::HasNormals() const
+	{
+		return Primitives.FindByPredicate([](const FPrimitive& Prim) { return Prim.HasNormals(); }) != nullptr;
+	}
+
+	inline bool FMesh::HasTangents() const
+	{
+		return Primitives.FindByPredicate([](const FPrimitive& Prim) { return Prim.HasTangents(); }) != nullptr;
+	}
+
+	inline bool FMesh::HasTexCoords(uint32 Index) const
+	{
+		return Primitives.FindByPredicate([Index](const FPrimitive& Prim) { return Prim.HasTexCoords(Index); }) != nullptr;
+	}
+
+	inline bool FMesh::HasColors() const
+	{
+		return Primitives.FindByPredicate([](const FPrimitive& Prim) { return Prim.HasColors(); }) != nullptr;
+	}
+
+	inline bool FMesh::HasJointWeights() const
+	{
+		const bool Result = Primitives.FindByPredicate([](const auto& Prim) { return Prim.HasJointWeights(); }) != nullptr;
+
+		if (Result)
+		{
+			// According to spec, *all* primitives of a skinned mesh must have joint weights.
+			int32 Count = 0;
+			for (const FPrimitive& Prim : Primitives)
+			{
+				Count += Prim.HasJointWeights();
+			}
+			ensure(Primitives.Num() == Count);
+		}
+
+		return Result;
+	}
+
+	inline void FMesh::GenerateIsValidCache(bool GenerateIsValidCacheForPrimitives)
+	{
+		if (GenerateIsValidCacheForPrimitives)
+		{
+			for (FPrimitive& Primitive : Primitives)
+			{
+				Primitive.GenerateIsValidCache();
+			}
+		}
+
+		bIsValidCache = IsValidPrivate();
+	}
+	inline bool FMesh::IsValid() const
+	{
+		if (bIsValidCache.IsSet())
+		{
+			return bIsValidCache.GetValue();
+		}
+		else
+		{
+			return IsValidPrivate();
+		}
+	}
+
+	inline bool FMesh::IsValidPrivate() const
+	{
+		//Validate Primitives:
+		bool bIsValid = Primitives.FindByPredicate([](const FPrimitive& Prim) { return !Prim.IsValid(); }) == nullptr;
+
+		//if MorphTargetNames are not set, but mesh has morph targets (which is likely), this will generate isValid false;
+		//IsValidCache will have to be (re-)genereated post FAsset::GenerateNames (which is called at the end of GLTF::FReader) to overcome this issue.
+
+		//Validate Morph Target Names and Morph Target Weights:
+		if ((MorphTargetNames.Num() > 0) && (MorphTargetWeights.Num() > 0))
+		{
+			if (MorphTargetNames.Num() != MorphTargetWeights.Num())
+			{
+				bIsValid = false;
+			}
+		}
+
+		//Validate Morph Target (and Morph Target Names) Counts:
+		int32 MorphTargetCounter = MorphTargetNames.Num();
+
+		for (const FPrimitive& Primitive : Primitives)
+		{
+			if (MorphTargetCounter != Primitive.MorphTargets.Num())
+			{
+				bIsValid = false;
+				break;
+			}
+		}
+
+		return bIsValid;
+	}
+
+}  // namespace GLTF
diff --git a/Engine/Plugins/Interchange/Runtime/Source/Parsers/GLTFCore/Public/GLTFMeshFactory.h b/Engine/Plugins/Interchange/Runtime/Source/Parsers/GLTFCore/Public/GLTFMeshFactory.h
new file mode 100644
--- /dev/null
+++ b/Engine/Plugins/Interchange/Runtime/Source/Parsers/GLTFCore/Public/GLTFMeshFactory.h
@@ -0,0 +1,42 @@
+// Copyright Epic Games, Inc. All Rights Reserved.
+
+#pragma once
+
+#include "CoreMinimal.h"
+#include "MeshTypes.h"
+#include "GLTFLogger.h"
+
+struct FVertexID;
+
+struct FMeshDescription;
+
+namespace GLTF
+{
+	struct FMesh;
+	class FMeshFactoryImpl;
+
+	class GLTFCORE_API FMeshFactory
+	{
+	public:
+		FMeshFactory();
+		~FMeshFactory();
+
+		using FIndexVertexIdMap = TMap<int32, FVertexID>;
+
+		void FillMeshDescription(const GLTF::FMesh &Mesh, FMeshDescription* MeshDescription, const TArray<float>& MorphTargetWeights = TArray<float>());
+
+		float GetUniformScale() const;
+		void  SetUniformScale(float Scale);
+
+		const TArray<FLogMessage>&  GetLogMessages() const;
+
+		void SetReserveSize(uint32 Size);
+
+		TArray<FMeshFactory::FIndexVertexIdMap>& GetPositionIndexToVertexIdPerPrim() const;
+
+		void CleanUp();
+
+	private:
+		TUniquePtr<FMeshFactoryImpl> Impl;
+	};
+}
diff --git a/Engine/Plugins/Interchange/Runtime/Source/Parsers/GLTFCore/Public/GLTFNode.h b/Engine/Plugins/Interchange/Runtime/Source/Parsers/GLTFCore/Public/GLTFNode.h
new file mode 100644
--- /dev/null
+++ b/Engine/Plugins/Interchange/Runtime/Source/Parsers/GLTFCore/Public/GLTFNode.h
@@ -0,0 +1,154 @@
+// Copyright Epic Games, Inc. All Rights Reserved.
+
+#pragma once
+
+#include "CoreMinimal.h"
+#include "GLTFMaterial.h"
+#include "GLTFMesh.h"
+#include "GLTFTexture.h"
+
+namespace GLTF { struct FAccessor; }
+
+namespace GLTF
+{
+	struct GLTFCORE_API FNode
+	{
+		enum class EType
+		{
+			None,
+			Transform,
+			Joint,
+			Mesh,
+			MeshSkinned,
+			Camera,
+			Light
+		};
+
+		FString       Name;
+		FTransform    Transform;
+		TArray<int32> Children;  // each is an index into FAsset::Nodes
+		EType         Type;
+
+		int32 MeshIndex;  // index into FAsset::Meshes
+
+		// Skindex is the skin used by the mesh at this node.
+		// It's not the joints belonging *to* a skin
+		int32 Skindex;      // index into FAsset::Skins
+		int32 CameraIndex;  // index into FAsset::Cameras
+		int32 LightIndex;   // index into FAsset::Lights
+
+		int32 ParentIndex;
+		int32 RootJointIndex; //only valid if node is of Joint Type
+
+		TArray<float> MorphTargetWeights; //for the instantiated mesh with morph targets.
+
+		FString UniqueId; //will be generated in FAsset::GenerateNames
+
+		FNode()
+		    : Type(EType::None)
+		    , MeshIndex(INDEX_NONE)
+		    , Skindex(INDEX_NONE)
+		    , CameraIndex(INDEX_NONE)
+		    , LightIndex(INDEX_NONE)
+			, ParentIndex(INDEX_NONE)
+			, RootJointIndex(INDEX_NONE)
+		{
+		}
+	};
+
+	struct GLTFCORE_API FCamera
+	{
+		struct FPerspective
+		{
+			// vertical field of view in radians
+			float Fov;
+			// aspect ratio of the field of view
+			float AspectRatio;
+		};
+		struct FOrthographic
+		{
+			// horizontal magnification of the view
+			float XMagnification;
+			// vertical magnification of the view
+			float YMagnification;
+		};
+
+		const FNode& Node;
+		FString      Name;
+		union {
+			FOrthographic Orthographic;
+			FPerspective  Perspective;
+		};
+		float ZNear;
+		float ZFar;
+		bool  bIsPerspective;
+		FString UniqueId; //will be generated in FAsset::GenerateNames
+
+		FCamera(const FNode& Node)
+		    : Node(Node)
+		    , ZNear(0.f)
+		    , bIsPerspective(true)
+		{
+			Perspective.Fov         = 0.f;
+			Perspective.AspectRatio = 1.f;
+		}
+	};
+
+	struct GLTFCORE_API FLight
+	{
+		enum class EType
+		{
+			Directional,
+			Point,
+			Spot
+		};
+
+		struct FSpot
+		{
+			float InnerConeAngle;
+			float OuterConeAngle;
+
+			FSpot()
+			    : InnerConeAngle(0.f)
+			    , OuterConeAngle(PI / 4.f)
+			{
+			}
+		};
+
+		const FNode* Node;
+		FString      Name;
+		EType        Type;
+		FVector      Color;
+		float        Intensity;
+		// Must be > 0. When undefined, range is assumed to be infinite.
+		float Range;
+		FSpot Spot;
+
+		FString UniqueId; //will be generated in FAsset::GenerateNames
+
+		FLight(const FNode* Node)
+		    : Node(Node)
+		    , Color(1.f)
+		    , Intensity(1.f)
+		    , Range(1e+20)
+		{
+		}
+	};
+
+	struct GLTFCORE_API FSkinInfo
+	{
+		const FAccessor& InverseBindMatrices;
+		FString          Name;
+		TArray<int32>    Joints;    // each is an index into FAsset::Nodes
+		int32            Skeleton;  // root node, index into FAsset::Nodes
+
+		FString          UniqueId; //will be generated in FAsset::GenerateNames
+
+		FSkinInfo(const FAccessor& InverseBindMatrices)
+		    : InverseBindMatrices(InverseBindMatrices)
+		    , Skeleton(INDEX_NONE)
+		{
+		}
+	};
+
+}  // namespace GLTF
diff --git a/Engine/Plugins/Interchange/Runtime/Source/Parsers/GLTFCore/Public/GLTFReader.h b/Engine/Plugins/Interchange/Runtime/Source/Parsers/GLTFCore/Public/GLTFReader.h
new file mode 100644
--- /dev/null
+++ b/Engine/Plugins/Interchange/Runtime/Source/Parsers/GLTFCore/Public/GLTFReader.h
@@ -0,0 +1,87 @@
+// Copyright Epic Games, Inc. All Rights Reserved.
+
+#pragma once
+
+#include "GLTFLogger.h"
+
+#include "CoreMinimal.h"
+
+class FJsonObject;
+
+namespace GLTF
+{
+	struct FAsset;
+	struct FMesh;
+	struct FPrimitive;
+	struct FTextureMap;
+	struct FMaterial;
+	class FBinaryFileReader;
+	class FExtensionsHandler;
+
+	class GLTFCORE_API FFileReader : public FBaseLogger
+	{
+	public:
+		FFileReader();
+		~FFileReader();
+
+		/**
+		 * Loads the contents of a glTF file into the given asset.
+		 *
+		 * @param InFilePath - disk file path
+		 * @param bInLoadImageData - if false then doesn't loads the data field of FImage
+		 * @param bInLoadMetadata - load extra asset metadata
+		 * images from a disk path, images from data sources will always be loaded.
+		 * @param OutAsset - asset data destination
+		 */
+		void ReadFile(const FString& InFilePath, bool bInLoadImageData, bool bInLoadMetadata, GLTF::FAsset& OutAsset);
+
+	private:
+		void LoadMetadata(GLTF::FAsset& Asset);
+		void ImportAsset(const FString& FilePath, bool bInLoadImageData, GLTF::FAsset& Asset);
+		void AllocateExtraData(const FString& InResourcesPath, bool bInLoadImageData, TArray<uint8>& OutExtraData);
+
+		void SetupBuffer(const FJsonObject& Object, const FString& Path);
+		void SetupBufferView(const FJsonObject& Object) const;
+		void SetupAccessor(const FJsonObject& Object) const;
+		void SetupMorphTarget(const FJsonObject& Object, GLTF::FPrimitive& Primitive) const;
+		void SetupPrimitive(const FJsonObject& Object, GLTF::FMesh& Mesh) const;
+		void SetupMesh(const FJsonObject& Object) const;
+
+		void SetupScene(const FJsonObject& Object) const;
+		void SetupNode(const FJsonObject& Object) const;
+		void SetupCamera(const FJsonObject& Object) const;
+		void SetupAnimation(const FJsonObject& Object) const;
+		void SetupSkin(const FJsonObject& Object) const;
+
+		void SetupImage(const FJsonObject& Object, const FString& Path, bool bInLoadImageData);
+		void SetupSampler(const FJsonObject& Object) const;
+		void SetupTexture(const FJsonObject& Object) const;
+		void SetupMaterial(const FJsonObject& Object) const;
+
+		template <typename SetupFunc>
+		void SetupObjects(uint32 ObjectCount, const TCHAR* FieldName, SetupFunc Func) const;
+		void SetupNodesType() const;
+
+		void BuildParentIndices(int32 ParentNodeIndex, int32 CurrentNodeIndex) const;
+		int32 FindRootJointIndex(int32 CurrentNodeIndex) const;
+		void BuildRootJoints() const;
+
+	private:
+		uint32 BufferCount;
+		uint32 BufferViewCount;
+		uint32 ImageCount;
+
+		TSharedPtr<FJsonObject>  JsonRoot;
+		FString                        JsonBuffer;
+		TUniquePtr<FBinaryFileReader>  BinaryReader;
+		TUniquePtr<FExtensionsHandler> ExtensionsHandler;
+
+		// Each mesh has its own primitives array.
+		// Each primitive refers to the asset's global accessors array.
+		GLTF::FAsset* Asset;
+
+		// Current offset for extra data buffer.
+		uint8* CurrentBufferOffset;
+	};
+
+}  // namespace GLTF
diff --git a/Engine/Plugins/Interchange/Runtime/Source/Parsers/GLTFCore/Public/GLTFTexture.h b/Engine/Plugins/Interchange/Runtime/Source/Parsers/GLTFCore/Public/GLTFTexture.h
new file mode 100644
--- /dev/null
+++ b/Engine/Plugins/Interchange/Runtime/Source/Parsers/GLTFCore/Public/GLTFTexture.h
@@ -0,0 +1,93 @@
+// Copyright Epic Games, Inc. All Rights Reserved.
+
+#pragma once
+
+#include "CoreMinimal.h"
+
+struct FScriptContainerElement;
+
+namespace GLTF
+{
+	struct GLTFCORE_API FImage
+	{
+		enum class EFormat
+		{
+			Unknown,
+			PNG,
+			JPEG
+		};
+
+		FString Name;
+		FString URI;
+		EFormat Format;
+		FString FilePath;
+
+		// Image data is kept encoded in Format, to be decoded when needed by Unreal.
+		uint32       DataByteLength;
+		const uint8* Data;
+
+		FString      UniqueId; //will be generated in FAsset::GenerateNames
+
+		FImage()
+		    : DataByteLength(0)
+		    , Data(nullptr)
+		{
+		}
+	};
+
+	struct GLTFCORE_API FSampler
+	{
+		enum class EFilter
+		{
+			// valid for Min & Mag
+			Nearest = 9728,
+			Linear  = 9729,
+			// valid for Min only
+			NearestMipmapNearest = 9984,
+			LinearMipmapNearest  = 9985,
+			NearestMipmapLinear  = 9986,
+			LinearMipmapLinear   = 9987
+		};
+
+		enum class EWrap
+		{
+			Repeat         = 10497,
+			MirroredRepeat = 33648,
+			ClampToEdge    = 33071
+		};
+
+		static const FSampler DefaultSampler;
+
+		EFilter MinFilter;
+		EFilter MagFilter;
+
+		EWrap WrapS;
+		EWrap WrapT;
+
+		FSampler()
+		    // Spec defines no default min/mag filter.
+		    : MinFilter(EFilter::Linear)
+		    , MagFilter(EFilter::Linear)
+		    , WrapS(EWrap::Repeat)
+		    , WrapT(EWrap::Repeat)
+		{
+		}
+	};
+
+	struct GLTFCORE_API FTexture
+	{
+		const FImage&   Source;
+		const FSampler& Sampler;
+		FString         Name;
+
+		FString			UniqueId; //will be generated in FAsset::GenerateNames
+
+		FTexture(const FString& InName, const FImage& InSource, const FSampler& InSampler)
+		    : Source(InSource)
+		    , Sampler(InSampler)
+		    , Name(InName)
+		{
+		}
+	};
+
+}  // namespace GLTF
diff --git a/Templates/TP_AEC_ArchvisBP/TP_AEC_ArchvisBP.uproject b/Templates/TP_AEC_ArchvisBP/TP_AEC_ArchvisBP.uproject
--- a/Templates/TP_AEC_ArchvisBP/TP_AEC_ArchvisBP.uproject
+++ b/Templates/TP_AEC_ArchvisBP/TP_AEC_ArchvisBP.uproject
@@ -9,7 +9,7 @@
 			"Enabled": true
 		},
 		{
-			"Name": "DatasmithGLTFImporter",
+			"Name": "InterchangeEditor",
 			"Enabled": true
 		},
 		{
diff --git a/Templates/TP_AEC_BlankBP/TP_AEC_BlankBP.uproject b/Templates/TP_AEC_BlankBP/TP_AEC_BlankBP.uproject
--- a/Templates/TP_AEC_BlankBP/TP_AEC_BlankBP.uproject
+++ b/Templates/TP_AEC_BlankBP/TP_AEC_BlankBP.uproject
@@ -9,7 +9,7 @@
 			"Enabled": true
 		},
 		{
-			"Name": "DatasmithGLTFImporter",
+			"Name": "InterchangeEditor",
 			"Enabled": true
 		},
 		{
diff --git a/Templates/TP_AEC_CollabBP/TP_AEC_CollabBP.uproject b/Templates/TP_AEC_CollabBP/TP_AEC_CollabBP.uproject
--- a/Templates/TP_AEC_CollabBP/TP_AEC_CollabBP.uproject
+++ b/Templates/TP_AEC_CollabBP/TP_AEC_CollabBP.uproject
@@ -40,7 +40,7 @@
 			]
 		},
 		{
-			"Name": "DatasmithGLTFImporter",
+			"Name": "InterchangeEditor",
 			"Enabled": true
 		},
 		{
diff --git a/Templates/TP_MFG_CollabBP/TP_MFG_CollabBP.uproject b/Templates/TP_MFG_CollabBP/TP_MFG_CollabBP.uproject
--- a/Templates/TP_MFG_CollabBP/TP_MFG_CollabBP.uproject
+++ b/Templates/TP_MFG_CollabBP/TP_MFG_CollabBP.uproject
@@ -41,7 +41,7 @@
 			]
 		},
 		{
-			"Name": "DatasmithGLTFImporter",
+			"Name": "InterchangeEditor",
 			"Enabled": true
 		},
 		{
diff --git a/Templates/TP_PhotoStudioBP/TP_PhotoStudioBP.uproject b/Templates/TP_PhotoStudioBP/TP_PhotoStudioBP.uproject
--- a/Templates/TP_PhotoStudioBP/TP_PhotoStudioBP.uproject
+++ b/Templates/TP_PhotoStudioBP/TP_PhotoStudioBP.uproject
@@ -41,7 +41,7 @@
 			]
 		},
 		{
-			"Name": "DatasmithGLTFImporter",
+			"Name": "InterchangeEditor",
 			"Enabled": true
 		},
 		{
