[Interchange] OBJ translator distinguishes multiple polygon groups with different materials within a group

diff --git a/Engine/Plugins/Interchange/Runtime/Source/Import/Private/Mesh/InterchangeOBJTranslator.cpp b/Engine/Plugins/Interchange/Runtime/Source/Import/Private/Mesh/InterchangeOBJTranslator.cpp
--- a/Engine/Plugins/Interchange/Runtime/Source/Import/Private/Mesh/InterchangeOBJTranslator.cpp
+++ b/Engine/Plugins/Interchange/Runtime/Source/Import/Private/Mesh/InterchangeOBJTranslator.cpp
@@ -83,13 +83,18 @@ struct FObjData
 		}
 	};
 
-	/** Defines the data for a group */
-	struct FGroupData
+	struct FFaceGroupData
 	{
 		FString MaterialName;
 		TArray<FFaceData> Faces;
 	};
 
+	/** Defines the data for a group */
+	struct FGroupData
+	{
+		TArray<FFaceGroupData> FaceGroups;
+	};
+
 	/** All the groups defined by the .obj */
 	TMap<FString, FGroupData> Groups;
 
@@ -155,11 +160,14 @@ TArray<int32> FObjData::GetVertexIndicesUsedByGroup(const FGroupData& GroupData)
 {
 	TSet<int32> VertexIndexSet;
 
-	for (const FFaceData& FaceData : GroupData.Faces)
+	for (const FFaceGroupData& FaceGroup : GroupData.FaceGroups)
 	{
-		for (const FVertexData& VertexData : FaceData.Vertices)
+		for (const FFaceData& FaceData : FaceGroup.Faces)
 		{
-			VertexIndexSet.Add(VertexData.VertexIndex);
+			for (const FVertexData& VertexData : FaceData.Vertices)
+			{
+				VertexIndexSet.Add(VertexData.VertexIndex);
+			}
 		}
 	}
 
@@ -171,13 +179,16 @@ TArray<int32> FObjData::GetUVIndicesUsedByGroup(const FGroupData& GroupData) con
 {
 	TSet<int32> UVIndexSet;
 
-	for (const FFaceData& FaceData : GroupData.Faces)
+	for (const FFaceGroupData& FaceGroup : GroupData.FaceGroups)
 	{
-		for (const FVertexData& VertexData : FaceData.Vertices)
+		for (const FFaceData& FaceData : FaceGroup.Faces)
 		{
-			if (VertexData.UVIndex != INDEX_NONE)
+			for (const FVertexData& VertexData : FaceData.Vertices)
 			{
-				UVIndexSet.Add(VertexData.UVIndex);
+				if (VertexData.UVIndex != INDEX_NONE)
+				{
+					UVIndexSet.Add(VertexData.UVIndex);
+				}
 			}
 		}
 	}
@@ -230,7 +241,10 @@ int32 FObjData::GetGroupPolygonCount(const FString& GroupName) const
 	}
 
 	const FGroupData& GroupData = *GroupDataPtr;
-	return GroupData.Faces.Num();
+	return Algo::Accumulate(GroupData.FaceGroups, 0, [](int32 Accum, const FFaceGroupData& FaceGroup)
+	{
+		return Accum + FaceGroup.Faces.Num();
+	});
 }
 
 
@@ -289,43 +303,47 @@ FMeshDescription FObjData::MakeMeshDescriptionForGroup(const FString& GroupName)
 		}
 	}
 
-	// Create polygon group
-
-	FPolygonGroupID PolygonGroupIndex = MeshDescription.CreatePolygonGroup();
-	const FString MaterialName = GroupData.MaterialName.IsEmpty() ? UMaterial::GetDefaultMaterial(MD_Surface)->GetName() : GroupData.MaterialName;
-	ensure(!MaterialName.IsEmpty());
-	Attributes.GetPolygonGroupMaterialSlotNames()[PolygonGroupIndex] = FName(MaterialName);
-
-	// Create faces.
-	// n-gons are preserved
+	TMap<FPolygonID, const FFaceData*> Polygons;
+	for (const FFaceGroupData& FaceGroup : GroupData.FaceGroups)
+	{
+		
+		// Create polygon group
+		FPolygonGroupID PolygonGroupIndex = MeshDescription.CreatePolygonGroup();
+		const FString MaterialName = FaceGroup.MaterialName.IsEmpty() ? UMaterial::GetDefaultMaterial(MD_Surface)->GetName() : FaceGroup.MaterialName;
+		ensure(!MaterialName.IsEmpty());
+		Attributes.GetPolygonGroupMaterialSlotNames()[PolygonGroupIndex] = FName(MaterialName);
 
-	MeshDescription.ReserveNewTriangles(GroupData.Faces.Num());
-	MeshDescription.ReserveNewPolygons(GroupData.Faces.Num());
+		// Create faces.
+		// n-gons are preserved
 
-	TArray<FVertexInstanceID, TInlineAllocator<8>> VertexInstanceIDs;
-	for (const FFaceData& FaceData : GroupData.Faces)
-	{
-		VertexInstanceIDs.Reset();
-		MeshDescription.ReserveNewVertexInstances(FaceData.Vertices.Num());
+		MeshDescription.ReserveNewTriangles(FaceGroup.Faces.Num());
+		MeshDescription.ReserveNewPolygons(FaceGroup.Faces.Num());
 
-		for (const FVertexData& VertexData : FaceData.Vertices)
+		TArray<FVertexInstanceID, TInlineAllocator<8>> VertexInstanceIDs;
+		for (const FFaceData& FaceData : FaceGroup.Faces)
 		{
-			FVertexID VertexID = Algo::BinarySearch(VertexIndexMapping, VertexData.VertexIndex);
-			FVertexInstanceID VertexInstanceID = MeshDescription.CreateVertexInstance(VertexID);
-			VertexInstanceIDs.Add(VertexInstanceID);
+			VertexInstanceIDs.Reset();
+			MeshDescription.ReserveNewVertexInstances(FaceData.Vertices.Num());
 
-			if (VertexData.NormalIndex != INDEX_NONE && Normals.IsValidIndex(VertexData.NormalIndex))
+			for (const FVertexData& VertexData : FaceData.Vertices)
 			{
-				Attributes.GetVertexInstanceNormals()[VertexInstanceID] = PositionToUEBasis(Normals[VertexData.NormalIndex]);
-			}
+				FVertexID VertexID = Algo::BinarySearch(VertexIndexMapping, VertexData.VertexIndex);
+				FVertexInstanceID VertexInstanceID = MeshDescription.CreateVertexInstance(VertexID);
+				VertexInstanceIDs.Add(VertexInstanceID);
 
-			if (VertexData.UVIndex != INDEX_NONE && UVs.IsValidIndex(VertexData.UVIndex))
-			{
-				Attributes.GetVertexInstanceUVs()[VertexInstanceID] = UVToUEBasis(UVs[VertexData.UVIndex]);
+				if (VertexData.NormalIndex != INDEX_NONE && Normals.IsValidIndex(VertexData.NormalIndex))
+				{
+					Attributes.GetVertexInstanceNormals()[VertexInstanceID] = PositionToUEBasis(Normals[VertexData.NormalIndex]);
+				}
+
+				if (VertexData.UVIndex != INDEX_NONE && UVs.IsValidIndex(VertexData.UVIndex))
+				{
+					Attributes.GetVertexInstanceUVs()[VertexInstanceID] = UVToUEBasis(UVs[VertexData.UVIndex]);
+				}
 			}
-		}
 
-		int32 PolygonIndex = MeshDescription.CreatePolygon(PolygonGroupIndex, VertexInstanceIDs);
+			Polygons.Add(MeshDescription.CreatePolygon(PolygonGroupIndex, VertexInstanceIDs), &FaceData);
+		}
 	}
 
 	// Determine edge hardnesses
@@ -354,13 +372,13 @@ FMeshDescription FObjData::MakeMeshDescriptionForGroup(const FString& GroupName)
 
 					// For the vertex we are considering, look at the first face adjacent to the edge, and find the vertex data which includes it
 					// This is a baseline we will use to compare all the other adjacent faces
-					const FVertexData& VertexData = GroupData.Faces[EdgePolygonIDs[0]].GetVertexDataContainingVertexIndex(ObjVertexIndex);
+					const FVertexData& VertexData = Polygons[EdgePolygonIDs[0]]->GetVertexDataContainingVertexIndex(ObjVertexIndex);
 
 					for (int32 Index = 1; Index < EdgePolygonIDs.Num(); Index++)
 					{
 						// For all other adjacent faces, find the vertex data which includes the vertex we are considering.
 						// If the normal index is not the same as the baseline, we know this must be a hard edge.
-						const FVertexData& VertexDataToCompare = GroupData.Faces[EdgePolygonIDs[Index]].GetVertexDataContainingVertexIndex(ObjVertexIndex);
+						const FVertexData& VertexDataToCompare = Polygons[EdgePolygonIDs[Index]]->GetVertexDataContainingVertexIndex(ObjVertexIndex);
 
 						if (VertexData.NormalIndex != VertexDataToCompare.NormalIndex)
 						{
@@ -615,8 +633,14 @@ namespace ObjParser
 		if (FaceData.Vertices.Num() > 2)
 		{
 			FObjData::FGroupData& GroupData = ObjData.Groups.FindOrAdd(ObjData.CurrentGroup);
-			GroupData.Faces.Emplace(FaceData);
-			GroupData.MaterialName = ObjData.CurrentMaterial;
+			// Start new face group if last one had different material
+			if (GroupData.FaceGroups.IsEmpty() || (GroupData.FaceGroups.Last().MaterialName != ObjData.CurrentMaterial))
+			{
+				GroupData.FaceGroups.Emplace();
+				GroupData.FaceGroups.Last().MaterialName = ObjData.CurrentMaterial;
+			}
+
+			GroupData.FaceGroups.Last().Faces.Emplace(FaceData);
 		}
 		else
 		{
@@ -990,10 +1014,13 @@ bool UInterchangeOBJTranslator::Translate(UInterchangeBaseNodeContainer& BaseNod
 			MeshNode->SetCustomHasSmoothGroup(false);
 			MeshNode->SetCustomHasVertexColor(false);
 
-			if (!Group.Value.MaterialName.IsEmpty())
+			for (const FObjData::FFaceGroupData& FaceGroup : Group.Value.FaceGroups)
 			{
-				const FString MaterialDependencyUid = UInterchangeShaderGraphNode::MakeNodeUid(MakeShaderGraphNodeName(Group.Value.MaterialName));
-				MeshNode->SetSlotMaterialDependencyUid(Group.Value.MaterialName, MaterialDependencyUid);
+				if (!FaceGroup.MaterialName.IsEmpty())
+				{
+					const FString MaterialDependencyUid = UInterchangeShaderGraphNode::MakeNodeUid(MakeShaderGraphNodeName(FaceGroup.MaterialName));
+					MeshNode->SetSlotMaterialDependencyUid(FaceGroup.MaterialName, MaterialDependencyUid);
+				}
 			}
 		}
 
