MaterialExpression:
- Adding Swizzle material expression

MaterialX Translator:
- Abort import if a node doesn't have a node definition

MaterialX:
- Refactor to support multi output types
- Adding support of <separate2/3/4> nodes
- Adding support of <swizzle> node
- Support multi uv channel for <texcoord> node
- Fixing wrong input taken for <magnitude> and <rotate2d>

Generic Material Pipeline:
- Handle Swizzle Material Expression

diff --git a/Engine/Plugins/Interchange/Runtime/Source/Import/Private/MaterialX/InterchangeMaterialXDefinitions.h b/Engine/Plugins/Interchange/Runtime/Source/Import/Private/MaterialX/InterchangeMaterialXDefinitions.h
--- a/Engine/Plugins/Interchange/Runtime/Source/Import/Private/MaterialX/InterchangeMaterialXDefinitions.h
+++ b/Engine/Plugins/Interchange/Runtime/Source/Import/Private/MaterialX/InterchangeMaterialXDefinitions.h
@@ -318,6 +318,7 @@ MATERIALX_NAMESPACE_BEGIN
 		static constexpr const char* Separate2 = "separate2";
 		static constexpr const char* Separate3 = "separate3";
 		static constexpr const char* Separate4 = "separate4";
+		static constexpr const char* Swizzle = "swizzle";
 		// Procedural nodes 
 		static constexpr const char* Constant = "constant";
 		// Procedural2D nodes 
diff --git a/Engine/Plugins/Interchange/Runtime/Source/Import/Private/MaterialX/InterchangeMaterialXTranslator.cpp b/Engine/Plugins/Interchange/Runtime/Source/Import/Private/MaterialX/InterchangeMaterialXTranslator.cpp
--- a/Engine/Plugins/Interchange/Runtime/Source/Import/Private/MaterialX/InterchangeMaterialXTranslator.cpp
+++ b/Engine/Plugins/Interchange/Runtime/Source/Import/Private/MaterialX/InterchangeMaterialXTranslator.cpp
@@ -147,6 +147,16 @@ bool UInterchangeMaterialXTranslator::Translate(UInterchangeBaseNodeContainer& B
 
 			if(Node)
 			{
+				// Validate that all nodes in the file are strictly respecting their node definition
+				if(!Node->getNodeDef())
+				{
+					UInterchangeResultError_Generic* Message = AddMessage<UInterchangeResultError_Generic>();
+					Message->Text = FText::Format(LOCTEXT("NodeDefNotFound", "<{0}> has no matching NodeDef, aborting import..."),
+												  FText::FromString(Node->getName().c_str()));
+					bIsReferencesValid = false;
+					break;
+				}
+
 				bool bIsMaterialShader = Node->getType() == mx::Type::Material;
 				bool bIsLightShader = Node->getType() == mx::Type::LightShader;
 
diff --git a/Engine/Plugins/Interchange/Runtime/Source/Import/Private/MaterialX/MaterialExpressions/MaterialExpressionSwizzle.cpp b/Engine/Plugins/Interchange/Runtime/Source/Import/Private/MaterialX/MaterialExpressions/MaterialExpressionSwizzle.cpp
new file mode 100644
--- /dev/null
+++ b/Engine/Plugins/Interchange/Runtime/Source/Import/Private/MaterialX/MaterialExpressions/MaterialExpressionSwizzle.cpp
@@ -0,0 +1,146 @@
+// Copyright Epic Games, Inc. All Rights Reserved.
+
+#include "MaterialX/MaterialExpressions/MaterialExpressionSwizzle.h"
+#include "MaterialCompiler.h"
+
+#include UE_INLINE_GENERATED_CPP_BY_NAME(MaterialExpressionSwizzle)
+
+#define LOCTEXT_NAMESPACE "MaterialExpressionMaterialXSwizzle"
+
+UMaterialExpressionMaterialXSwizzle::UMaterialExpressionMaterialXSwizzle(const FObjectInitializer& ObjectInitializer)
+	: Super(ObjectInitializer)
+{
+	// Structure to hold one-time initialization
+	struct FConstructorStatics
+	{
+		FText NAME_MaterialX;
+		FConstructorStatics()
+			: NAME_MaterialX(LOCTEXT("MaterialX", "MaterialX"))
+		{}
+	};
+	static FConstructorStatics ConstructorStatics;
+
+#if WITH_EDITORONLY_DATA
+	MenuCategories.Add(ConstructorStatics.NAME_MaterialX);
+#endif
+}
+
+#if WITH_EDITOR
+int32 UMaterialExpressionMaterialXSwizzle::Compile(FMaterialCompiler* Compiler, int32 OutputIndex)
+{
+	if(!Input.GetTracedInput().Expression)
+	{
+		return Compiler->Errorf(TEXT("Missing Swizzle input"));
+	}
+
+	int32 Index = Input.Compile(Compiler);
+
+	if(Channels.Len() > 4)
+	{
+		return Compiler->Errorf(TEXT("Too many channels"));
+	}
+
+	if(Channels.IsEmpty())
+	{
+		return Index;
+	}
+
+	//Don't allow mixing xyzw and rgba or invalid channels
+	bool bHasRGBA = false;
+	bool bHasXYZW = false;
+
+	bool SwizzleX[4]{ false, false, false, false };
+	bool SwizzleY[4]{ false, false, false, false };
+	bool SwizzleZ[4]{ false, false, false, false };
+	bool SwizzleW[4]{ false, false, false, false };
+	
+	for(int32 i = 0; i < Channels.Len(); ++i)
+	{
+		const TCHAR Channel = Channels[i];
+		if(Channel == TEXT('r') || Channel == TEXT('g') || Channel == TEXT('b') || Channel == TEXT('a'))
+		{
+			bHasRGBA = true;
+		}
+		else if(Channel == TEXT('x') || Channel == TEXT('y') || Channel == TEXT('z') || Channel == TEXT('w'))
+		{
+			bHasXYZW = true;
+		}
+		if(bHasRGBA && bHasXYZW)
+		{
+			return Compiler->Errorf(TEXT("Cannot mix rgba and xyzw channels"));
+		}
+		if(!bHasRGBA && !bHasXYZW)
+		{
+			return Compiler->Errorf(TEXT("%c is not recognized as a valid channel"), Channel);
+		}
+
+		switch(Channel)
+		{
+		case TEXT('x'):
+		case TEXT('r'):
+			SwizzleX[i] = true;
+			break;
+
+		case TEXT('y'):
+		case TEXT('g'):
+			SwizzleY[i] = true;
+			break;
+
+		case TEXT('z'):
+		case TEXT('b'):
+			SwizzleZ[i] = true;
+			break;
+
+		case TEXT('w'):
+		case TEXT('a'):
+			SwizzleW[i] = true;
+			break;
+		}
+	}
+	
+	int32 Mask[4]{
+		Channels.Len() >= 1 ? Compiler->ComponentMask(Index, SwizzleX[0], SwizzleY[0], SwizzleZ[0], SwizzleW[0]) : -1,
+		Channels.Len() >= 2 ? Compiler->ComponentMask(Index, SwizzleX[1], SwizzleY[1], SwizzleZ[1], SwizzleW[1]) : -1,
+		Channels.Len() >= 3 ? Compiler->ComponentMask(Index, SwizzleX[2], SwizzleY[2], SwizzleZ[2], SwizzleW[2]) : -1,
+		Channels.Len() >= 4 ? Compiler->ComponentMask(Index, SwizzleX[3], SwizzleY[3], SwizzleZ[3], SwizzleW[3]) : -1,
+	};
+
+	Input.Expression->GetOutputType(Input.OutputIndex);
+	int32 Res = INDEX_NONE;
+
+	switch(Channels.Len())
+	{
+	case 1:
+		Res = Mask[0];
+		break;
+
+	case 2:
+		Res = Compiler->AppendVector(Mask[0], Mask[1]);
+		break;
+
+	case 3:
+		Res = Compiler->AppendVector(Compiler->AppendVector(Mask[0], Mask[1]), Mask[2]);
+		break;
+
+	case 4:
+		Res = Compiler->AppendVector(Compiler->AppendVector(Compiler->AppendVector(Mask[0], Mask[1]), Mask[2]), Mask[3]);
+		break;
+	}
+
+	return Res;
+}
+
+void UMaterialExpressionMaterialXSwizzle::GetCaption(TArray<FString>& OutCaptions) const
+{
+	OutCaptions.Add(TEXT("MaterialX Swizzle"));
+}
+
+void UMaterialExpressionMaterialXSwizzle::GetExpressionToolTip(TArray<FString>& OutToolTip)
+{
+	ConvertToMultilineToolTip(TEXT("Perform an arbitrary permutation of the channels of the input stream, returning a new "
+								   "stream of the specified type. Individual channels may be replicated or omitted, and the output "
+								   "stream may have a different number of channels than the input."), 40, OutToolTip);
+}
+#endif
+
+#undef LOCTEXT_NAMESPACE 
\ No newline at end of file
diff --git a/Engine/Plugins/Interchange/Runtime/Source/Import/Private/MaterialX/MaterialXUtils/MaterialXBase.cpp b/Engine/Plugins/Interchange/Runtime/Source/Import/Private/MaterialX/MaterialXUtils/MaterialXBase.cpp
--- a/Engine/Plugins/Interchange/Runtime/Source/Import/Private/MaterialX/MaterialXUtils/MaterialXBase.cpp
+++ b/Engine/Plugins/Interchange/Runtime/Source/Import/Private/MaterialX/MaterialXUtils/MaterialXBase.cpp
@@ -69,6 +69,7 @@ void FMaterialXBase::AddTexCoordToTiledImageNodes(MaterialX::ElementPtr Graph)
 													  {},
 													  { {"index", FAttributeValueArray{{"type", "integer"}, {"value", "0"}}} });
 
+				NodeTexCoord->setType(mx::Type::Vector2);
 				InputTexCoord = Node->addInput(mx::NodeGroup::Texture2D::Inputs::TexCoord, mx::Type::Vector2);
 				InputTexCoord->setNodeName(NodeTexCoord->getName());
 			}
diff --git a/Engine/Plugins/Interchange/Runtime/Source/Import/Private/MaterialX/MaterialXUtils/MaterialXManager.cpp b/Engine/Plugins/Interchange/Runtime/Source/Import/Private/MaterialX/MaterialXUtils/MaterialXManager.cpp
--- a/Engine/Plugins/Interchange/Runtime/Source/Import/Private/MaterialX/MaterialXUtils/MaterialXManager.cpp
+++ b/Engine/Plugins/Interchange/Runtime/Source/Import/Private/MaterialX/MaterialXUtils/MaterialXManager.cpp
@@ -67,7 +67,6 @@ FMaterialXManager::FMaterialXManager()
 		{{MaterialX::Category::Inside,         TEXT("mask")},      TEXT("B")},			  // Inside is treated as a Multiply node
 		{{MaterialX::Category::Invert,         TEXT("amount")},    TEXT("A")},
 		{{MaterialX::Category::Invert,         TEXT("in")},        TEXT("B")},
-		{{MaterialX::Category::Magnitude,      TEXT("in")},        TEXT("A")},
 		{{MaterialX::Category::Mix,            TEXT("fg")},        TEXT("B")},
 		{{MaterialX::Category::Mix,            TEXT("bg")},        TEXT("A")},
 		{{MaterialX::Category::Mix,            TEXT("mix")},		TEXT("Factor")},
@@ -81,7 +80,6 @@ FMaterialXManager::FMaterialXManager()
 		{{MaterialX::Category::Power,          TEXT("in1")},       TEXT("Base")},
 		{{MaterialX::Category::Power,	       TEXT("in2")},       TEXT("Exponent")},
 		{{MaterialX::Category::Rotate2D,       TEXT("amount")},    TEXT("RotationAngle")},
-		{{MaterialX::Category::Rotate2D,       TEXT("in")},        TEXT("Position")},
 		{{MaterialX::Category::Rotate3D,       TEXT("amount")},    TEXT("RotationAngle")},
 		{{MaterialX::Category::Rotate3D,       TEXT("axis")},		TEXT("NormalizedRotationAxis")},
 		{{MaterialX::Category::Rotate3D,       TEXT("in")},        TEXT("Position")},
@@ -155,7 +153,6 @@ FMaterialXManager::FMaterialXManager()
 		{MaterialX::Category::Fractal3D,    TEXT("MaterialXFractal3D")},
 		// Geometric nodes 
 		{MaterialX::Category::GeomColor,	TEXT("VertexColor")},
-		{MaterialX::Category::TexCoord,		TEXT("TextureCoordinate")},
 		// Adjustment nodes,
 		{MaterialX::Category::HsvToRgb,		TEXT("MaterialXHsvToRgb")},
 		{MaterialX::Category::Luminance,    TEXT("MaterialXLuminance")},
diff --git a/Engine/Plugins/Interchange/Runtime/Source/Import/Private/MaterialX/MaterialXUtils/MaterialXStandardSurfaceShader.cpp b/Engine/Plugins/Interchange/Runtime/Source/Import/Private/MaterialX/MaterialXUtils/MaterialXStandardSurfaceShader.cpp
--- a/Engine/Plugins/Interchange/Runtime/Source/Import/Private/MaterialX/MaterialXUtils/MaterialXStandardSurfaceShader.cpp
+++ b/Engine/Plugins/Interchange/Runtime/Source/Import/Private/MaterialX/MaterialXUtils/MaterialXStandardSurfaceShader.cpp
@@ -37,7 +37,7 @@ void FMaterialXStandardSurfaceShader::Translate(mx::NodePtr StandardSurfaceNode)
 		StandardSurfaceShaderNode->SetCustomMaterialFunction(TEXT("/Interchange/Functions/MX_StandardSurface.MX_StandardSurface"));
 		NodeContainer.AddNode(StandardSurfaceShaderNode);
 
-		ShaderNodes.Add(NodeName, StandardSurfaceShaderNode);
+		ShaderNodes.Add({ NodeName, StandardSurfaceNode->getNodeDef(mx::EMPTY_STRING, true)->getActiveOutputs()[0]->getName().c_str() }, StandardSurfaceShaderNode);
 	}
 
 	// Inputs
diff --git a/Engine/Plugins/Interchange/Runtime/Source/Import/Private/MaterialX/MaterialXUtils/MaterialXSurfaceShaderAbstract.cpp b/Engine/Plugins/Interchange/Runtime/Source/Import/Private/MaterialX/MaterialXUtils/MaterialXSurfaceShaderAbstract.cpp
--- a/Engine/Plugins/Interchange/Runtime/Source/Import/Private/MaterialX/MaterialXUtils/MaterialXSurfaceShaderAbstract.cpp
+++ b/Engine/Plugins/Interchange/Runtime/Source/Import/Private/MaterialX/MaterialXUtils/MaterialXSurfaceShaderAbstract.cpp
@@ -19,6 +19,9 @@
 
 namespace mx = MaterialX;
 
+FString FMaterialXSurfaceShaderAbstract::EmptyString{};
+FString FMaterialXSurfaceShaderAbstract::DefaultOutput{TEXT("out")};
+
 FMaterialXSurfaceShaderAbstract::FMaterialXSurfaceShaderAbstract(UInterchangeBaseNodeContainer& BaseNodeContainer)
 	: FMaterialXBase(BaseNodeContainer)
 {}
@@ -147,20 +150,20 @@ bool FMaterialXSurfaceShaderAbstract::ConnectNodeGraphOutputToInput(MaterialX::I
 	return bHasNodeGraph;
 }
 
-bool FMaterialXSurfaceShaderAbstract::ConnectMatchingNodeOutputToInput(MaterialX::NodePtr Node, UInterchangeShaderNode* ParentShaderNode, const FString& InputChannelName)
+bool FMaterialXSurfaceShaderAbstract::ConnectMatchingNodeOutputToInput(const FConnectNode& Connect)
 {
 	FMaterialXManager& Manager = FMaterialXManager::GetInstance();
 
 	bool bIsConnected = false;
 
 	// First search a matching Material Expression
-	if(const FString* ShaderType = Manager.FindMatchingMaterialExpression(Node->getCategory().c_str()))
+	if(const FString* ShaderType = Manager.FindMatchingMaterialExpression(Connect.UpstreamNode->getCategory().c_str()))
 	{
 		UInterchangeShaderNode* OperatorNode = nullptr;
+		
+		OperatorNode = CreateShaderNode(Connect.UpstreamNode->getName().c_str(), *ShaderType);
 
-		OperatorNode = CreateShaderNode(Node->getName().c_str(), *ShaderType);
-
-		for(mx::InputPtr Input : Node->getInputs())
+		for(mx::InputPtr Input : Connect.UpstreamNode->getInputs())
 		{
 			if(const FString* InputNameFound = Manager.FindMaterialExpressionInput(GetInputName(Input)))
 			{
@@ -168,11 +171,11 @@ bool FMaterialXSurfaceShaderAbstract::ConnectMatchingNodeOutputToInput(MaterialX
 			}
 		}
 
-		bIsConnected = UInterchangeShaderPortsAPI::ConnectDefaultOuputToInput(ParentShaderNode, InputChannelName, OperatorNode->GetUniqueID());
+		bIsConnected = UInterchangeShaderPortsAPI::ConnectDefaultOuputToInput(Connect.ParentShaderNode, Connect.InputChannelName, OperatorNode->GetUniqueID());
 	}
-	else if(FOnConnectNodeOutputToInput* Delegate = MatchingConnectNodeDelegates.Find(Node->getCategory().c_str()))
+	else if(FOnConnectNodeOutputToInput* Delegate = MatchingConnectNodeDelegates.Find(Connect.UpstreamNode->getCategory().c_str()))
 	{
-		bIsConnected = Delegate->ExecuteIfBound(Node, ParentShaderNode, InputChannelName);
+		bIsConnected = Delegate->ExecuteIfBound(Connect);
 	}
 
 	return bIsConnected;
@@ -182,31 +185,52 @@ void FMaterialXSurfaceShaderAbstract::ConnectNodeCategoryOutputToInput(const Mat
 {
 	if(mx::NodePtr UpstreamNode = Edge.getUpstreamElement()->asA<mx::Node>())
 	{
-		UInterchangeShaderNode* ParentShaderNode = ShaderNode;
+		// We need to connect the different descending nodes to all the outputs of a node
+		// At least one output connected to the root shader node
+		TArray<UInterchangeShaderNode*> ParentShaderNodeOutputs{ShaderNode};
 		FString InputChannelName = ParentInputName;
 
 		//Replace the input's name by the one used in UE
 		SetMatchingInputsNames(UpstreamNode);
 
+		FString OutputChannelName = DefaultOutput;
+
 		if(mx::ElementPtr DownstreamElement = Edge.getDownstreamElement())
 		{
 			if(mx::NodePtr DownstreamNode = DownstreamElement->asA<mx::Node>())
 			{
-				if(UInterchangeShaderNode** FoundNode = ShaderNodes.Find(GetAttributeParentName(DownstreamNode)))
+				if(mx::InputPtr ConnectedInput = Edge.getConnectingElement()->asA<mx::Input>())
 				{
-					ParentShaderNode = *FoundNode;
+					InputChannelName = GetInputName(ConnectedInput);
+					if(ConnectedInput->hasOutputString())
+					{
+						OutputChannelName = ConnectedInput->getOutputString().c_str();
+					}
 				}
 
-				if(mx::InputPtr ConnectedInput = Edge.getConnectingElement()->asA<mx::Input>())
+				std::vector<mx::OutputPtr> Outputs{ DownstreamNode->getOutputs()};
+				if(Outputs.empty())
 				{
-					InputChannelName = GetInputName(ConnectedInput);
+					Outputs = DownstreamNode->getNodeDef(mx::EMPTY_STRING, true)->getOutputs();
+				}
+
+				ParentShaderNodeOutputs.Empty();
+				for(std::size_t i = 0; i < Outputs.size(); ++i)
+				{
+					if(UInterchangeShaderNode** FoundNode = ShaderNodes.Find({ GetAttributeParentName(DownstreamNode), Outputs[i]->getName().c_str()}))
+					{
+						ParentShaderNodeOutputs.Emplace(*FoundNode);
+					}
 				}
 			}
 		}
 
-		if(!ConnectMatchingNodeOutputToInput(UpstreamNode, ParentShaderNode, InputChannelName))
+		for(UInterchangeShaderNode* ParentShaderNode: ParentShaderNodeOutputs)
 		{
-			UE_LOG(LogInterchangeImport, Warning, TEXT("<%s> is not supported yet"), ANSI_TO_TCHAR(UpstreamNode->getCategory().c_str()));
+			if(!ConnectMatchingNodeOutputToInput({ UpstreamNode, ParentShaderNode, InputChannelName, OutputChannelName }))
+			{
+				UE_LOG(LogInterchangeImport, Warning, TEXT("<%s> is not supported yet"), ANSI_TO_TCHAR(UpstreamNode->getCategory().c_str()));
+			}
 		}
 	}
 }
@@ -245,18 +269,18 @@ bool FMaterialXSurfaceShaderAbstract::ConnectNodeNameOutputToInput(MaterialX::In
 	return true;
 }
 
-void FMaterialXSurfaceShaderAbstract::ConnectConstantInputToOutput(MaterialX::NodePtr UpstreamNode, UInterchangeShaderNode* ParentShaderNode, const FString& InputChannelName)
+void FMaterialXSurfaceShaderAbstract::ConnectConstantInputToOutput(const FConnectNode& Connect)
 {
-	AddAttributeFromValueOrInterface(UpstreamNode->getInput("value"), InputChannelName, ParentShaderNode);
+	AddAttributeFromValueOrInterface(Connect.UpstreamNode->getInput("value"), Connect.InputChannelName, Connect.ParentShaderNode);
 }
 
-void FMaterialXSurfaceShaderAbstract::ConnectExtractInputToOutput(MaterialX::NodePtr UpstreamNode, UInterchangeShaderNode* ParentShaderNode, const FString& InputChannelName)
+void FMaterialXSurfaceShaderAbstract::ConnectExtractInputToOutput(const FConnectNode& Connect)
 {
 	using namespace UE::Interchange::Materials::Standard::Nodes;
 
-	UInterchangeShaderNode* MaskShaderNode = CreateShaderNode(UpstreamNode->getName().c_str(), Mask::Name.ToString());
+	UInterchangeShaderNode* MaskShaderNode = CreateShaderNode(Connect.UpstreamNode->getName().c_str(), Mask::Name.ToString());
 
-	if(mx::InputPtr InputIndex = UpstreamNode->getInput("index"))
+	if(mx::InputPtr InputIndex = Connect.UpstreamNode->getInput("index"))
 	{
 		const int32 Index = mx::fromValueString<int>(InputIndex->getValueString());
 		switch(Index)
@@ -271,95 +295,97 @@ void FMaterialXSurfaceShaderAbstract::ConnectExtractInputToOutput(MaterialX::Nod
 		}
 	}
 
-	UInterchangeShaderPortsAPI::ConnectDefaultOuputToInput(ParentShaderNode, InputChannelName, MaskShaderNode->GetUniqueID());
+	UInterchangeShaderPortsAPI::ConnectDefaultOuputToInput(Connect.ParentShaderNode, Connect.InputChannelName, MaskShaderNode->GetUniqueID());
 }
 
-void FMaterialXSurfaceShaderAbstract::ConnectDotInputToOutput(MaterialX::NodePtr UpstreamNode, UInterchangeShaderNode* ParentShaderNode, const FString& InputChannelName)
+void FMaterialXSurfaceShaderAbstract::ConnectDotInputToOutput(const FConnectNode& Connect)
 {
-	if(mx::InputPtr Input = UpstreamNode->getInput("in"))
+	if(mx::InputPtr Input = Connect.UpstreamNode->getInput("in"))
 	{
-		SetAttributeNewName(Input, TCHAR_TO_UTF8(*InputChannelName)); //let's take the parent node's input name
-		ShaderNodes.Add(UpstreamNode->getName().c_str(), ParentShaderNode);
+		SetAttributeNewName(Input, TCHAR_TO_UTF8(*Connect.InputChannelName)); //let's take the parent node's input name
+		ShaderNodes.Add({ Connect.UpstreamNode->getName().c_str(), Connect.OutputName}, Connect.ParentShaderNode);
 	}
 }
 
-void FMaterialXSurfaceShaderAbstract::ConnectTransformPositionInputToOutput(MaterialX::NodePtr UpstreamNode, UInterchangeShaderNode* ParentShaderNode, const FString& InputChannelName)
+void FMaterialXSurfaceShaderAbstract::ConnectTransformPositionInputToOutput(const FConnectNode& Connect)
 {
-	UInterchangeShaderNode* TransformNode = CreateShaderNode(UpstreamNode->getName().c_str(), TEXT("TransformPosition"));
-	AddAttributeFromValueOrInterface(UpstreamNode->getInput("in"), TEXT("Input"), TransformNode);
-	UInterchangeShaderPortsAPI::ConnectDefaultOuputToInput(ParentShaderNode, InputChannelName, TransformNode->GetUniqueID());
+	UInterchangeShaderNode* TransformNode = CreateShaderNode(Connect.UpstreamNode->getName().c_str(), TEXT("TransformPosition"));
+	AddAttributeFromValueOrInterface(Connect.UpstreamNode->getInput("in"), TEXT("Input"), TransformNode);
+	UInterchangeShaderPortsAPI::ConnectDefaultOuputToInput(Connect.ParentShaderNode, Connect.InputChannelName, TransformNode->GetUniqueID());
 }
 
-void FMaterialXSurfaceShaderAbstract::ConnectTransformVectorInputToOutput(MaterialX::NodePtr UpstreamNode, UInterchangeShaderNode* ParentShaderNode, const FString& InputChannelName)
+void FMaterialXSurfaceShaderAbstract::ConnectTransformVectorInputToOutput(const FConnectNode& Connect)
 {
-	UInterchangeShaderNode* TransformNode = CreateShaderNode(UpstreamNode->getName().c_str(), TEXT("Transform"));
-	AddAttributeFromValueOrInterface(UpstreamNode->getInput("in"), TEXT("Input"), TransformNode);
-	UInterchangeShaderPortsAPI::ConnectDefaultOuputToInput(ParentShaderNode, InputChannelName, TransformNode->GetUniqueID());
+	UInterchangeShaderNode* TransformNode = CreateShaderNode(Connect.UpstreamNode->getName().c_str(), TEXT("Transform"));
+	AddAttributeFromValueOrInterface(Connect.UpstreamNode->getInput("in"), TEXT("Input"), TransformNode);
+	UInterchangeShaderPortsAPI::ConnectDefaultOuputToInput(Connect.ParentShaderNode, Connect.InputChannelName, TransformNode->GetUniqueID());
 }
 
-void FMaterialXSurfaceShaderAbstract::ConnectRotate3DInputToOutput(mx::NodePtr UpstreamNode, UInterchangeShaderNode* ParentShaderNode, const FString& InputChannelName)
+void FMaterialXSurfaceShaderAbstract::ConnectRotate3DInputToOutput(const FConnectNode& Connect)
 {
-	UInterchangeShaderNode* Rotate3DNode = CreateShaderNode(UpstreamNode->getName().c_str(), TEXT("RotateAboutAxis"));
+	UInterchangeShaderNode* Rotate3DNode = CreateShaderNode(Connect.UpstreamNode->getName().c_str(), TEXT("RotateAboutAxis"));
 	Rotate3DNode->AddLinearColorAttribute(UInterchangeShaderPortsAPI::MakeInputValueKey(TEXT("PivotPoint")), FLinearColor(0.5, 0.5, 0));
 
-	AddAttributeFromValueOrInterface(UpstreamNode->getInput("in"), TEXT("Position"), Rotate3DNode);
-	AddAttributeFromValueOrInterface(UpstreamNode->getInput("axis"), TEXT("NormalizedRotationAxis"), Rotate3DNode);
+	AddAttributeFromValueOrInterface(Connect.UpstreamNode->getInput("in"), TEXT("Position"), Rotate3DNode);
+	AddAttributeFromValueOrInterface(Connect.UpstreamNode->getInput("axis"), TEXT("NormalizedRotationAxis"), Rotate3DNode);
 
 	// we create a Divide node to convert MaterialX angle in degrees to UE's angle which is a value between [0,1]
-	if(mx::InputPtr Input = UpstreamNode->getInput("amount"))
+	if(mx::InputPtr Input = Connect.UpstreamNode->getInput("amount"))
 	{
 		if(!AddAttributeFromValueOrInterface(Input, TEXT("RotationAngle"), Rotate3DNode))
 		{
 			//copy "in" input into "in1" and create "amount" input as "in2" with the value 360 because UE's angle is a value between [0,1]
-			mx::NodePtr NewDivideNode = CreateNode(UpstreamNode->getParent()->asA<mx::NodeGraph>(),
-												   UpstreamNode->getName().c_str(),
+			mx::NodePtr NewDivideNode = CreateNode(Connect.UpstreamNode->getParent()->asA<mx::NodeGraph>(),
+												   Connect.UpstreamNode->getName().c_str(),
 												   mx::Category::Divide,
 												   { {"in1", Input} }, // rename input to match <divide> input "in1"
 												   { {"in2", FAttributeValueArray{{"type", "float"}, {"value", "360"}}} });
 
+			//We need to set the type otherwise we won't be able to find a nodedef
+			NewDivideNode->setType(Input->getType());
 			// Input now points to the new node
 			Input->setNodeName(NewDivideNode->getName());
 		}
 	}
 
-	UInterchangeShaderPortsAPI::ConnectDefaultOuputToInput(ParentShaderNode, InputChannelName, Rotate3DNode->GetUniqueID());
+	UInterchangeShaderPortsAPI::ConnectDefaultOuputToInput(Connect.ParentShaderNode, Connect.InputChannelName, Rotate3DNode->GetUniqueID());
 }
 
-void FMaterialXSurfaceShaderAbstract::ConnectImageInputToOutput(MaterialX::NodePtr UpstreamNode, UInterchangeShaderNode* ParentShaderNode, const FString& InputChannelName)
+void FMaterialXSurfaceShaderAbstract::ConnectImageInputToOutput(const FConnectNode& Connect)
 {
 	using namespace UE::Interchange::Materials::Standard::Nodes;
 
-	if(UInterchangeTextureNode* TextureNode = CreateTextureNode<UInterchangeTexture2DNode>(UpstreamNode))
+	if(UInterchangeTextureNode* TextureNode = CreateTextureNode<UInterchangeTexture2DNode>(Connect.UpstreamNode))
 	{
 		//By default set the output of a texture to RGB
 		FString OutputChannel{ TEXT("RGB") };
 
-		if(UpstreamNode->getType() == mx::Type::Vector4 || UpstreamNode->getType() == mx::Type::Color4)
+		if(Connect.UpstreamNode->getType() == mx::Type::Vector4 || Connect.UpstreamNode->getType() == mx::Type::Color4)
 		{
 			OutputChannel = TEXT("RGBA");
 		}
-		else if(UpstreamNode->getType() == mx::Type::Float)
+		else if(Connect.UpstreamNode->getType() == mx::Type::Float)
 		{
 			OutputChannel = TEXT("R");
 		}
 
-		UInterchangeShaderNode* TextureShaderNode = CreateShaderNode(UpstreamNode->getName().c_str(), TextureSample::Name.ToString());
+		UInterchangeShaderNode* TextureShaderNode = CreateShaderNode(Connect.UpstreamNode->getName().c_str(), TextureSample::Name.ToString());
 		TextureShaderNode->AddStringAttribute(UInterchangeShaderPortsAPI::MakeInputValueKey(TextureSample::Inputs::Texture.ToString()), TextureNode->GetUniqueID());
-		UInterchangeShaderPortsAPI::ConnectOuputToInputByName(ParentShaderNode, InputChannelName, TextureShaderNode->GetUniqueID(), OutputChannel);
+		UInterchangeShaderPortsAPI::ConnectOuputToInputByName(Connect.ParentShaderNode, Connect.InputChannelName, TextureShaderNode->GetUniqueID(), OutputChannel);
 	}
 	else
 	{
-		AddAttributeFromValueOrInterface(UpstreamNode->getInput(mx::NodeGroup::Texture2D::Inputs::Default), InputChannelName, ParentShaderNode);
+		AddAttributeFromValueOrInterface(Connect.UpstreamNode->getInput(mx::NodeGroup::Texture2D::Inputs::Default), Connect.InputChannelName, Connect.ParentShaderNode);
 	}
 }
 
-void FMaterialXSurfaceShaderAbstract::ConnectConvertInputToOutput(MaterialX::NodePtr UpstreamNode, UInterchangeShaderNode* ParentShaderNode, const FString& InputChannelName)
+void FMaterialXSurfaceShaderAbstract::ConnectConvertInputToOutput(const FConnectNode& Connect)
 {
 	using namespace UE::Interchange::Materials::Standard::Nodes;
 
 	//In case of an upwards conversion, let's do an append, downwards a mask, otherwise leave as is
-	mx::InputPtr Input = UpstreamNode->getInput("in");
-	const std::string& NodeType = UpstreamNode->getType();
+	mx::InputPtr Input = Connect.UpstreamNode->getInput("in");
+	const std::string& NodeType = Connect.UpstreamNode->getType();
 	const std::string& InputType = Input->getType();
 
 	//Ensure that the types are supported
@@ -411,32 +437,32 @@ void FMaterialXSurfaceShaderAbstract::ConnectConvertInputToOutput(MaterialX::Nod
 	{
 		UInterchangeShaderNode* MaskShaderNode = nullptr;
 		if(NodeN == '3')
-		{
-			MaskShaderNode = CreateMaskShaderNode(0b1110, UpstreamNode->getName().c_str());
+		{			
+			MaskShaderNode = CreateMaskShaderNode(0b1110, Connect.UpstreamNode->getName().c_str());
 		}
 		else if(NodeN == '2')
 		{
-			MaskShaderNode = CreateMaskShaderNode(0b1100, UpstreamNode->getName().c_str());
+			MaskShaderNode = CreateMaskShaderNode(0b1100, Connect.UpstreamNode->getName().c_str());
 		}
 		else
 		{
-			MaskShaderNode = CreateMaskShaderNode(0b1000, UpstreamNode->getName().c_str());
+			MaskShaderNode = CreateMaskShaderNode(0b1000, Connect.UpstreamNode->getName().c_str());
 		}
 
 		AddAttributeFromValueOrInterface(Input, TEXT("Input"), MaskShaderNode);
 
-		UInterchangeShaderPortsAPI::ConnectDefaultOuputToInput(ParentShaderNode, InputChannelName, MaskShaderNode->GetUniqueID());
+		UInterchangeShaderPortsAPI::ConnectDefaultOuputToInput(Connect.ParentShaderNode, Connect.InputChannelName, MaskShaderNode->GetUniqueID());
 	}
 	else // Append 
 	{
 		//same as dot, just connect the next output to this parent input
-		SetAttributeNewName(Input, TCHAR_TO_UTF8(*InputChannelName)); //let's take the parent node's input name
-		ShaderNodes.Add(UpstreamNode->getName().c_str(), ParentShaderNode);
+		SetAttributeNewName(Input, TCHAR_TO_UTF8(*Connect.InputChannelName)); //let's take the parent node's input name
+		ShaderNodes.Add({ Connect.UpstreamNode->getName().c_str(), Connect.OutputName }, Connect.ParentShaderNode);
 
 		//No need to create a node, since the input and the node have the same channels, we just check if there's a value
 		if(NodeN == InputN || (NodeN < '2' && InputN < '2'))
 		{
-			AddAttributeFromValueOrInterface(Input, InputChannelName, ParentShaderNode);
+			AddAttributeFromValueOrInterface(Input, Connect.InputChannelName, Connect.ParentShaderNode);
 			return;
 		}
 		std::string Category;
@@ -476,39 +502,47 @@ void FMaterialXSurfaceShaderAbstract::ConnectConvertInputToOutput(MaterialX::Nod
 		}
 
 		//copy "in" input into "in1" and create "amount" input as "in2" with the value 360 because UE's angle is a value between [0,1]
-		mx::NodePtr CombineNode = CreateNode(UpstreamNode->getParent()->asA<mx::NodeGraph>(),
-											 UpstreamNode->getName().c_str(),
+		mx::NodePtr CombineNode = CreateNode(Connect.UpstreamNode->getParent()->asA<mx::NodeGraph>(),
+											 Connect.UpstreamNode->getName().c_str(),
 											 Category.c_str(),
 											 InputsToCopy, // rename input to match <divide> input "in1"
 											 InputsToCreate);
 
+		//We set the type and an output to find it during connect node resolution phase
+		//Because the case Combine2 vector2->vector3 has no nodedef in the standard library
+		CombineNode->setType(Connect.UpstreamNode->getType());
+		if(!CombineNode->getOutput("out"))
+		{
+			CombineNode->addOutput("out", Connect.UpstreamNode->getType());
+		}
+
 		// Input now points to the new node
 		Input->setNodeName(CombineNode->getName());
 	}
 }
 
-void FMaterialXSurfaceShaderAbstract::ConnectIfGreaterInputToOutput(MaterialX::NodePtr UpstreamNode, UInterchangeShaderNode* ParentShaderNode, const FString& InputChannelName)
+void FMaterialXSurfaceShaderAbstract::ConnectIfGreaterInputToOutput(const FConnectNode& Connect)
 {
-	UInterchangeShaderNode* NodeIf = CreateShaderNode(UpstreamNode->getName().c_str(), TEXT("If"));
-	UInterchangeShaderPortsAPI::ConnectDefaultOuputToInput(ParentShaderNode, InputChannelName, NodeIf->GetUniqueID());
+	UInterchangeShaderNode* NodeIf = CreateShaderNode(Connect.UpstreamNode->getName().c_str(), TEXT("If"));
+	UInterchangeShaderPortsAPI::ConnectDefaultOuputToInput(Connect.ParentShaderNode, Connect.InputChannelName, NodeIf->GetUniqueID());
 
-	if(mx::InputPtr Input = UpstreamNode->getInput("value1"))
+	if(mx::InputPtr Input = Connect.UpstreamNode->getInput("value1"))
 	{
 		AddAttributeFromValueOrInterface(Input, GetInputName(Input), NodeIf);
 	}
 
-	if(mx::InputPtr Input = UpstreamNode->getInput("value2"))
+	if(mx::InputPtr Input = Connect.UpstreamNode->getInput("value2"))
 	{
 		AddAttributeFromValueOrInterface(Input, GetInputName(Input), NodeIf);
 	}
 
-	if(mx::InputPtr Input = UpstreamNode->getInput("in1"))
+	if(mx::InputPtr Input = Connect.UpstreamNode->getInput("in1"))
 	{
 		AddAttributeFromValueOrInterface(Input, GetInputName(Input), NodeIf);
 	}
 
 	//In that case we also need to add an attribute to AEqualsB
-	mx::InputPtr Input = UpstreamNode->getInput("in2");
+	mx::InputPtr Input = Connect.UpstreamNode->getInput("in2");
 	if(Input)
 	{
 		bool bHasValue = AddAttributeFromValueOrInterface(Input, GetInputName(Input), NodeIf);
@@ -517,35 +551,35 @@ void FMaterialXSurfaceShaderAbstract::ConnectIfGreaterInputToOutput(MaterialX::N
 		if(bHasValue)
 		{
 			//Let's add a new input that is a copy of in2 to connect it to the equal input
-			mx::InputPtr Input3 = UpstreamNode->addInput("in3");
+			mx::InputPtr Input3 = Connect.UpstreamNode->addInput("in3");
 			Input3->copyContentFrom(Input);
 			SetAttributeNewName(Input3, "AEqualsB");
 		}
 	}
 }
 
-void FMaterialXSurfaceShaderAbstract::ConnectIfGreaterEqInputToOutput(MaterialX::NodePtr UpstreamNode, UInterchangeShaderNode* ParentShaderNode, const FString& InputChannelName)
+void FMaterialXSurfaceShaderAbstract::ConnectIfGreaterEqInputToOutput(const FConnectNode& Connect)
 {
-	UInterchangeShaderNode* NodeIf = CreateShaderNode(UpstreamNode->getName().c_str(), TEXT("If"));
-	UInterchangeShaderPortsAPI::ConnectDefaultOuputToInput(ParentShaderNode, InputChannelName, NodeIf->GetUniqueID());
+	UInterchangeShaderNode* NodeIf = CreateShaderNode(Connect.UpstreamNode->getName().c_str(), TEXT("If"));
+	UInterchangeShaderPortsAPI::ConnectDefaultOuputToInput(Connect.ParentShaderNode, Connect.InputChannelName, NodeIf->GetUniqueID());
 
-	if(mx::InputPtr Input = UpstreamNode->getInput("value1"))
+	if(mx::InputPtr Input = Connect.UpstreamNode->getInput("value1"))
 	{
 		AddAttributeFromValueOrInterface(Input, GetInputName(Input), NodeIf);
 	}
 
-	if(mx::InputPtr Input = UpstreamNode->getInput("value2"))
+	if(mx::InputPtr Input = Connect.UpstreamNode->getInput("value2"))
 	{
 		AddAttributeFromValueOrInterface(Input, GetInputName(Input), NodeIf);
 	}
 
-	if(mx::InputPtr Input = UpstreamNode->getInput("in2"))
+	if(mx::InputPtr Input = Connect.UpstreamNode->getInput("in2"))
 	{
 		AddAttributeFromValueOrInterface(Input, GetInputName(Input), NodeIf);
 	}
 
 	//In that case we also need to add an attribute to AEqualsB
-	mx::InputPtr Input = UpstreamNode->getInput("in1");
+	mx::InputPtr Input = Connect.UpstreamNode->getInput("in1");
 	if(Input)
 	{
 		bool bHasValue = AddAttributeFromValueOrInterface(Input, GetInputName(Input), NodeIf);
@@ -554,35 +588,35 @@ void FMaterialXSurfaceShaderAbstract::ConnectIfGreaterEqInputToOutput(MaterialX:
 		if(bHasValue)
 		{
 			//Let's add a new input that is a copy of in2 to connect it to the equal input
-			mx::InputPtr Input3 = UpstreamNode->addInput("in3");
+			mx::InputPtr Input3 = Connect.UpstreamNode->addInput("in3");
 			Input3->copyContentFrom(Input);
 			SetAttributeNewName(Input3, "AEqualsB");
 		}
 	}
 }
 
-void FMaterialXSurfaceShaderAbstract::ConnectIfEqualInputToOutput(MaterialX::NodePtr UpstreamNode, UInterchangeShaderNode* ParentShaderNode, const FString& InputChannelName)
+void FMaterialXSurfaceShaderAbstract::ConnectIfEqualInputToOutput(const FConnectNode& Connect)
 {
-	UInterchangeShaderNode* NodeIf = CreateShaderNode(UpstreamNode->getName().c_str(), TEXT("If"));
-	UInterchangeShaderPortsAPI::ConnectDefaultOuputToInput(ParentShaderNode, InputChannelName, NodeIf->GetUniqueID());
+	UInterchangeShaderNode* NodeIf = CreateShaderNode(Connect.UpstreamNode->getName().c_str(), TEXT("If"));
+	UInterchangeShaderPortsAPI::ConnectDefaultOuputToInput(Connect.ParentShaderNode, Connect.InputChannelName, NodeIf->GetUniqueID());
 
-	if(mx::InputPtr Input = UpstreamNode->getInput("value1"); Input && Input->hasValue())
+	if(mx::InputPtr Input = Connect.UpstreamNode->getInput("value1"); Input && Input->hasValue())
 	{
 		AddAttributeFromValueOrInterface(Input, GetInputName(Input), NodeIf);
 	}
 
-	if(mx::InputPtr Input = UpstreamNode->getInput("value2"); Input && Input->hasValue())
+	if(mx::InputPtr Input = Connect.UpstreamNode->getInput("value2"); Input && Input->hasValue())
 	{
 		AddAttributeFromValueOrInterface(Input, GetInputName(Input), NodeIf);
 	}
 
-	if(mx::InputPtr Input = UpstreamNode->getInput("in1"); Input && Input->hasValue())
+	if(mx::InputPtr Input = Connect.UpstreamNode->getInput("in1"); Input && Input->hasValue())
 	{
 		AddAttributeFromValueOrInterface(Input, GetInputName(Input), NodeIf);
 	}
 
 	//In that case we also need to add an attribute to AGreaterThanB
-	mx::InputPtr Input = UpstreamNode->getInput("in2");
+	mx::InputPtr Input = Connect.UpstreamNode->getInput("in2");
 	if(Input)
 	{
 		bool bHasValue = AddAttributeFromValueOrInterface(Input, GetInputName(Input), NodeIf);
@@ -591,26 +625,26 @@ void FMaterialXSurfaceShaderAbstract::ConnectIfEqualInputToOutput(MaterialX::Nod
 		if(bHasValue)
 		{
 			//Let's add a new input that is a copy of in2 to connect it to the equal input
-			mx::InputPtr Input3 = UpstreamNode->addInput("in3");
+			mx::InputPtr Input3 = Connect.UpstreamNode->addInput("in3");
 			Input3->copyContentFrom(Input);
 			SetAttributeNewName(Input3, "AGreaterThanB");
 		}
 	}
 }
 
-void FMaterialXSurfaceShaderAbstract::ConnectOutsideInputToOutput(MaterialX::NodePtr UpstreamNode, UInterchangeShaderNode* ParentShaderNode, const FString& InputChannelName)
+void FMaterialXSurfaceShaderAbstract::ConnectOutsideInputToOutput(const FConnectNode& Connect)
 {
 	//in * (1 - mask)
-	UInterchangeShaderNode* NodeMultiply = CreateShaderNode(UpstreamNode->getName().c_str(), TEXT("Multiply"));
-	AddAttributeFromValueOrInterface(UpstreamNode->getInput("in"), TEXT("A"), NodeMultiply);
-	UInterchangeShaderPortsAPI::ConnectDefaultOuputToInput(ParentShaderNode, InputChannelName, NodeMultiply->GetUniqueID());
+	UInterchangeShaderNode* NodeMultiply = CreateShaderNode(Connect.UpstreamNode->getName().c_str(), TEXT("Multiply"));
+	AddAttributeFromValueOrInterface(Connect.UpstreamNode->getInput("in"), TEXT("A"), NodeMultiply);
+	UInterchangeShaderPortsAPI::ConnectDefaultOuputToInput(Connect.ParentShaderNode, Connect.InputChannelName, NodeMultiply->GetUniqueID());
 
-	UInterchangeShaderNode* NodeOneMinus = CreateShaderNode(UpstreamNode->getName().c_str() + FString(TEXT("_OneMinus")), TEXT("OneMinus"));
-	AddAttributeFromValueOrInterface(UpstreamNode->getInput("mask"), TEXT("Input"), NodeOneMinus);
+	UInterchangeShaderNode* NodeOneMinus = CreateShaderNode(Connect.UpstreamNode->getName().c_str() + FString(TEXT("_OneMinus")), TEXT("OneMinus"));
+	AddAttributeFromValueOrInterface(Connect.UpstreamNode->getInput("mask"), TEXT("Input"), NodeOneMinus);
 	UInterchangeShaderPortsAPI::ConnectDefaultOuputToInput(NodeMultiply, TEXT("B"), NodeOneMinus->GetUniqueID());
 }
 
-void FMaterialXSurfaceShaderAbstract::ConnectPositionInputToOutput(MaterialX::NodePtr UpstreamNode, UInterchangeShaderNode* ParentShaderNode, const FString& InputChannelName)
+void FMaterialXSurfaceShaderAbstract::ConnectPositionInputToOutput(const FConnectNode& Connect)
 {
 	// MaterialX defines the space as: object, model, world
 	// model: The local coordinate space of the geometry, before any local deformations or global transforms have been applied.
@@ -619,19 +653,19 @@ void FMaterialXSurfaceShaderAbstract::ConnectPositionInputToOutput(MaterialX::No
 
 	// For the moment we don't have the distinction between model/object, so let's just create an UMaterialExpressionWorldPosition
 	// In case of model/object we need to add a TransformPoint from world to local space
-	UInterchangeShaderNode* PositionNode = CreateShaderNode(UpstreamNode->getName().c_str() + FString(TEXT("_Position")), TEXT("WorldPosition"));
+	UInterchangeShaderNode* PositionNode = CreateShaderNode(Connect.UpstreamNode->getName().c_str() + FString(TEXT("_Position")), TEXT("WorldPosition"));
 
 	// In case of the position node, it seems that the unit is different, we assume for now a conversion from mm -> m, even if UE by default is cm
 	// See standard_surface_marble_solid file, especially on the fractal3d node
-	UInterchangeShaderNode* UnitNode = CreateShaderNode(UpstreamNode->getName().c_str(), TEXT("Multiply"));
+	UInterchangeShaderNode* UnitNode = CreateShaderNode(Connect.UpstreamNode->getName().c_str(), TEXT("Multiply"));
 	UnitNode->AddFloatAttribute(UInterchangeShaderPortsAPI::MakeInputValueKey(TEXT("B")), 0.001f);
 	UInterchangeShaderPortsAPI::ConnectDefaultOuputToInput(UnitNode, TEXT("A"), PositionNode->GetUniqueID());
 
-	UInterchangeShaderNode* NodeToConnectTo = ParentShaderNode;
-	FString InputToConnectTo = InputChannelName;
+	UInterchangeShaderNode* NodeToConnectTo = Connect.ParentShaderNode;
+	FString InputToConnectTo = Connect.InputChannelName;
 
 	std::string Space = "world";
-	mx::InputPtr InputSpace = UpstreamNode->getInput("space");
+	mx::InputPtr InputSpace = Connect.UpstreamNode->getInput("space");
 
 	if(InputSpace)
 	{
@@ -643,18 +677,18 @@ void FMaterialXSurfaceShaderAbstract::ConnectPositionInputToOutput(MaterialX::No
 	{
 		using namespace UE::Interchange::Materials::Standard::Nodes;
 
-		UInterchangeShaderNode* TransformNode = CreateShaderNode(UpstreamNode->getName().c_str() + FString(TEXT("_Transform")), TransformPosition::Name.ToString());
+		UInterchangeShaderNode* TransformNode = CreateShaderNode(Connect.UpstreamNode->getName().c_str() + FString(TEXT("_Transform")), TransformPosition::Name.ToString());
 		NodeToConnectTo = TransformNode;
 		InputToConnectTo = TransformPosition::Inputs::Input.ToString();
 		TransformNode->AddInt32Attribute(TransformPosition::Attributes::TransformSourceType.ToString(), EMaterialPositionTransformSource::TRANSFORMPOSSOURCE_World);
 		TransformNode->AddInt32Attribute(TransformPosition::Attributes::TransformType.ToString(), EMaterialPositionTransformSource::TRANSFORMPOSSOURCE_Local);
-		UInterchangeShaderPortsAPI::ConnectDefaultOuputToInput(ParentShaderNode, InputChannelName, TransformNode->GetUniqueID());
+		UInterchangeShaderPortsAPI::ConnectDefaultOuputToInput(Connect.ParentShaderNode, Connect.InputChannelName, TransformNode->GetUniqueID());
 	}
 
 	UInterchangeShaderPortsAPI::ConnectDefaultOuputToInput(NodeToConnectTo, InputToConnectTo, UnitNode->GetUniqueID());
 }
 
-void FMaterialXSurfaceShaderAbstract::ConnectNormalInputToOutput(MaterialX::NodePtr UpstreamNode, UInterchangeShaderNode* ParentShaderNode, const FString& InputChannelName)
+void FMaterialXSurfaceShaderAbstract::ConnectNormalInputToOutput(const FConnectNode& Connect)
 {
 	// MaterialX defines the space as: object, model, world
 	// model: The local coordinate space of the geometry, before any local deformations or global transforms have been applied.
@@ -663,12 +697,12 @@ void FMaterialXSurfaceShaderAbstract::ConnectNormalInputToOutput(MaterialX::Node
 
 	// For the moment we don't have the distinction between model/object, so let's just create an UMaterialExpressionVertexNormalWS
 	// In case of model/object we need to add a TransformVector from world to local space
-	UInterchangeShaderNode* NormalNode = CreateShaderNode(UpstreamNode->getName().c_str(), TEXT("VertexNormalWS"));
-	UInterchangeShaderNode* NodeToConnectTo = ParentShaderNode;
-	FString InputToConnectTo = InputChannelName;
+	UInterchangeShaderNode* NormalNode = CreateShaderNode(Connect.UpstreamNode->getName().c_str(), TEXT("VertexNormalWS"));
+	UInterchangeShaderNode* NodeToConnectTo = Connect.ParentShaderNode;
+	FString InputToConnectTo = Connect.InputChannelName;
 
 	std::string Space = "world";
-	mx::InputPtr InputSpace = UpstreamNode->getInput("space");
+	mx::InputPtr InputSpace = Connect.UpstreamNode->getInput("space");
 
 	if(InputSpace)
 	{
@@ -680,18 +714,18 @@ void FMaterialXSurfaceShaderAbstract::ConnectNormalInputToOutput(MaterialX::Node
 	{
 		using namespace UE::Interchange::Materials::Standard::Nodes;
 
-		UInterchangeShaderNode* TransformNode = CreateShaderNode(UpstreamNode->getName().c_str() + FString(TEXT("_Transform")), TransformVector::Name.ToString());
+		UInterchangeShaderNode* TransformNode = CreateShaderNode(Connect.UpstreamNode->getName().c_str() + FString(TEXT("_Transform")), TransformVector::Name.ToString());
 		NodeToConnectTo = TransformNode;
 		InputToConnectTo = TransformVector::Inputs::Input.ToString();
 		TransformNode->AddInt32Attribute(TransformVector::Attributes::TransformSourceType.ToString(), EMaterialVectorCoordTransformSource::TRANSFORMSOURCE_World);
 		TransformNode->AddInt32Attribute(TransformVector::Attributes::TransformType.ToString(), EMaterialVectorCoordTransform::TRANSFORM_Local);
-		UInterchangeShaderPortsAPI::ConnectDefaultOuputToInput(ParentShaderNode, InputChannelName, TransformNode->GetUniqueID());
+		UInterchangeShaderPortsAPI::ConnectDefaultOuputToInput(Connect.ParentShaderNode, Connect.InputChannelName, TransformNode->GetUniqueID());
 	}
 
 	UInterchangeShaderPortsAPI::ConnectDefaultOuputToInput(NodeToConnectTo, InputToConnectTo, NormalNode->GetUniqueID());
 }
 
-void FMaterialXSurfaceShaderAbstract::ConnectTangentInputToOutput(MaterialX::NodePtr UpstreamNode, UInterchangeShaderNode* ParentShaderNode, const FString& InputChannelName)
+void FMaterialXSurfaceShaderAbstract::ConnectTangentInputToOutput(const FConnectNode& Connect)
 {
 	// MaterialX defines the space as: object, model, world
 	// model: The local coordinate space of the geometry, before any local deformations or global transforms have been applied.
@@ -700,12 +734,12 @@ void FMaterialXSurfaceShaderAbstract::ConnectTangentInputToOutput(MaterialX::Nod
 
 	// For the moment we don't have the distinction between model/object, so let's just create an UMaterialExpressionVertexTangentWS
 	// In case of model/object we need to add a TransformVector from world to local space
-	UInterchangeShaderNode* TangentNode = CreateShaderNode(UpstreamNode->getName().c_str(), TEXT("VertexTangentWS"));
-	UInterchangeShaderNode* NodeToConnectTo = ParentShaderNode;
-	FString InputToConnectTo = InputChannelName;
+	UInterchangeShaderNode* TangentNode = CreateShaderNode(Connect.UpstreamNode->getName().c_str(), TEXT("VertexTangentWS"));
+	UInterchangeShaderNode* NodeToConnectTo = Connect.ParentShaderNode;
+	FString InputToConnectTo = Connect.InputChannelName;
 
 	std::string Space = "world";
-	mx::InputPtr InputSpace = UpstreamNode->getInput("space");
+	mx::InputPtr InputSpace = Connect.UpstreamNode->getInput("space");
 
 	if(InputSpace)
 	{
@@ -717,18 +751,18 @@ void FMaterialXSurfaceShaderAbstract::ConnectTangentInputToOutput(MaterialX::Nod
 	{
 		using namespace UE::Interchange::Materials::Standard::Nodes;
 
-		UInterchangeShaderNode* TransformNode = CreateShaderNode(UpstreamNode->getName().c_str() + FString(TEXT("_Transform")), TransformVector::Name.ToString());
+		UInterchangeShaderNode* TransformNode = CreateShaderNode(Connect.UpstreamNode->getName().c_str() + FString(TEXT("_Transform")), TransformVector::Name.ToString());
 		NodeToConnectTo = TransformNode;
 		InputToConnectTo = TransformVector::Inputs::Input.ToString();
 		TransformNode->AddInt32Attribute(TransformVector::Attributes::TransformSourceType.ToString(), EMaterialVectorCoordTransformSource::TRANSFORMSOURCE_World);
 		TransformNode->AddInt32Attribute(TransformVector::Attributes::TransformType.ToString(), EMaterialVectorCoordTransform::TRANSFORM_Local);
-		UInterchangeShaderPortsAPI::ConnectDefaultOuputToInput(ParentShaderNode, InputChannelName, TransformNode->GetUniqueID());
+		UInterchangeShaderPortsAPI::ConnectDefaultOuputToInput(Connect.ParentShaderNode, Connect.InputChannelName, TransformNode->GetUniqueID());
 	}
 
 	UInterchangeShaderPortsAPI::ConnectDefaultOuputToInput(NodeToConnectTo, InputToConnectTo, TangentNode->GetUniqueID());
 }
 
-void FMaterialXSurfaceShaderAbstract::ConnectBitangentInputToOutput(MaterialX::NodePtr UpstreamNode, UInterchangeShaderNode* ParentShaderNode, const FString& InputChannelName)
+void FMaterialXSurfaceShaderAbstract::ConnectBitangentInputToOutput(const FConnectNode& Connect)
 {
 	// MaterialX defines the space as: object, model, world
 	// model: The local coordinate space of the geometry, before any local deformations or global transforms have been applied.
@@ -737,19 +771,19 @@ void FMaterialXSurfaceShaderAbstract::ConnectBitangentInputToOutput(MaterialX::N
 
 	// For the moment we don't have the distinction between model/object, so let's just do the cross product between the normal and the tangent
 	// In case of model/object we need to add a TransformVector from world to local space
-	UInterchangeShaderNode* NormalNode = CreateShaderNode(UpstreamNode->getName().c_str() + FString(TEXT("_Normal")), TEXT("VertexNormalWS"));
-	UInterchangeShaderNode* TangentNode = CreateShaderNode(UpstreamNode->getName().c_str() + FString(TEXT("_Tangent")), TEXT("VertexTangentWS"));
-	UInterchangeShaderNode* BitangentNode = CreateShaderNode(UpstreamNode->getName().c_str(), TEXT("CrossProduct"));
+	UInterchangeShaderNode* NormalNode = CreateShaderNode(Connect.UpstreamNode->getName().c_str() + FString(TEXT("_Normal")), TEXT("VertexNormalWS"));
+	UInterchangeShaderNode* TangentNode = CreateShaderNode(Connect.UpstreamNode->getName().c_str() + FString(TEXT("_Tangent")), TEXT("VertexTangentWS"));
+	UInterchangeShaderNode* BitangentNode = CreateShaderNode(Connect.UpstreamNode->getName().c_str(), TEXT("CrossProduct"));
 
 	UInterchangeShaderPortsAPI::ConnectDefaultOuputToInput(BitangentNode, TEXT("A"), NormalNode->GetUniqueID());
 	UInterchangeShaderPortsAPI::ConnectDefaultOuputToInput(BitangentNode, TEXT("B"), TangentNode->GetUniqueID());
 
 
-	UInterchangeShaderNode* NodeToConnectTo = ParentShaderNode;
-	FString InputToConnectTo = InputChannelName;
+	UInterchangeShaderNode* NodeToConnectTo = Connect.ParentShaderNode;
+	FString InputToConnectTo = Connect.InputChannelName;
 
 	std::string Space = "world";
-	mx::InputPtr InputSpace = UpstreamNode->getInput("space");
+	mx::InputPtr InputSpace = Connect.UpstreamNode->getInput("space");
 
 	if(InputSpace)
 	{
@@ -761,30 +795,30 @@ void FMaterialXSurfaceShaderAbstract::ConnectBitangentInputToOutput(MaterialX::N
 	{
 		using namespace UE::Interchange::Materials::Standard::Nodes;
 
-		UInterchangeShaderNode* TransformNode = CreateShaderNode(UpstreamNode->getName().c_str() + FString(TEXT("_Transform")), TransformVector::Name.ToString());
+		UInterchangeShaderNode* TransformNode = CreateShaderNode(Connect.UpstreamNode->getName().c_str() + FString(TEXT("_Transform")), TransformVector::Name.ToString());
 		NodeToConnectTo = TransformNode;
 		InputToConnectTo = TransformVector::Inputs::Input.ToString();
 		TransformNode->AddInt32Attribute(TransformVector::Attributes::TransformSourceType.ToString(), EMaterialVectorCoordTransformSource::TRANSFORMSOURCE_World);
 		TransformNode->AddInt32Attribute(TransformVector::Attributes::TransformType.ToString(), EMaterialVectorCoordTransform::TRANSFORM_Local);
-		UInterchangeShaderPortsAPI::ConnectDefaultOuputToInput(ParentShaderNode, InputChannelName, TransformNode->GetUniqueID());
+		UInterchangeShaderPortsAPI::ConnectDefaultOuputToInput(Connect.ParentShaderNode, Connect.InputChannelName, TransformNode->GetUniqueID());
 	}
 
 	UInterchangeShaderPortsAPI::ConnectDefaultOuputToInput(NodeToConnectTo, InputToConnectTo, BitangentNode->GetUniqueID());
 }
 
-void FMaterialXSurfaceShaderAbstract::ConnectTimeInputToOutput(MaterialX::NodePtr UpstreamNode, UInterchangeShaderNode* ParentShaderNode, const FString& InputChannelName)
+void FMaterialXSurfaceShaderAbstract::ConnectTimeInputToOutput(const FConnectNode& Connect)
 {
 	using namespace UE::Interchange::Materials::Standard::Nodes;
-	UInterchangeShaderNode* TimeNode = CreateShaderNode(UpstreamNode->getName().c_str(), TEXT("Time"));
+	UInterchangeShaderNode* TimeNode = CreateShaderNode(Connect.UpstreamNode->getName().c_str(), TEXT("Time"));
 	TimeNode->AddBooleanAttribute(Time::Attributes::OverridePeriod.ToString(), true);
 
 	float FPS;
-	mx::InputPtr Input = UpstreamNode->getInput("fps");
+	mx::InputPtr Input = Connect.UpstreamNode->getInput("fps");
 
 	//Take the default value from the node definition
 	if(!Input)
 	{
-		Input = UpstreamNode->getNodeDef()->getInput("fps");
+		Input = Connect.UpstreamNode->getNodeDef(mx::EMPTY_STRING, true)->getInput("fps");
 	}
 
 	FPS = mx::fromValueString<float>(Input->getValueString());
@@ -792,15 +826,15 @@ void FMaterialXSurfaceShaderAbstract::ConnectTimeInputToOutput(MaterialX::NodePt
 	//UE is a period
 	TimeNode->AddFloatAttribute(Time::Attributes::Period.ToString(), 1.f / FPS);
 
-	UInterchangeShaderPortsAPI::ConnectDefaultOuputToInput(ParentShaderNode, InputChannelName, TimeNode->GetUniqueID());
+	UInterchangeShaderPortsAPI::ConnectDefaultOuputToInput(Connect.ParentShaderNode, Connect.InputChannelName, TimeNode->GetUniqueID());
 }
 
-void FMaterialXSurfaceShaderAbstract::ConnectNoise3DInputToOutput(MaterialX::NodePtr UpstreamNode, UInterchangeShaderNode* ParentShaderNode, const FString& InputChannelName)
+void FMaterialXSurfaceShaderAbstract::ConnectNoise3DInputToOutput(const FConnectNode& Connect)
 {
 	using namespace UE::Interchange::Materials::Standard::Nodes;
 
 	// MaterialX defines the Noise3d as Perlin Noise which is multiplied by the Amplitude then Added to Pivot
-	UInterchangeShaderNode* NoiseNode = CreateShaderNode(UpstreamNode->getName().c_str(), Noise::Name.ToString());
+	UInterchangeShaderNode* NoiseNode = CreateShaderNode(Connect.UpstreamNode->getName().c_str(), Noise::Name.ToString());
 	NoiseNode->AddInt32Attribute(Noise::Attributes::Function.ToString(), ENoiseFunction::NOISEFUNCTION_GradientTex);
 	NoiseNode->AddBooleanAttribute(Noise::Attributes::Turbulence.ToString(), false);
 	NoiseNode->AddFloatAttribute(Noise::Attributes::OutputMin.ToString(), 0);
@@ -815,7 +849,7 @@ void FMaterialXSurfaceShaderAbstract::ConnectNoise3DInputToOutput(MaterialX::Nod
 			return nullptr;
 		}
 
-		const FString ShaderNodeName = UpstreamNode->getName().c_str() + FString{ TEXT("_") } + ShaderType;
+		const FString ShaderNodeName = Connect.UpstreamNode->getName().c_str() + FString{ TEXT("_") } + ShaderType;
 		UInterchangeShaderNode* ShaderNode = CreateShaderNode(ShaderNodeName, ShaderType);
 
 		UInterchangeShaderPortsAPI::ConnectDefaultOuputToInput(ShaderNode, TEXT("A"), NodeToConnectTo->GetUniqueID());
@@ -823,120 +857,120 @@ void FMaterialXSurfaceShaderAbstract::ConnectNoise3DInputToOutput(MaterialX::Nod
 		// Connect the amplitude node to the shader node not the noise
 		// it will be handle during the upstream-downstream connection phase
 		// The index is here for the unicity of the attribute
-		UpstreamNode->setAttribute(mx::Attributes::ParentName + std::to_string(IndexAttrib), TCHAR_TO_ANSI(*ShaderNodeName));
+		Connect.UpstreamNode->setAttribute(mx::Attributes::ParentName + std::to_string(IndexAttrib), TCHAR_TO_ANSI(*ShaderNodeName));
 		AddAttributeFromValueOrInterface(Input, TEXT("B"), ShaderNode);
 
 		return ShaderNode;
 	};
 
-	if(UInterchangeShaderNode* MultiplyNode = ConnectNodeToInput(UpstreamNode->getInput("amplitude"), NoiseNode, TEXT("Multiply"), 0))
+	if(UInterchangeShaderNode* MultiplyNode = ConnectNodeToInput(Connect.UpstreamNode->getInput("amplitude"), NoiseNode, TEXT("Multiply"), 0))
 	{
 		NodeToConnect = MultiplyNode;
 	}
 
 
-	if(UInterchangeShaderNode* AddNode = ConnectNodeToInput(UpstreamNode->getInput("pivot"), NodeToConnect, TEXT("Add"), 1))
+	if(UInterchangeShaderNode* AddNode = ConnectNodeToInput(Connect.UpstreamNode->getInput("pivot"), NodeToConnect, TEXT("Add"), 1))
 	{
 		NodeToConnect = AddNode;
 	}
 
-	UInterchangeShaderPortsAPI::ConnectDefaultOuputToInput(ParentShaderNode, InputChannelName, NodeToConnect->GetUniqueID());
+	UInterchangeShaderPortsAPI::ConnectDefaultOuputToInput(Connect.ParentShaderNode, Connect.InputChannelName, NodeToConnect->GetUniqueID());
 }
 
-void FMaterialXSurfaceShaderAbstract::ConnectCellNoise3DInputToOutput(MaterialX::NodePtr UpstreamNode, UInterchangeShaderNode* ParentShaderNode, const FString& InputChannelName)
+void FMaterialXSurfaceShaderAbstract::ConnectCellNoise3DInputToOutput(const FConnectNode& Connect)
 {
 	using namespace UE::Interchange::Materials::Standard::Nodes;
 
 	// Let's use a vector noise for this one, the only one that is close to MaterialX implementation
-	UInterchangeShaderNode* NoiseNode = CreateShaderNode(UpstreamNode->getName().c_str(), VectorNoise::Name.ToString());
+	UInterchangeShaderNode* NoiseNode = CreateShaderNode(Connect.UpstreamNode->getName().c_str(), VectorNoise::Name.ToString());
 	NoiseNode->AddInt32Attribute(VectorNoise::Attributes::Function.ToString(), EVectorNoiseFunction::VNF_CellnoiseALU);
 
-	UInterchangeShaderPortsAPI::ConnectDefaultOuputToInput(ParentShaderNode, InputChannelName, NoiseNode->GetUniqueID());
+	UInterchangeShaderPortsAPI::ConnectDefaultOuputToInput(Connect.ParentShaderNode, Connect.InputChannelName, NoiseNode->GetUniqueID());
 }
 
-void FMaterialXSurfaceShaderAbstract::ConnectWorleyNoise3DInputToOutput(MaterialX::NodePtr UpstreamNode, UInterchangeShaderNode* ParentShaderNode, const FString& InputChannelName)
+void FMaterialXSurfaceShaderAbstract::ConnectWorleyNoise3DInputToOutput(const FConnectNode& Connect)
 {
 	using namespace UE::Interchange::Materials::Standard::Nodes;
 
 	//Also called Voronoi, the implementation is a bit different in UE, especially we don't have access to the jitter
-	UInterchangeShaderNode* NoiseNode = CreateShaderNode(UpstreamNode->getName().c_str(), Noise::Name.ToString());
+	UInterchangeShaderNode* NoiseNode = CreateShaderNode(Connect.UpstreamNode->getName().c_str(), Noise::Name.ToString());
 	NoiseNode->AddInt32Attribute(Noise::Attributes::Function.ToString(), ENoiseFunction::NOISEFUNCTION_VoronoiALU);
 
-	UInterchangeShaderPortsAPI::ConnectDefaultOuputToInput(ParentShaderNode, InputChannelName, NoiseNode->GetUniqueID());
+	UInterchangeShaderPortsAPI::ConnectDefaultOuputToInput(Connect.ParentShaderNode, Connect.InputChannelName, NoiseNode->GetUniqueID());
 }
 
-void FMaterialXSurfaceShaderAbstract::ConnectHeightToNormalInputToOutput(MaterialX::NodePtr UpstreamNode, UInterchangeShaderNode* ParentShaderNode, const FString& InputChannelName)
+void FMaterialXSurfaceShaderAbstract::ConnectHeightToNormalInputToOutput(const FConnectNode& Connect)
 {
-	if(mx::InputPtr Input = UpstreamNode->getInput("in"))
+	if(mx::InputPtr Input = Connect.UpstreamNode->getInput("in"))
 	{
 		// Image node will become this node
 		if(mx::NodePtr ConnectedNode = Input->getConnectedNode();
 		   ConnectedNode && ConnectedNode->getCategory() == mx::Category::Image)
 		{
 			//we need to copy the content of the image node to this node
-			UpstreamNode->copyContentFrom(ConnectedNode);
+			Connect.UpstreamNode->copyContentFrom(ConnectedNode);
 
-			SetMatchingInputsNames(UpstreamNode);
+			SetMatchingInputsNames(Connect.UpstreamNode);
 
 			//the copy overwrite every attribute of the node, so we need to get them back, essentially the type and the renaming
 			// the output is always a vec3
-			UpstreamNode->setType(mx::Type::Vector3);
+			Connect.UpstreamNode->setType(mx::Type::Vector3);
 
-			mx::NodeGraphPtr Graph = UpstreamNode->getParent()->asA<mx::NodeGraph>();
+			mx::NodeGraphPtr Graph = Connect.UpstreamNode->getParent()->asA<mx::NodeGraph>();
 			Graph->removeNode(ConnectedNode->getName());
 
 			using namespace UE::Interchange::Materials::Standard::Nodes;
 
-			if(UInterchangeTextureNode* TextureNode = CreateTextureNode<UInterchangeTexture2DNode>(UpstreamNode))
+			if(UInterchangeTextureNode* TextureNode = CreateTextureNode<UInterchangeTexture2DNode>(Connect.UpstreamNode))
 			{
-				UInterchangeShaderNode* HeightMapNode = CreateShaderNode(UpstreamNode->getName().c_str(), NormalFromHeightMap::Name.ToString());
-				UInterchangeShaderPortsAPI::ConnectDefaultOuputToInput(ParentShaderNode, InputChannelName, HeightMapNode->GetUniqueID());
+				UInterchangeShaderNode* HeightMapNode = CreateShaderNode(Connect.UpstreamNode->getName().c_str(), NormalFromHeightMap::Name.ToString());
+				UInterchangeShaderPortsAPI::ConnectDefaultOuputToInput(Connect.ParentShaderNode, Connect.InputChannelName, HeightMapNode->GetUniqueID());
 
-				const FString TextureNodeName = UpstreamNode->getName().c_str() + FString{ "_texture" };
+				const FString TextureNodeName = Connect.UpstreamNode->getName().c_str() + FString{ "_texture" };
 				UInterchangeShaderNode* TextureShaderNode = CreateShaderNode(TextureNodeName, TextureObject::Name.ToString());
 				TextureShaderNode->AddStringAttribute(UInterchangeShaderPortsAPI::MakeInputValueKey(TextureObject::Inputs::Texture.ToString()), TextureNode->GetUniqueID());
 				UInterchangeShaderPortsAPI::ConnectDefaultOuputToInput(HeightMapNode, NormalFromHeightMap::Inputs::HeightMap.ToString(), TextureShaderNode->GetUniqueID());
 
-				AddAttributeFromValueOrInterface(UpstreamNode->getInput("scale"), NormalFromHeightMap::Inputs::Intensity.ToString(), HeightMapNode);
+				AddAttributeFromValueOrInterface(Connect.UpstreamNode->getInput("scale"), NormalFromHeightMap::Inputs::Intensity.ToString(), HeightMapNode);
 			}
 			else
 			{
-				AddAttributeFromValueOrInterface(UpstreamNode->getInput(mx::NodeGroup::Texture2D::Inputs::Default), InputChannelName, ParentShaderNode);
+				AddAttributeFromValueOrInterface(Connect.UpstreamNode->getInput(mx::NodeGroup::Texture2D::Inputs::Default), Connect.InputChannelName, Connect.ParentShaderNode);
 			}
 		}
 		else
 		{
 			// For the moment it doesn't make sense to plug a value to it, so let's plug directly the child to the parent, in the future we could implement a Sobel and handle a multi output
-			SetAttributeNewName(Input, TCHAR_TO_UTF8(*InputChannelName)); //let's take the parent node's input name
-			ShaderNodes.Add(UpstreamNode->getName().c_str(), ParentShaderNode);
+			SetAttributeNewName(Input, TCHAR_TO_UTF8(*Connect.InputChannelName)); //let's take the parent node's input name
+			ShaderNodes.Add({ Connect.UpstreamNode->getName().c_str(), Connect.OutputName }, Connect.ParentShaderNode);
 		}
 	}
 }
 
-void FMaterialXSurfaceShaderAbstract::ConnectBlurInputToOutput(MaterialX::NodePtr UpstreamNode, UInterchangeShaderNode* ParentShaderNode, const FString& InputChannelName)
+void FMaterialXSurfaceShaderAbstract::ConnectBlurInputToOutput(const FConnectNode& Connect)
 {
-	if(mx::InputPtr Input = UpstreamNode->getInput("in"))
+	if(mx::InputPtr Input = Connect.UpstreamNode->getInput("in"))
 	{
 		// Image node will become this node
 		if(mx::NodePtr ConnectedNode = Input->getConnectedNode();
 		   ConnectedNode && ConnectedNode->getCategory() == mx::Category::Image)
 		{
-			std::string NodeType = UpstreamNode->getType();
+			std::string NodeType = Connect.UpstreamNode->getType();
 
 			//we need to copy the content of the image node to this node
-			UpstreamNode->copyContentFrom(ConnectedNode);
+			Connect.UpstreamNode->copyContentFrom(ConnectedNode);
 
-			SetMatchingInputsNames(UpstreamNode);
+			SetMatchingInputsNames(Connect.UpstreamNode);
 
 			//the copy overwrites every attribute of the node, so we need to get them back, essentially the type and the renaming
-			UpstreamNode->setType(NodeType);
+			Connect.UpstreamNode->setType(NodeType);
 
-			mx::NodeGraphPtr Graph = UpstreamNode->getParent()->asA<mx::NodeGraph>();
+			mx::NodeGraphPtr Graph = Connect.UpstreamNode->getParent()->asA<mx::NodeGraph>();
 			Graph->removeNode(ConnectedNode->getName());
 
 			using namespace UE::Interchange::Materials::Standard::Nodes;
 
-			if(UInterchangeTextureNode* TextureNode = CreateTextureNode<UInterchangeTextureBlurNode>(UpstreamNode))
+			if(UInterchangeTextureNode* TextureNode = CreateTextureNode<UInterchangeTextureBlurNode>(Connect.UpstreamNode))
 			{
 				FString OutputChannel{ TEXT("RGB") };
 
@@ -949,11 +983,11 @@ void FMaterialXSurfaceShaderAbstract::ConnectBlurInputToOutput(MaterialX::NodePt
 					OutputChannel = TEXT("R");
 				}
 
-				UInterchangeShaderNode* TextureShaderNode = CreateShaderNode(UpstreamNode->getName().c_str(), TextureSampleBlur::Name.ToString());
+				UInterchangeShaderNode* TextureShaderNode = CreateShaderNode(Connect.UpstreamNode->getName().c_str(), TextureSampleBlur::Name.ToString());
 				TextureShaderNode->AddStringAttribute(UInterchangeShaderPortsAPI::MakeInputValueKey(TextureSampleBlur::Inputs::Texture.ToString()), TextureNode->GetUniqueID());
-				UInterchangeShaderPortsAPI::ConnectOuputToInputByName(ParentShaderNode, InputChannelName, TextureShaderNode->GetUniqueID(), OutputChannel);
+				UInterchangeShaderPortsAPI::ConnectOuputToInputByName(Connect.ParentShaderNode, Connect.InputChannelName, TextureShaderNode->GetUniqueID(), OutputChannel);
 
-				if(mx::InputPtr InputKernel = UpstreamNode->getInput("filtertype"))
+				if(mx::InputPtr InputKernel = Connect.UpstreamNode->getInput("filtertype"))
 				{
 					// By default TextureSampleBox uses gaussian filter
 					if(InputKernel->getValueString() == "box")
@@ -962,7 +996,7 @@ void FMaterialXSurfaceShaderAbstract::ConnectBlurInputToOutput(MaterialX::NodePt
 					}
 				}
 
-				if(mx::InputPtr InputKernel = UpstreamNode->getInput("size"))
+				if(mx::InputPtr InputKernel = Connect.UpstreamNode->getInput("size"))
 				{
 					if(InputKernel->hasValueString())
 					{
@@ -980,32 +1014,104 @@ void FMaterialXSurfaceShaderAbstract::ConnectBlurInputToOutput(MaterialX::NodePt
 					}
 					else
 					{
-						UE_LOG(LogInterchangeImport, Warning, TEXT("<%s>: input 'size' must have a value"), ANSI_TO_TCHAR(UpstreamNode->getName().c_str()));
+						UE_LOG(LogInterchangeImport, Warning, TEXT("<%s>: input 'size' must have a value"), ANSI_TO_TCHAR(Connect.UpstreamNode->getName().c_str()));
 					}
 				}
 			}
 			else
 			{
-				AddAttributeFromValueOrInterface(UpstreamNode->getInput(mx::NodeGroup::Texture2D::Inputs::Default), InputChannelName, ParentShaderNode);
+				AddAttributeFromValueOrInterface(Connect.UpstreamNode->getInput(mx::NodeGroup::Texture2D::Inputs::Default), Connect.InputChannelName, Connect.ParentShaderNode);
 			}
 		}
 		else
 		{
 			// For a blur it doesn't make sense if there's no image input
-			SetAttributeNewName(Input, TCHAR_TO_UTF8(*InputChannelName)); //let's take the parent node's input name
-			ShaderNodes.Add(UpstreamNode->getName().c_str(), ParentShaderNode);
+			SetAttributeNewName(Input, TCHAR_TO_UTF8(*Connect.InputChannelName)); //let's take the parent node's input name
+			ShaderNodes.Add({ Connect.UpstreamNode->getName().c_str(), Connect.OutputName }, Connect.ParentShaderNode);
+		}
+	}
+}
+
+void FMaterialXSurfaceShaderAbstract::ConnectTexCoordInputToOutput(const FConnectNode& Connect)
+{
+	using namespace UE::Interchange::Materials::Standard::Nodes;
+	UInterchangeShaderNode* TexCoord = CreateShaderNode(Connect.UpstreamNode->getName().c_str(), TextureCoordinate::Name.ToString());
+	if(mx::InputPtr Input = Connect.UpstreamNode->getInput("index"))
+	{
+		TexCoord->AddInt32Attribute(UInterchangeShaderPortsAPI::MakeInputValueKey(TextureCoordinate::Inputs::Index.ToString()), mx::fromValueString<int>(Input->getValueString()));
+	}
+	UInterchangeShaderPortsAPI::ConnectDefaultOuputToInput(Connect.ParentShaderNode, Connect.InputChannelName, TexCoord->GetUniqueID());
+}
+
+void FMaterialXSurfaceShaderAbstract::ConnectSeparateInputToOutput(const FConnectNode& Connect)
+{
+	if(Connect.OutputName == TEXT("outx") || Connect.OutputName == TEXT("outr"))
+	{
+		UInterchangeShaderNode* OutXNode = CreateMaskShaderNode(0b1000, Connect.UpstreamNode->getName().c_str(), Connect.OutputName);
+		if(mx::InputPtr Input = Connect.UpstreamNode->getInput("in"))
+		{
+			AddAttributeFromValueOrInterface(Input, GetInputName(Input), OutXNode);
+		}
+		UInterchangeShaderPortsAPI::ConnectDefaultOuputToInput(Connect.ParentShaderNode, Connect.InputChannelName, OutXNode->GetUniqueID());
+	}
+	else if(Connect.OutputName == TEXT("outy") || Connect.OutputName == TEXT("outg"))
+	{
+		UInterchangeShaderNode* OutYNode = CreateMaskShaderNode(0b0100, Connect.UpstreamNode->getName().c_str(), Connect.OutputName);
+		if(mx::InputPtr Input = Connect.UpstreamNode->getInput("in"))
+		{
+			AddAttributeFromValueOrInterface(Input, GetInputName(Input), OutYNode);
+		}
+		UInterchangeShaderPortsAPI::ConnectDefaultOuputToInput(Connect.ParentShaderNode, Connect.InputChannelName, OutYNode->GetUniqueID());
+	}
+	else if(Connect.OutputName == TEXT("outz") || Connect.OutputName == TEXT("outb"))
+	{
+		UInterchangeShaderNode* OutZNode = CreateMaskShaderNode(0b0010, Connect.UpstreamNode->getName().c_str(), Connect.OutputName);
+		if(mx::InputPtr Input = Connect.UpstreamNode->getInput("in"))
+		{
+			AddAttributeFromValueOrInterface(Input, GetInputName(Input), OutZNode);
+		}
+		UInterchangeShaderPortsAPI::ConnectDefaultOuputToInput(Connect.ParentShaderNode, Connect.InputChannelName, OutZNode->GetUniqueID());
+	}
+	else if(Connect.OutputName == TEXT("outw") || Connect.OutputName == TEXT("outa"))
+	{
+		UInterchangeShaderNode* OutWNode = CreateMaskShaderNode(0b0001, Connect.UpstreamNode->getName().c_str(), Connect.OutputName);
+		if(mx::InputPtr Input = Connect.UpstreamNode->getInput("in"))
+		{
+			AddAttributeFromValueOrInterface(Input, GetInputName(Input), OutWNode);
 		}
+		UInterchangeShaderPortsAPI::ConnectDefaultOuputToInput(Connect.ParentShaderNode, Connect.InputChannelName, OutWNode->GetUniqueID());
+	}
+	else
+	{
+		UE_LOG(LogInterchangeImport, Warning, TEXT("output <%s> not defined in <%s>"), *Connect.OutputName, ANSI_TO_TCHAR(Connect.UpstreamNode->getCategory().c_str()));
 	}
 }
 
-UInterchangeShaderNode* FMaterialXSurfaceShaderAbstract::CreateMaskShaderNode(uint8 RGBA, const FString& NodeName)
+void FMaterialXSurfaceShaderAbstract::ConnectSwizzleInputToOutput(const FConnectNode& Connect)
+{
+	using namespace UE::Interchange::Materials::Standard::Nodes;
+
+	UInterchangeShaderNode* SwizzleNode = CreateShaderNode(Connect.UpstreamNode->getName().c_str(), Swizzle::Name.ToString());
+	if(mx::InputPtr Input = Connect.UpstreamNode->getInput("in"))
+	{
+		AddAttributeFromValueOrInterface(Input, GetInputName(Input), SwizzleNode);
+	}
+	if(mx::InputPtr Input = Connect.UpstreamNode->getInput("channels"))
+	{
+		SwizzleNode->AddStringAttribute(Swizzle::Attributes::Channels.ToString(), Input->getValueString().c_str());
+	}
+
+	UInterchangeShaderPortsAPI::ConnectDefaultOuputToInput(Connect.ParentShaderNode, Connect.InputChannelName, SwizzleNode->GetUniqueID());
+}
+
+UInterchangeShaderNode* FMaterialXSurfaceShaderAbstract::CreateMaskShaderNode(uint8 RGBA, const FString& NodeName, const FString& OutputName)
 {
 	bool bR = (0b1000 & RGBA) >> 3;
 	bool bG = (0b0100 & RGBA) >> 2;
 	bool bB = (0b0010 & RGBA) >> 1;
 	bool bA = (0b0001 & RGBA) >> 0;
 	using namespace UE::Interchange::Materials::Standard::Nodes;
-	UInterchangeShaderNode* MaskShaderNode = CreateShaderNode(NodeName, Mask::Name.ToString());
+	UInterchangeShaderNode* MaskShaderNode = CreateShaderNode(NodeName, Mask::Name.ToString(), OutputName);
 	MaskShaderNode->AddBooleanAttribute(Mask::Attributes::R.ToString(), bR);
 	MaskShaderNode->AddBooleanAttribute(Mask::Attributes::G.ToString(), bG);
 	MaskShaderNode->AddBooleanAttribute(Mask::Attributes::B.ToString(), bB);
@@ -1014,11 +1120,11 @@ UInterchangeShaderNode* FMaterialXSurfaceShaderAbstract::CreateMaskShaderNode(ui
 	return MaskShaderNode;
 }
 
-UInterchangeShaderNode* FMaterialXSurfaceShaderAbstract::CreateShaderNode(const FString& NodeName, const FString& ShaderType)
+UInterchangeShaderNode* FMaterialXSurfaceShaderAbstract::CreateShaderNode(const FString& NodeName, const FString& ShaderType, const FString& OutputName)
 {
 	UInterchangeShaderNode* Node;
 
-	const FString NodeUID = UInterchangeShaderNode::MakeNodeUid(NodeName, FStringView{});
+	const FString NodeUID = UInterchangeShaderNode::MakeNodeUid(NodeName + TEXT('_') + OutputName, FStringView{});
 
 	//Test directly in the NodeContainer, because the ShaderNodes can be altered during the node graph either by the parent (dot/normalmap),
 	//or by putting an intermediary node between the child and the parent (tiledimage)
@@ -1029,7 +1135,7 @@ UInterchangeShaderNode* FMaterialXSurfaceShaderAbstract::CreateShaderNode(const
 		NodeContainer.AddNode(Node);
 		Node->SetCustomShaderType(ShaderType);
 
-		ShaderNodes.Add(NodeName, Node);
+		ShaderNodes.Add({ NodeName, OutputName }, Node);
 	}
 
 	return Node;
@@ -1037,8 +1143,6 @@ UInterchangeShaderNode* FMaterialXSurfaceShaderAbstract::CreateShaderNode(const
 
 const FString& FMaterialXSurfaceShaderAbstract::GetMatchedInputName(MaterialX::NodePtr Node, MaterialX::InputPtr Input) const
 {
-	static FString EmptyString;
-
 	FMaterialXManager& Manager = FMaterialXManager::GetInstance();
 
 	if(Input)
@@ -1139,30 +1243,35 @@ FString FMaterialXSurfaceShaderAbstract::GetAttributeParentName(MaterialX::NodeP
 
 void FMaterialXSurfaceShaderAbstract::RegisterConnectNodeOutputToInputDelegates()
 {
-	MatchingConnectNodeDelegates.Add(mx::Category::Constant, FOnConnectNodeOutputToInput::CreateSP(this, &FMaterialXSurfaceShaderAbstract::ConnectConstantInputToOutput));
-	MatchingConnectNodeDelegates.Add(mx::Category::Extract, FOnConnectNodeOutputToInput::CreateSP(this, &FMaterialXSurfaceShaderAbstract::ConnectExtractInputToOutput));
-	MatchingConnectNodeDelegates.Add(mx::Category::Dot, FOnConnectNodeOutputToInput::CreateSP(this, &FMaterialXSurfaceShaderAbstract::ConnectDotInputToOutput));
-	MatchingConnectNodeDelegates.Add(mx::Category::NormalMap, FOnConnectNodeOutputToInput::CreateSP(this, &FMaterialXSurfaceShaderAbstract::ConnectDotInputToOutput));
-	MatchingConnectNodeDelegates.Add(mx::Category::TransformPoint, FOnConnectNodeOutputToInput::CreateSP(this, &FMaterialXSurfaceShaderAbstract::ConnectTransformPositionInputToOutput));
+	MatchingConnectNodeDelegates.Add(mx::Category::Constant,		FOnConnectNodeOutputToInput::CreateSP(this, &FMaterialXSurfaceShaderAbstract::ConnectConstantInputToOutput));
+	MatchingConnectNodeDelegates.Add(mx::Category::Extract,			FOnConnectNodeOutputToInput::CreateSP(this, &FMaterialXSurfaceShaderAbstract::ConnectExtractInputToOutput));
+	MatchingConnectNodeDelegates.Add(mx::Category::Dot,				FOnConnectNodeOutputToInput::CreateSP(this, &FMaterialXSurfaceShaderAbstract::ConnectDotInputToOutput));
+	MatchingConnectNodeDelegates.Add(mx::Category::NormalMap,		FOnConnectNodeOutputToInput::CreateSP(this, &FMaterialXSurfaceShaderAbstract::ConnectDotInputToOutput));
+	MatchingConnectNodeDelegates.Add(mx::Category::TransformPoint,	FOnConnectNodeOutputToInput::CreateSP(this, &FMaterialXSurfaceShaderAbstract::ConnectTransformPositionInputToOutput));
 	MatchingConnectNodeDelegates.Add(mx::Category::TransformVector, FOnConnectNodeOutputToInput::CreateSP(this, &FMaterialXSurfaceShaderAbstract::ConnectTransformVectorInputToOutput));
 	MatchingConnectNodeDelegates.Add(mx::Category::TransformNormal, FOnConnectNodeOutputToInput::CreateSP(this, &FMaterialXSurfaceShaderAbstract::ConnectTransformVectorInputToOutput));
-	MatchingConnectNodeDelegates.Add(mx::Category::Rotate3D, FOnConnectNodeOutputToInput::CreateSP(this, &FMaterialXSurfaceShaderAbstract::ConnectRotate3DInputToOutput));
-	MatchingConnectNodeDelegates.Add(mx::Category::Image, FOnConnectNodeOutputToInput::CreateSP(this, &FMaterialXSurfaceShaderAbstract::ConnectImageInputToOutput));
-	MatchingConnectNodeDelegates.Add(mx::Category::Convert, FOnConnectNodeOutputToInput::CreateSP(this, &FMaterialXSurfaceShaderAbstract::ConnectConvertInputToOutput));
-	MatchingConnectNodeDelegates.Add(mx::Category::IfGreater, FOnConnectNodeOutputToInput::CreateSP(this, &FMaterialXSurfaceShaderAbstract::ConnectIfGreaterInputToOutput));
-	MatchingConnectNodeDelegates.Add(mx::Category::IfGreaterEq, FOnConnectNodeOutputToInput::CreateSP(this, &FMaterialXSurfaceShaderAbstract::ConnectIfGreaterEqInputToOutput));
-	MatchingConnectNodeDelegates.Add(mx::Category::IfEqual, FOnConnectNodeOutputToInput::CreateSP(this, &FMaterialXSurfaceShaderAbstract::ConnectIfEqualInputToOutput));
-	MatchingConnectNodeDelegates.Add(mx::Category::Outside, FOnConnectNodeOutputToInput::CreateSP(this, &FMaterialXSurfaceShaderAbstract::ConnectOutsideInputToOutput));
-	MatchingConnectNodeDelegates.Add(mx::Category::Position, FOnConnectNodeOutputToInput::CreateSP(this, &FMaterialXSurfaceShaderAbstract::ConnectPositionInputToOutput));
-	MatchingConnectNodeDelegates.Add(mx::Category::Normal, FOnConnectNodeOutputToInput::CreateSP(this, &FMaterialXSurfaceShaderAbstract::ConnectNormalInputToOutput));
-	MatchingConnectNodeDelegates.Add(mx::Category::Tangent, FOnConnectNodeOutputToInput::CreateSP(this, &FMaterialXSurfaceShaderAbstract::ConnectTangentInputToOutput));
-	MatchingConnectNodeDelegates.Add(mx::Category::Bitangent, FOnConnectNodeOutputToInput::CreateSP(this, &FMaterialXSurfaceShaderAbstract::ConnectBitangentInputToOutput));
-	MatchingConnectNodeDelegates.Add(mx::Category::Time, FOnConnectNodeOutputToInput::CreateSP(this, &FMaterialXSurfaceShaderAbstract::ConnectTimeInputToOutput));
-	MatchingConnectNodeDelegates.Add(mx::Category::Noise3D, FOnConnectNodeOutputToInput::CreateSP(this, &FMaterialXSurfaceShaderAbstract::ConnectNoise3DInputToOutput));
-	MatchingConnectNodeDelegates.Add(mx::Category::CellNoise3D, FOnConnectNodeOutputToInput::CreateSP(this, &FMaterialXSurfaceShaderAbstract::ConnectCellNoise3DInputToOutput));
-	MatchingConnectNodeDelegates.Add(mx::Category::WorleyNoise3D, FOnConnectNodeOutputToInput::CreateSP(this, &FMaterialXSurfaceShaderAbstract::ConnectWorleyNoise3DInputToOutput));
-	MatchingConnectNodeDelegates.Add(mx::Category::Blur, FOnConnectNodeOutputToInput::CreateSP(this, &FMaterialXSurfaceShaderAbstract::ConnectBlurInputToOutput));
-	MatchingConnectNodeDelegates.Add(mx::Category::HeightToNormal, FOnConnectNodeOutputToInput::CreateSP(this, &FMaterialXSurfaceShaderAbstract::ConnectHeightToNormalInputToOutput));
+	MatchingConnectNodeDelegates.Add(mx::Category::Rotate3D,		FOnConnectNodeOutputToInput::CreateSP(this, &FMaterialXSurfaceShaderAbstract::ConnectRotate3DInputToOutput));
+	MatchingConnectNodeDelegates.Add(mx::Category::Image,			FOnConnectNodeOutputToInput::CreateSP(this, &FMaterialXSurfaceShaderAbstract::ConnectImageInputToOutput));
+	MatchingConnectNodeDelegates.Add(mx::Category::Convert,			FOnConnectNodeOutputToInput::CreateSP(this, &FMaterialXSurfaceShaderAbstract::ConnectConvertInputToOutput));
+	MatchingConnectNodeDelegates.Add(mx::Category::IfGreater,		FOnConnectNodeOutputToInput::CreateSP(this, &FMaterialXSurfaceShaderAbstract::ConnectIfGreaterInputToOutput));
+	MatchingConnectNodeDelegates.Add(mx::Category::IfGreaterEq,		FOnConnectNodeOutputToInput::CreateSP(this, &FMaterialXSurfaceShaderAbstract::ConnectIfGreaterEqInputToOutput));
+	MatchingConnectNodeDelegates.Add(mx::Category::IfEqual,			FOnConnectNodeOutputToInput::CreateSP(this, &FMaterialXSurfaceShaderAbstract::ConnectIfEqualInputToOutput));
+	MatchingConnectNodeDelegates.Add(mx::Category::Outside,			FOnConnectNodeOutputToInput::CreateSP(this, &FMaterialXSurfaceShaderAbstract::ConnectOutsideInputToOutput));
+	MatchingConnectNodeDelegates.Add(mx::Category::Position,		FOnConnectNodeOutputToInput::CreateSP(this, &FMaterialXSurfaceShaderAbstract::ConnectPositionInputToOutput));
+	MatchingConnectNodeDelegates.Add(mx::Category::Normal,			FOnConnectNodeOutputToInput::CreateSP(this, &FMaterialXSurfaceShaderAbstract::ConnectNormalInputToOutput));
+	MatchingConnectNodeDelegates.Add(mx::Category::Tangent,			FOnConnectNodeOutputToInput::CreateSP(this, &FMaterialXSurfaceShaderAbstract::ConnectTangentInputToOutput));
+	MatchingConnectNodeDelegates.Add(mx::Category::Bitangent,		FOnConnectNodeOutputToInput::CreateSP(this, &FMaterialXSurfaceShaderAbstract::ConnectBitangentInputToOutput));
+	MatchingConnectNodeDelegates.Add(mx::Category::Time,			FOnConnectNodeOutputToInput::CreateSP(this, &FMaterialXSurfaceShaderAbstract::ConnectTimeInputToOutput));
+	MatchingConnectNodeDelegates.Add(mx::Category::Noise3D,			FOnConnectNodeOutputToInput::CreateSP(this, &FMaterialXSurfaceShaderAbstract::ConnectNoise3DInputToOutput));
+	MatchingConnectNodeDelegates.Add(mx::Category::CellNoise3D,		FOnConnectNodeOutputToInput::CreateSP(this, &FMaterialXSurfaceShaderAbstract::ConnectCellNoise3DInputToOutput));
+	MatchingConnectNodeDelegates.Add(mx::Category::WorleyNoise3D,	FOnConnectNodeOutputToInput::CreateSP(this, &FMaterialXSurfaceShaderAbstract::ConnectWorleyNoise3DInputToOutput));
+	MatchingConnectNodeDelegates.Add(mx::Category::Blur,			FOnConnectNodeOutputToInput::CreateSP(this, &FMaterialXSurfaceShaderAbstract::ConnectBlurInputToOutput));
+	MatchingConnectNodeDelegates.Add(mx::Category::HeightToNormal,	FOnConnectNodeOutputToInput::CreateSP(this, &FMaterialXSurfaceShaderAbstract::ConnectHeightToNormalInputToOutput));
+	MatchingConnectNodeDelegates.Add(mx::Category::Separate2,		FOnConnectNodeOutputToInput::CreateSP(this, &FMaterialXSurfaceShaderAbstract::ConnectSeparateInputToOutput));
+	MatchingConnectNodeDelegates.Add(mx::Category::Separate3,		FOnConnectNodeOutputToInput::CreateSP(this, &FMaterialXSurfaceShaderAbstract::ConnectSeparateInputToOutput));
+	MatchingConnectNodeDelegates.Add(mx::Category::Separate4,		FOnConnectNodeOutputToInput::CreateSP(this, &FMaterialXSurfaceShaderAbstract::ConnectSeparateInputToOutput));
+	MatchingConnectNodeDelegates.Add(mx::Category::TexCoord,		FOnConnectNodeOutputToInput::CreateSP(this, &FMaterialXSurfaceShaderAbstract::ConnectTexCoordInputToOutput));
+	MatchingConnectNodeDelegates.Add(mx::Category::Swizzle,			FOnConnectNodeOutputToInput::CreateSP(this, &FMaterialXSurfaceShaderAbstract::ConnectSwizzleInputToOutput));
 }
 
 void FMaterialXSurfaceShaderAbstract::SetMatchingInputsNames(MaterialX::NodePtr Node) const
diff --git a/Engine/Plugins/Interchange/Runtime/Source/Import/Private/MaterialX/MaterialXUtils/MaterialXSurfaceUnlitShader.cpp b/Engine/Plugins/Interchange/Runtime/Source/Import/Private/MaterialX/MaterialXUtils/MaterialXSurfaceUnlitShader.cpp
--- a/Engine/Plugins/Interchange/Runtime/Source/Import/Private/MaterialX/MaterialXUtils/MaterialXSurfaceUnlitShader.cpp
+++ b/Engine/Plugins/Interchange/Runtime/Source/Import/Private/MaterialX/MaterialXUtils/MaterialXSurfaceUnlitShader.cpp
@@ -35,7 +35,7 @@ void FMaterialXSurfaceUnlitShader::Translate(mx::NodePtr SurfaceUnlitNode)
 		SurfaceUnlitShaderNode->SetCustomMaterialFunction(TEXT("/Interchange/Functions/MX_SurfaceUnlit.MX_SurfaceUnlit"));
 		NodeContainer.AddNode(SurfaceUnlitShaderNode);
 
-		ShaderNodes.Add(NodeName, SurfaceUnlitShaderNode);
+		ShaderNodes.Add({ NodeName, SurfaceUnlitNode->getNodeDef(mx::EMPTY_STRING, true)->getActiveOutputs()[0]->getName().c_str() }, SurfaceUnlitShaderNode);
 	}
 
 	using namespace UE::Interchange::Materials;
diff --git a/Engine/Plugins/Interchange/Runtime/Source/Import/Private/MaterialX/MaterialXUtils/MaterialXUsdPreviewSurface.cpp b/Engine/Plugins/Interchange/Runtime/Source/Import/Private/MaterialX/MaterialXUtils/MaterialXUsdPreviewSurface.cpp
--- a/Engine/Plugins/Interchange/Runtime/Source/Import/Private/MaterialX/MaterialXUtils/MaterialXUsdPreviewSurface.cpp
+++ b/Engine/Plugins/Interchange/Runtime/Source/Import/Private/MaterialX/MaterialXUtils/MaterialXUsdPreviewSurface.cpp
@@ -38,7 +38,7 @@ void FMaterialXUsdPreviewSurface::Translate(MaterialX::NodePtr UsdPreviewSurface
 		UsdPreviewSurfaceShaderNode->SetCustomMaterialFunction(TEXT("/Interchange/Functions/MX_UsdPreviewSurface.MX_UsdPreviewSurface"));
 		NodeContainer.AddNode(UsdPreviewSurfaceShaderNode);
 
-		ShaderNodes.Add(NodeName, UsdPreviewSurfaceShaderNode);
+		ShaderNodes.Add({ NodeName, UsdPreviewSurfaceNode->getNodeDef(mx::EMPTY_STRING, true)->getActiveOutputs()[0]->getName().c_str()}, UsdPreviewSurfaceShaderNode);
 	}
 
 	// Inputs
diff --git a/Engine/Plugins/Interchange/Runtime/Source/Import/Public/MaterialX/MaterialExpressions/MaterialExpressionSwizzle.h b/Engine/Plugins/Interchange/Runtime/Source/Import/Public/MaterialX/MaterialExpressions/MaterialExpressionSwizzle.h
new file mode 100644
--- /dev/null
+++ b/Engine/Plugins/Interchange/Runtime/Source/Import/Public/MaterialX/MaterialExpressions/MaterialExpressionSwizzle.h
@@ -0,0 +1,28 @@
+// Copyright Epic Games, Inc. All Rights Reserved.
+#pragma once
+
+#include "CoreMinimal.h"
+#include "Materials/MaterialExpression.h"
+
+#include "MaterialExpressionSwizzle.generated.h"
+
+UCLASS(collapsecategories, hidecategories = Object, MinimalAPI)
+class UMaterialExpressionMaterialXSwizzle : public UMaterialExpression
+{
+	GENERATED_UCLASS_BODY()
+
+	UPROPERTY()
+	FExpressionInput Input;
+
+	UPROPERTY(EditAnywhere, Category = Swizzle, Meta = (ShowAsInputPin = "Advanced"))
+	FString Channels;
+
+	//~ Begin UMaterialExpressionMaterialX Interface
+#if WITH_EDITOR
+	virtual int32 Compile(class FMaterialCompiler* Compiler, int32 OutputIndex) override;
+	virtual void GetCaption(TArray<FString>& OutCaptions) const override;
+	void GetExpressionToolTip(TArray<FString>& OutToolTip) override;
+#endif
+	//~ End UMaterialExpressionMaterialX Interface
+};
+
diff --git a/Engine/Plugins/Interchange/Runtime/Source/Import/Public/MaterialX/MaterialXUtils/MaterialXSurfaceShaderAbstract.h b/Engine/Plugins/Interchange/Runtime/Source/Import/Public/MaterialX/MaterialXUtils/MaterialXSurfaceShaderAbstract.h
--- a/Engine/Plugins/Interchange/Runtime/Source/Import/Public/MaterialX/MaterialXUtils/MaterialXSurfaceShaderAbstract.h
+++ b/Engine/Plugins/Interchange/Runtime/Source/Import/Public/MaterialX/MaterialXUtils/MaterialXSurfaceShaderAbstract.h
@@ -14,9 +14,29 @@ class FMaterialXSurfaceShaderAbstract : public FMaterialXBase
 {
 protected:
 
+	static FString EmptyString;
+
+	// MaterialX states the default output name of the different nodes is 'out'
+	static FString DefaultOutput;
+
 	friend class FMaterialXSurfaceMaterial;
 
-	DECLARE_DELEGATE_ThreeParams(FOnConnectNodeOutputToInput, MaterialX::NodePtr, UInterchangeShaderNode*, const FString&);
+	struct FConnectNode
+	{
+		// The MaterialX node of a given type used to create the appropriate shader node
+		MaterialX::NodePtr UpstreamNode;
+
+		// The shader node to connect to
+		UInterchangeShaderNode* ParentShaderNode;
+
+		// The input of the ParentShaderNode to connect to
+		FString InputChannelName;
+
+		// The output name of the MaterialX node, default name is 'out' as stated by the standard library
+		FString OutputName{ DefaultOutput };
+	};
+		
+	DECLARE_DELEGATE_OneParam(FOnConnectNodeOutputToInput, const FConnectNode&);
 
 	FMaterialXSurfaceShaderAbstract(UInterchangeBaseNodeContainer & BaseNodeContainer);
 
@@ -142,7 +162,7 @@ protected:
 	 *
 	 * @return true if a shader node has been successfully created and is connected to the given input
 	 */
-	bool ConnectMatchingNodeOutputToInput(MaterialX::NodePtr Node, UInterchangeShaderNode* ParentShaderNode, const FString& InputChannelName);
+	bool ConnectMatchingNodeOutputToInput(const FConnectNode& Connect);
 
 	/**
 	 * Create and Connect manually the output of a MaterialX node to a shader node
@@ -169,70 +189,79 @@ protected:
 	/** Begin Connect MaterialX nodes*/
 
 	/** <constant> */
-	void ConnectConstantInputToOutput(MaterialX::NodePtr UpstreamNode, UInterchangeShaderNode* ParentShaderNode, const FString& InputChannelName);
+	void ConnectConstantInputToOutput(const FConnectNode& Connect);
 
 	/** <extract> */
-	void ConnectExtractInputToOutput(MaterialX::NodePtr UpstreamNode, UInterchangeShaderNode* ParentShaderNode, const FString& InputChannelName);
+	void ConnectExtractInputToOutput(const FConnectNode& Connect);
 
 	/** <dot> */
-	void ConnectDotInputToOutput(MaterialX::NodePtr UpstreamNode, UInterchangeShaderNode* ParentShaderNode, const FString& InputChannelName);
+	void ConnectDotInputToOutput(const FConnectNode& Connect);
 
 	/** <transformpoint> */
-	void ConnectTransformPositionInputToOutput(MaterialX::NodePtr UpstreamNode, UInterchangeShaderNode* ParentShaderNode, const FString& InputChannelName);
+	void ConnectTransformPositionInputToOutput(const FConnectNode& Connect);
 
 	/** <transformvector> */
-	void ConnectTransformVectorInputToOutput(MaterialX::NodePtr UpstreamNode, UInterchangeShaderNode* ParentShaderNode, const FString& InputChannelName);
+	void ConnectTransformVectorInputToOutput(const FConnectNode& Connect);
 
 	/** <rotate3d> */
-	void ConnectRotate3DInputToOutput(MaterialX::NodePtr UpstreamNode, UInterchangeShaderNode* ParentShaderNode, const FString& InputChannelName);
+	void ConnectRotate3DInputToOutput(const FConnectNode& Connect);
 
 	/** <image> */
-	void ConnectImageInputToOutput(MaterialX::NodePtr UpstreamNode, UInterchangeShaderNode* ParentShaderNode, const FString& InputChannelName);
+	void ConnectImageInputToOutput(const FConnectNode& Connect);
 
 	/** <convert> */
-	void ConnectConvertInputToOutput(MaterialX::NodePtr UpstreamNode, UInterchangeShaderNode* ParentShaderNode, const FString& InputChannelName);
+	void ConnectConvertInputToOutput(const FConnectNode& Connect);
 
 	/** <ifgreater> */
-	void ConnectIfGreaterInputToOutput(MaterialX::NodePtr UpstreamNode, UInterchangeShaderNode* ParentShaderNode, const FString& InputChannelName);
+	void ConnectIfGreaterInputToOutput(const FConnectNode& Connect);
 
 	/** <ifgreatereq> */
-	void ConnectIfGreaterEqInputToOutput(MaterialX::NodePtr UpstreamNode, UInterchangeShaderNode* ParentShaderNode, const FString& InputChannelName);
+	void ConnectIfGreaterEqInputToOutput(const FConnectNode& Connect);
 
 	/** <ifequal> */
-	void ConnectIfEqualInputToOutput(MaterialX::NodePtr UpstreamNode, UInterchangeShaderNode* ParentShaderNode, const FString& InputChannelName);
+	void ConnectIfEqualInputToOutput(const FConnectNode& Connect);
 
 	/** <outside> */
-	void ConnectOutsideInputToOutput(MaterialX::NodePtr UpstreamNode, UInterchangeShaderNode* ParentShaderNode, const FString& InputChannelName);
+	void ConnectOutsideInputToOutput(const FConnectNode& Connect);
 
 	/** <position> */
-	void ConnectPositionInputToOutput(MaterialX::NodePtr UpstreamNode, UInterchangeShaderNode* ParentShaderNode, const FString& InputChannelName);
+	void ConnectPositionInputToOutput(const FConnectNode& Connect);
 
 	/** <normal> */
-	void ConnectNormalInputToOutput(MaterialX::NodePtr UpstreamNode, UInterchangeShaderNode* ParentShaderNode, const FString& InputChannelName);
+	void ConnectNormalInputToOutput(const FConnectNode& Connect);
 
 	/** <tangent> */
-	void ConnectTangentInputToOutput(MaterialX::NodePtr UpstreamNode, UInterchangeShaderNode* ParentShaderNode, const FString& InputChannelName);
+	void ConnectTangentInputToOutput(const FConnectNode& Connect);
 
 	/** <bitangent> */
-	void ConnectBitangentInputToOutput(MaterialX::NodePtr UpstreamNode, UInterchangeShaderNode* ParentShaderNode, const FString& InputChannelName);
+	void ConnectBitangentInputToOutput(const FConnectNode& Connect);
 
 	/** <time> */
-	void ConnectTimeInputToOutput(MaterialX::NodePtr UpstreamNode, UInterchangeShaderNode* ParentShaderNode, const FString& InputChannelName);
+	void ConnectTimeInputToOutput(const FConnectNode& Connect);
 
 	/** <noise3d> */
-	void ConnectNoise3DInputToOutput(MaterialX::NodePtr UpstreamNode, UInterchangeShaderNode* ParentShaderNode, const FString& InputChannelName);
+	void ConnectNoise3DInputToOutput(const FConnectNode& Connect);
 
 	/** <cellnoise3d> */
-	void ConnectCellNoise3DInputToOutput(MaterialX::NodePtr UpstreamNode, UInterchangeShaderNode* ParentShaderNode, const FString& InputChannelName);
+	void ConnectCellNoise3DInputToOutput(const FConnectNode& Connect);
 
 	/** <worleynoise3d> */
-	void ConnectWorleyNoise3DInputToOutput(MaterialX::NodePtr UpstreamNode, UInterchangeShaderNode* ParentShaderNode, const FString& InputChannelName);
+	void ConnectWorleyNoise3DInputToOutput(const FConnectNode& Connect);
 
 	/** <heighttonormal> */
-	void ConnectHeightToNormalInputToOutput(MaterialX::NodePtr UpstreamNode, UInterchangeShaderNode* ParentShaderNode, const FString& InputChannelName);
+	void ConnectHeightToNormalInputToOutput(const FConnectNode& Connect);
 
 	/** <blur> */
-	void ConnectBlurInputToOutput(MaterialX::NodePtr UpstreamNode, UInterchangeShaderNode* ParentShaderNode, const FString& InputChannelName);
+	void ConnectBlurInputToOutput(const FConnectNode& Connect);
+
+	/** <texcoord> */
+	void ConnectTexCoordInputToOutput(const FConnectNode& Connect);
+
+	/** <separate2/3/4> */
+	void ConnectSeparateInputToOutput(const FConnectNode& Connect);
+
+	/** <swizzle> */
+	void ConnectSwizzleInputToOutput(const FConnectNode& Connect);
 
 	/** End Connect MaterialX nodes*/
 
@@ -241,20 +270,21 @@ protected:
 	 *
 	 * @param RGBA - The mask component, e.g: 0b1011 -> Only RBA are toggled
 	 * @param NodeName - the name of the shader node
-	 *
+	 * @param OutputName - the name of the output of the MaterialX node, default name is 'out' as stated by the standard library
 	 * @return The ComponentMask node
 	 */
-	UInterchangeShaderNode* CreateMaskShaderNode(uint8 RGBA, const FString& NodeName);
+	UInterchangeShaderNode* CreateMaskShaderNode(uint8 RGBA, const FString& NodeName, const FString& OutputName = TEXT("out"));
 
 	/**
 	 * Helper function to create an InterchangeShaderNode
 	 *
 	 * @param NodeName - The name of the shader node
 	 * @param ShaderType - The shader node's type we want to create
+	 * @param OutputName - The output name of the MaterialX node, default name is 'out' as stated by the standard library
 	 *
 	 * @return The shader node that was created
 	 */
-	UInterchangeShaderNode* CreateShaderNode(const FString& NodeName, const FString& ShaderType);
+	UInterchangeShaderNode* CreateShaderNode(const FString& NodeName, const FString& ShaderType, const FString& OutputName = TEXT("out"));
 
 	/**
 	 * Helper function to create an InterchangeTextureNode
@@ -405,8 +435,14 @@ protected:
 
 protected:
 
+	/** 
+	 * @param Get<0>: node
+	 * @param Get<1>: output
+	 */
+	using FNodeOutput = TPair<FString, FString>;
+
 	/** Store the shader nodes only when we create the shader graph node*/
-	TMap<FString, UInterchangeShaderNode*> ShaderNodes;
+	TMap<FNodeOutput, UInterchangeShaderNode*> ShaderNodes;
 
 	/** Matching MaterialX category and Connect function*/
 	TMap<FString, FOnConnectNodeOutputToInput> MatchingConnectNodeDelegates;
diff --git a/Engine/Plugins/Interchange/Runtime/Source/Nodes/Public/InterchangeMaterialDefinitions.h b/Engine/Plugins/Interchange/Runtime/Source/Nodes/Public/InterchangeMaterialDefinitions.h
--- a/Engine/Plugins/Interchange/Runtime/Source/Nodes/Public/InterchangeMaterialDefinitions.h
+++ b/Engine/Plugins/Interchange/Runtime/Source/Nodes/Public/InterchangeMaterialDefinitions.h
@@ -137,6 +137,21 @@ namespace Interchange
 					}
 				}
 
+				namespace Swizzle
+				{
+					const FName Name = TEXT("MaterialXSwizzle");
+
+					namespace Inputs
+					{
+						const FName Input = TEXT("Input");
+					}
+
+					namespace Attributes
+					{
+						const FName Channels = TEXT("Channels");
+					}
+				}
+
 				namespace TextureCoordinate
 				{
 					const FName Name = TEXT("TextureCoordinate");
diff --git a/Engine/Plugins/Interchange/Runtime/Source/Pipelines/Private/InterchangeGenericMaterialPipeline.cpp b/Engine/Plugins/Interchange/Runtime/Source/Pipelines/Private/InterchangeGenericMaterialPipeline.cpp
--- a/Engine/Plugins/Interchange/Runtime/Source/Pipelines/Private/InterchangeGenericMaterialPipeline.cpp
+++ b/Engine/Plugins/Interchange/Runtime/Source/Pipelines/Private/InterchangeGenericMaterialPipeline.cpp
@@ -45,6 +45,7 @@
 #include "Materials/MaterialFunction.h"
 #include "Materials/MaterialInstanceConstant.h"
 #include "Materials/MaterialInstanceDynamic.h"
+#include "MaterialX/MaterialExpressions/MaterialExpressionSwizzle.h"
 #include "MaterialX/MaterialExpressions/MaterialExpressionTextureSampleParameterBlur.h"
 #include "Misc/CoreMisc.h"
 #include "Misc/Paths.h"
@@ -2245,6 +2246,29 @@ void UInterchangeGenericMaterialPipeline::HandleVectorNoiseNode(const UInterchan
 	}
 }
 
+void UInterchangeGenericMaterialPipeline::HandleSwizzleNode(const UInterchangeShaderNode* ShaderNode, UInterchangeBaseMaterialFactoryNode* MaterialFactoryNode, UInterchangeMaterialExpressionFactoryNode* SwizzleFactoryNode)
+{
+	using namespace UE::Interchange::Materials::Standard::Nodes;
+	SwizzleFactoryNode->SetCustomExpressionClassName(UMaterialExpressionMaterialXSwizzle::StaticClass()->GetName());
+	// Input
+	{
+		TTuple<UInterchangeMaterialExpressionFactoryNode*, FString> InputExpression =
+			CreateMaterialExpressionForInput(MaterialFactoryNode, ShaderNode, Swizzle::Inputs::Input.ToString(), SwizzleFactoryNode->GetUniqueID());
+		if(InputExpression.Get<0>())
+		{
+			UInterchangeShaderPortsAPI::ConnectOuputToInputByName(SwizzleFactoryNode, GET_MEMBER_NAME_CHECKED(UMaterialExpressionMaterialXSwizzle, Input).ToString(),
+																  InputExpression.Get<0>()->GetUniqueID(), InputExpression.Get<1>());
+		}
+	}
+	// Channels
+	if(FString Channels; ShaderNode->GetStringAttribute(Swizzle::Attributes::Channels.ToString(), Channels))
+	{
+		const FName ChannelsMemberName = GET_MEMBER_NAME_CHECKED(UMaterialExpressionMaterialXSwizzle, Channels);
+		SwizzleFactoryNode->AddStringAttribute(ChannelsMemberName.ToString(), Channels);
+		SwizzleFactoryNode->AddApplyAndFillDelegates<FString>(ChannelsMemberName.ToString(), UMaterialExpressionMaterialXSwizzle::StaticClass(), ChannelsMemberName);
+	}
+}
+
 UInterchangeMaterialExpressionFactoryNode* UInterchangeGenericMaterialPipeline::CreateMaterialExpressionForShaderNode(UInterchangeBaseMaterialFactoryNode* MaterialFactoryNode,
 	const UInterchangeShaderNode* ShaderNode, const FString& ParentUid)
 {
@@ -2307,6 +2331,10 @@ UInterchangeMaterialExpressionFactoryNode* UInterchangeGenericMaterialPipeline::
 	{
 		HandleNormalFromHeightMapNode(ShaderNode, MaterialFactoryNode, MaterialExpression);
 	}
+	else if(*ShaderType == Nodes::Swizzle::Name)
+	{
+		HandleSwizzleNode(ShaderNode, MaterialFactoryNode, MaterialExpression);
+	}
 	else if (*ShaderType == Nodes::TextureCoordinate::Name)
 	{
 		HandleTextureCoordinateNode(ShaderNode, MaterialFactoryNode, MaterialExpression);
diff --git a/Engine/Plugins/Interchange/Runtime/Source/Pipelines/Public/InterchangeGenericMaterialPipeline.h b/Engine/Plugins/Interchange/Runtime/Source/Pipelines/Public/InterchangeGenericMaterialPipeline.h
--- a/Engine/Plugins/Interchange/Runtime/Source/Pipelines/Public/InterchangeGenericMaterialPipeline.h
+++ b/Engine/Plugins/Interchange/Runtime/Source/Pipelines/Public/InterchangeGenericMaterialPipeline.h
@@ -190,6 +190,7 @@ private:
 	void HandleTransformVectorNode(const UInterchangeShaderNode* ShaderNode, UInterchangeBaseMaterialFactoryNode* MaterialFactoryNode, UInterchangeMaterialExpressionFactoryNode* TransformVectorFactoryNode);
 	void HandleNoiseNode(const UInterchangeShaderNode* ShaderNode, UInterchangeBaseMaterialFactoryNode* MaterialFactoryNode, UInterchangeMaterialExpressionFactoryNode* NoiseFactoryNode);
 	void HandleVectorNoiseNode(const UInterchangeShaderNode* ShaderNode, UInterchangeBaseMaterialFactoryNode* MaterialFactoryNode, UInterchangeMaterialExpressionFactoryNode* NoiseFactoryNode);
+	void HandleSwizzleNode(const UInterchangeShaderNode* ShaderNode, UInterchangeBaseMaterialFactoryNode* MaterialFactoryNode, UInterchangeMaterialExpressionFactoryNode* SwizzleFactoryNode);
 
 	UInterchangeMaterialExpressionFactoryNode* CreateMaterialExpressionForShaderNode(UInterchangeBaseMaterialFactoryNode* MaterialFactoryNode, const UInterchangeShaderNode* ShaderNode, const FString& ParentUid);
 	TTuple<UInterchangeMaterialExpressionFactoryNode*, FString> CreateMaterialExpressionForInput(UInterchangeBaseMaterialFactoryNode* MaterialFactoryNode, const UInterchangeShaderNode* ShaderNode, const FString& InputName, const FString& ParentUid);
