Added ability to override a pipeline stack on a translator basis
Added MaterialX and glTF pipelines - To be completed

diff --git a/Engine/Config/BaseEngine.ini b/Engine/Config/BaseEngine.ini
--- a/Engine/Config/BaseEngine.ini
+++ b/Engine/Config/BaseEngine.ini
@@ -3300,8 +3300,11 @@ EnvironmentCollisionThickness = 2.f
 EnvironmentCollisionMaxShapes = 32
 
 [/Script/InterchangeEngine.InterchangeProjectSettings]
+ContentImportSettings=(DefaultPipelineStackOverride=((Materials, "Materials"),(Textures, "Textures")),ShowPipelineStacksConfigurationDialogOverride=((Textures, False)),PipelineStacks=(("Assets", (Pipelines=("/Interchange/Pipelines/DefaultAssetsPipeline.DefaultAssetsPipeline"),PerTranslatorPipelines=((Translator="/Script/InterchangeImport.InterchangeMaterialXTranslator",Pipelines=("/Interchange/Pipelines/DefaultAssetsPipeline.DefaultAssetsPipeline","/Interchange/Pipelines/DefaultMaterialXPipeline.DefaultMaterialXPipeline")),(Translator="/Script/InterchangeImport.InterchangeGltfTranslator",Pipelines=("/Interchange/Pipelines/DefaultglTFAssetsPipeline.DefaultglTFAssetsPipeline","/Interchange/Pipelines/DefaultGLTFPipeline.DefaultGLTFPipeline"))))),("Materials", (Pipelines=("/Interchange/Pipelines/DefaultMaterialPipeline.DefaultMaterialPipeline"),PerTranslatorPipelines=((Translator="/Script/InterchangeImport.InterchangeMaterialXTranslator",Pipelines=("/Interchange/Pipelines/DefaultMaterialPipeline.DefaultMaterialPipeline","/Interchange/Pipelines/DefaultMaterialXPipeline.DefaultMaterialXPipeline"))))),("Textures", (Pipelines=("/Interchange/Pipelines/DefaultTexturePipeline.DefaultTexturePipeline")))),DefaultPipelineStack="Assets",PipelineConfigurationDialogClass="/Script/InterchangeEditorPipelines.InterchangePipelineConfigurationGeneric",bShowPipelineStacksConfigurationDialog=True)
+SceneImportSettings=(PipelineStacks=(("Scene", (Pipelines=("/Interchange/Pipelines/DefaultSceneAssetsPipeline.DefaultSceneAssetsPipeline","/Interchange/Pipelines/DefaultSceneLevelPipeline.DefaultSceneLevelPipeline"),PerTranslatorPipelines=((Translator="/Script/InterchangeImport.InterchangeMaterialXTranslator",Pipelines=("/Interchange/Pipelines/DefaultSceneAssetsPipeline.DefaultSceneAssetsPipeline","/Interchange/Pipelines/DefaultSceneLevelPipeline.DefaultSceneLevelPipeline","/Interchange/Pipelines/DefaultMaterialXPipeline.DefaultMaterialXPipeline")),(Translator="/Script/InterchangeImport.InterchangeGltfTranslator",Pipelines=("/Interchange/Pipelines/DefaultglTFSceneAssetsPipeline.DefaultglTFSceneAssetsPipeline","/Interchange/Pipelines/DefaultSceneLevelPipeline.DefaultSceneLevelPipeline","/Interchange/Pipelines/DefaultGLTFPipeline.DefaultGLTFPipeline")))))),DefaultPipelineStack="Scene",PipelineConfigurationDialogClass="/Script/InterchangeEditorPipelines.InterchangePipelineConfigurationGeneric",bShowPipelineStacksConfigurationDialog=True)
 FilePickerClass=/Script/InterchangeEditorUtilities.InterchangeFilePickerGeneric
 bStaticMeshUseSmoothEdgesIfSmoothingInformationIsMissing=True
-ContentImportSettings=(DefaultPipelineStackOverride=((Materials, "Materials"),(Textures, "Textures")),ShowPipelineStacksConfigurationDialogOverride=((Textures, False)),PipelineStacks=(("Assets", (Pipelines=("/Interchange/Pipelines/DefaultAssetsPipeline.DefaultAssetsPipeline"))),("Materials", (Pipelines=("/Interchange/Pipelines/DefaultMaterialPipeline.DefaultMaterialPipeline"))),("Textures", (Pipelines=("/Interchange/Pipelines/DefaultTexturePipeline.DefaultTexturePipeline")))),DefaultPipelineStack="Assets",PipelineConfigurationDialogClass="/Script/InterchangeEditorPipelines.InterchangePipelineConfigurationGeneric",bShowPipelineStacksConfigurationDialog=True)
-SceneImportSettings=(PipelineStacks=(("Scene", (Pipelines=("/Interchange/Pipelines/DefaultSceneAssetsPipeline.DefaultSceneAssetsPipeline","/Interchange/Pipelines/DefaultSceneLevelPipeline.DefaultSceneLevelPipeline")))),DefaultPipelineStack="Scene",PipelineConfigurationDialogClass="/Script/InterchangeEditorPipelines.InterchangePipelineConfigurationGeneric",bShowPipelineStacksConfigurationDialog=True)
 GenericPipelineClass=/Script/InterchangePipelines.InterchangeGenericAssetsPipeline
+
+[/Script/InterchangePipelines.MaterialXPipelineSettings]
+PredefinedSurfaceShaders=((StandardSurface, "/Interchange/Functions/MX_StandardSurface.MX_StandardSurface"),(StandardSurfaceTransmission, "/Interchange/Functions/MX_TransmissionSurface.MX_TransmissionSurface"),(SurfaceUnlit, "/Interchange/Functions/MX_SurfaceUnlit.MX_SurfaceUnlit"),(UsdPreviewSurface, "/Interchange/Functions/MX_UsdPreviewSurface.MX_UsdPreviewSurface"))
diff --git a/Engine/Plugins/Interchange/Editor/Source/Pipelines/Private/InterchangeEditorPipelinesModule.cpp b/Engine/Plugins/Interchange/Editor/Source/Pipelines/Private/InterchangeEditorPipelinesModule.cpp
--- a/Engine/Plugins/Interchange/Editor/Source/Pipelines/Private/InterchangeEditorPipelinesModule.cpp
+++ b/Engine/Plugins/Interchange/Editor/Source/Pipelines/Private/InterchangeEditorPipelinesModule.cpp
@@ -4,9 +4,9 @@
 
 #include "InterchangeEditorPipelineDetails.h"
 #include "InterchangeEditorPipelineStyle.h"
-#include "InterchangeGenericMaterialPipeline.h"
 #include "InterchangeManager.h"
-#include "InterchangeMaterialXPipelineSettingsCustomization.h"
+#include "InterchangeMaterialXPipeline.h"
+#include "InterchangeMaterialXPipelineCustomizations.h"
 #include "InterchangePipelineBase.h"
 #include "InterchangePipelineFactories.h"
 #include "Misc/CoreDelegates.h"
@@ -79,8 +79,11 @@ void FInterchangeEditorPipelinesModule::AcquireResources()
 	ClassesToUnregisterOnShutdown.Add(UInterchangePipelineBase::StaticClass()->GetFName());
 	PropertyEditorModule.RegisterCustomClassLayout(ClassesToUnregisterOnShutdown.Last(), FOnGetDetailCustomizationInstance::CreateStatic(&FInterchangePipelineBaseDetailsCustomization::MakeInstance));
 
-	PropertiesTypesToUnregisterOnShutdown.Add(FMaterialXPipelineSettings::StaticStruct()->GetFName());
-	PropertyEditorModule.RegisterCustomPropertyTypeLayout(PropertiesTypesToUnregisterOnShutdown.Last(), FOnGetPropertyTypeCustomizationInstance::CreateStatic(&FInterchangeMaterialXPipelineSettingsCustomization::MakeInstance));
+	ClassesToUnregisterOnShutdown.Add(UInterchangeMaterialXPipeline::StaticClass()->GetFName());
+	PropertyEditorModule.RegisterCustomClassLayout(ClassesToUnregisterOnShutdown.Last(), FOnGetDetailCustomizationInstance::CreateStatic(&FInterchangeMaterialXPipelineCustomization::MakeInstance));
+
+	ClassesToUnregisterOnShutdown.Add(UMaterialXPipelineSettings::StaticClass()->GetFName());
+	PropertyEditorModule.RegisterCustomClassLayout(ClassesToUnregisterOnShutdown.Last(), FOnGetDetailCustomizationInstance::CreateStatic(&FInterchangeMaterialXPipelineSettingsCustomization::MakeInstance));
 
 	if (!InterchangeEditorPipelineStyle.IsValid())
 	{
diff --git a/Engine/Plugins/Interchange/Editor/Source/Pipelines/Private/InterchangeMaterialXPipelineCustomizations.cpp b/Engine/Plugins/Interchange/Editor/Source/Pipelines/Private/InterchangeMaterialXPipelineCustomizations.cpp
new file mode 100644
--- /dev/null
+++ b/Engine/Plugins/Interchange/Editor/Source/Pipelines/Private/InterchangeMaterialXPipelineCustomizations.cpp
@@ -0,0 +1,171 @@
+// Copyright Epic Games, Inc. All Rights Reserved.
+
+#include "InterchangeMaterialXPipelineCustomizations.h"
+
+#include "InterchangeMaterialDefinitions.h"
+#include "InterchangeMaterialXPipeline.h"
+#include "Materials/MaterialFunction.h"
+
+#include "DetailLayoutBuilder.h"
+#include "Modules/ModuleManager.h"
+#include "PropertyCustomizationHelpers.h"
+//#include "DetailWidgetRow.h"
+//#include "IDetailGroup.h"
+//#include "InterchangePipelineBase.h"
+//#include "Nodes/InterchangeBaseNode.h"
+//#include "ScopedTransaction.h"
+//#include "Styling/StyleColors.h"
+//#include "Widgets/Input/SNumericEntryBox.h"
+//#include "Widgets/Layout/SBox.h"
+
+TSharedRef<IDetailCustomization> FInterchangeMaterialXPipelineCustomization::MakeInstance()
+{
+	return MakeShareable(new FInterchangeMaterialXPipelineCustomization());
+}
+
+TSharedRef<IDetailCustomization> FInterchangeMaterialXPipelineSettingsCustomization::MakeInstance()
+{
+	return MakeShareable(new FInterchangeMaterialXPipelineSettingsCustomization());
+}
+
+void FInterchangeMaterialXPipelineCustomization::CustomizeDetails(IDetailLayoutBuilder& DetailBuilder)
+{
+	TArray<TWeakObjectPtr<UObject>> EditingObjects;
+	DetailBuilder.GetObjectsBeingCustomized(EditingObjects);
+	check(EditingObjects.Num() == 1);
+
+	TWeakObjectPtr<UInterchangeMaterialXPipeline> Pipeline = Cast<UInterchangeMaterialXPipeline>(EditingObjects[0].Get());
+
+	if (!ensure(Pipeline.IsValid()))
+	{
+		return;
+	}
+
+	IDetailCategoryBuilder& MaterialXCategory = DetailBuilder.EditCategory("MaterialX");
+
+	MaterialXCategory.SetDisplayName(NSLOCTEXT("InterchangeMaterialXPipelineCustomization", "CategoryDisplayName", "MaterialX Settings"));
+
+	FPropertyEditorModule& PropertyEditorModule = FModuleManager::GetModuleChecked<FPropertyEditorModule>("PropertyEditor");
+
+	FDetailsViewArgs DetailsViewArgs;
+	DetailsViewArgs.NameAreaSettings = FDetailsViewArgs::HideNameArea;
+	TSharedPtr<IDetailsView> DetailsView = PropertyEditorModule.CreateDetailView(DetailsViewArgs);
+
+	DetailsView->SetObject(Pipeline->MaterialXSettings);
+
+	MaterialXCategory.AddCustomRow(NSLOCTEXT("InterchangeMaterialXPipelineCustomization", "MaterialXPredefined", "Predefined Material Functions"))
+	[
+		SNew( SVerticalBox )
+		+ SVerticalBox::Slot()
+		[
+			DetailsView.ToSharedRef()
+		]
+	];
+}
+
+void FInterchangeMaterialXPipelineSettingsCustomization::CustomizeDetails(IDetailLayoutBuilder& DetailBuilder)
+{
+	using OnShouldFilterAssetFunc = bool (FInterchangeMaterialXPipelineSettingsCustomization::*)(const FAssetData&);
+
+	TArray<TWeakObjectPtr<UObject>> EditingObjects;
+	DetailBuilder.GetObjectsBeingCustomized(EditingObjects);
+	check(EditingObjects.Num() == 1);
+
+	MaterialXSettings = Cast<UMaterialXPipelineSettings>(EditingObjects[0].Get());
+
+	if (!ensure(MaterialXSettings.IsValid()))
+	{
+		return;
+	}
+
+	TSharedRef< IPropertyHandle > PairingsHandle = DetailBuilder.GetProperty(GET_MEMBER_NAME_CHECKED(UMaterialXPipelineSettings, PredefinedSurfaceShaders));
+	if (!PairingsHandle->IsValidHandle())
+	{
+		return;
+	}
+
+	DetailBuilder.HideProperty(PairingsHandle);
+
+	IDetailCategoryBuilder& MaterialXPredefinedCategory = DetailBuilder.EditCategory("MaterialXPredefined");
+
+	MaterialXPredefinedCategory.SetDisplayName(NSLOCTEXT("InterchangeMaterialXPipelineSettingsCustomization", "CategoryDisplayName", "MaterialX Predefined Surface Shaders"));
+
+	uint32 NumChildren = 0;
+	PairingsHandle->GetNumChildren(NumChildren);
+
+	for(uint32 i = 0; i < NumChildren; ++i)
+	{
+		TSharedPtr<IPropertyHandle> ChildPropertyHandle = PairingsHandle->GetChildHandle(i);
+		TSharedPtr<IPropertyHandle> KeyPropertyHandle = ChildPropertyHandle->GetKeyHandle();
+
+		OnShouldFilterAssetFunc OnShouldFilterAsset = &FInterchangeMaterialXPipelineSettingsCustomization::OnShouldFilterAsset;
+
+		if(KeyPropertyHandle.IsValid())
+		{
+			FString Str;
+			KeyPropertyHandle->GetValueAsDisplayString(Str);
+			if(Str == TEXT("Standard Surface"))
+			{
+				OnShouldFilterAsset = &FInterchangeMaterialXPipelineSettingsCustomization::OnShouldFilterAssetStandardSurface;
+			}
+			else if(Str == TEXT("Standard Surface Transmission"))
+			{
+				OnShouldFilterAsset = &FInterchangeMaterialXPipelineSettingsCustomization::OnShouldFilterAssetStandardSurfaceTransmission;
+			}
+			else if(Str == TEXT("Surface Unlit"))
+			{
+				OnShouldFilterAsset = &FInterchangeMaterialXPipelineSettingsCustomization::OnShouldFilterAssetSurfaceUnlit;
+			}
+			else if(Str == TEXT("Usd Preview Surface"))
+			{
+				OnShouldFilterAsset = &FInterchangeMaterialXPipelineSettingsCustomization::OnShouldFilterAssetUsdPreviewSurface;
+			}
+		}
+
+		IDetailPropertyRow& SurfaceShaderTypeRow = DetailBuilder.AddPropertyToCategory(ChildPropertyHandle.ToSharedRef());
+		SurfaceShaderTypeRow.ShowPropertyButtons(false);
+
+		FDetailWidgetRow& DetailWidgetRow = SurfaceShaderTypeRow.CustomWidget();
+		DetailWidgetRow
+		.NameContent()
+		[
+			SNullWidget::NullWidget
+		]
+		.ValueContent()
+		.MinDesiredWidth(1.0f)
+		.MaxDesiredWidth(1.0f)
+		[
+			SNew(SHorizontalBox)
+			+ SHorizontalBox::Slot()
+			.HAlign(HAlign_Fill)
+			[
+				SNew(SObjectPropertyEntryBox)
+				.AllowedClass(UMaterialFunction::StaticClass())
+				.PropertyHandle(ChildPropertyHandle)
+				.OnShouldFilterAsset(this, OnShouldFilterAsset)
+			]
+		];
+	}
+
+	DetailBuilder.HideCategory(GET_MEMBER_NAME_CHECKED(UMaterialXPipelineSettings, PredefinedSurfaceShaders));
+}
+
+bool FInterchangeMaterialXPipelineSettingsCustomization::OnShouldFilterAssetStandardSurface(const FAssetData & InAssetData)
+{
+	return UMaterialXPipelineSettings::ShouldFilterAssets(Cast<UMaterialFunction>(InAssetData.GetAsset()), UMaterialXPipelineSettings::StandardSurfaceInputs, UMaterialXPipelineSettings::StandardSurfaceOutputs);
+}
+
+bool FInterchangeMaterialXPipelineSettingsCustomization::OnShouldFilterAssetStandardSurfaceTransmission(const FAssetData& InAssetData)
+{
+	return UMaterialXPipelineSettings::ShouldFilterAssets(Cast<UMaterialFunction>(InAssetData.GetAsset()), UMaterialXPipelineSettings::TransmissionSurfaceInputs, UMaterialXPipelineSettings::TransmissionSurfaceOutputs);
+}
+
+bool FInterchangeMaterialXPipelineSettingsCustomization::OnShouldFilterAssetSurfaceUnlit(const FAssetData& InAssetData)
+{
+	return UMaterialXPipelineSettings::ShouldFilterAssets(Cast<UMaterialFunction>(InAssetData.GetAsset()), UMaterialXPipelineSettings::SurfaceUnlitInputs, UMaterialXPipelineSettings::SurfaceUnlitOutputs);
+}
+
+bool FInterchangeMaterialXPipelineSettingsCustomization::OnShouldFilterAssetUsdPreviewSurface(const FAssetData& InAssetData)
+{
+	return UMaterialXPipelineSettings::ShouldFilterAssets(Cast<UMaterialFunction>(InAssetData.GetAsset()), UMaterialXPipelineSettings::UsdPreviewSurfaceInputs, UMaterialXPipelineSettings::UsdPreviewSurfaceOutputs);
+}
diff --git a/Engine/Plugins/Interchange/Editor/Source/Pipelines/Private/InterchangeMaterialXPipelineCustomizations.h b/Engine/Plugins/Interchange/Editor/Source/Pipelines/Private/InterchangeMaterialXPipelineCustomizations.h
new file mode 100644
--- /dev/null
+++ b/Engine/Plugins/Interchange/Editor/Source/Pipelines/Private/InterchangeMaterialXPipelineCustomizations.h
@@ -0,0 +1,48 @@
+// Copyright Epic Games, Inc. All Rights Reserved.
+
+#pragma once
+
+#include "IDetailCustomization.h"
+#include "UObject/WeakObjectPtr.h"
+
+class IDetailLayoutBuilder;
+class UMaterialXPipelineSettings;
+struct FAssetData;
+
+class FInterchangeMaterialXPipelineSettingsCustomization : public IDetailCustomization
+{
+public:
+
+	/**
+	 * Creates an instance of this class
+	 */
+	static TSharedRef<IDetailCustomization> MakeInstance();
+
+	/** IDetailCustomization interface */
+	virtual void CustomizeDetails(IDetailLayoutBuilder& DetailBuilder) override;
+	/** End IDetailCustomization interface */
+
+protected:
+	bool OnShouldFilterAsset(const FAssetData&) { return false; }
+	bool OnShouldFilterAssetStandardSurface(const FAssetData& InAssetData);
+	bool OnShouldFilterAssetStandardSurfaceTransmission(const FAssetData& InAssetData);
+	bool OnShouldFilterAssetSurfaceUnlit(const FAssetData& InAssetData);
+	bool OnShouldFilterAssetUsdPreviewSurface(const FAssetData& InAssetData);
+
+private:
+	TWeakObjectPtr<UMaterialXPipelineSettings> MaterialXSettings;
+};
+
+class FInterchangeMaterialXPipelineCustomization : public IDetailCustomization
+{
+public:
+
+	/**
+	 * Creates an instance of this class
+	 */
+	static TSharedRef<IDetailCustomization> MakeInstance();
+
+	/** IDetailCustomization interface */
+	virtual void CustomizeDetails(IDetailLayoutBuilder& DetailBuilder) override;
+	/** End IDetailCustomization interface */
+};
\ No newline at end of file
diff --git a/Engine/Plugins/Interchange/Editor/Source/Pipelines/Private/InterchangeMaterialXPipelineSettingsCustomization.cpp b/Engine/Plugins/Interchange/Editor/Source/Pipelines/Private/InterchangeMaterialXPipelineSettingsCustomization.cpp
deleted file mode 100644
diff --git a/Engine/Plugins/Interchange/Editor/Source/Pipelines/Private/InterchangeMaterialXPipelineSettingsCustomization.h b/Engine/Plugins/Interchange/Editor/Source/Pipelines/Private/InterchangeMaterialXPipelineSettingsCustomization.h
deleted file mode 100644
diff --git a/Engine/Plugins/Interchange/Runtime/Source/Import/Public/MaterialX/InterchangeMaterialXTranslator.h b/Engine/Plugins/Interchange/Runtime/Source/Import/Public/MaterialX/InterchangeMaterialXTranslator.h
--- a/Engine/Plugins/Interchange/Runtime/Source/Import/Public/MaterialX/InterchangeMaterialXTranslator.h
+++ b/Engine/Plugins/Interchange/Runtime/Source/Import/Public/MaterialX/InterchangeMaterialXTranslator.h
@@ -13,7 +13,7 @@ class UInterchangeBaseLightNode;
 class UInterchangeSceneNode;
 
 UCLASS(BlueprintType)
-class UInterchangeMaterialXTranslator : public UInterchangeTranslatorBase, public IInterchangeTexturePayloadInterface
+class INTERCHANGEIMPORT_API UInterchangeMaterialXTranslator : public UInterchangeTranslatorBase, public IInterchangeTexturePayloadInterface
 {
 	GENERATED_BODY()
 
diff --git a/Engine/Plugins/Interchange/Runtime/Source/Pipelines/InterchangePipelines.Build.cs b/Engine/Plugins/Interchange/Runtime/Source/Pipelines/InterchangePipelines.Build.cs
--- a/Engine/Plugins/Interchange/Runtime/Source/Pipelines/InterchangePipelines.Build.cs
+++ b/Engine/Plugins/Interchange/Runtime/Source/Pipelines/InterchangePipelines.Build.cs
@@ -24,6 +24,7 @@ namespace UnrealBuildTool.Rules
 				{
 					"CinematicCamera",
 					"InterchangeImport",
+					"DeveloperSettings",
 				}
 			);
 
diff --git a/Engine/Plugins/Interchange/Runtime/Source/Pipelines/Private/InterchangeGenericMaterialPipeline.cpp b/Engine/Plugins/Interchange/Runtime/Source/Pipelines/Private/InterchangeGenericMaterialPipeline.cpp
--- a/Engine/Plugins/Interchange/Runtime/Source/Pipelines/Private/InterchangeGenericMaterialPipeline.cpp
+++ b/Engine/Plugins/Interchange/Runtime/Source/Pipelines/Private/InterchangeGenericMaterialPipeline.cpp
@@ -79,262 +79,6 @@ FString LexToString(UInterchangeGenericMaterialPipeline::EMaterialInputType Valu
 	}
 }
 
-FMaterialXPipelineSettings::FMaterialXPipelineSettings()
-{
-	SurfaceShader.FindOrAdd(EInterchangeMaterialXShaders::StandardSurface, FSoftObjectPath{ TEXT("MaterialFunction'/Interchange/Functions/MX_StandardSurface.MX_StandardSurface'") });
-	SurfaceShader.FindOrAdd(EInterchangeMaterialXShaders::StandardSurfaceTransmission, FSoftObjectPath{ TEXT("MaterialFunction'/Interchange/Functions/MX_TransmissionSurface.MX_TransmissionSurface'") });
-	SurfaceShader.FindOrAdd(EInterchangeMaterialXShaders::SurfaceUnlit, FSoftObjectPath{ TEXT("MaterialFunction'/Interchange/Functions/MX_SurfaceUnlit.MX_SurfaceUnlit'") });
-	SurfaceShader.FindOrAdd(EInterchangeMaterialXShaders::UsdPreviewSurface, FSoftObjectPath{ TEXT("MaterialFunction'/Interchange/Functions/MX_UsdPreviewSurface.MX_UsdPreviewSurface'") });
-}
-
-bool FMaterialXPipelineSettings::AreRequiredPackagesLoaded()
-{
-	auto ArePackagesLoaded = [&](const TMap<EInterchangeMaterialXShaders, FSoftObjectPath>& ObjectPaths)
-	{
-		bool bAllLoaded = true;
-
-		for(const TPair<EInterchangeMaterialXShaders, FSoftObjectPath>& Pair : ObjectPaths)
-		{
-			const FSoftObjectPath& ObjectPath = Pair.Get<1>();
-
-			if(!ObjectPath.ResolveObject())
-			{
-				FString PackagePath = ObjectPath.GetLongPackageName();
-				if(FPackageName::DoesPackageExist(PackagePath))
-				{
-					UObject* Asset = ObjectPath.TryLoad();
-					if(!Asset)
-					{
-						UE_LOG(LogInterchangePipeline, Warning, TEXT("Couldn't load %s"), *PackagePath);
-						bAllLoaded = false;
-					}
-#if WITH_EDITOR
-					else
-					{
-						if(Pair.Get<EInterchangeMaterialXShaders>() == EInterchangeMaterialXShaders::StandardSurface)
-						{
-							bAllLoaded = !ShouldFilterAssets(Cast<UMaterialFunction>(Asset), StandardSurfaceInputs, StandardSurfaceOutputs);
-						}
-						else if(Pair.Get<EInterchangeMaterialXShaders>() == EInterchangeMaterialXShaders::StandardSurfaceTransmission)
-						{
-							bAllLoaded = !ShouldFilterAssets(Cast<UMaterialFunction>(Asset), TransmissionSurfaceInputs, TransmissionSurfaceOutputs);
-						}
-						else if(Pair.Get<EInterchangeMaterialXShaders>() == EInterchangeMaterialXShaders::SurfaceUnlit)
-						{
-							bAllLoaded = !ShouldFilterAssets(Cast<UMaterialFunction>(Asset), SurfaceUnlitInputs, SurfaceUnlitOutputs);
-						}
-						else if(Pair.Get<EInterchangeMaterialXShaders>() == EInterchangeMaterialXShaders::UsdPreviewSurface)
-						{
-							bAllLoaded = !ShouldFilterAssets(Cast<UMaterialFunction>(Asset), UsdPreviewSurfaceInputs, UsdPreviewSurfaceOutputs);
-						}
-					}
-#endif // WITH_EDITOR
-				}
-				else
-				{
-					UE_LOG(LogInterchangePipeline, Warning, TEXT("Couldn't find %s"), *PackagePath);
-					bAllLoaded = false;
-				}
-			}
-		}
-
-		return bAllLoaded;
-	};
-
-	return ArePackagesLoaded(SurfaceShader);
-}
-
-FString FMaterialXPipelineSettings::GetAssetPathString(EInterchangeMaterialXShaders ShaderType) const
-{
-	FString AssetPath;
-	if(const FSoftObjectPath* ObjectPath = SurfaceShader.Find(ShaderType))
-	{
-		AssetPath = ObjectPath->GetAssetPathString();
-	}
-
-	return AssetPath;
-}
-
-#if WITH_EDITOR
-TSet<FName> FMaterialXPipelineSettings::StandardSurfaceInputs
-{
-	UE::Interchange::Materials::StandardSurface::Parameters::Base,
-	UE::Interchange::Materials::StandardSurface::Parameters::BaseColor,
-	UE::Interchange::Materials::StandardSurface::Parameters::DiffuseRoughness,
-	UE::Interchange::Materials::StandardSurface::Parameters::Metalness,
-	UE::Interchange::Materials::StandardSurface::Parameters::Specular,
-	UE::Interchange::Materials::StandardSurface::Parameters::SpecularRoughness,
-	UE::Interchange::Materials::StandardSurface::Parameters::SpecularIOR,
-	UE::Interchange::Materials::StandardSurface::Parameters::SpecularAnisotropy,
-	UE::Interchange::Materials::StandardSurface::Parameters::SpecularRotation,
-	UE::Interchange::Materials::StandardSurface::Parameters::Subsurface,
-	UE::Interchange::Materials::StandardSurface::Parameters::SubsurfaceColor,
-	UE::Interchange::Materials::StandardSurface::Parameters::SubsurfaceRadius,
-	UE::Interchange::Materials::StandardSurface::Parameters::SubsurfaceScale,
-	UE::Interchange::Materials::StandardSurface::Parameters::Sheen,
-	UE::Interchange::Materials::StandardSurface::Parameters::SheenColor,
-	UE::Interchange::Materials::StandardSurface::Parameters::SheenRoughness,
-	UE::Interchange::Materials::StandardSurface::Parameters::Coat,
-	UE::Interchange::Materials::StandardSurface::Parameters::CoatColor,
-	UE::Interchange::Materials::StandardSurface::Parameters::CoatRoughness,
-	UE::Interchange::Materials::StandardSurface::Parameters::CoatNormal,
-	UE::Interchange::Materials::StandardSurface::Parameters::ThinFilmThickness,
-	UE::Interchange::Materials::StandardSurface::Parameters::Emission,
-	UE::Interchange::Materials::StandardSurface::Parameters::EmissionColor,
-	UE::Interchange::Materials::StandardSurface::Parameters::Normal,
-	UE::Interchange::Materials::StandardSurface::Parameters::Tangent
-};
-
-TSet<FName> FMaterialXPipelineSettings::StandardSurfaceOutputs
-{
-	TEXT("Base Color"), // MX_StandardSurface has BaseColor with a whitespace, this should be fixed in further release
-	UE::Interchange::Materials::PBR::Parameters::Metallic,
-	UE::Interchange::Materials::PBR::Parameters::Specular,
-	UE::Interchange::Materials::PBR::Parameters::Roughness,
-	UE::Interchange::Materials::PBR::Parameters::Anisotropy,
-	UE::Interchange::Materials::PBR::Parameters::EmissiveColor,
-	UE::Interchange::Materials::PBR::Parameters::Opacity,
-	UE::Interchange::Materials::PBR::Parameters::Normal,
-	UE::Interchange::Materials::PBR::Parameters::Tangent,
-	UE::Interchange::Materials::Sheen::Parameters::SheenRoughness,
-	UE::Interchange::Materials::Sheen::Parameters::SheenColor,
-	UE::Interchange::Materials::Subsurface::Parameters::SubsurfaceColor,
-	UE::Interchange::Materials::ClearCoat::Parameters::ClearCoat,
-	UE::Interchange::Materials::ClearCoat::Parameters::ClearCoatRoughness,
-	UE::Interchange::Materials::ClearCoat::Parameters::ClearCoatNormal
-};
-
-TSet<FName> FMaterialXPipelineSettings::TransmissionSurfaceInputs
-{
-	UE::Interchange::Materials::StandardSurface::Parameters::Base,
-	UE::Interchange::Materials::StandardSurface::Parameters::BaseColor,
-	UE::Interchange::Materials::StandardSurface::Parameters::DiffuseRoughness,
-	UE::Interchange::Materials::StandardSurface::Parameters::Metalness,
-	UE::Interchange::Materials::StandardSurface::Parameters::Specular,
-	UE::Interchange::Materials::StandardSurface::Parameters::SpecularRoughness,
-	UE::Interchange::Materials::StandardSurface::Parameters::SpecularIOR,
-	UE::Interchange::Materials::StandardSurface::Parameters::SpecularAnisotropy,
-	UE::Interchange::Materials::StandardSurface::Parameters::SpecularRotation,
-	UE::Interchange::Materials::StandardSurface::Parameters::Transmission,
-	UE::Interchange::Materials::StandardSurface::Parameters::TransmissionColor,
-	UE::Interchange::Materials::StandardSurface::Parameters::TransmissionDepth,
-	UE::Interchange::Materials::StandardSurface::Parameters::TransmissionScatter,
-	UE::Interchange::Materials::StandardSurface::Parameters::TransmissionScatterAnisotropy,
-	UE::Interchange::Materials::StandardSurface::Parameters::TransmissionDispersion,
-	UE::Interchange::Materials::StandardSurface::Parameters::TransmissionExtraRoughness,
-	UE::Interchange::Materials::StandardSurface::Parameters::Subsurface,
-	UE::Interchange::Materials::StandardSurface::Parameters::SubsurfaceColor,
-	UE::Interchange::Materials::StandardSurface::Parameters::SubsurfaceRadius,
-	UE::Interchange::Materials::StandardSurface::Parameters::SubsurfaceScale,
-	UE::Interchange::Materials::StandardSurface::Parameters::Sheen,
-	UE::Interchange::Materials::StandardSurface::Parameters::SheenColor,
-	UE::Interchange::Materials::StandardSurface::Parameters::SheenRoughness,
-	UE::Interchange::Materials::StandardSurface::Parameters::Coat,
-	UE::Interchange::Materials::StandardSurface::Parameters::CoatColor,
-	UE::Interchange::Materials::StandardSurface::Parameters::CoatRoughness,
-	UE::Interchange::Materials::StandardSurface::Parameters::CoatNormal,
-	UE::Interchange::Materials::StandardSurface::Parameters::ThinFilmThickness,
-	UE::Interchange::Materials::StandardSurface::Parameters::Emission,
-	UE::Interchange::Materials::StandardSurface::Parameters::EmissionColor,
-	UE::Interchange::Materials::StandardSurface::Parameters::Normal,
-};
-
-TSet<FName> FMaterialXPipelineSettings::TransmissionSurfaceOutputs
-{
-	UE::Interchange::Materials::PBR::Parameters::BaseColor,
-	UE::Interchange::Materials::PBR::Parameters::Metallic,
-	UE::Interchange::Materials::PBR::Parameters::Specular,
-	UE::Interchange::Materials::PBR::Parameters::Roughness,
-	UE::Interchange::Materials::PBR::Parameters::Anisotropy,
-	UE::Interchange::Materials::PBR::Parameters::EmissiveColor,
-	UE::Interchange::Materials::PBR::Parameters::Opacity,
-	UE::Interchange::Materials::PBR::Parameters::Normal,
-	UE::Interchange::Materials::PBR::Parameters::Tangent,
-	UE::Interchange::Materials::PBR::Parameters::Refraction,
-	UE::Interchange::Materials::ThinTranslucent::Parameters::TransmissionColor
-};
-
-TSet<FName> FMaterialXPipelineSettings::SurfaceUnlitInputs
-{
-	UE::Interchange::Materials::SurfaceUnlit::Parameters::Emission,
-	UE::Interchange::Materials::SurfaceUnlit::Parameters::EmissionColor,
-	UE::Interchange::Materials::SurfaceUnlit::Parameters::Transmission,
-	UE::Interchange::Materials::SurfaceUnlit::Parameters::TransmissionColor,
-	UE::Interchange::Materials::SurfaceUnlit::Parameters::Opacity
-};
-
-TSet<FName> FMaterialXPipelineSettings::SurfaceUnlitOutputs
-{
-	UE::Interchange::Materials::Common::Parameters::EmissiveColor,
-	UE::Interchange::Materials::Common::Parameters::Opacity,
-	TEXT("OpacityMask")
-};
-
-TSet<FName> FMaterialXPipelineSettings::UsdPreviewSurfaceInputs
-{
-	UE::Interchange::Materials::UsdPreviewSurface::Parameters::DiffuseColor,
-	UE::Interchange::Materials::UsdPreviewSurface::Parameters::EmissiveColor,
-	UE::Interchange::Materials::UsdPreviewSurface::Parameters::SpecularColor,
-	UE::Interchange::Materials::UsdPreviewSurface::Parameters::Metallic,
-	UE::Interchange::Materials::UsdPreviewSurface::Parameters::Roughness,
-	UE::Interchange::Materials::UsdPreviewSurface::Parameters::Clearcoat,
-	UE::Interchange::Materials::UsdPreviewSurface::Parameters::ClearcoatRoughness,
-	UE::Interchange::Materials::UsdPreviewSurface::Parameters::Opacity,
-	UE::Interchange::Materials::UsdPreviewSurface::Parameters::OpacityThreshold,
-	UE::Interchange::Materials::UsdPreviewSurface::Parameters::IOR,
-	UE::Interchange::Materials::UsdPreviewSurface::Parameters::Normal,
-	UE::Interchange::Materials::UsdPreviewSurface::Parameters::Displacement,
-	UE::Interchange::Materials::UsdPreviewSurface::Parameters::Occlusion
-};
-
-TSet<FName> FMaterialXPipelineSettings::UsdPreviewSurfaceOutputs
-{
-	UE::Interchange::Materials::PBR::Parameters::BaseColor,
-	UE::Interchange::Materials::PBR::Parameters::Metallic,
-	UE::Interchange::Materials::PBR::Parameters::Specular,
-	UE::Interchange::Materials::PBR::Parameters::Roughness,
-	UE::Interchange::Materials::PBR::Parameters::EmissiveColor,
-	UE::Interchange::Materials::PBR::Parameters::Opacity,
-	UE::Interchange::Materials::PBR::Parameters::Normal,
-	UE::Interchange::Materials::Common::Parameters::Refraction,
-	UE::Interchange::Materials::Common::Parameters::Occlusion,
-	UE::Interchange::Materials::ClearCoat::Parameters::ClearCoat,
-	UE::Interchange::Materials::ClearCoat::Parameters::ClearCoatRoughness,
-};
-
-bool FMaterialXPipelineSettings::ShouldFilterAssets(UMaterialFunction* Asset, const TSet<FName>& Inputs, const TSet<FName>& Outputs)
-{
-	int32 InputMatches = 0;
-	int32 OutputMatches = 0;
-
-	if(Asset != nullptr)
-	{
-		TArray<FFunctionExpressionInput> ExpressionInputs;
-		TArray<FFunctionExpressionOutput> ExpressionOutputs;
-		Asset->GetInputsAndOutputs(ExpressionInputs, ExpressionOutputs);
-
-		for(const FFunctionExpressionInput& ExpressionInput : ExpressionInputs)
-		{
-			if(Inputs.Find(ExpressionInput.Input.InputName))
-			{
-				InputMatches++;
-			}
-		}
-
-		for(const FFunctionExpressionOutput& ExpressionOutput : ExpressionOutputs)
-		{
-			if(Outputs.Find(ExpressionOutput.Output.OutputName))
-			{
-				OutputMatches++;
-			}
-		}
-	}
-
-	// we allow at least one input of the same name, but we should have exactly the same outputs
-	return !(InputMatches > 0 && OutputMatches == Outputs.Num());
-}
-#endif // WITH_EDITOR
-
 namespace UE::Interchange::InterchangeGenericMaterialPipeline::Private
 {
 	bool AreRequiredPackagesLoaded()
@@ -508,11 +252,6 @@ void UInterchangeGenericMaterialPipeline::AdjustSettingsForContext(EInterchangeP
 	{
 		UE_LOG(LogInterchangePipeline, Warning, TEXT("UInterchangeGenericMaterialPipeline: Some required packages are missing. Material import might be wrong"));
 	}
-
-	if(!MaterialXSettings.AreRequiredPackagesLoaded())
-	{
-		UE_LOG(LogInterchangePipeline, Warning, TEXT("UInterchangeGenericMaterialPipeline: Some required packages are missing. Material import might be wrong"));
-	}
 }
 
 void UInterchangeGenericMaterialPipeline::ExecutePipeline(UInterchangeBaseNodeContainer* InBaseNodeContainer, const TArray<UInterchangeSourceData*>& InSourceDatas)
diff --git a/Engine/Plugins/Interchange/Runtime/Source/Pipelines/Private/InterchangeMaterialXPipeline.cpp b/Engine/Plugins/Interchange/Runtime/Source/Pipelines/Private/InterchangeMaterialXPipeline.cpp
new file mode 100644
--- /dev/null
+++ b/Engine/Plugins/Interchange/Runtime/Source/Pipelines/Private/InterchangeMaterialXPipeline.cpp
@@ -0,0 +1,573 @@
+// Copyright Epic Games, Inc. All Rights Reserved. 
+
+#include "InterchangeMaterialXPipeline.h"
+
+#include "InterchangePipelineLog.h"
+
+#include "InterchangeManager.h"
+#include "InterchangeMaterialDefinitions.h"
+#include "InterchangeMaterialFactoryNode.h"
+
+#include "Materials/MaterialExpressionMaterialFunctionCall.h"
+#include "Materials/MaterialFunction.h"
+#include "Misc/PackageName.h"
+
+#if WITH_EDITOR
+namespace UE::Interchange::InterchangeMaterialXPipeline::Private
+{
+	namespace StandardSurface
+	{
+		const FName Name = TEXT("standard_surface");
+
+		namespace Parameters
+		{
+			const FName Base = TEXT("base");
+			const FName BaseColor = TEXT("base_color");
+			const FName DiffuseRoughness = TEXT("diffuse_roughness");
+			const FName Metalness = TEXT("metalness");
+			const FName Specular = TEXT("specular");
+			const FName SpecularColor = TEXT("specular_color");
+			const FName SpecularRoughness = TEXT("specular_roughness");
+			const FName SpecularIOR = TEXT("specular_IOR");
+			const FName SpecularAnisotropy = TEXT("specular_anisotropy");
+			const FName SpecularRotation = TEXT("specular_rotation");
+			const FName Transmission = TEXT("transmission");
+			const FName TransmissionColor = TEXT("transmission_color");
+			const FName TransmissionDepth = TEXT("transmission_depth");
+			const FName TransmissionScatter = TEXT("transmission_scatter");
+			const FName TransmissionScatterAnisotropy = TEXT("transmission_scatter_anisotropy");
+			const FName TransmissionDispersion = TEXT("transmission_dispersion");
+			const FName TransmissionExtraRoughness = TEXT("transmission_extra_roughness");
+			const FName Subsurface = TEXT("subsurface");
+			const FName SubsurfaceColor = TEXT("subsurface_color");
+			const FName SubsurfaceRadius = TEXT("subsurface_radius");
+			const FName SubsurfaceScale = TEXT("subsurface_scale");
+			const FName SubsurfaceAnisotropy = TEXT("subsurface_anisotropy");
+			const FName Sheen = TEXT("sheen");
+			const FName SheenColor = TEXT("sheen_color");
+			const FName SheenRoughness = TEXT("sheen_roughness");
+			const FName Coat = TEXT("coat");
+			const FName CoatColor = TEXT("coat_color");
+			const FName CoatRoughness = TEXT("coat_roughness");
+			const FName CoatAnisotropy = TEXT("coat_anisotropy");
+			const FName CoatRotation = TEXT("coat_rotation");
+			const FName CoatIOR = TEXT("coat_IOR");
+			const FName CoatNormal = TEXT("coat_normal");
+			const FName CoatAffectColor = TEXT("coat_affect_color");
+			const FName CoatAffectRoughness = TEXT("coat_affect_roughness");
+			const FName ThinFilmThickness = TEXT("thin_film_thickness");
+			const FName ThinFilmIOR = TEXT("thin_film_IOR");
+			const FName Emission = TEXT("emission");
+			const FName EmissionColor = TEXT("emission_color");
+			const FName Opacity = TEXT("opacity");
+			const FName ThinWalled = TEXT("thin_walled");
+			const FName Normal = TEXT("normal");
+			const FName Tangent = TEXT("tangent");
+		}
+	}
+
+	namespace SurfaceUnlit
+	{
+		const FName Name = TEXT("surface_unlit");
+
+		namespace Parameters
+		{
+			const FName Emission = TEXT("emission");
+			const FName EmissionColor = TEXT("emission_color");
+			const FName Transmission = TEXT("transmission");
+			const FName TransmissionColor = TEXT("transmission_color");
+			const FName Opacity = TEXT("opacity");
+		}
+	}
+
+	TSet<FName> StandardSurfaceInputs
+	{
+		StandardSurface::Parameters::Base,
+		StandardSurface::Parameters::BaseColor,
+		StandardSurface::Parameters::DiffuseRoughness,
+		StandardSurface::Parameters::Metalness,
+		StandardSurface::Parameters::Specular,
+		StandardSurface::Parameters::SpecularRoughness,
+		StandardSurface::Parameters::SpecularIOR,
+		StandardSurface::Parameters::SpecularAnisotropy,
+		StandardSurface::Parameters::SpecularRotation,
+		StandardSurface::Parameters::Subsurface,
+		StandardSurface::Parameters::SubsurfaceColor,
+		StandardSurface::Parameters::SubsurfaceRadius,
+		StandardSurface::Parameters::SubsurfaceScale,
+		StandardSurface::Parameters::Sheen,
+		StandardSurface::Parameters::SheenColor,
+		StandardSurface::Parameters::SheenRoughness,
+		StandardSurface::Parameters::Coat,
+		StandardSurface::Parameters::CoatColor,
+		StandardSurface::Parameters::CoatRoughness,
+		StandardSurface::Parameters::CoatNormal,
+		StandardSurface::Parameters::ThinFilmThickness,
+		StandardSurface::Parameters::Emission,
+		StandardSurface::Parameters::EmissionColor,
+		StandardSurface::Parameters::Normal,
+		StandardSurface::Parameters::Tangent
+	};
+
+	TSet<FName> StandardSurfaceOutputs
+	{
+		TEXT("Base Color"), // MX_StandardSurface has BaseColor with a whitespace, this should be fixed in further release
+		UE::Interchange::Materials::PBR::Parameters::Metallic,
+		UE::Interchange::Materials::PBR::Parameters::Specular,
+		UE::Interchange::Materials::PBR::Parameters::Roughness,
+		UE::Interchange::Materials::PBR::Parameters::Anisotropy,
+		UE::Interchange::Materials::PBR::Parameters::EmissiveColor,
+		UE::Interchange::Materials::PBR::Parameters::Opacity,
+		UE::Interchange::Materials::PBR::Parameters::Normal,
+		UE::Interchange::Materials::PBR::Parameters::Tangent,
+		UE::Interchange::Materials::Sheen::Parameters::SheenRoughness,
+		UE::Interchange::Materials::Sheen::Parameters::SheenColor,
+		UE::Interchange::Materials::Subsurface::Parameters::SubsurfaceColor,
+		UE::Interchange::Materials::ClearCoat::Parameters::ClearCoat,
+		UE::Interchange::Materials::ClearCoat::Parameters::ClearCoatRoughness,
+		UE::Interchange::Materials::ClearCoat::Parameters::ClearCoatNormal
+	};
+
+	TSet<FName> TransmissionSurfaceInputs
+	{
+		StandardSurface::Parameters::Base,
+		StandardSurface::Parameters::BaseColor,
+		StandardSurface::Parameters::DiffuseRoughness,
+		StandardSurface::Parameters::Metalness,
+		StandardSurface::Parameters::Specular,
+		StandardSurface::Parameters::SpecularRoughness,
+		StandardSurface::Parameters::SpecularIOR,
+		StandardSurface::Parameters::SpecularAnisotropy,
+		StandardSurface::Parameters::SpecularRotation,
+		StandardSurface::Parameters::Transmission,
+		StandardSurface::Parameters::TransmissionColor,
+		StandardSurface::Parameters::TransmissionDepth,
+		StandardSurface::Parameters::TransmissionScatter,
+		StandardSurface::Parameters::TransmissionScatterAnisotropy,
+		StandardSurface::Parameters::TransmissionDispersion,
+		StandardSurface::Parameters::TransmissionExtraRoughness,
+		StandardSurface::Parameters::Subsurface,
+		StandardSurface::Parameters::SubsurfaceColor,
+		StandardSurface::Parameters::SubsurfaceRadius,
+		StandardSurface::Parameters::SubsurfaceScale,
+		StandardSurface::Parameters::Sheen,
+		StandardSurface::Parameters::SheenColor,
+		StandardSurface::Parameters::SheenRoughness,
+		StandardSurface::Parameters::Coat,
+		StandardSurface::Parameters::CoatColor,
+		StandardSurface::Parameters::CoatRoughness,
+		StandardSurface::Parameters::CoatNormal,
+		StandardSurface::Parameters::ThinFilmThickness,
+		StandardSurface::Parameters::Emission,
+		StandardSurface::Parameters::EmissionColor,
+		StandardSurface::Parameters::Normal,
+	};
+
+	TSet<FName> TransmissionSurfaceOutputs
+	{
+		UE::Interchange::Materials::PBR::Parameters::BaseColor,
+		UE::Interchange::Materials::PBR::Parameters::Metallic,
+		UE::Interchange::Materials::PBR::Parameters::Specular,
+		UE::Interchange::Materials::PBR::Parameters::Roughness,
+		UE::Interchange::Materials::PBR::Parameters::Anisotropy,
+		UE::Interchange::Materials::PBR::Parameters::EmissiveColor,
+		UE::Interchange::Materials::PBR::Parameters::Opacity,
+		UE::Interchange::Materials::PBR::Parameters::Normal,
+		UE::Interchange::Materials::PBR::Parameters::Tangent,
+		UE::Interchange::Materials::PBR::Parameters::Refraction,
+		UE::Interchange::Materials::ThinTranslucent::Parameters::TransmissionColor
+	};
+
+	TSet<FName> SurfaceUnlitInputs
+	{
+		SurfaceUnlit::Parameters::Emission,
+		SurfaceUnlit::Parameters::EmissionColor,
+		SurfaceUnlit::Parameters::Transmission,
+		SurfaceUnlit::Parameters::TransmissionColor,
+		SurfaceUnlit::Parameters::Opacity
+	};
+
+	TSet<FName> SurfaceUnlitOutputs
+	{
+		UE::Interchange::Materials::Common::Parameters::EmissiveColor,
+		UE::Interchange::Materials::Common::Parameters::Opacity,
+	};
+
+	bool ShouldFilterAssets(UMaterialFunction* Asset, const TSet<FName>& Inputs, const TSet<FName>& Outputs)
+	{
+		int32 InputMatches = 0;
+		int32 OutputMatches = 0;
+
+		if (Asset != nullptr)
+		{
+			TArray<FFunctionExpressionInput> ExpressionInputs;
+			TArray<FFunctionExpressionOutput> ExpressionOutputs;
+			Asset->GetInputsAndOutputs(ExpressionInputs, ExpressionOutputs);
+
+			for (const FFunctionExpressionInput& ExpressionInput : ExpressionInputs)
+			{
+				if (Inputs.Find(ExpressionInput.Input.InputName))
+				{
+					InputMatches++;
+				}
+			}
+
+			for (const FFunctionExpressionOutput& ExpressionOutput : ExpressionOutputs)
+			{
+				if (Outputs.Find(ExpressionOutput.Output.OutputName))
+				{
+					OutputMatches++;
+				}
+			}
+		}
+
+		// we allow at least one input of the same name, but we should have exactly the same outputs
+		return !(InputMatches > 0 && OutputMatches == Outputs.Num());
+	}
+
+	bool AreRequiredPackagesLoaded(TMap<EInterchangeMaterialXShaders, FSoftObjectPath>& SurfaceShaderMapping)
+	{
+		auto ArePackagesLoaded = [&](const TMap<EInterchangeMaterialXShaders, FSoftObjectPath>& ObjectPaths)
+		{
+			bool bAllLoaded = true;
+
+			for (const TPair<EInterchangeMaterialXShaders, FSoftObjectPath>& Pair : ObjectPaths)
+			{
+				const FSoftObjectPath& ObjectPath = Pair.Get<1>();
+
+				if (!ObjectPath.ResolveObject())
+				{
+					FString PackagePath = ObjectPath.GetLongPackageName();
+					if (FPackageName::DoesPackageExist(PackagePath))
+					{
+						UObject* Asset = ObjectPath.TryLoad();
+						if (!Asset)
+						{
+							UE_LOG(LogInterchangePipeline, Warning, TEXT("Couldn't load %s"), *PackagePath);
+							bAllLoaded = false;
+						}
+#if WITH_EDITOR
+						else
+						{
+							if (Pair.Get<EInterchangeMaterialXShaders>() == EInterchangeMaterialXShaders::StandardSurface)
+							{
+								bAllLoaded = !ShouldFilterAssets(Cast<UMaterialFunction>(Asset), StandardSurfaceInputs, StandardSurfaceOutputs);
+							}
+							else if (Pair.Get<EInterchangeMaterialXShaders>() == EInterchangeMaterialXShaders::StandardSurfaceTransmission)
+							{
+								bAllLoaded = !ShouldFilterAssets(Cast<UMaterialFunction>(Asset), TransmissionSurfaceInputs, TransmissionSurfaceOutputs);
+							}
+							else if (Pair.Get<EInterchangeMaterialXShaders>() == EInterchangeMaterialXShaders::SurfaceUnlit)
+							{
+								bAllLoaded = !ShouldFilterAssets(Cast<UMaterialFunction>(Asset), SurfaceUnlitInputs, SurfaceUnlitOutputs);
+							}
+
+						}
+#endif // WITH_EDITOR
+					}
+					else
+					{
+						UE_LOG(LogInterchangePipeline, Warning, TEXT("Couldn't find %s"), *PackagePath);
+						bAllLoaded = false;
+					}
+				}
+			}
+
+			return bAllLoaded;
+		};
+
+		return ArePackagesLoaded(SurfaceShaderMapping);
+	}
+}
+#endif
+
+UInterchangeMaterialXPipeline::UInterchangeMaterialXPipeline()
+	: MaterialXSettings(UMaterialXPipelineSettings::StaticClass()->GetDefaultObject< UMaterialXPipelineSettings>())
+{
+	for (const TPair<EInterchangeMaterialXShaders, FSoftObjectPath>& Entry : MaterialXSettings->PredefinedSurfaceShaders)
+	{
+		PathToEnumMapping.FindOrAdd(Entry.Value.GetAssetPathString(), Entry.Key);
+	}
+}
+
+void UInterchangeMaterialXPipeline::AdjustSettingsForContext(EInterchangePipelineContext ImportType, TObjectPtr<UObject> ReimportAsset)
+{
+	Super::AdjustSettingsForContext(ImportType, ReimportAsset);
+
+	if (!MaterialXSettings->AreRequiredPackagesLoaded())
+	{
+		UE_LOG(LogInterchangePipeline, Warning, TEXT("UInterchangeGenericMaterialPipeline: Some required packages are missing. Material import might be wrong"));
+	}
+}
+
+void UInterchangeMaterialXPipeline::ExecutePipeline(UInterchangeBaseNodeContainer* NodeContainer, const TArray<UInterchangeSourceData*>& InSourceDatas)
+{
+	Super::ExecutePipeline(NodeContainer, InSourceDatas);
+
+#if WITH_EDITOR
+	auto UpdateMaterialXNodes = [this](const FString& NodeUid, UInterchangeMaterialFunctionCallExpressionFactoryNode* FactorNode)
+	{
+		const FString MaterialFunctionMemberName = GET_MEMBER_NAME_CHECKED(UMaterialExpressionMaterialFunctionCall, MaterialFunction).ToString();
+
+		FString MaterialFunctionPath;
+		if (FactorNode->GetStringAttribute(MaterialFunctionMemberName, MaterialFunctionPath))
+		{
+			if (const EInterchangeMaterialXShaders* EnumPtr = PathToEnumMapping.Find(MaterialFunctionPath))
+			{
+				FactorNode->AddStringAttribute(MaterialFunctionMemberName, MaterialXSettings->GetAssetPathString(*EnumPtr));
+			}
+		}
+	};
+
+	//Find all translated node we need for this pipeline
+	NodeContainer->IterateNodesOfType<UInterchangeMaterialFunctionCallExpressionFactoryNode>(UpdateMaterialXNodes);
+#endif
+}
+
+bool UMaterialXPipelineSettings::AreRequiredPackagesLoaded()
+{
+	auto ArePackagesLoaded = [&](const TMap<EInterchangeMaterialXShaders, FSoftObjectPath>& ObjectPaths)
+	{
+		bool bAllLoaded = true;
+
+		for (const TPair<EInterchangeMaterialXShaders, FSoftObjectPath>& Pair : ObjectPaths)
+		{
+			const FSoftObjectPath& ObjectPath = Pair.Get<1>();
+
+			if (!ObjectPath.ResolveObject())
+			{
+				FString PackagePath = ObjectPath.GetLongPackageName();
+				if (FPackageName::DoesPackageExist(PackagePath))
+				{
+					UObject* Asset = ObjectPath.TryLoad();
+					if (!Asset)
+					{
+						UE_LOG(LogInterchangePipeline, Warning, TEXT("Couldn't load %s"), *PackagePath);
+						bAllLoaded = false;
+					}
+#if WITH_EDITOR
+					else
+					{
+						if (Pair.Get<EInterchangeMaterialXShaders>() == EInterchangeMaterialXShaders::StandardSurface)
+						{
+							bAllLoaded = !ShouldFilterAssets(Cast<UMaterialFunction>(Asset), StandardSurfaceInputs, StandardSurfaceOutputs);
+						}
+						else if (Pair.Get<EInterchangeMaterialXShaders>() == EInterchangeMaterialXShaders::StandardSurfaceTransmission)
+						{
+							bAllLoaded = !ShouldFilterAssets(Cast<UMaterialFunction>(Asset), TransmissionSurfaceInputs, TransmissionSurfaceOutputs);
+						}
+						else if (Pair.Get<EInterchangeMaterialXShaders>() == EInterchangeMaterialXShaders::SurfaceUnlit)
+						{
+							bAllLoaded = !ShouldFilterAssets(Cast<UMaterialFunction>(Asset), SurfaceUnlitInputs, SurfaceUnlitOutputs);
+						}
+						else if (Pair.Get<EInterchangeMaterialXShaders>() == EInterchangeMaterialXShaders::UsdPreviewSurface)
+						{
+							bAllLoaded = !ShouldFilterAssets(Cast<UMaterialFunction>(Asset), UsdPreviewSurfaceInputs, UsdPreviewSurfaceOutputs);
+						}
+					}
+#endif // WITH_EDITOR
+				}
+				else
+				{
+					UE_LOG(LogInterchangePipeline, Warning, TEXT("Couldn't find %s"), *PackagePath);
+					bAllLoaded = false;
+				}
+			}
+		}
+
+		return bAllLoaded;
+	};
+
+	return ArePackagesLoaded(PredefinedSurfaceShaders);
+}
+
+FString UMaterialXPipelineSettings::GetAssetPathString(EInterchangeMaterialXShaders ShaderType) const
+{
+	if (const FSoftObjectPath* ObjectPath = PredefinedSurfaceShaders.Find(ShaderType))
+	{
+		return ObjectPath->GetAssetPathString();
+	}
+
+	return {};
+}
+
+#if WITH_EDITOR
+TSet<FName> UMaterialXPipelineSettings::StandardSurfaceInputs
+{
+	UE::Interchange::Materials::StandardSurface::Parameters::Base,
+	UE::Interchange::Materials::StandardSurface::Parameters::BaseColor,
+	UE::Interchange::Materials::StandardSurface::Parameters::DiffuseRoughness,
+	UE::Interchange::Materials::StandardSurface::Parameters::Metalness,
+	UE::Interchange::Materials::StandardSurface::Parameters::Specular,
+	UE::Interchange::Materials::StandardSurface::Parameters::SpecularRoughness,
+	UE::Interchange::Materials::StandardSurface::Parameters::SpecularIOR,
+	UE::Interchange::Materials::StandardSurface::Parameters::SpecularAnisotropy,
+	UE::Interchange::Materials::StandardSurface::Parameters::SpecularRotation,
+	UE::Interchange::Materials::StandardSurface::Parameters::Subsurface,
+	UE::Interchange::Materials::StandardSurface::Parameters::SubsurfaceColor,
+	UE::Interchange::Materials::StandardSurface::Parameters::SubsurfaceRadius,
+	UE::Interchange::Materials::StandardSurface::Parameters::SubsurfaceScale,
+	UE::Interchange::Materials::StandardSurface::Parameters::Sheen,
+	UE::Interchange::Materials::StandardSurface::Parameters::SheenColor,
+	UE::Interchange::Materials::StandardSurface::Parameters::SheenRoughness,
+	UE::Interchange::Materials::StandardSurface::Parameters::Coat,
+	UE::Interchange::Materials::StandardSurface::Parameters::CoatColor,
+	UE::Interchange::Materials::StandardSurface::Parameters::CoatRoughness,
+	UE::Interchange::Materials::StandardSurface::Parameters::CoatNormal,
+	UE::Interchange::Materials::StandardSurface::Parameters::ThinFilmThickness,
+	UE::Interchange::Materials::StandardSurface::Parameters::Emission,
+	UE::Interchange::Materials::StandardSurface::Parameters::EmissionColor,
+	UE::Interchange::Materials::StandardSurface::Parameters::Normal,
+	UE::Interchange::Materials::StandardSurface::Parameters::Tangent
+};
+
+TSet<FName> UMaterialXPipelineSettings::StandardSurfaceOutputs
+{
+	TEXT("Base Color"), // MX_StandardSurface has BaseColor with a whitespace, this should be fixed in further release
+	UE::Interchange::Materials::PBR::Parameters::Metallic,
+	UE::Interchange::Materials::PBR::Parameters::Specular,
+	UE::Interchange::Materials::PBR::Parameters::Roughness,
+	UE::Interchange::Materials::PBR::Parameters::Anisotropy,
+	UE::Interchange::Materials::PBR::Parameters::EmissiveColor,
+	UE::Interchange::Materials::PBR::Parameters::Opacity,
+	UE::Interchange::Materials::PBR::Parameters::Normal,
+	UE::Interchange::Materials::PBR::Parameters::Tangent,
+	UE::Interchange::Materials::Sheen::Parameters::SheenRoughness,
+	UE::Interchange::Materials::Sheen::Parameters::SheenColor,
+	UE::Interchange::Materials::Subsurface::Parameters::SubsurfaceColor,
+	UE::Interchange::Materials::ClearCoat::Parameters::ClearCoat,
+	UE::Interchange::Materials::ClearCoat::Parameters::ClearCoatRoughness,
+	UE::Interchange::Materials::ClearCoat::Parameters::ClearCoatNormal
+};
+
+TSet<FName> UMaterialXPipelineSettings::TransmissionSurfaceInputs
+{
+	UE::Interchange::Materials::StandardSurface::Parameters::Base,
+	UE::Interchange::Materials::StandardSurface::Parameters::BaseColor,
+	UE::Interchange::Materials::StandardSurface::Parameters::DiffuseRoughness,
+	UE::Interchange::Materials::StandardSurface::Parameters::Metalness,
+	UE::Interchange::Materials::StandardSurface::Parameters::Specular,
+	UE::Interchange::Materials::StandardSurface::Parameters::SpecularRoughness,
+	UE::Interchange::Materials::StandardSurface::Parameters::SpecularIOR,
+	UE::Interchange::Materials::StandardSurface::Parameters::SpecularAnisotropy,
+	UE::Interchange::Materials::StandardSurface::Parameters::SpecularRotation,
+	UE::Interchange::Materials::StandardSurface::Parameters::Transmission,
+	UE::Interchange::Materials::StandardSurface::Parameters::TransmissionColor,
+	UE::Interchange::Materials::StandardSurface::Parameters::TransmissionDepth,
+	UE::Interchange::Materials::StandardSurface::Parameters::TransmissionScatter,
+	UE::Interchange::Materials::StandardSurface::Parameters::TransmissionScatterAnisotropy,
+	UE::Interchange::Materials::StandardSurface::Parameters::TransmissionDispersion,
+	UE::Interchange::Materials::StandardSurface::Parameters::TransmissionExtraRoughness,
+	UE::Interchange::Materials::StandardSurface::Parameters::Subsurface,
+	UE::Interchange::Materials::StandardSurface::Parameters::SubsurfaceColor,
+	UE::Interchange::Materials::StandardSurface::Parameters::SubsurfaceRadius,
+	UE::Interchange::Materials::StandardSurface::Parameters::SubsurfaceScale,
+	UE::Interchange::Materials::StandardSurface::Parameters::Sheen,
+	UE::Interchange::Materials::StandardSurface::Parameters::SheenColor,
+	UE::Interchange::Materials::StandardSurface::Parameters::SheenRoughness,
+	UE::Interchange::Materials::StandardSurface::Parameters::Coat,
+	UE::Interchange::Materials::StandardSurface::Parameters::CoatColor,
+	UE::Interchange::Materials::StandardSurface::Parameters::CoatRoughness,
+	UE::Interchange::Materials::StandardSurface::Parameters::CoatNormal,
+	UE::Interchange::Materials::StandardSurface::Parameters::ThinFilmThickness,
+	UE::Interchange::Materials::StandardSurface::Parameters::Emission,
+	UE::Interchange::Materials::StandardSurface::Parameters::EmissionColor,
+	UE::Interchange::Materials::StandardSurface::Parameters::Normal,
+};
+
+TSet<FName> UMaterialXPipelineSettings::TransmissionSurfaceOutputs
+{
+	UE::Interchange::Materials::PBR::Parameters::BaseColor,
+	UE::Interchange::Materials::PBR::Parameters::Metallic,
+	UE::Interchange::Materials::PBR::Parameters::Specular,
+	UE::Interchange::Materials::PBR::Parameters::Roughness,
+	UE::Interchange::Materials::PBR::Parameters::Anisotropy,
+	UE::Interchange::Materials::PBR::Parameters::EmissiveColor,
+	UE::Interchange::Materials::PBR::Parameters::Opacity,
+	UE::Interchange::Materials::PBR::Parameters::Normal,
+	UE::Interchange::Materials::PBR::Parameters::Tangent,
+	UE::Interchange::Materials::PBR::Parameters::Refraction,
+	UE::Interchange::Materials::ThinTranslucent::Parameters::TransmissionColor
+};
+
+TSet<FName> UMaterialXPipelineSettings::SurfaceUnlitInputs
+{
+	UE::Interchange::Materials::SurfaceUnlit::Parameters::Emission,
+	UE::Interchange::Materials::SurfaceUnlit::Parameters::EmissionColor,
+	UE::Interchange::Materials::SurfaceUnlit::Parameters::Transmission,
+	UE::Interchange::Materials::SurfaceUnlit::Parameters::TransmissionColor,
+	UE::Interchange::Materials::SurfaceUnlit::Parameters::Opacity
+};
+
+TSet<FName> UMaterialXPipelineSettings::SurfaceUnlitOutputs
+{
+	UE::Interchange::Materials::Common::Parameters::EmissiveColor,
+	UE::Interchange::Materials::Common::Parameters::Opacity,
+	TEXT("OpacityMask")
+};
+
+TSet<FName> UMaterialXPipelineSettings::UsdPreviewSurfaceInputs
+{
+	UE::Interchange::Materials::UsdPreviewSurface::Parameters::DiffuseColor,
+	UE::Interchange::Materials::UsdPreviewSurface::Parameters::EmissiveColor,
+	UE::Interchange::Materials::UsdPreviewSurface::Parameters::SpecularColor,
+	UE::Interchange::Materials::UsdPreviewSurface::Parameters::Metallic,
+	UE::Interchange::Materials::UsdPreviewSurface::Parameters::Roughness,
+	UE::Interchange::Materials::UsdPreviewSurface::Parameters::Clearcoat,
+	UE::Interchange::Materials::UsdPreviewSurface::Parameters::ClearcoatRoughness,
+	UE::Interchange::Materials::UsdPreviewSurface::Parameters::Opacity,
+	UE::Interchange::Materials::UsdPreviewSurface::Parameters::OpacityThreshold,
+	UE::Interchange::Materials::UsdPreviewSurface::Parameters::IOR,
+	UE::Interchange::Materials::UsdPreviewSurface::Parameters::Normal,
+	UE::Interchange::Materials::UsdPreviewSurface::Parameters::Displacement,
+	UE::Interchange::Materials::UsdPreviewSurface::Parameters::Occlusion
+};
+
+TSet<FName> UMaterialXPipelineSettings::UsdPreviewSurfaceOutputs
+{
+	UE::Interchange::Materials::PBR::Parameters::BaseColor,
+	UE::Interchange::Materials::PBR::Parameters::Metallic,
+	UE::Interchange::Materials::PBR::Parameters::Specular,
+	UE::Interchange::Materials::PBR::Parameters::Roughness,
+	UE::Interchange::Materials::PBR::Parameters::EmissiveColor,
+	UE::Interchange::Materials::PBR::Parameters::Opacity,
+	UE::Interchange::Materials::PBR::Parameters::Normal,
+	UE::Interchange::Materials::Common::Parameters::Refraction,
+	UE::Interchange::Materials::Common::Parameters::Occlusion,
+	UE::Interchange::Materials::ClearCoat::Parameters::ClearCoat,
+	UE::Interchange::Materials::ClearCoat::Parameters::ClearCoatRoughness,
+};
+
+bool UMaterialXPipelineSettings::ShouldFilterAssets(UMaterialFunction* Asset, const TSet<FName>& Inputs, const TSet<FName>& Outputs)
+{
+	int32 InputMatches = 0;
+	int32 OutputMatches = 0;
+
+	if (Asset != nullptr)
+	{
+		TArray<FFunctionExpressionInput> ExpressionInputs;
+		TArray<FFunctionExpressionOutput> ExpressionOutputs;
+		Asset->GetInputsAndOutputs(ExpressionInputs, ExpressionOutputs);
+
+		for (const FFunctionExpressionInput& ExpressionInput : ExpressionInputs)
+		{
+			if (Inputs.Find(ExpressionInput.Input.InputName))
+			{
+				InputMatches++;
+			}
+		}
+
+		for (const FFunctionExpressionOutput& ExpressionOutput : ExpressionOutputs)
+		{
+			if (Outputs.Find(ExpressionOutput.Output.OutputName))
+			{
+				OutputMatches++;
+			}
+		}
+	}
+
+	// we allow at least one input of the same name, but we should have exactly the same outputs
+	return !(InputMatches > 0 && OutputMatches == Outputs.Num());
+}
+#endif // WITH_EDITOR
+
diff --git a/Engine/Plugins/Interchange/Runtime/Source/Pipelines/Private/InterchangePipelinesModule.cpp b/Engine/Plugins/Interchange/Runtime/Source/Pipelines/Private/InterchangePipelinesModule.cpp
--- a/Engine/Plugins/Interchange/Runtime/Source/Pipelines/Private/InterchangePipelinesModule.cpp
+++ b/Engine/Plugins/Interchange/Runtime/Source/Pipelines/Private/InterchangePipelinesModule.cpp
@@ -3,11 +3,7 @@
 #include "InterchangePipelinesModule.h"
 
 #include "CoreMinimal.h"
-#include "Engine/Engine.h"
-//#include "InterchangeGenericAssetsPipeline.h"
-#include "InterchangeManager.h"
 #include "InterchangePipelineLog.h"
-#include "Misc/CoreDelegates.h"
 #include "Modules/ModuleManager.h"
 
 DEFINE_LOG_CATEGORY(LogInterchangePipeline);
@@ -15,38 +11,8 @@ DEFINE_LOG_CATEGORY(LogInterchangePipeline);
 class FInterchangePipelinesModule : public IInterchangePipelinesModule
 {
 	/** IModuleInterface implementation */
-	virtual void StartupModule() override;
-	virtual void ShutdownModule() override;
+	virtual void StartupModule() override {}
+	virtual void ShutdownModule() override {}
 };
 
 IMPLEMENT_MODULE(FInterchangePipelinesModule, InterchangePipelines)
-
-
-
-void FInterchangePipelinesModule::StartupModule()
-{
-	auto RegisterItems = []()
-	{
-		UInterchangeManager& InterchangeManager = UInterchangeManager::GetInterchangeManager();
-		//Register the Pipeline
-		//InterchangeManager.RegisterPipeline(UInterchangeGenericMeshPipeline::StaticClass());
-	};
-
-	if (GEngine)
-	{
-		RegisterItems();
-	}
-	else
-	{
-		FCoreDelegates::OnPostEngineInit.AddLambda(RegisterItems);
-	}
-}
-
-
-void FInterchangePipelinesModule::ShutdownModule()
-{
-
-}
-
-
-
diff --git a/Engine/Plugins/Interchange/Runtime/Source/Pipelines/Private/InterchangeglTFPipeline.cpp b/Engine/Plugins/Interchange/Runtime/Source/Pipelines/Private/InterchangeglTFPipeline.cpp
new file mode 100644
--- /dev/null
+++ b/Engine/Plugins/Interchange/Runtime/Source/Pipelines/Private/InterchangeglTFPipeline.cpp
@@ -0,0 +1,15 @@
+// Copyright Epic Games, Inc. All Rights Reserved. 
+
+#include "InterchangeglTFPipeline.h"
+
+#include "InterchangePipelineLog.h"
+
+#include "InterchangeMeshFactoryNode.h"
+
+void UInterchangeglTFPipeline::ExecutePipeline(UInterchangeBaseNodeContainer* NodeContainer, const TArray<UInterchangeSourceData*>& InSourceDatas)
+{
+	Super::ExecutePipeline(NodeContainer, InSourceDatas);
+
+	// TODO: Put logic to convert material graphs to material instances when applicable.
+}
+
diff --git a/Engine/Plugins/Interchange/Runtime/Source/Pipelines/Public/InterchangeGenericMaterialPipeline.h b/Engine/Plugins/Interchange/Runtime/Source/Pipelines/Public/InterchangeGenericMaterialPipeline.h
--- a/Engine/Plugins/Interchange/Runtime/Source/Pipelines/Public/InterchangeGenericMaterialPipeline.h
+++ b/Engine/Plugins/Interchange/Runtime/Source/Pipelines/Public/InterchangeGenericMaterialPipeline.h
@@ -31,55 +31,6 @@ enum class EInterchangeMaterialImportOption : uint8
 	ImportAsMaterialInstances,
 };
 
-
-UENUM(BlueprintType)
-enum class EInterchangeMaterialXShaders : uint8
-{
-	/** Default settings for Autodesk's Standard Surface shader	*/
-	StandardSurface,
-
-	/** Standard Surface shader used for translucency	*/
-	StandardSurfaceTransmission,
-
-	/** Shader used for unlit surface*/
-	SurfaceUnlit,
-
-	UsdPreviewSurface
-};
-
-USTRUCT(BlueprintType)
-struct INTERCHANGEPIPELINES_API FMaterialXPipelineSettings
-{
-	GENERATED_USTRUCT_BODY()
-
-public:
-	FMaterialXPipelineSettings();
-
-	UPROPERTY(EditAnywhere, BlueprintReadWrite, Category = "MaterialX", meta=(AllowedClasses="/Script/Engine.MaterialFunction"))
-	TMap<EInterchangeMaterialXShaders, FSoftObjectPath> SurfaceShader;
-
-	bool AreRequiredPackagesLoaded();
-
-	FString GetAssetPathString(EInterchangeMaterialXShaders ShaderType) const;
-
-private:
-
-#if WITH_EDITOR
-	friend class FInterchangeMaterialXPipelineSettingsCustomization;
-
-	static bool ShouldFilterAssets(UMaterialFunction* Asset, const TSet<FName>& Inputs, const TSet<FName>& Outputs);
-
-	static TSet<FName> StandardSurfaceInputs;
-	static TSet<FName> StandardSurfaceOutputs;
-	static TSet<FName> TransmissionSurfaceInputs;
-	static TSet<FName> TransmissionSurfaceOutputs;
-	static TSet<FName> SurfaceUnlitInputs;
-	static TSet<FName> SurfaceUnlitOutputs;
-	static TSet<FName> UsdPreviewSurfaceInputs;
-	static TSet<FName> UsdPreviewSurfaceOutputs;
-#endif // WITH_EDITOR
-};
-
 UCLASS(BlueprintType, editinlinenew)
 class INTERCHANGEPIPELINES_API UInterchangeGenericMaterialPipeline : public UInterchangePipelineBase
 {
@@ -102,9 +53,6 @@ public:
 	UPROPERTY(VisibleAnywhere, Instanced, Category = "Textures")
 	TObjectPtr<UInterchangeGenericTexturePipeline> TexturePipeline;
 
-	UPROPERTY(EditAnywhere, BlueprintReadWrite, Category = "MaterialX", DisplayName = "Settings")
-	FMaterialXPipelineSettings MaterialXSettings;
-
 	/** BEGIN UInterchangePipelineBase overrides */
 	virtual void PreDialogCleanup(const FName PipelineStackName) override;
 	virtual bool IsSettingsAreValid(TOptional<FText>& OutInvalidReason) const override;
diff --git a/Engine/Plugins/Interchange/Runtime/Source/Pipelines/Public/InterchangeMaterialXPipeline.h b/Engine/Plugins/Interchange/Runtime/Source/Pipelines/Public/InterchangeMaterialXPipeline.h
new file mode 100644
--- /dev/null
+++ b/Engine/Plugins/Interchange/Runtime/Source/Pipelines/Public/InterchangeMaterialXPipeline.h
@@ -0,0 +1,101 @@
+// Copyright Epic Games, Inc. All Rights Reserved.
+
+#pragma once
+
+#include "CoreMinimal.h"
+
+#include "InterchangePipelineBase.h"
+#include "InterchangeSourceData.h"
+#include "Nodes/InterchangeBaseNodeContainer.h"
+
+#include "Engine/DeveloperSettings.h"
+#include "Templates/Tuple.h"
+#include "UObject/Object.h"
+#include "UObject/ObjectMacros.h"
+#include "UObject/SoftObjectPath.h"
+
+#include "InterchangeMaterialXPipeline.generated.h"
+
+class UInterchangeDatasmithPbrMaterialNode;
+class UInterchangeFactoryBaseNode;
+class UInterchangeMaterialFactoryNode;
+class UInterchangeMaterialInstanceFactoryNode;
+class UInterchangeShaderGraphNode;
+class UInterchangeMaterialInstanceNode;
+
+class UMaterialFunction;
+class UMaterialInterface;
+
+UENUM(BlueprintType)
+enum class EInterchangeMaterialXShaders : uint8
+{
+	/** Default settings for Autodesk's Standard Surface shader	*/
+	StandardSurface,
+
+	/** Standard Surface shader used for translucency	*/
+	StandardSurfaceTransmission,
+
+	/** Shader used for unlit surface*/
+	SurfaceUnlit,
+
+	/** Default settings for USD's Surface shader	*/
+	UsdPreviewSurface,
+
+	MaxShaderCount
+};
+
+UCLASS(config = Engine, meta = (DisplayName = "Interchange MaterialX"))
+class INTERCHANGEPIPELINES_API UMaterialXPipelineSettings : public UDeveloperSettings
+{
+	GENERATED_BODY()
+
+public:
+	UPROPERTY(EditAnywhere, config, Category = "MaterialXPredefined", meta = (DisplayName = "MaterialX Predefined Surface Shaders"))
+	TMap<EInterchangeMaterialXShaders, FSoftObjectPath> PredefinedSurfaceShaders;
+
+	bool AreRequiredPackagesLoaded();
+
+	FString GetAssetPathString(EInterchangeMaterialXShaders ShaderType) const;
+
+private:
+
+#if WITH_EDITOR
+	friend class FInterchangeMaterialXPipelineSettingsCustomization;
+
+	static bool ShouldFilterAssets(UMaterialFunction* Asset, const TSet<FName>& Inputs, const TSet<FName>& Outputs);
+
+	static TSet<FName> StandardSurfaceInputs;
+	static TSet<FName> StandardSurfaceOutputs;
+	static TSet<FName> TransmissionSurfaceInputs;
+	static TSet<FName> TransmissionSurfaceOutputs;
+	static TSet<FName> SurfaceUnlitInputs;
+	static TSet<FName> SurfaceUnlitOutputs;
+	static TSet<FName> UsdPreviewSurfaceInputs;
+	static TSet<FName> UsdPreviewSurfaceOutputs;
+#endif // WITH_EDITOR
+};
+
+UCLASS(BlueprintType)
+class INTERCHANGEPIPELINES_API UInterchangeMaterialXPipeline : public UInterchangePipelineBase
+{
+	GENERATED_BODY()
+
+	UInterchangeMaterialXPipeline();
+
+public:
+	UPROPERTY(VisibleAnywhere, BlueprintReadWrite, Category = "MaterialX")
+	TObjectPtr<UMaterialXPipelineSettings> MaterialXSettings;
+
+protected:
+	virtual void AdjustSettingsForContext(EInterchangePipelineContext ImportType, TObjectPtr<UObject> ReimportAsset) override;
+	virtual void ExecutePipeline(UInterchangeBaseNodeContainer* BaseNodeContainer, const TArray<UInterchangeSourceData*>& SourceDatas) override;
+
+	virtual bool CanExecuteOnAnyThread(EInterchangePipelineTask PipelineTask) override
+	{
+		// This pipeline creates UObjects and assets. Not safe to execute outside of main thread.
+		return true;
+	}
+
+private:
+	TMap<FString, EInterchangeMaterialXShaders> PathToEnumMapping;
+};
\ No newline at end of file
diff --git a/Engine/Plugins/Interchange/Runtime/Source/Pipelines/Public/InterchangeglTFPipeline.h b/Engine/Plugins/Interchange/Runtime/Source/Pipelines/Public/InterchangeglTFPipeline.h
new file mode 100644
--- /dev/null
+++ b/Engine/Plugins/Interchange/Runtime/Source/Pipelines/Public/InterchangeglTFPipeline.h
@@ -0,0 +1,29 @@
+// Copyright Epic Games, Inc. All Rights Reserved.
+
+#pragma once
+
+#include "CoreMinimal.h"
+
+#include "InterchangePipelineBase.h"
+#include "InterchangeSourceData.h"
+#include "Nodes/InterchangeBaseNodeContainer.h"
+
+#include "InterchangeglTFPipeline.generated.h"
+
+UCLASS(BlueprintType)
+class INTERCHANGEPIPELINES_API UInterchangeglTFPipeline : public UInterchangePipelineBase
+{
+	GENERATED_BODY()
+
+public:
+	// TODO: Put properties to manage the creation or not of material instances instead of material graphs.
+
+protected:
+	virtual void ExecutePipeline(UInterchangeBaseNodeContainer* BaseNodeContainer, const TArray<UInterchangeSourceData*>& SourceDatas) override;
+
+	virtual bool CanExecuteOnAnyThread(EInterchangePipelineTask PipelineTask) override
+	{
+		// This pipeline creates UObjects and assets. Not safe to execute outside of main thread.
+		return true;
+	}
+};
\ No newline at end of file
diff --git a/Engine/Source/Runtime/Interchange/Engine/Private/InterchangeManager.cpp b/Engine/Source/Runtime/Interchange/Engine/Private/InterchangeManager.cpp
--- a/Engine/Source/Runtime/Interchange/Engine/Private/InterchangeManager.cpp
+++ b/Engine/Source/Runtime/Interchange/Engine/Private/InterchangeManager.cpp
@@ -12,6 +12,7 @@
 #include "CoreGlobals.h"
 #include "CoreMinimal.h"
 #include "Engine/Blueprint.h"
+#include "Engine/World.h"
 #include "EngineAnalytics.h"
 #include "Framework/Notifications/NotificationManager.h"
 #include "InterchangeAssetImportData.h"
@@ -1006,6 +1007,7 @@ UE::Interchange::FAssetImportResultRef UInterchangeManager::ImportAssetAsync(con
 bool UInterchangeManager::ImportScene(const FString& ContentPath, const UInterchangeSourceData* SourceData, const FImportAssetParameters& ImportAssetParameters)
 {
 	using namespace UE::Interchange;
+
 	TTuple<FAssetImportResultRef, FSceneImportResultRef> ImportResults = ImportInternal(ContentPath, SourceData, ImportAssetParameters, UE::Interchange::EImportType::ImportType_Scene);
 	
 	ImportResults.Get<0>()->WaitUntilDone();
@@ -1233,17 +1235,38 @@ UInterchangeManager::ImportInternal(const FString& ContentPath, const UInterchan
 				}
 			}
 		}
-		for (const TPair<FName, FInterchangePipelineStack>& PipelineStackInfo : DefaultPipelineStacks)
+
 		{
-			FName StackName = PipelineStackInfo.Key;
-			FInterchangeStackInfo& StackInfo = PipelineStacks.AddDefaulted_GetRef();
-			StackInfo.StackName = StackName;
-			for (int32 PipelineIndex = 0; PipelineIndex < PipelineStackInfo.Value.Pipelines.Num(); ++PipelineIndex)
+			UE::Interchange::FScopedTranslator ScopedTranslator(SourceData);
+			const UInterchangeTranslatorBase* Translator = ScopedTranslator.GetTranslator();
+
+			for (const TPair<FName, FInterchangePipelineStack>& PipelineStackInfo : DefaultPipelineStacks)
 			{
-				if (UInterchangePipelineBase* GeneratedPipeline = UE::Interchange::GeneratePipelineInstance(PipelineStackInfo.Value.Pipelines[PipelineIndex]))
+				FName StackName = PipelineStackInfo.Key;
+				FInterchangeStackInfo& StackInfo = PipelineStacks.AddDefaulted_GetRef();
+				StackInfo.StackName = StackName;
+
+				const FInterchangePipelineStack& PipelineStack = PipelineStackInfo.Value;
+				const TArray<FSoftObjectPath>* Pipelines = &PipelineStack.Pipelines;
+
+				// If applicable, check to see if a specific pipeline stack is associated with this translator
+				for (const FInterchangeTranslatorPipelines& TranslatorPipelines : PipelineStack.PerTranslatorPipelines)
 				{
-					AdjustPipelineSettingForContext(GeneratedPipeline);
-					StackInfo.Pipelines.Add(GeneratedPipeline);
+					const UClass* TranslatorClass = TranslatorPipelines.Translator.LoadSynchronous();
+					if (Translator->IsA(TranslatorClass))
+					{
+						Pipelines = &TranslatorPipelines.Pipelines;
+						break;
+					}
+				}
+
+				for (int32 PipelineIndex = 0; PipelineIndex < Pipelines->Num(); ++PipelineIndex)
+				{
+					if (UInterchangePipelineBase* GeneratedPipeline = UE::Interchange::GeneratePipelineInstance((*Pipelines)[PipelineIndex]))
+					{
+						AdjustPipelineSettingForContext(GeneratedPipeline);
+						StackInfo.Pipelines.Add(GeneratedPipeline);
+					}
 				}
 			}
 		}
diff --git a/Engine/Source/Runtime/Interchange/Engine/Public/InterchangeProjectSettings.h b/Engine/Source/Runtime/Interchange/Engine/Public/InterchangeProjectSettings.h
--- a/Engine/Source/Runtime/Interchange/Engine/Public/InterchangeProjectSettings.h
+++ b/Engine/Source/Runtime/Interchange/Engine/Public/InterchangeProjectSettings.h
@@ -13,6 +13,18 @@
 
 class UInterchangeSourceData;
 
+USTRUCT()
+struct FInterchangeTranslatorPipelines
+{
+	GENERATED_BODY()
+
+	UPROPERTY(EditAnywhere, Category = "TranslatorPipelines", meta = (AllowedClasses = "/Script/InterchangeCore.InterchangeTranslatorBase"))
+	TSoftClassPtr<UInterchangeTranslatorBase> Translator;
+
+	UPROPERTY(EditAnywhere, Category = "TranslatorPipelines", meta = (AllowedClasses = "/Script/InterchangeCore.InterchangePipelineBase, /Script/InterchangeEngine.InterchangeBlueprintPipelineBase, /Script/InterchangeEngine.InterchangePythonPipelineAsset"))
+	TArray<FSoftObjectPath> Pipelines;
+};
+
 USTRUCT()
 struct FInterchangePipelineStack
 {
@@ -20,6 +32,10 @@ struct FInterchangePipelineStack
 	
 	UPROPERTY(EditAnywhere, Category = "Pipelines", meta = (AllowedClasses = "/Script/InterchangeCore.InterchangePipelineBase, /Script/InterchangeEngine.InterchangeBlueprintPipelineBase, /Script/InterchangeEngine.InterchangePythonPipelineAsset"))
 	TArray<FSoftObjectPath> Pipelines;
+
+	/** This tells Interchange which pipeline to add based on the type of the source */
+	UPROPERTY(EditAnywhere, Category = "TranslatorPipelines")
+	TArray<FInterchangeTranslatorPipelines> PerTranslatorPipelines;
 };
 
 USTRUCT()
