Fix C++20 error

#rnx

diff --git a/Engine/Plugins/Interchange/Runtime/Source/Import/Private/Gltf/InterchangeGLTFMaterialInstance.cpp b/Engine/Plugins/Interchange/Runtime/Source/Import/Private/Gltf/InterchangeGLTFMaterialInstance.cpp
--- a/Engine/Plugins/Interchange/Runtime/Source/Import/Private/Gltf/InterchangeGLTFMaterialInstance.cpp
+++ b/Engine/Plugins/Interchange/Runtime/Source/Import/Private/Gltf/InterchangeGLTFMaterialInstance.cpp
@@ -155,7 +155,7 @@ FGLTFMaterialInstanceSettings::FGLTFMaterialInstanceSettings()
 {
 	static const FString AssetFolder = TEXT("/Interchange/gltf/MaterialInstances/");
 
-	auto AddMaterialParent = [=](EShadingModel ShadingModel, GLTF::FMaterial::EAlphaMode AlphaMode, bool bTwoSided)
+	auto AddMaterialParent = [this](EShadingModel ShadingModel, GLTF::FMaterial::EAlphaMode AlphaMode, bool bTwoSided)
 	{
 		const FString Identifier = GetIdentifier(ShadingModel, AlphaMode, bTwoSided);
 		const FString AssetName = Identifier + TEXT(".") + Identifier;
