Renaming FMaterialXUsdPreviewSurface to FMaterialXUsdPreviewSurfaceShader

diff --git a/Engine/Plugins/Interchange/Runtime/Source/Import/Private/MaterialX/MaterialXUtils/MaterialXManager.cpp b/Engine/Plugins/Interchange/Runtime/Source/Import/Private/MaterialX/MaterialXUtils/MaterialXManager.cpp
--- a/Engine/Plugins/Interchange/Runtime/Source/Import/Private/MaterialX/MaterialXUtils/MaterialXManager.cpp
+++ b/Engine/Plugins/Interchange/Runtime/Source/Import/Private/MaterialX/MaterialXUtils/MaterialXManager.cpp
@@ -7,7 +7,7 @@
 #include "InterchangeShaderGraphNode.h"
 #include "MaterialX/MaterialXUtils/MaterialXSurfaceShaderAbstract.h"
 #include "MaterialX/MaterialXUtils/MaterialXStandardSurfaceShader.h"
-#include "MaterialX/MaterialXUtils/MaterialXUsdPreviewSurface.h"
+#include "MaterialX/MaterialXUtils/MaterialXUsdPreviewSurfaceShader.h"
 #include "MaterialX/MaterialXUtils/MaterialXSurfaceUnlitShader.h"
 #include "MaterialX/MaterialXUtils/MaterialXPointLightShader.h"
 #include "MaterialX/MaterialXUtils/MaterialXDirectionalLightShader.h"
@@ -197,7 +197,7 @@ FMaterialXManager::FMaterialXManager()
 	, MaterialXContainerDelegates{
 		{mx::Category::SurfaceUnlit, FMaterialXManager::FOnGetMaterialXInstance::CreateStatic(&FMaterialXSurfaceUnlitShader::MakeInstance)},
 		{mx::Category::StandardSurface, FMaterialXManager::FOnGetMaterialXInstance::CreateStatic(&FMaterialXStandardSurfaceShader::MakeInstance)},
-		{mx::Category::UsdPreviewSurface, FMaterialXManager::FOnGetMaterialXInstance::CreateStatic(&FMaterialXUsdPreviewSurface::MakeInstance)},
+		{mx::Category::UsdPreviewSurface, FMaterialXManager::FOnGetMaterialXInstance::CreateStatic(&FMaterialXUsdPreviewSurfaceShader::MakeInstance)},
 		{mx::Category::PointLight, FMaterialXManager::FOnGetMaterialXInstance::CreateStatic(&FMaterialXPointLightShader::MakeInstance)},
 		{mx::Category::DirectionalLight, FMaterialXManager::FOnGetMaterialXInstance::CreateStatic(&FMaterialXDirectionalLightShader::MakeInstance)},
 		{mx::Category::SpotLight, FMaterialXManager::FOnGetMaterialXInstance::CreateStatic(&FMaterialXSpotLightShader::MakeInstance)},
diff --git a/Engine/Plugins/Interchange/Runtime/Source/Import/Private/MaterialX/MaterialXUtils/MaterialXUsdPreviewSurface.cpp b/Engine/Plugins/Interchange/Runtime/Source/Import/Private/MaterialX/MaterialXUtils/MaterialXUsdPreviewSurface.cpp
deleted file mode 100644
diff --git a/Engine/Plugins/Interchange/Runtime/Source/Import/Private/MaterialX/MaterialXUtils/MaterialXUsdPreviewSurface.h b/Engine/Plugins/Interchange/Runtime/Source/Import/Private/MaterialX/MaterialXUtils/MaterialXUsdPreviewSurface.h
deleted file mode 100644
diff --git a/Engine/Plugins/Interchange/Runtime/Source/Import/Private/MaterialX/MaterialXUtils/MaterialXUsdPreviewSurfaceShader.cpp b/Engine/Plugins/Interchange/Runtime/Source/Import/Private/MaterialX/MaterialXUtils/MaterialXUsdPreviewSurfaceShader.cpp
new file mode 100644
--- /dev/null
+++ b/Engine/Plugins/Interchange/Runtime/Source/Import/Private/MaterialX/MaterialXUtils/MaterialXUsdPreviewSurfaceShader.cpp
@@ -0,0 +1,102 @@
+// Copyright Epic Games, Inc. All Rights Reserved.
+
+#if WITH_EDITOR
+
+#include "MaterialXUsdPreviewSurfaceShader.h"
+
+namespace mx = MaterialX;
+
+FMaterialXUsdPreviewSurfaceShader::FMaterialXUsdPreviewSurfaceShader(UInterchangeBaseNodeContainer& BaseNodeContainer)
+	: FMaterialXSurfaceShaderAbstract(BaseNodeContainer)
+{
+	NodeDefinition = mx::NodeDefinition::UsdPreviewSurface;
+}
+
+TSharedRef<FMaterialXBase> FMaterialXUsdPreviewSurfaceShader::MakeInstance(UInterchangeBaseNodeContainer& BaseNodeContainer)
+{
+	TSharedRef<FMaterialXUsdPreviewSurfaceShader> Result = MakeShared<FMaterialXUsdPreviewSurfaceShader>(BaseNodeContainer);
+	Result->RegisterConnectNodeOutputToInputDelegates();
+	return Result;
+}
+
+void FMaterialXUsdPreviewSurfaceShader::Translate(MaterialX::NodePtr UsdPreviewSurfaceNode)
+{
+	this->SurfaceShaderNode = UsdPreviewSurfaceNode;
+
+	using namespace UE::Interchange::Materials;
+
+	UInterchangeFunctionCallShaderNode* UsdPreviewSurfaceShaderNode;
+
+	const FString NodeUID = UInterchangeShaderNode::MakeNodeUid(ANSI_TO_TCHAR(UsdPreviewSurfaceNode->getName().c_str()), FStringView{});
+
+	if(UsdPreviewSurfaceShaderNode = const_cast<UInterchangeFunctionCallShaderNode*>(Cast<UInterchangeFunctionCallShaderNode>(NodeContainer.GetNode(NodeUID))); !UsdPreviewSurfaceShaderNode)
+	{
+		const FString NodeName = UsdPreviewSurfaceNode->getName().c_str();
+		UsdPreviewSurfaceShaderNode = NewObject<UInterchangeFunctionCallShaderNode>(&NodeContainer);
+		UsdPreviewSurfaceShaderNode->InitializeNode(NodeUID, NodeName, EInterchangeNodeContainerType::TranslatedAsset);
+
+		UsdPreviewSurfaceShaderNode->SetCustomMaterialFunction(TEXT("/Interchange/Functions/MX_UsdPreviewSurface.MX_UsdPreviewSurface"));
+		NodeContainer.AddNode(UsdPreviewSurfaceShaderNode);
+
+		ShaderNodes.Add({ NodeName, UsdPreviewSurfaceNode->getNodeDef(mx::EMPTY_STRING, true)->getActiveOutputs()[0]->getName().c_str()}, UsdPreviewSurfaceShaderNode);
+	}
+
+	// Inputs
+	//Diffuse Color
+	ConnectNodeOutputToInput(mx::UsdPreviewSurface::Input::DiffuseColor, UsdPreviewSurfaceShaderNode, UsdPreviewSurface::Parameters::DiffuseColor.ToString(), mx::UsdPreviewSurface::DefaultValue::Color3::DiffuseColor);
+
+	//Emissive Color
+	ConnectNodeOutputToInput(mx::UsdPreviewSurface::Input::EmissiveColor, UsdPreviewSurfaceShaderNode, UsdPreviewSurface::Parameters::EmissiveColor.ToString(), mx::UsdPreviewSurface::DefaultValue::Color3::EmissiveColor);
+
+	//Specular Color
+	ConnectNodeOutputToInput(mx::UsdPreviewSurface::Input::SpecularColor, UsdPreviewSurfaceShaderNode, UsdPreviewSurface::Parameters::SpecularColor.ToString(), mx::UsdPreviewSurface::DefaultValue::Color3::SpecularColor);
+
+	//Metallic
+	ConnectNodeOutputToInput(mx::UsdPreviewSurface::Input::Metallic, UsdPreviewSurfaceShaderNode, UsdPreviewSurface::Parameters::Metallic.ToString(), mx::UsdPreviewSurface::DefaultValue::Float::Metallic);
+
+	//Roughness
+	ConnectNodeOutputToInput(mx::UsdPreviewSurface::Input::Roughness, UsdPreviewSurfaceShaderNode, UsdPreviewSurface::Parameters::Roughness.ToString(), mx::UsdPreviewSurface::DefaultValue::Float::Roughness);
+
+	//Clearcoat
+	ConnectNodeOutputToInput(mx::UsdPreviewSurface::Input::Clearcoat, UsdPreviewSurfaceShaderNode, UsdPreviewSurface::Parameters::Clearcoat.ToString(), mx::UsdPreviewSurface::DefaultValue::Float::Clearcoat);
+
+	//Clearcoat Roughness
+	ConnectNodeOutputToInput(mx::UsdPreviewSurface::Input::ClearcoatRoughness, UsdPreviewSurfaceShaderNode, UsdPreviewSurface::Parameters::ClearcoatRoughness.ToString(), mx::UsdPreviewSurface::DefaultValue::Float::ClearcoatRoughness);
+
+	//Opacity
+	ConnectNodeOutputToInput(mx::UsdPreviewSurface::Input::Opacity, UsdPreviewSurfaceShaderNode, UsdPreviewSurface::Parameters::Opacity.ToString(), mx::UsdPreviewSurface::DefaultValue::Float::Opacity);
+
+	//Opacity Threshold
+	ConnectNodeOutputToInput(mx::UsdPreviewSurface::Input::OpacityThreshold, UsdPreviewSurfaceShaderNode, UsdPreviewSurface::Parameters::OpacityThreshold.ToString(), mx::UsdPreviewSurface::DefaultValue::Float::OpacityThreshold);
+
+	//IOR
+	ConnectNodeOutputToInput(mx::UsdPreviewSurface::Input::IOR, UsdPreviewSurfaceShaderNode, UsdPreviewSurface::Parameters::IOR.ToString(), mx::UsdPreviewSurface::DefaultValue::Float::IOR);
+
+	//Normal
+	ConnectNodeOutputToInput(mx::UsdPreviewSurface::Input::Normal, UsdPreviewSurfaceShaderNode, UsdPreviewSurface::Parameters::Normal.ToString(), mx::UsdPreviewSurface::DefaultValue::Vector3::Normal);
+
+	//Displacement
+	ConnectNodeOutputToInput(mx::UsdPreviewSurface::Input::Displacement, ShaderGraphNode, UsdPreviewSurface::Parameters::Displacement.ToString(), mx::UsdPreviewSurface::DefaultValue::Float::Displacement);
+
+	//Occlusion
+	ConnectNodeOutputToInput(mx::UsdPreviewSurface::Input::Occlusion, UsdPreviewSurfaceShaderNode, UsdPreviewSurface::Parameters::Occlusion.ToString(), mx::UsdPreviewSurface::DefaultValue::Float::Occlusion);
+
+	// Outputs
+	UInterchangeShaderPortsAPI::ConnectOuputToInputByName(ShaderGraphNode, PBR::Parameters::BaseColor.ToString(), UsdPreviewSurfaceShaderNode->GetUniqueID(), PBR::Parameters::BaseColor.ToString());
+	UInterchangeShaderPortsAPI::ConnectOuputToInputByName(ShaderGraphNode, PBR::Parameters::Metallic.ToString(), UsdPreviewSurfaceShaderNode->GetUniqueID(), PBR::Parameters::Metallic.ToString());
+	UInterchangeShaderPortsAPI::ConnectOuputToInputByName(ShaderGraphNode, PBR::Parameters::Specular.ToString(), UsdPreviewSurfaceShaderNode->GetUniqueID(), PBR::Parameters::Specular.ToString());
+	UInterchangeShaderPortsAPI::ConnectOuputToInputByName(ShaderGraphNode, PBR::Parameters::Roughness.ToString(), UsdPreviewSurfaceShaderNode->GetUniqueID(), PBR::Parameters::Roughness.ToString());
+	UInterchangeShaderPortsAPI::ConnectOuputToInputByName(ShaderGraphNode, PBR::Parameters::EmissiveColor.ToString(), UsdPreviewSurfaceShaderNode->GetUniqueID(), PBR::Parameters::EmissiveColor.ToString());
+	UInterchangeShaderPortsAPI::ConnectOuputToInputByName(ShaderGraphNode, PBR::Parameters::Normal.ToString(), UsdPreviewSurfaceShaderNode->GetUniqueID(), PBR::Parameters::Normal.ToString());
+	
+	if(UInterchangeShaderPortsAPI::HasInput(UsdPreviewSurfaceShaderNode, UsdPreviewSurface::Parameters::Opacity))
+	{
+		UInterchangeShaderPortsAPI::ConnectOuputToInputByName(ShaderGraphNode, PBR::Parameters::Opacity.ToString(), UsdPreviewSurfaceShaderNode->GetUniqueID(), PBR::Parameters::Opacity.ToString());
+	}
+
+	UInterchangeShaderPortsAPI::ConnectOuputToInputByName(ShaderGraphNode, PBR::Parameters::Occlusion.ToString(), UsdPreviewSurfaceShaderNode->GetUniqueID(), PBR::Parameters::Occlusion.ToString());
+	UInterchangeShaderPortsAPI::ConnectOuputToInputByName(ShaderGraphNode, PBR::Parameters::Refraction.ToString(), UsdPreviewSurfaceShaderNode->GetUniqueID(), PBR::Parameters::Refraction.ToString());
+	UInterchangeShaderPortsAPI::ConnectOuputToInputByName(ShaderGraphNode, ClearCoat::Parameters::ClearCoat.ToString(), UsdPreviewSurfaceShaderNode->GetUniqueID(), ClearCoat::Parameters::ClearCoat.ToString());
+	UInterchangeShaderPortsAPI::ConnectOuputToInputByName(ShaderGraphNode, ClearCoat::Parameters::ClearCoatRoughness.ToString(), UsdPreviewSurfaceShaderNode->GetUniqueID(), ClearCoat::Parameters::ClearCoatRoughness.ToString());
+}
+#endif
\ No newline at end of file
diff --git a/Engine/Plugins/Interchange/Runtime/Source/Import/Private/MaterialX/MaterialXUtils/MaterialXUsdPreviewSurfaceShader.h b/Engine/Plugins/Interchange/Runtime/Source/Import/Private/MaterialX/MaterialXUtils/MaterialXUsdPreviewSurfaceShader.h
new file mode 100644
--- /dev/null
+++ b/Engine/Plugins/Interchange/Runtime/Source/Import/Private/MaterialX/MaterialXUtils/MaterialXUsdPreviewSurfaceShader.h
@@ -0,0 +1,24 @@
+// Copyright Epic Games, Inc. All Rights Reserved.
+
+#pragma once
+
+#if WITH_EDITOR
+
+#include "MaterialX/MaterialXUtils/MaterialXSurfaceShaderAbstract.h"
+
+class FMaterialXUsdPreviewSurfaceShader : public FMaterialXSurfaceShaderAbstract
+{
+protected:
+
+	template<typename T, ESPMode>
+	friend class SharedPointerInternals::TIntrusiveReferenceController;
+
+	FMaterialXUsdPreviewSurfaceShader(UInterchangeBaseNodeContainer& BaseNodeContainer);
+
+public:
+
+	static TSharedRef<FMaterialXBase> MakeInstance(UInterchangeBaseNodeContainer& BaseNodeContainer);
+
+	virtual void Translate(MaterialX::NodePtr UsdPreviewSurfaceNode) override;
+};
+#endif
\ No newline at end of file
