[Interchange] Warn if imported glTF has ClearCoatBottomNormal but dual normals disabled in render settings

diff --git a/Engine/Plugins/Interchange/Runtime/Source/Import/Private/Gltf/InterchangeGltfTranslator.cpp b/Engine/Plugins/Interchange/Runtime/Source/Import/Private/Gltf/InterchangeGltfTranslator.cpp
--- a/Engine/Plugins/Interchange/Runtime/Source/Import/Private/Gltf/InterchangeGltfTranslator.cpp
+++ b/Engine/Plugins/Interchange/Runtime/Source/Import/Private/Gltf/InterchangeGltfTranslator.cpp
@@ -40,6 +40,7 @@
 #include "Gltf/InterchangeGltfPrivate.h"
 
 #include "EngineAnalytics.h"
+#include "Engine/RendererSettings.h"
 
 #include UE_INLINE_GENERATED_CPP_BY_NAME(InterchangeGltfTranslator)
 
@@ -724,6 +725,12 @@ void UInterchangeGltfTranslator::HandleGltfClearCoat( UInterchangeBaseNodeContai
 
 		HandleGltfMaterialParameter( NodeContainer, GltfMaterial.ClearCoat.NormalMap, ShaderGraphNode, bSwapNormalAndClearCoatNormal ? Common::Parameters::Normal.ToString() : ClearCoat::Parameters::ClearCoatNormal.ToString(),
 			ClearCoatNormalFactor, Standard::Nodes::TextureSample::Outputs::RGB.ToString(), bInverse, bIsNormal );
+
+		//check if ClearCoat second normal is enabled in the render settings:
+		if (GltfAsset.Textures.IsValidIndex(GltfMaterial.Normal.TextureIndex) && !bRenderSettingsClearCoatEnableSecondNormal)
+		{
+			UE_LOG(LogInterchangeImport, Warning, TEXT("GLTF Material[%s] uses ClearCoat and has Normal map, however ClearCoat Second Normal is disabled in the Render Settings."), *GltfMaterial.Name);
+		}
 	}
 }
 
@@ -2236,6 +2243,8 @@ UInterchangeGltfTranslator::UInterchangeGltfTranslator()
 				UE_LOG(LogInterchangeImport, Warning, TEXT("%s"), *MaterialInstanceIssue);
 			}
 		}
+
+		bRenderSettingsClearCoatEnableSecondNormal = GetDefault<URendererSettings>()->bClearCoatEnableSecondNormal != 0;
 	}
 }
 
diff --git a/Engine/Plugins/Interchange/Runtime/Source/Import/Public/Gltf/InterchangeGltfTranslator.h b/Engine/Plugins/Interchange/Runtime/Source/Import/Public/Gltf/InterchangeGltfTranslator.h
--- a/Engine/Plugins/Interchange/Runtime/Source/Import/Public/Gltf/InterchangeGltfTranslator.h
+++ b/Engine/Plugins/Interchange/Runtime/Source/Import/Public/Gltf/InterchangeGltfTranslator.h
@@ -97,6 +97,8 @@ private:
 
 	GLTF::FAsset GltfAsset;
 	mutable FNodeUidMap NodeUidMap;
+
+	bool bRenderSettingsClearCoatEnableSecondNormal = false;
 };
 
 
