Fix material order tests.
Avoid adding warning from fbx parser when running the automation tests.

#rnx

diff --git a/Engine/Plugins/Interchange/Runtime/Source/Import/Private/Fbx/InterchangeFbxTranslator.cpp b/Engine/Plugins/Interchange/Runtime/Source/Import/Private/Fbx/InterchangeFbxTranslator.cpp
--- a/Engine/Plugins/Interchange/Runtime/Source/Import/Private/Fbx/InterchangeFbxTranslator.cpp
+++ b/Engine/Plugins/Interchange/Runtime/Source/Import/Private/Fbx/InterchangeFbxTranslator.cpp
@@ -31,6 +31,26 @@ static FAutoConsoleVariableRef CCvarInterchangeEnableFBXImport(
 	TEXT("[Experimental] Whether FBX support is enabled."),
 	ECVF_Default);
 
+namespace UE::Interchange::Private
+{
+	void ApplyTranslatorMessage(const UInterchangeFbxTranslator* Translator, const FString& JsonMessage)
+	{
+		UInterchangeResult* InterchangeResult = UInterchangeResult::FromJson(JsonMessage);
+		if (InterchangeResult)
+		{
+			//Downgrade warning message to display log when we are in automation
+			if (GIsAutomationTesting && InterchangeResult->IsA(UInterchangeResultWarning::StaticClass()))
+			{
+				UE_LOG(LogInterchangeImport, Display, TEXT("%s"), *InterchangeResult->GetText().ToString());
+			}
+			else
+			{
+				Translator->AddMessage(InterchangeResult);
+			}
+		}
+	}
+} //ns UE::Interchange::Private
+
 UInterchangeFbxTranslator::UInterchangeFbxTranslator()
 {
 	Dispatcher = nullptr;
@@ -127,7 +147,7 @@ bool UInterchangeFbxTranslator::Translate(UInterchangeBaseNodeContainer& BaseNod
 	// Parse the Json messages into UInterchangeResults
 	for (const FString& JsonMessage : JsonMessages)
 	{
-		AddMessage(UInterchangeResult::FromJson(JsonMessage));
+		UE::Interchange::Private::ApplyTranslatorMessage(this, JsonMessage);
 	}
 
 	if(TaskState != UE::Interchange::ETaskState::ProcessOk)
@@ -208,7 +228,7 @@ TFuture<TOptional<UE::Interchange::FMeshPayloadData>> UInterchangeFbxTranslator:
 		// Parse the Json messages into UInterchangeResults
 		for (const FString& JsonMessage : JsonMessages)
 		{
-			AddMessage(UInterchangeResult::FromJson(JsonMessage));
+			UE::Interchange::Private::ApplyTranslatorMessage(this, JsonMessage);
 		}
 
 		if (TaskState != UE::Interchange::ETaskState::ProcessOk)
@@ -314,7 +334,7 @@ TFuture<TOptional<UE::Interchange::FAnimationPayloadData>> UInterchangeFbxTransl
 			// Parse the Json messages into UInterchangeResults
 			for (const FString& JsonMessage : JsonMessages)
 			{
-				AddMessage(UInterchangeResult::FromJson(JsonMessage));
+				UE::Interchange::Private::ApplyTranslatorMessage(this, JsonMessage);
 			}
 
 			if (TaskState != UE::Interchange::ETaskState::ProcessOk)
