Interchange: Refactor the way async factory task retrieve the UObject create in the previous game thread task.

#rnx

diff --git a/Engine/Plugins/Experimental/Enterprise/DatasmithInterchange/Source/DatasmithInterchange/Private/InterchangeDatasmithSceneFactory.cpp b/Engine/Plugins/Experimental/Enterprise/DatasmithInterchange/Source/DatasmithInterchange/Private/InterchangeDatasmithSceneFactory.cpp
--- a/Engine/Plugins/Experimental/Enterprise/DatasmithInterchange/Source/DatasmithInterchange/Private/InterchangeDatasmithSceneFactory.cpp
+++ b/Engine/Plugins/Experimental/Enterprise/DatasmithInterchange/Source/DatasmithInterchange/Private/InterchangeDatasmithSceneFactory.cpp
@@ -92,19 +92,15 @@ UInterchangeFactoryBase::FImportAssetResult UInterchangeDatasmithSceneFactory::I
 		return ImportAssetResult;
 	}
 
-	// create an asset if it doesn't exist
-	UObject* ExistingAsset = StaticFindObject(nullptr, Arguments.Parent, *Arguments.AssetName);
-	UDatasmithScene* DatasmithScene = nullptr;
-
-	// create a new texture or overwrite existing asset, if possible
-	if (!ExistingAsset)
+	UObject* ExistingAsset = nullptr;
+	FSoftObjectPath ReferenceObject;
+	if (Arguments.AssetNode->GetCustomReferenceObject(ReferenceObject))
 	{
-		//NewObject is not thread safe, the asset registry directory watcher tick on the main thread can trig before we finish initializing the UObject and will crash
-		//The UObject should have been create by calling CreateEmptyAsset on the main thread.
-		check(IsInGameThread());
-		DatasmithScene = NewObject<UDatasmithScene>(Arguments.Parent, DatasmithSceneClass, *Arguments.AssetName, RF_Public | RF_Standalone);
+		ExistingAsset = ReferenceObject.TryLoad();
 	}
-	else if (ExistingAsset->GetClass()->IsChildOf(DatasmithSceneClass))
+	UDatasmithScene* DatasmithScene = nullptr;
+
+	if (ExistingAsset && ExistingAsset->GetClass()->IsChildOf(DatasmithSceneClass))
 	{
 		//This is a reimport, we are just re-updating the source data
 		DatasmithScene = static_cast<UDatasmithScene*>(ExistingAsset);
diff --git a/Engine/Plugins/Interchange/Runtime/Source/Import/Private/Animation/InterchangeAnimSequenceFactory.cpp b/Engine/Plugins/Interchange/Runtime/Source/Import/Private/Animation/InterchangeAnimSequenceFactory.cpp
--- a/Engine/Plugins/Interchange/Runtime/Source/Import/Private/Animation/InterchangeAnimSequenceFactory.cpp
+++ b/Engine/Plugins/Interchange/Runtime/Source/Import/Private/Animation/InterchangeAnimSequenceFactory.cpp
@@ -781,7 +781,7 @@ UInterchangeFactoryBase::FImportAssetResult UInterchangeAnimSequenceFactory::Beg
 		return ImportAssetResult;
 	}
 
-	const UInterchangeAnimSequenceFactoryNode* AnimSequenceFactoryNode = Cast<UInterchangeAnimSequenceFactoryNode>(Arguments.AssetNode);
+	UInterchangeAnimSequenceFactoryNode* AnimSequenceFactoryNode = Cast<UInterchangeAnimSequenceFactoryNode>(Arguments.AssetNode);
 	if (AnimSequenceFactoryNode == nullptr)
 	{
 		return ImportAssetResult;
@@ -813,8 +813,11 @@ UInterchangeFactoryBase::FImportAssetResult UInterchangeAnimSequenceFactory::Beg
 		return ImportAssetResult;
 	}
 
+	AnimSequenceFactoryNode->SetCustomReferenceObject(FSoftObjectPath(AnimSequence));
+
 	AnimSequence->PreEditChange(nullptr);
 
+
 	ImportAssetResult.ImportedObject = ImportObjectSourceData(Arguments);
 	return ImportAssetResult;
 #endif //else !WITH_EDITOR || !WITH_EDITORONLY_DATA
@@ -899,57 +902,23 @@ UObject* UInterchangeAnimSequenceFactory::ImportObjectSourceData(const FImportAs
 		return nullptr;
 	}
 
-	const UClass* AnimSequenceClass = AnimSequenceFactoryNode->GetObjectClass();
-	check(AnimSequenceClass && AnimSequenceClass->IsChildOf(GetFactoryClass()));
-
-	// create an asset if it doesn't exist
-	UObject* ExistingAsset = StaticFindObject(nullptr, Arguments.Parent, *Arguments.AssetName);
-
-	UObject* AnimSequenceObject = nullptr;
-	// create a new material or overwrite existing asset, if possible
-	if (!ExistingAsset)
-	{
-		//NewObject is not thread safe, the asset registry directory watcher tick on the main thread can trig before we finish initializing the UObject and will crash
-		//The UObject should have been create by calling CreateEmptyAsset on the main thread.
-		if (IsInGameThread())
-		{
-			AnimSequenceObject = NewObject<UObject>(Arguments.Parent, AnimSequenceClass, *Arguments.AssetName, RF_Public | RF_Standalone);
-		}
-		else
-		{
-			UE_LOG(LogInterchangeImport, Error, TEXT("Could not create AnimSequence asset [%s] outside of the game thread"), *Arguments.AssetName);
-			return nullptr;
-		}
-	}
-	else if (ExistingAsset->GetClass()->IsChildOf(AnimSequenceClass))
-	{
-		//This is a reimport, we are just re-updating the source data
-		AnimSequenceObject = ExistingAsset;
-	}
+	UObject* AnimSequenceObject = UE::Interchange::FFactoryCommon::AsyncFindObject(AnimSequenceFactoryNode, GetFactoryClass(), Arguments.Parent, Arguments.AssetName);
 
 	if (!AnimSequenceObject)
 	{
-		UE_LOG(LogInterchangeImport, Error, TEXT("Could not create AnimSequence asset %s"), *Arguments.AssetName);
+		UE_LOG(LogInterchangeImport, Error, TEXT("Could not import the AnimSequence asset %s, because the asset do not exist."), *Arguments.AssetName);
 		return nullptr;
 	}
 
 	UAnimSequence* AnimSequence = Cast<UAnimSequence>(AnimSequenceObject);
-
-	const bool bIsReImport = (Arguments.ReimportObject != nullptr);
-
 	if (!ensure(AnimSequence))
 	{
-		if (!bIsReImport)
-		{
-			UE_LOG(LogInterchangeImport, Error, TEXT("Could not create AnimSequence asset %s"), *Arguments.AssetName);
-		}
-		else
-		{
-			UE_LOG(LogInterchangeImport, Error, TEXT("Could not find reimported AnimSequence asset %s"), *Arguments.AssetName);
-		}
+		UE_LOG(LogInterchangeImport, Error, TEXT("Could not cast to AnimSequence asset %s"), *Arguments.AssetName);
 		return nullptr;
 	}
 
+	const bool bIsReImport = (Arguments.ReimportObject != nullptr);
+
 	//Fill the animsequence data, we need to retrieve the skeleton and then ask the payload for every joint
 	{
 		FFrameRate FrameRate(30, 1);
diff --git a/Engine/Plugins/Interchange/Runtime/Source/Import/Private/Animation/InterchangeAnimationTrackSetFactory.cpp b/Engine/Plugins/Interchange/Runtime/Source/Import/Private/Animation/InterchangeAnimationTrackSetFactory.cpp
--- a/Engine/Plugins/Interchange/Runtime/Source/Import/Private/Animation/InterchangeAnimationTrackSetFactory.cpp
+++ b/Engine/Plugins/Interchange/Runtime/Source/Import/Private/Animation/InterchangeAnimationTrackSetFactory.cpp
@@ -6,6 +6,7 @@
 #include "InterchangeAnimationTrackSetFactoryNode.h"
 #include "InterchangeAnimationTrackSetNode.h"
 #include "InterchangeAnimSequenceFactoryNode.h"
+#include "InterchangeImportCommon.h"
 #include "InterchangeImportLog.h"
 #include "InterchangeResult.h"
 #include "InterchangeSourceData.h"
@@ -677,7 +678,7 @@ UInterchangeFactoryBase::FImportAssetResult UInterchangeAnimationTrackSetFactory
 		return ImportAssetResult;
 	}
 
-	const UInterchangeAnimationTrackSetFactoryNode* FactoryNode = Cast<UInterchangeAnimationTrackSetFactoryNode>(Arguments.AssetNode);
+	UInterchangeAnimationTrackSetFactoryNode* FactoryNode = Cast<UInterchangeAnimationTrackSetFactoryNode>(Arguments.AssetNode);
 	if (FactoryNode == nullptr)
 	{
 		return ImportAssetResult;
@@ -709,6 +710,8 @@ UInterchangeFactoryBase::FImportAssetResult UInterchangeAnimationTrackSetFactory
 		return ImportAssetResult;
 	}
 
+	FactoryNode->SetCustomReferenceObject(LevelSequence);
+
 	LevelSequence->PreEditChange(nullptr);
 
 	ImportAssetResult.ImportedObject = ImportObjectSourceData(Arguments);
@@ -758,34 +761,16 @@ UObject* UInterchangeAnimationTrackSetFactory::ImportObjectSourceData(const FImp
 		return nullptr;
 	}
 
-	const UClass* LevelSequenceClass = FactoryNode->GetObjectClass();
-	check(LevelSequenceClass && LevelSequenceClass->IsChildOf(GetFactoryClass()));
-
-	// create an asset if it doesn't exist
-	UObject* ExistingAsset = StaticFindObject(nullptr, Arguments.Parent, *Arguments.AssetName);
+	UObject* ExistingAsset = UE::Interchange::FFactoryCommon::AsyncFindObject(FactoryNode, GetFactoryClass(), Arguments.Parent, Arguments.AssetName);
 
-	ULevelSequence* LevelSequence = nullptr;
-	// create a new level sequence or overwrite existing asset, if possible
 	if (!ExistingAsset)
 	{
-		//NewObject is not thread safe, the asset registry directory watcher tick on the main thread can trig before we finish initializing the UObject and will crash
-		//The UObject should have been create by calling CreateEmptyAsset on the main thread.
-		if (IsInGameThread())
-		{
-			LevelSequence = NewObject<ULevelSequence>(Arguments.Parent, LevelSequenceClass, *Arguments.AssetName, RF_Public | RF_Standalone);
-		}
-		else
-		{
-			UE_LOG(LogInterchangeImport, Error, TEXT("Could not create LevelSequence asset [%s] outside of the game thread"), *Arguments.AssetName);
-			return nullptr;
-		}
-	}
-	else if (ExistingAsset->GetClass()->IsChildOf(LevelSequenceClass))
-	{
-		//This is a reimport, we are just re-updating the source data
-		LevelSequence =  Cast<ULevelSequence>(ExistingAsset);
+		UE_LOG(LogInterchangeImport, Error, TEXT("Could not import the LevelSequence asset %s, because the asset do not exist."), *Arguments.AssetName);
+		return nullptr;
 	}
 
+	ULevelSequence* LevelSequence = Cast<ULevelSequence>(ExistingAsset);
+
 	if (!ensure(LevelSequence))
 	{
 		UInterchangeResultError_Generic* Message = AddMessage<UInterchangeResultError_Generic>();
diff --git a/Engine/Plugins/Interchange/Runtime/Source/Import/Private/Material/InterchangeMaterialFactory.cpp b/Engine/Plugins/Interchange/Runtime/Source/Import/Private/Material/InterchangeMaterialFactory.cpp
--- a/Engine/Plugins/Interchange/Runtime/Source/Import/Private/Material/InterchangeMaterialFactory.cpp
+++ b/Engine/Plugins/Interchange/Runtime/Source/Import/Private/Material/InterchangeMaterialFactory.cpp
@@ -630,11 +630,7 @@ UInterchangeFactoryBase::FImportAssetResult UInterchangeMaterialFactory::ImportA
 		return ImportAssetResult;
 	}
 
-	const UClass* MaterialClass = MaterialFactoryNode->GetObjectClass();
-	check(MaterialClass && MaterialClass->IsChildOf(GetFactoryClass()));
-
-	// create an asset if it doesn't exist
-	UObject* ExistingAsset = StaticFindObject(nullptr, Arguments.Parent, *Arguments.AssetName);
+	UObject* MaterialObject = UE::Interchange::FFactoryCommon::AsyncFindObject(MaterialFactoryNode, GetFactoryClass(), Arguments.Parent, Arguments.AssetName);
 
 	//Do not override an asset we skip
 	if (bSkipImport)
@@ -642,32 +638,18 @@ UInterchangeFactoryBase::FImportAssetResult UInterchangeMaterialFactory::ImportA
 		ImportAssetResult.ImportedObject = ExistingAsset;
 		return ImportAssetResult;
 	}
+	const bool bReimport = Arguments.ReimportObject && MaterialObject;
 
-	UObject* MaterialObject = nullptr;
-	// create a new material or overwrite existing asset, if possible
-	if (!ExistingAsset)
-	{
-		//NewObject is not thread safe, the asset registry directory watcher tick on the main thread can trig before we finish initializing the UObject and will crash
-		//The UObject should have been created by calling CreateEmptyAsset on the main thread.
-		if (IsInGameThread())
-		{
-			MaterialObject = NewObject<UObject>(Arguments.Parent, MaterialClass, *Arguments.AssetName, RF_Public | RF_Standalone);
-		}
-		else
-		{
-			UE_LOG(LogInterchangeImport, Error, TEXT("Could not create Material asset [%s] outside of the game thread"), *Arguments.AssetName);
-			return ImportAssetResult;
-		}
-	}
-	else if(ExistingAsset->GetClass()->IsChildOf(MaterialClass))
+	if (!MaterialObject)
 	{
-		//This is a reimport, we are just re-updating the source data
-		MaterialObject = ExistingAsset;
+		UE_LOG(LogInterchangeImport, Error, TEXT("Could not import the Material asset %s, because the asset do not exist."), *Arguments.AssetName);
+		return ImportAssetResult;
 	}
 
-	if (!MaterialObject)
+	UMaterialInterface* MaterialInterface = Cast<UMaterialInterface>(MaterialObject);
+	if (!ensure(MaterialInterface))
 	{
-		UE_LOG(LogInterchangeImport, Warning, TEXT("Could not create Material asset %s"), *Arguments.AssetName);
+		UE_LOG(LogInterchangeImport, Error, TEXT("Could not cast to Material asset %s"), *Arguments.AssetName);
 		return ImportAssetResult;
 	}
 
@@ -675,7 +657,7 @@ UInterchangeFactoryBase::FImportAssetResult UInterchangeMaterialFactory::ImportA
 	//TODO: Design a re-import process for the material (expressions and input connections)
 	if(!Arguments.ReimportObject)
 	{
-		if (UMaterialInterface* MaterialInterface = Cast<UMaterialInterface>(MaterialObject))
+		if (MaterialInterface)
 		{
 #if WITH_EDITOR
 			if (UMaterial* Material = Cast<UMaterial>(MaterialObject))
@@ -706,9 +688,8 @@ UInterchangeFactoryBase::FImportAssetResult UInterchangeMaterialFactory::ImportA
 
 	//Getting the file Hash will cache it into the source data
 	Arguments.SourceData->GetFileContentHash();
-
-	//The interchange completion task (call in the GameThread after the factories pass), will call PostEditChange which will trig another asynchronous system that will build all material in parallel
-	ImportAssetResult.ImportedObject = MaterialObject;
+		//The interchange completion task (call in the GameThread after the factories pass), will call PostEditChange which will trig another asynchronous system that will build all material in parallel
+	}
 	return ImportAssetResult;
 }
 
@@ -1475,11 +1456,7 @@ UInterchangeFactoryBase::FImportAssetResult UInterchangeMaterialFunctionFactory:
 		return ImportAssetResult;
 	}
 
-	const UClass* MaterialClass = MaterialFactoryNode->GetObjectClass();
-	check(MaterialClass && MaterialClass->IsChildOf(GetFactoryClass()));
-
-	// create an asset if it doesn't exist
-	UObject* ExistingAsset = StaticFindObject(nullptr, Arguments.Parent, *Arguments.AssetName);
+	UObject* MaterialObject = UE::Interchange::FFactoryCommon::AsyncFindObject(Arguments.AssetNode, GetFactoryClass(), Arguments.Parent, Arguments.AssetName);
 
 	//Do not override an asset we skip
 	if (bSkipImport)
@@ -1487,57 +1464,31 @@ UInterchangeFactoryBase::FImportAssetResult UInterchangeMaterialFunctionFactory:
 		ImportAssetResult.ImportedObject = ExistingAsset;
 		return ImportAssetResult;
 	}
-
-	UObject* MaterialObject = nullptr;
-	// create a new material or overwrite existing asset, if possible
-	if (!ExistingAsset)
-	{
-		//NewObject is not thread safe, the asset registry directory watcher tick on the main thread can trig before we finish initializing the UObject and will crash
-		//The UObject should have been created by calling CreateEmptyAsset on the main thread.
-		if (IsInGameThread())
-		{
-			MaterialObject = NewObject<UObject>(Arguments.Parent, MaterialClass, *Arguments.AssetName, RF_Public | RF_Standalone);
-		}
-		else
-		{
-			UE_LOG(LogInterchangeImport, Error, TEXT("Could not create Material asset [%s] outside of the game thread"), *Arguments.AssetName);
-			return ImportAssetResult;
-		}
-	}
-	else if (ExistingAsset->GetClass()->IsChildOf(MaterialClass))
-	{
-		//This is a reimport, we are just re-updating the source data
-		MaterialObject = ExistingAsset;
-	}
+	const bool bReimport = Arguments.ReimportObject && MaterialObject;
 
 	if (!MaterialObject)
 	{
-		UE_LOG(LogInterchangeImport, Warning, TEXT("Could not create Material asset %s"), *Arguments.AssetName);
+		UE_LOG(LogInterchangeImport, Error, TEXT("Could not import the Material asset %s, because the asset do not exist."), *Arguments.AssetName);
 		return ImportAssetResult;
 	}
 
-	if (MaterialObject)
+	//Currently material re-import will not touch the material at all
+	//TODO design a re-import process for the material (expressions and input connections)
+	if (!Arguments.ReimportObject)
 	{
-		//Currently material re-import will not touch the material at all
-		//TODO design a re-import process for the material (expressions and input connections)
-		if (!Arguments.ReimportObject)
+		if (UMaterialFunction* MaterialFunction = Cast<UMaterialFunction>(MaterialObject))
 		{
-			if (UMaterialFunction* MaterialFunction = Cast<UMaterialFunction>(MaterialObject))
-			{
 #if WITH_EDITOR
-				SetupMaterial(MaterialFunction, Arguments, MaterialFactoryNode);
+			SetupMaterial(MaterialFunction, Arguments, MaterialFactoryNode);
 #endif // #if WITH_EDITOR
 
-				MaterialFactoryNode->ApplyAllCustomAttributeToObject(MaterialObject);
-			}
+			MaterialFactoryNode->ApplyAllCustomAttributeToObject(MaterialObject);
 		}
-
-		//Getting the file Hash will cache it into the source data
-		Arguments.SourceData->GetFileContentHash();
-
-		//The interchange completion task (call in the GameThread after the factories pass), will call PostEditChange which will trig another asynchronous system that will build all material in parallel
 	}
 
+	//Getting the file Hash will cache it into the source data
+	Arguments.SourceData->GetFileContentHash();
+
 	ImportAssetResult.ImportedObject = MaterialObject;
 	return ImportAssetResult;
 }
diff --git a/Engine/Plugins/Interchange/Runtime/Source/Import/Private/Mesh/InterchangePhysicsAssetFactory.cpp b/Engine/Plugins/Interchange/Runtime/Source/Import/Private/Mesh/InterchangePhysicsAssetFactory.cpp
--- a/Engine/Plugins/Interchange/Runtime/Source/Import/Private/Mesh/InterchangePhysicsAssetFactory.cpp
+++ b/Engine/Plugins/Interchange/Runtime/Source/Import/Private/Mesh/InterchangePhysicsAssetFactory.cpp
@@ -82,67 +82,31 @@ UInterchangeFactoryBase::FImportAssetResult UInterchangePhysicsAssetFactory::Imp
 		return ImportAssetResult;
 	}
 
-	const UClass* PhysicsAssetClass = PhysicsAssetNode->GetObjectClass();
-	check(PhysicsAssetClass && PhysicsAssetClass->IsChildOf(GetFactoryClass()));
+	UObject* PhysicsAssetObject = UE::Interchange::FFactoryCommon::AsyncFindObject(PhysicsAssetNode, GetFactoryClass(), Arguments.Parent, Arguments.AssetName);
 
-	// create an asset if it doesn't exist
-	UObject* ExistingAsset = StaticFindObject(nullptr, Arguments.Parent, *Arguments.AssetName);
-
-	UObject* PhysicsAssetObject = nullptr;
-	// create a new PhysicsAsset or overwrite existing asset, if possible
-	if (!ExistingAsset)
-	{
-		//NewObject is not thread safe, the asset registry directory watcher tick on the main thread can trig before we finish initializing the UObject and will crash
-		//The UObject should have been create by calling BeginImportAssetObject_GameThread on the main thread.
-		if (IsInGameThread())
-		{
-			PhysicsAssetObject = NewObject<UObject>(Arguments.Parent, PhysicsAssetClass, *Arguments.AssetName, RF_Public | RF_Standalone);
-		}
-		else
-		{
-			UE_LOG(LogInterchangeImport, Error, TEXT("Could not create Physics asset [%s] outside of the game thread"), *Arguments.AssetName);
-			return ImportAssetResult;
-		}
-	}
-	else if(ExistingAsset->GetClass()->IsChildOf(PhysicsAssetClass))
+	if (!PhysicsAssetObject)
 	{
-		//This is a reimport, we are just re-updating the source data
-		PhysicsAssetObject = ExistingAsset;
+		UE_LOG(LogInterchangeImport, Error, TEXT("Could not import the PhysicsAsset asset %s, because the asset do not exist."), *Arguments.AssetName);
+		return ImportAssetResult;
 	}
 
-	if (!PhysicsAssetObject)
+	UPhysicsAsset* PhysicsAsset = Cast<UPhysicsAsset>(PhysicsAssetObject);
+	if (!ensure(PhysicsAsset))
 	{
-		UE_LOG(LogInterchangeImport, Warning, TEXT("Could not create PhysicsAsset asset %s"), *Arguments.AssetName);
+		UE_LOG(LogInterchangeImport, Error, TEXT("Could not cast to PhysicsAsset asset %s"), *Arguments.AssetName);
 		return ImportAssetResult;
 	}
 
-	if (PhysicsAssetObject)
+	//Currently PhysicsAsset re-import will not touch the PhysicsAsset at all
+	//TODO design a re-import process for the PhysicsAsset
+	if(!Arguments.ReimportObject)
 	{
-		//Currently PhysicsAsset re-import will not touch the PhysicsAsset at all
-		//TODO design a re-import process for the PhysicsAsset
-		if(!Arguments.ReimportObject)
-		{
-			UPhysicsAsset* PhysicsAsset = Cast<UPhysicsAsset>(PhysicsAssetObject);
-			if (!ensure(PhysicsAsset))
-			{
-				UE_LOG(LogInterchangeImport, Warning, TEXT("Could not create PhysicsAsset asset %s"), *Arguments.AssetName);
-				return ImportAssetResult;
-			}
-			
-			PhysicsAssetNode->SetCustomReferenceObject(FSoftObjectPath(PhysicsAsset));
-		}
+		PhysicsAssetNode->SetCustomReferenceObject(FSoftObjectPath(PhysicsAsset));
+	}
 		
-		//Getting the file Hash will cache it into the source data
-		Arguments.SourceData->GetFileContentHash();
+	//Getting the file Hash will cache it into the source data
+	Arguments.SourceData->GetFileContentHash();
 
-		//The interchange completion task (call in the GameThread after the factories pass), will call PostEditChange.
-	}
-	else
-	{
-		//The PhysicsAsset is not a UPhysicsAsset
-		PhysicsAssetObject->RemoveFromRoot();
-		PhysicsAssetObject->MarkAsGarbage();
-	}
 	ImportAssetResult.ImportedObject = PhysicsAssetObject;
 	return ImportAssetResult;
 #endif
diff --git a/Engine/Plugins/Interchange/Runtime/Source/Import/Private/Mesh/InterchangeSkeletalMeshFactory.cpp b/Engine/Plugins/Interchange/Runtime/Source/Import/Private/Mesh/InterchangeSkeletalMeshFactory.cpp
--- a/Engine/Plugins/Interchange/Runtime/Source/Import/Private/Mesh/InterchangeSkeletalMeshFactory.cpp
+++ b/Engine/Plugins/Interchange/Runtime/Source/Import/Private/Mesh/InterchangeSkeletalMeshFactory.cpp
@@ -979,54 +979,18 @@ UInterchangeFactoryBase::FImportAssetResult UInterchangeSkeletalMeshFactory::Imp
 		return ImportAssetResult;
 	}
 
-	const UClass* SkeletalMeshClass = SkeletalMeshFactoryNode->GetObjectClass();
-	check(SkeletalMeshClass && SkeletalMeshClass->IsChildOf(GetFactoryClass()));
-
-	// create an asset if it doesn't exist
-	UObject* ExistingAsset = StaticFindObject(nullptr, Arguments.Parent, *Arguments.AssetName);
-
-	UObject* SkeletalMeshObject = nullptr;
-	// create a new material or overwrite existing asset, if possible
-	if (!ExistingAsset)
-	{
-		//NewObject is not thread safe, the asset registry directory watcher tick on the main thread can trig before we finish initializing the UObject and will crash
-		//The UObject should have been create by calling CreateEmptyAsset on the main thread.
-		if (IsInGameThread())
-		{
-			SkeletalMeshObject = NewObject<UObject>(Arguments.Parent, SkeletalMeshClass, *Arguments.AssetName, RF_Public | RF_Standalone);
-		}
-		else
-		{
-			UE_LOG(LogInterchangeImport, Error, TEXT("Could not create SkeletalMesh asset [%s] outside of the game thread"), *Arguments.AssetName);
-			return ImportAssetResult;
-		}
-	}
-	else if(ExistingAsset->GetClass()->IsChildOf(SkeletalMeshClass))
-	{
-		//This is a reimport, we are just re-updating the source data
-		SkeletalMeshObject = ExistingAsset;
-	}
+	UObject* SkeletalMeshObject = UE::Interchange::FFactoryCommon::AsyncFindObject(SkeletalMeshFactoryNode, GetFactoryClass(), Arguments.Parent, Arguments.AssetName);
 
 	if (!SkeletalMeshObject)
 	{
-		if (!Arguments.ReimportObject)
-		{
-			UE_LOG(LogInterchangeImport, Error, TEXT("Could not create SkeletalMesh asset %s"), *Arguments.AssetName);
-		}
+		UE_LOG(LogInterchangeImport, Error, TEXT("Could not import the SkeletalMesh asset %s, because the asset do not exist."), *Arguments.AssetName);
 		return ImportAssetResult;
 	}
 
 	USkeletalMesh* SkeletalMesh = Cast<USkeletalMesh>(SkeletalMeshObject);
 	if (!ensure(SkeletalMesh))
 	{
-		if (Arguments.ReimportObject == nullptr)
-		{
-			UE_LOG(LogInterchangeImport, Error, TEXT("Could not create skeletalMesh asset %s"), *Arguments.AssetName);
-		}
-		else
-		{
-			UE_LOG(LogInterchangeImport, Error, TEXT("Could not find reimported skeletalMesh asset %s"), *Arguments.AssetName);
-		}
+		UE_LOG(LogInterchangeImport, Error, TEXT("Could not cast to SkeletalMesh asset %s"), *Arguments.AssetName);
 		return ImportAssetResult;
 	}
 
diff --git a/Engine/Plugins/Interchange/Runtime/Source/Import/Private/Mesh/InterchangeStaticMeshFactory.cpp b/Engine/Plugins/Interchange/Runtime/Source/Import/Private/Mesh/InterchangeStaticMeshFactory.cpp
--- a/Engine/Plugins/Interchange/Runtime/Source/Import/Private/Mesh/InterchangeStaticMeshFactory.cpp
+++ b/Engine/Plugins/Interchange/Runtime/Source/Import/Private/Mesh/InterchangeStaticMeshFactory.cpp
@@ -219,49 +219,19 @@ UInterchangeFactoryBase::FImportAssetResult UInterchangeStaticMeshFactory::Impor
 		return ImportAssetResult;
 	}
 
-	const UClass* StaticMeshClass = StaticMeshFactoryNode->GetObjectClass();
-	check(StaticMeshClass && StaticMeshClass->IsChildOf(GetFactoryClass()));
-
-	// create an asset if it doesn't exist
-	UObject* ExistingAsset = StaticFindObject(nullptr, Arguments.Parent, *Arguments.AssetName);
-
-	const bool bReimport = Arguments.ReimportObject&& ExistingAsset;
-	UObject* StaticMeshObject = nullptr;
-
-	// create a new static mesh or overwrite existing asset, if possible
-	if (!ExistingAsset)
-	{
-		// NewObject is not thread safe, the asset registry directory watcher tick on the main thread can trig before we finish initializing the UObject and will crash
-		// The UObject should have been created by calling CreateEmptyAsset on the main thread.
-		if (IsInGameThread())
-		{
-			StaticMeshObject = NewObject<UObject>(Arguments.Parent, StaticMeshClass, *Arguments.AssetName, RF_Public | RF_Standalone);
-		}
-		else
-		{
-			UE_LOG(LogInterchangeImport, Error, TEXT("Could not create StaticMesh asset [%s] outside of the game thread"), *Arguments.AssetName);
-			return ImportAssetResult;
-		}
-	}
-	else if(ExistingAsset->GetClass()->IsChildOf(StaticMeshClass))
-	{
-		//This is a reimport, we are just re-updating the source data
-		StaticMeshObject = ExistingAsset;
-	}
+	UObject* StaticMeshObject = UE::Interchange::FFactoryCommon::AsyncFindObject(StaticMeshFactoryNode, GetFactoryClass(), Arguments.Parent, Arguments.AssetName);
+	const bool bReimport = Arguments.ReimportObject && StaticMeshObject;
 
 	if (!StaticMeshObject)
 	{
-		if (!Arguments.ReimportObject)
-		{
-			UE_LOG(LogInterchangeImport, Error, TEXT("Could not create StaticMesh asset %s"), *Arguments.AssetName);
-		}
+		UE_LOG(LogInterchangeImport, Error, TEXT("Could not import the StaticMesh asset %s, because the asset do not exist."), *Arguments.AssetName);
 		return ImportAssetResult;
 	}
 
 	UStaticMesh* StaticMesh = Cast<UStaticMesh>(StaticMeshObject);
 	if (!ensure(StaticMesh))
 	{
-		UE_LOG(LogInterchangeImport, Error, TEXT("Could not create StaticMesh asset %s"), *Arguments.AssetName);
+		UE_LOG(LogInterchangeImport, Error, TEXT("Could not cast to StaticMesh asset %s"), *Arguments.AssetName);
 		return ImportAssetResult;
 	}
 
diff --git a/Engine/Plugins/Interchange/Runtime/Source/Import/Private/Texture/InterchangeTextureFactory.cpp b/Engine/Plugins/Interchange/Runtime/Source/Import/Private/Texture/InterchangeTextureFactory.cpp
--- a/Engine/Plugins/Interchange/Runtime/Source/Import/Private/Texture/InterchangeTextureFactory.cpp
+++ b/Engine/Plugins/Interchange/Runtime/Source/Import/Private/Texture/InterchangeTextureFactory.cpp
@@ -1714,26 +1714,22 @@ UInterchangeFactoryBase::FImportAssetResult UInterchangeTextureFactory::ImportAs
 		return ImportAssetResult;
 	}
 
+	UObject* ExistingAsset = nullptr;
+	FSoftObjectPath ReferenceObject;
+	if (Arguments.AssetNode->GetCustomReferenceObject(ReferenceObject))
+	{
+		ExistingAsset = ReferenceObject.TryLoad();
+	}
+
 	UTexture* Texture = nullptr;
 	// create a new texture or overwrite existing asset, if possible
 	if (!ExistingAsset)
 	{
-		//NewObject is not thread safe, the asset registry directory watcher tick on the main thread can trig before we finish initializing the UObject and will crash
-		//The UObject should have been create by calling CreateEmptyAsset on the main thread.
-		if (IsInGameThread())
-		{
-			Texture = NewObject<UTexture>(Arguments.Parent, TextureClass, *Arguments.AssetName, RF_Public | RF_Standalone);
-		}
-		else
-		{
-			ImportTextureErrorLog(LOCTEXT("TextureFactory_Async_CannotCreateAsync", "UInterchangeTextureFactory: Could not create Texture asset outside of the game thread."));
-			return ImportAssetResult;
-		}
+		ImportTextureErrorLog(LOCTEXT("TextureFactory_Async_CannotCreateAsync", "UInterchangeTextureFactory: Could not create Texture asset outside of the game thread."));
+		return ImportAssetResult;
 	}
 	else if(ExistingAsset->GetClass()->IsChildOf(TextureClass))
 	{
-		//This is a reimport, we are just re-updating the source data  
-		// <- pretty sure this comment is wrong, this is hit in regular imports, because ExistingAsset was previously made by CreateNewAsset
 		Texture = static_cast<UTexture*>(ExistingAsset);
 	}
 
diff --git a/Engine/Source/Runtime/Interchange/Engine/Private/InterchangeImportCommon.cpp b/Engine/Source/Runtime/Interchange/Engine/Private/InterchangeImportCommon.cpp
--- a/Engine/Source/Runtime/Interchange/Engine/Private/InterchangeImportCommon.cpp
+++ b/Engine/Source/Runtime/Interchange/Engine/Private/InterchangeImportCommon.cpp
@@ -4,6 +4,7 @@
 #include "CoreMinimal.h"
 #include "EditorFramework/AssetImportData.h"
 #include "InterchangeAssetImportData.h"
+#include "InterchangeEngineLogPrivate.h"
 #include "InterchangePipelineBase.h"
 #include "InterchangePythonPipelineBase.h"
 #include "InterchangeSceneImportAsset.h"
@@ -310,6 +311,42 @@ namespace UE::Interchange
 				break;
 			}
 		}
+
+		UObject* FFactoryCommon::AsyncFindObject(UInterchangeFactoryBaseNode* FactoryNode, const UClass* FactoryClass, UObject* Parent, const FString& AssetName)
+		{
+			if (!FactoryNode)
+			{
+				return nullptr;
+			}
+
+			const UClass* AssetClass = FactoryNode->GetObjectClass();
+			if (!ensure(AssetClass && AssetClass->IsChildOf(FactoryClass)))
+			{
+				UE_LOG(LogInterchangeEngine, Error, TEXT("The asset class %s is not compatible with the factory class %s. Cannot import asset %s ."), *AssetClass->GetName(), *FactoryClass->GetName(), * AssetName);
+				return nullptr;
+			}
+
+			UObject* ExistingAsset = nullptr;
+			FSoftObjectPath ReferenceObject;
+			if (FactoryNode->GetCustomReferenceObject(ReferenceObject))
+			{
+				ExistingAsset = ReferenceObject.TryLoad();
+			}
+
+			// create a new static mesh or overwrite existing asset, if possible
+			if (!ExistingAsset)
+			{
+				UE_LOG(LogInterchangeEngine, Error, TEXT("Cannot import the %s asset [%s], because it was not create on the game thread."), *AssetClass->GetName(), *AssetName);
+				return nullptr;
+			}
+			
+			if (!ExistingAsset->GetClass()->IsChildOf(AssetClass))
+			{
+				UE_LOG(LogInterchangeEngine, Error, TEXT("Cannot import the %s asset [%s], because it will override an asset of a different class (%s)."), *AssetClass->GetName(), *AssetName, *ExistingAsset->GetClass()->GetName());
+				return nullptr;
+			}
+			return ExistingAsset;
+		}
 	}
 
 	UObject* FFactoryCommon::GetObjectToReimport(UObject* ReimportObject, const FString& PackageName, const FString& AssetName, const FString& SubPathString)
diff --git a/Engine/Source/Runtime/Interchange/Engine/Public/InterchangeImportCommon.h b/Engine/Source/Runtime/Interchange/Engine/Public/InterchangeImportCommon.h
--- a/Engine/Source/Runtime/Interchange/Engine/Public/InterchangeImportCommon.h
+++ b/Engine/Source/Runtime/Interchange/Engine/Public/InterchangeImportCommon.h
@@ -127,5 +127,10 @@ namespace UE::Interchange
 		 * @oaram SubPathString: Optional subobject name
 		 */
 		static 	const UInterchangeFactoryBaseNode* GetFactoryNode(UObject* ReimportObject, const FString& PackageName, const FString& AssetName, const FString& SubPathString = FString());
+		
+		/**
+		 * Factory helper use to find the object we have create in the game thread when we do the async part of the factory import.
+		 */
+		static UObject* AsyncFindObject(UInterchangeFactoryBaseNode* FactoryNode, const UClass* FactoryClass, UObject* Parent, const FString& AssetName);
 	};
 } //ns UE::Interchange
