Fix CIS build
#rnx

diff --git a/Engine/Plugins/Interchange/Runtime/Source/Import/Private/Material/InterchangeMaterialFactory.cpp b/Engine/Plugins/Interchange/Runtime/Source/Import/Private/Material/InterchangeMaterialFactory.cpp
--- a/Engine/Plugins/Interchange/Runtime/Source/Import/Private/Material/InterchangeMaterialFactory.cpp
+++ b/Engine/Plugins/Interchange/Runtime/Source/Import/Private/Material/InterchangeMaterialFactory.cpp
@@ -635,7 +635,7 @@ UInterchangeFactoryBase::FImportAssetResult UInterchangeMaterialFactory::ImportA
 	//Do not override an asset we skip
 	if (bSkipImport)
 	{
-		ImportAssetResult.ImportedObject = ExistingAsset;
+		ImportAssetResult.ImportedObject = MaterialObject;
 		return ImportAssetResult;
 	}
 	const bool bReimport = Arguments.ReimportObject && MaterialObject;
@@ -688,8 +688,7 @@ UInterchangeFactoryBase::FImportAssetResult UInterchangeMaterialFactory::ImportA
 
 	//Getting the file Hash will cache it into the source data
 	Arguments.SourceData->GetFileContentHash();
-		//The interchange completion task (call in the GameThread after the factories pass), will call PostEditChange which will trig another asynchronous system that will build all material in parallel
-	}
+
 	return ImportAssetResult;
 }
 
@@ -1461,9 +1460,10 @@ UInterchangeFactoryBase::FImportAssetResult UInterchangeMaterialFunctionFactory:
 	//Do not override an asset we skip
 	if (bSkipImport)
 	{
-		ImportAssetResult.ImportedObject = ExistingAsset;
+		ImportAssetResult.ImportedObject = MaterialObject;
 		return ImportAssetResult;
 	}
+
 	const bool bReimport = Arguments.ReimportObject && MaterialObject;
 
 	if (!MaterialObject)
diff --git a/Engine/Plugins/Interchange/Runtime/Source/Import/Private/Texture/InterchangeTextureFactory.cpp b/Engine/Plugins/Interchange/Runtime/Source/Import/Private/Texture/InterchangeTextureFactory.cpp
--- a/Engine/Plugins/Interchange/Runtime/Source/Import/Private/Texture/InterchangeTextureFactory.cpp
+++ b/Engine/Plugins/Interchange/Runtime/Source/Import/Private/Texture/InterchangeTextureFactory.cpp
@@ -1677,9 +1677,12 @@ UInterchangeFactoryBase::FImportAssetResult UInterchangeTextureFactory::ImportAs
 		return ImportAssetResult;
 	}
 
-	// create an asset if it doesn't exist
-	// typically ExistingAsset should already exist, made by CreateEmptyAsset on the main thread
-	UObject* ExistingAsset = StaticFindObject(nullptr, Arguments.Parent, *Arguments.AssetName);
+	UObject* ExistingAsset = nullptr;
+	FSoftObjectPath ReferenceObject;
+	if (Arguments.AssetNode->GetCustomReferenceObject(ReferenceObject))
+	{
+		ExistingAsset = ReferenceObject.TryLoad();
+	}
 
 	//Do not override an asset we skip
 	if (bSkipImport)
@@ -1714,13 +1717,6 @@ UInterchangeFactoryBase::FImportAssetResult UInterchangeTextureFactory::ImportAs
 		return ImportAssetResult;
 	}
 
-	UObject* ExistingAsset = nullptr;
-	FSoftObjectPath ReferenceObject;
-	if (Arguments.AssetNode->GetCustomReferenceObject(ReferenceObject))
-	{
-		ExistingAsset = ReferenceObject.TryLoad();
-	}
-
 	UTexture* Texture = nullptr;
 	// create a new texture or overwrite existing asset, if possible
 	if (!ExistingAsset)
