[Interchange] GLTFCore's normal generation re-evaluated.

diff --git a/Engine/Plugins/Interchange/Runtime/Source/Parsers/GLTFCore/Private/GLTFMeshFactory.cpp b/Engine/Plugins/Interchange/Runtime/Source/Parsers/GLTFCore/Private/GLTFMeshFactory.cpp
--- a/Engine/Plugins/Interchange/Runtime/Source/Parsers/GLTFCore/Private/GLTFMeshFactory.cpp
+++ b/Engine/Plugins/Interchange/Runtime/Source/Parsers/GLTFCore/Private/GLTFMeshFactory.cpp
@@ -149,7 +149,9 @@ namespace GLTF
 				const FVector3f& B = Positions[Indices[i + 1]];
 				const FVector3f& C = Positions[Indices[i + 2]];
 
-				const FVector3f Normal = FVector3f::CrossProduct(A - B, A - C).GetSafeNormal();
+				// We have a left-handed coordinate system, but a counter-clockwise winding order
+				// Hence normal calculation has to take the triangle vectors cross product in reverse.
+				const FVector3f Normal = FVector3f::CrossProduct(C - A, B - A).GetSafeNormal();
 				// Same for each corner of the triangle.
 				Normals[i] = Normal;
 				Normals[i + 1] = Normal;
