[Interchange - glTF] Can't import orthographic cameras.
Added CameraActor vs CineCameraActor handling to Interchange.
Refactored UInterchangeCameraNode to UInterchangeStandardCameraNode and UInterchangePhysicalCameraNode (+ FactoryNodes).
Added Standard.Perspective to Physical camera conversion to GenericScenePipeline.

diff --git a/Engine/Plugins/Experimental/Enterprise/DatasmithInterchange/Source/DatasmithInterchange/Private/InterchangeDatasmithTranslator.cpp b/Engine/Plugins/Experimental/Enterprise/DatasmithInterchange/Source/DatasmithInterchange/Private/InterchangeDatasmithTranslator.cpp
--- a/Engine/Plugins/Experimental/Enterprise/DatasmithInterchange/Source/DatasmithInterchange/Private/InterchangeDatasmithTranslator.cpp
+++ b/Engine/Plugins/Experimental/Enterprise/DatasmithInterchange/Source/DatasmithInterchange/Private/InterchangeDatasmithTranslator.cpp
@@ -381,7 +381,7 @@ void UInterchangeDatasmithTranslator::HandleDatasmithActor(UInterchangeBaseNodeC
 		TSharedRef<IDatasmithCameraActorElement> CameraActor = StaticCastSharedRef<IDatasmithCameraActorElement>(ActorElement);
 
 		// We need to add camera asset node and then instance it in the scene node.
-		UInterchangeCameraNode* CameraNode = AddCameraNode(BaseNodeContainer, CameraActor);
+		UInterchangePhysicalCameraNode* CameraNode = AddCameraNode(BaseNodeContainer, CameraActor);
 		InterchangeSceneNode->SetCustomAssetInstanceUid(CameraNode->GetUniqueID());
 	}
 	else if (ActorElement->IsA(EDatasmithElementType::Light))
@@ -402,11 +402,11 @@ void UInterchangeDatasmithTranslator::HandleDatasmithActor(UInterchangeBaseNodeC
 	}
 }
 
-UInterchangeCameraNode* UInterchangeDatasmithTranslator::AddCameraNode(UInterchangeBaseNodeContainer& BaseNodeContainer, const TSharedRef<IDatasmithCameraActorElement>& CameraActor) const
+UInterchangePhysicalCameraNode* UInterchangeDatasmithTranslator::AddCameraNode(UInterchangeBaseNodeContainer& BaseNodeContainer, const TSharedRef<IDatasmithCameraActorElement>& CameraActor) const
 {
 	using namespace UE::DatasmithInterchange;
 
-	UInterchangeCameraNode* CameraNode = NewObject<UInterchangeCameraNode>(&BaseNodeContainer);
+	UInterchangePhysicalCameraNode* CameraNode = NewObject<UInterchangePhysicalCameraNode>(&BaseNodeContainer);
 	const FString CameraUid = NodeUtils::CameraPrefix + CameraActor->GetName();
 	CameraNode->InitializeNode(CameraUid, CameraActor->GetLabel(), EInterchangeNodeContainerType::TranslatedAsset);
 	BaseNodeContainer.AddNode(CameraNode);
diff --git a/Engine/Plugins/Experimental/Enterprise/DatasmithInterchange/Source/DatasmithInterchange/Public/InterchangeDatasmithTranslator.h b/Engine/Plugins/Experimental/Enterprise/DatasmithInterchange/Source/DatasmithInterchange/Public/InterchangeDatasmithTranslator.h
--- a/Engine/Plugins/Experimental/Enterprise/DatasmithInterchange/Source/DatasmithInterchange/Public/InterchangeDatasmithTranslator.h
+++ b/Engine/Plugins/Experimental/Enterprise/DatasmithInterchange/Source/DatasmithInterchange/Public/InterchangeDatasmithTranslator.h
@@ -25,7 +25,7 @@ class IDatasmithCameraActorElement;
 class IDatasmithLightActorElement;
 class IDatasmithScene;
 class IDatasmithTransformAnimationElement;
-class UInterchangeCameraNode;
+class UInterchangePhysicalCameraNode;
 class UInterchangeBaseLightNode;
 class UInterchangeSceneNode;
 
@@ -97,7 +97,7 @@ private:
 
 	void HandleDatasmithActor(UInterchangeBaseNodeContainer& BaseNodeContainer, const TSharedRef<IDatasmithActorElement>& ActorElement, const UInterchangeSceneNode* ParentNode) const;
 
-	UInterchangeCameraNode* AddCameraNode(UInterchangeBaseNodeContainer& BaseNodeContainer, const TSharedRef<IDatasmithCameraActorElement>& CameraActor) const;
+	UInterchangePhysicalCameraNode* AddCameraNode(UInterchangeBaseNodeContainer& BaseNodeContainer, const TSharedRef<IDatasmithCameraActorElement>& CameraActor) const;
 
 	UInterchangeBaseLightNode* AddLightNode(UInterchangeBaseNodeContainer& BaseNodeContainer, const TSharedRef<IDatasmithLightActorElement>& LightActor) const;
 
diff --git a/Engine/Plugins/Interchange/Runtime/Source/FactoryNodes/Public/InterchangeCameraFactoryNode.h b/Engine/Plugins/Interchange/Runtime/Source/FactoryNodes/Public/InterchangeCameraFactoryNode.h
new file mode 100644
--- /dev/null
+++ b/Engine/Plugins/Interchange/Runtime/Source/FactoryNodes/Public/InterchangeCameraFactoryNode.h
@@ -0,0 +1,206 @@
+// Copyright Epic Games, Inc. All Rights Reserved.
+
+#pragma once
+
+#include "CoreMinimal.h"
+#include "InterchangeActorFactoryNode.h"
+
+#if WITH_ENGINE
+	#include "CineCameraComponent.h"
+#endif
+
+#include "InterchangeCameraFactoryNode.generated.h"
+
+UCLASS(BlueprintType)
+class INTERCHANGEFACTORYNODES_API UInterchangePhysicalCameraFactoryNode : public UInterchangeActorFactoryNode
+{
+	GENERATED_BODY()
+
+public:
+	UFUNCTION(BlueprintCallable, Category = "Interchange | Node | CameraFactory")
+	bool GetCustomFocalLength(float& AttributeValue) const
+	{
+		IMPLEMENT_NODE_ATTRIBUTE_GETTER(FocalLength, float);
+	}
+
+	UFUNCTION(BlueprintCallable, Category = "Interchange | Node | CameraFactory")
+	bool SetCustomFocalLength(const float& AttributeValue, bool bAddApplyDelegate = true)
+	{
+		IMPLEMENT_NODE_ATTRIBUTE_SETTER_WITH_CUSTOM_DELEGATE_WITH_CUSTOM_CLASS(UInterchangePhysicalCameraFactoryNode, FocalLength, float, UCineCameraComponent);
+	}
+
+	UFUNCTION(BlueprintCallable, Category = "Interchange | Node | CameraFactory")
+	bool GetCustomSensorWidth(float& AttributeValue) const
+	{
+		IMPLEMENT_NODE_ATTRIBUTE_GETTER(SensorWidth, float);
+	}
+
+	UFUNCTION(BlueprintCallable, Category = "Interchange | Node | CameraFactory")
+	bool SetCustomSensorWidth(const float& AttributeValue, bool bAddApplyDelegate = true)
+	{
+		IMPLEMENT_NODE_ATTRIBUTE_SETTER_WITH_CUSTOM_DELEGATE_WITH_CUSTOM_CLASS(UInterchangePhysicalCameraFactoryNode, SensorWidth, float, UCineCameraComponent);
+	}
+
+	UFUNCTION(BlueprintCallable, Category = "Interchange | Node | CameraFactory")
+	bool GetCustomSensorHeight(float& AttributeValue) const
+	{
+		IMPLEMENT_NODE_ATTRIBUTE_GETTER(SensorHeight, float);
+	}
+
+	UFUNCTION(BlueprintCallable, Category = "Interchange | Node | CameraFactory")
+	bool SetCustomSensorHeight(const float& AttributeValue, bool bAddApplyDelegate = true)
+	{
+		IMPLEMENT_NODE_ATTRIBUTE_SETTER_WITH_CUSTOM_DELEGATE_WITH_CUSTOM_CLASS(UInterchangePhysicalCameraFactoryNode, SensorHeight, float, UCineCameraComponent);
+	}
+
+	UFUNCTION(BlueprintCallable, Category = "Interchange | Node | CameraFactory")
+	bool GetCustomFocusMethod(ECameraFocusMethod& AttributeValue) const
+	{
+		IMPLEMENT_NODE_ATTRIBUTE_GETTER(FocusMethod, ECameraFocusMethod);
+	}
+
+	UFUNCTION(BlueprintCallable, Category = "Interchange | Node | CameraFactory")
+	bool SetCustomFocusMethod(const ECameraFocusMethod& AttributeValue, bool bAddApplyDelegate = true)
+	{
+		IMPLEMENT_NODE_ATTRIBUTE_SETTER_WITH_CUSTOM_DELEGATE_WITH_CUSTOM_CLASS(UInterchangePhysicalCameraFactoryNode, FocusMethod, ECameraFocusMethod, UCineCameraComponent);
+	}	
+
+	virtual void CopyWithObject(const UInterchangeFactoryBaseNode* SourceNode, UObject* Object) override
+	{
+		Super::CopyWithObject(SourceNode, Object);
+
+		if (const UInterchangePhysicalCameraFactoryNode* ActorFactoryNode = Cast<UInterchangePhysicalCameraFactoryNode>(SourceNode))
+		{
+			COPY_NODE_DELEGATES_WITH_CUSTOM_DELEGATE(ActorFactoryNode, UInterchangePhysicalCameraFactoryNode, FocalLength, float, UCineCameraComponent::StaticClass())
+			COPY_NODE_DELEGATES_WITH_CUSTOM_DELEGATE(ActorFactoryNode, UInterchangePhysicalCameraFactoryNode, SensorWidth, float, UCineCameraComponent::StaticClass())
+			COPY_NODE_DELEGATES_WITH_CUSTOM_DELEGATE(ActorFactoryNode, UInterchangePhysicalCameraFactoryNode, SensorHeight, float, UCineCameraComponent::StaticClass())
+			COPY_NODE_DELEGATES_WITH_CUSTOM_DELEGATE(ActorFactoryNode, UInterchangePhysicalCameraFactoryNode, FocusMethod, ECameraFocusMethod, UCineCameraComponent::StaticClass())
+		}
+	}
+
+private:
+	const UE::Interchange::FAttributeKey Macro_CustomFocalLengthKey = UE::Interchange::FAttributeKey(TEXT("FocalLength"));
+	const UE::Interchange::FAttributeKey Macro_CustomSensorWidthKey = UE::Interchange::FAttributeKey(TEXT("SensorWidth"));
+	const UE::Interchange::FAttributeKey Macro_CustomSensorHeightKey = UE::Interchange::FAttributeKey(TEXT("SensorHeight"));
+	const UE::Interchange::FAttributeKey Macro_CustomFocusMethodKey = UE::Interchange::FAttributeKey(TEXT("FocusMethod"));
+
+private:
+	IMPLEMENT_NODE_ATTRIBUTE_DELEGATE_BY_PROPERTYNAME(FocalLength, float, UCineCameraComponent, TEXT("CurrentFocalLength"));
+	IMPLEMENT_NODE_ATTRIBUTE_DELEGATE_BY_PROPERTYNAME(SensorWidth, float, UCineCameraComponent, TEXT("Filmback.SensorWidth"));
+	IMPLEMENT_NODE_ATTRIBUTE_DELEGATE_BY_PROPERTYNAME(SensorHeight, float, UCineCameraComponent, TEXT("Filmback.SensorHeight"));
+	IMPLEMENT_NODE_ATTRIBUTE_DELEGATE_BY_PROPERTYNAME(FocusMethod, ECameraFocusMethod, UCineCameraComponent, TEXT("FocusSettings.FocusMethod"));
+};
+
+
+/*
+* Used for common non-physical camera with orthographic or perspective projection.
+*/
+UCLASS(BlueprintType)
+class INTERCHANGEFACTORYNODES_API UInterchangeStandardCameraFactoryNode : public UInterchangeActorFactoryNode
+{
+	GENERATED_BODY()
+
+public:
+	UFUNCTION(BlueprintCallable, Category = "Interchange | Node | CameraFactory")
+		bool GetCustomProjectionMode(TEnumAsByte<ECameraProjectionMode::Type>& AttributeValue) const
+	{
+		IMPLEMENT_NODE_ATTRIBUTE_GETTER(ProjectionMode, TEnumAsByte<ECameraProjectionMode::Type>);
+	}
+	UFUNCTION(BlueprintCallable, Category = "Interchange | Node | CameraFactory")
+		bool SetCustomProjectionMode(const TEnumAsByte<ECameraProjectionMode::Type>& AttributeValue, bool bAddApplyDelegate = true)
+	{
+		IMPLEMENT_NODE_ATTRIBUTE_SETTER_WITH_CUSTOM_DELEGATE_WITH_CUSTOM_CLASS(UInterchangeStandardCameraFactoryNode, ProjectionMode, TEnumAsByte<ECameraProjectionMode::Type>, UCameraComponent);
+	}
+
+	UFUNCTION(BlueprintCallable, Category = "Interchange | Node | CameraFactory")
+	bool GetCustomWidth(float& AttributeValue) const
+	{
+		IMPLEMENT_NODE_ATTRIBUTE_GETTER(Width, float);
+	}
+	UFUNCTION(BlueprintCallable, Category = "Interchange | Node | CameraFactory")
+	bool SetCustomWidth(const float& AttributeValue, bool bAddApplyDelegate = true)
+	{
+		IMPLEMENT_NODE_ATTRIBUTE_SETTER_WITH_CUSTOM_DELEGATE_WITH_CUSTOM_CLASS(UInterchangeStandardCameraFactoryNode, Width, float, UCameraComponent);
+	}
+
+	UFUNCTION(BlueprintCallable, Category = "Interchange | Node | CameraFactory")
+	bool GetCustomNearClipPlane(float& AttributeValue) const
+	{
+		IMPLEMENT_NODE_ATTRIBUTE_GETTER(NearClipPlane, float);
+	}
+	UFUNCTION(BlueprintCallable, Category = "Interchange | Node | CameraFactory")
+	bool SetCustomNearClipPlane(const float& AttributeValue, bool bAddApplyDelegate = true)
+	{
+		IMPLEMENT_NODE_ATTRIBUTE_SETTER_WITH_CUSTOM_DELEGATE_WITH_CUSTOM_CLASS(UInterchangeStandardCameraFactoryNode, NearClipPlane, float, UCameraComponent);
+	}
+
+	UFUNCTION(BlueprintCallable, Category = "Interchange | Node | CameraFactory")
+	bool GetCustomFarClipPlane(float& AttributeValue) const
+	{
+		IMPLEMENT_NODE_ATTRIBUTE_GETTER(FarClipPlane, float);
+	}
+	UFUNCTION(BlueprintCallable, Category = "Interchange | Node | CameraFactory")
+	bool SetCustomFarClipPlane(const float& AttributeValue, bool bAddApplyDelegate = true)
+	{
+		IMPLEMENT_NODE_ATTRIBUTE_SETTER_WITH_CUSTOM_DELEGATE_WITH_CUSTOM_CLASS(UInterchangeStandardCameraFactoryNode, FarClipPlane, float, UCameraComponent);
+	}
+
+	UFUNCTION(BlueprintCallable, Category = "Interchange | Node | CameraFactory")
+	bool GetCustomAspectRatio(float& AttributeValue) const
+	{
+		IMPLEMENT_NODE_ATTRIBUTE_GETTER(AspectRatio, float);
+	}
+	UFUNCTION(BlueprintCallable, Category = "Interchange | Node | CameraFactory")
+	bool SetCustomAspectRatio(const float& AttributeValue, bool bAddApplyDelegate = true)
+	{
+		IMPLEMENT_NODE_ATTRIBUTE_SETTER_WITH_CUSTOM_DELEGATE_WITH_CUSTOM_CLASS(UInterchangeStandardCameraFactoryNode, AspectRatio, float, UCameraComponent);
+	}
+
+	UFUNCTION(BlueprintCallable, Category = "Interchange | Node | CameraFactory")
+		bool GetCustomFieldOfView(float& AttributeValue) const
+	{
+		IMPLEMENT_NODE_ATTRIBUTE_GETTER(FieldOfView, float);
+	}
+	UFUNCTION(BlueprintCallable, Category = "Interchange | Node | CameraFactory")
+		bool SetCustomFieldOfView(const float& AttributeValue, bool bAddApplyDelegate = true)
+	{
+		IMPLEMENT_NODE_ATTRIBUTE_SETTER_WITH_CUSTOM_DELEGATE_WITH_CUSTOM_CLASS(UInterchangeStandardCameraFactoryNode, FieldOfView, float, UCameraComponent);
+	}
+
+	virtual void CopyWithObject(const UInterchangeFactoryBaseNode* SourceNode, UObject* Object) override
+	{
+		Super::CopyWithObject(SourceNode, Object);
+
+		if (const UInterchangeStandardCameraFactoryNode* ActorFactoryNode = Cast<UInterchangeStandardCameraFactoryNode>(SourceNode))
+		{
+			COPY_NODE_DELEGATES_WITH_CUSTOM_DELEGATE(ActorFactoryNode, UInterchangeStandardCameraFactoryNode, ProjectionMode, TEnumAsByte<ECameraProjectionMode::Type>, UCameraComponent::StaticClass())
+
+			COPY_NODE_DELEGATES_WITH_CUSTOM_DELEGATE(ActorFactoryNode, UInterchangeStandardCameraFactoryNode, Width, float, UCameraComponent::StaticClass())
+			COPY_NODE_DELEGATES_WITH_CUSTOM_DELEGATE(ActorFactoryNode, UInterchangeStandardCameraFactoryNode, NearClipPlane, float, UCameraComponent::StaticClass())
+			COPY_NODE_DELEGATES_WITH_CUSTOM_DELEGATE(ActorFactoryNode, UInterchangeStandardCameraFactoryNode, FarClipPlane, float, UCameraComponent::StaticClass())
+			COPY_NODE_DELEGATES_WITH_CUSTOM_DELEGATE(ActorFactoryNode, UInterchangeStandardCameraFactoryNode, AspectRatio, float, UCameraComponent::StaticClass())
+
+			COPY_NODE_DELEGATES_WITH_CUSTOM_DELEGATE(ActorFactoryNode, UInterchangeStandardCameraFactoryNode, FieldOfView, float, UCameraComponent::StaticClass())
+		}
+	}
+
+private:
+	const UE::Interchange::FAttributeKey Macro_CustomProjectionModeKey = UE::Interchange::FAttributeKey(TEXT("ProjectionMode"));
+
+	const UE::Interchange::FAttributeKey Macro_CustomWidthKey = UE::Interchange::FAttributeKey(TEXT("Width"));
+	const UE::Interchange::FAttributeKey Macro_CustomNearClipPlaneKey = UE::Interchange::FAttributeKey(TEXT("NearClip"));
+	const UE::Interchange::FAttributeKey Macro_CustomFarClipPlaneKey = UE::Interchange::FAttributeKey(TEXT("FarClipPlane"));
+	const UE::Interchange::FAttributeKey Macro_CustomAspectRatioKey = UE::Interchange::FAttributeKey(TEXT("AspectRatio"));
+	
+	const UE::Interchange::FAttributeKey Macro_CustomFieldOfViewKey = UE::Interchange::FAttributeKey(TEXT("FieldOfView"));
+
+private:
+	IMPLEMENT_NODE_ATTRIBUTE_DELEGATE_BY_PROPERTYNAME(ProjectionMode, TEnumAsByte<ECameraProjectionMode::Type>, UCameraComponent, TEXT("ProjectionMode"));
+
+	IMPLEMENT_NODE_ATTRIBUTE_DELEGATE_BY_PROPERTYNAME(Width, float, UCameraComponent, TEXT("OrthoWidth"));
+	IMPLEMENT_NODE_ATTRIBUTE_DELEGATE_BY_PROPERTYNAME(NearClipPlane, float, UCameraComponent, TEXT("OrthoNearClipPlane"));
+	IMPLEMENT_NODE_ATTRIBUTE_DELEGATE_BY_PROPERTYNAME(FarClipPlane, float, UCameraComponent, TEXT("OrthoFarClipPlane"));
+	IMPLEMENT_NODE_ATTRIBUTE_DELEGATE_BY_PROPERTYNAME(AspectRatio, float, UCameraComponent, TEXT("AspectRatio"));
+
+	IMPLEMENT_NODE_ATTRIBUTE_DELEGATE_BY_PROPERTYNAME(FieldOfView, float, UCameraComponent, TEXT("FieldOfView"));
+};
\ No newline at end of file
diff --git a/Engine/Plugins/Interchange/Runtime/Source/FactoryNodes/Public/InterchangeCineCameraFactoryNode.h b/Engine/Plugins/Interchange/Runtime/Source/FactoryNodes/Public/InterchangeCineCameraFactoryNode.h
deleted file mode 100644
diff --git a/Engine/Plugins/Interchange/Runtime/Source/Import/Private/Gltf/InterchangeGltfTranslator.cpp b/Engine/Plugins/Interchange/Runtime/Source/Import/Private/Gltf/InterchangeGltfTranslator.cpp
--- a/Engine/Plugins/Interchange/Runtime/Source/Import/Private/Gltf/InterchangeGltfTranslator.cpp
+++ b/Engine/Plugins/Interchange/Runtime/Source/Import/Private/Gltf/InterchangeGltfTranslator.cpp
@@ -1192,31 +1192,29 @@ bool UInterchangeGltfTranslator::Translate( UInterchangeBaseNodeContainer& NodeC
 		int32 CameraIndex = 0;
 		for ( const GLTF::FCamera& GltfCamera : GltfAsset.Cameras )
 		{
-			UInterchangeCameraNode* CameraNode = NewObject< UInterchangeCameraNode >( &NodeContainer );
+			UInterchangeStandardCameraNode* CameraNode = NewObject< UInterchangeStandardCameraNode >(&NodeContainer);
 			FString CameraNodeUid = TEXT("\\Camera\\") + GltfCamera.UniqueId;
-			CameraNode->InitializeNode( CameraNodeUid, GltfCamera.Name, EInterchangeNodeContainerType::TranslatedAsset );
+			CameraNode->InitializeNode(CameraNodeUid, GltfCamera.Name, EInterchangeNodeContainerType::TranslatedAsset);
 
-			float       AspectRatio;
-			float       FocalLength;
-			const float SensorWidth = 36.f;  // mm
-			CameraNode->SetCustomSensorWidth(SensorWidth);
 			if (GltfCamera.bIsPerspective)
 			{
-				AspectRatio = GltfCamera.Perspective.AspectRatio;
-				FocalLength = (SensorWidth / GltfCamera.Perspective.AspectRatio) / (2.0 * tan(GltfCamera.Perspective.Fov / 2.0));
+				CameraNode->SetCustomProjectionMode(EInterchangeCameraProjectionType::Perspective);
+
+				CameraNode->SetCustomFieldOfView(FMath::RadiansToDegrees(GltfCamera.Perspective.Fov));
+				CameraNode->SetCustomAspectRatio(GltfCamera.Perspective.AspectRatio);
 			}
 			else
 			{
-				AspectRatio = GltfCamera.Orthographic.XMagnification / GltfCamera.Orthographic.YMagnification;
-				FocalLength = (SensorWidth / AspectRatio) / (AspectRatio * tan(AspectRatio / 4.0));  // can only approximate Fov
-			}
+				CameraNode->SetCustomProjectionMode(EInterchangeCameraProjectionType::Orthographic);
 
-			CameraNode->SetCustomSensorHeight(SensorWidth / AspectRatio);
-			CameraNode->SetCustomFocalLength(FocalLength);
-			CameraNode->SetCustomEnableDepthOfField(false);
-			// ignoring znear and zfar
+				CameraNode->SetCustomWidth(GltfCamera.Orthographic.XMagnification * GltfUnitConversionMultiplier);
+				CameraNode->SetCustomNearClipPlane(GltfCamera.ZNear * GltfUnitConversionMultiplier);
+				CameraNode->SetCustomFarClipPlane(GltfCamera.ZFar * GltfUnitConversionMultiplier);
+
+				CameraNode->SetCustomAspectRatio(GltfCamera.Orthographic.XMagnification / GltfCamera.Orthographic.YMagnification);
+			}
 
-			NodeContainer.AddNode( CameraNode );
+			NodeContainer.AddNode(CameraNode);
 			++CameraIndex;
 		}
 	}
diff --git a/Engine/Plugins/Interchange/Runtime/Source/Import/Private/InterchangeImportModule.cpp b/Engine/Plugins/Interchange/Runtime/Source/Import/Private/InterchangeImportModule.cpp
--- a/Engine/Plugins/Interchange/Runtime/Source/Import/Private/InterchangeImportModule.cpp
+++ b/Engine/Plugins/Interchange/Runtime/Source/Import/Private/InterchangeImportModule.cpp
@@ -20,7 +20,7 @@
 #include "Misc/CoreDelegates.h"
 #include "Modules/ModuleManager.h"
 #include "Scene/InterchangeActorFactory.h"
-#include "Scene/InterchangeCineCameraActorFactory.h"
+#include "Scene/InterchangeCameraActorFactory.h"
 #include "Scene/InterchangeSceneImportAssetFactory.h"
 #include "Scene/InterchangeLightActorFactory.h"
 #include "Scene/InterchangeSceneVariantSetsFactory.h"
@@ -84,6 +84,7 @@ void FInterchangeImportModule::StartupModule()
 		InterchangeManager.RegisterFactory(UInterchangeAnimationTrackSetFactory::StaticClass());
 		InterchangeManager.RegisterFactory(UInterchangeAnimSequenceFactory::StaticClass());
 		InterchangeManager.RegisterFactory(UInterchangeCineCameraActorFactory::StaticClass());
+		InterchangeManager.RegisterFactory(UInterchangeCameraActorFactory::StaticClass());
 		InterchangeManager.RegisterFactory(UInterchangeStaticMeshActorFactory::StaticClass());
 		InterchangeManager.RegisterFactory(UInterchangeSkeletalMeshActorFactory::StaticClass());
 		InterchangeManager.RegisterFactory(UInterchangeSceneVariantSetsFactory::StaticClass());
diff --git a/Engine/Plugins/Interchange/Runtime/Source/Import/Private/Scene/InterchangeCameraActorFactory.cpp b/Engine/Plugins/Interchange/Runtime/Source/Import/Private/Scene/InterchangeCameraActorFactory.cpp
new file mode 100644
--- /dev/null
+++ b/Engine/Plugins/Interchange/Runtime/Source/Import/Private/Scene/InterchangeCameraActorFactory.cpp
@@ -0,0 +1,40 @@
+// Copyright Epic Games, Inc. All Rights Reserved.
+
+#include "Scene/InterchangeCameraActorFactory.h"
+
+#include "InterchangeCameraFactoryNode.h"
+#include "Scene/InterchangeActorFactory.h"
+#include "Scene/InterchangeActorHelper.h"
+
+#include "CineCameraActor.h"
+#include "CineCameraComponent.h"
+#include "CoreMinimal.h"
+#include "UObject/Object.h"
+#include "UObject/ObjectMacros.h"
+
+#include UE_INLINE_GENERATED_CPP_BY_NAME(InterchangeCameraActorFactory)
+
+UClass* UInterchangeCineCameraActorFactory::GetFactoryClass() const
+{
+	return ACineCameraActor::StaticClass();
+}
+
+UObject* UInterchangeCineCameraActorFactory::ProcessActor(AActor& SpawnedActor, const UInterchangeActorFactoryNode& /*FactoryNode*/, const UInterchangeBaseNodeContainer& /*NodeContainer*/)
+{
+	ACineCameraActor* CineCameraActor = Cast<ACineCameraActor>(&SpawnedActor);
+
+	return CineCameraActor ? CineCameraActor->GetCineCameraComponent() : nullptr;
+}
+
+
+UClass* UInterchangeCameraActorFactory::GetFactoryClass() const
+{
+	return ACameraActor::StaticClass();
+}
+
+UObject* UInterchangeCameraActorFactory::ProcessActor(AActor& SpawnedActor, const UInterchangeActorFactoryNode& /*FactoryNode*/, const UInterchangeBaseNodeContainer& /*NodeContainer*/)
+{
+	ACameraActor* CameraActor = Cast<ACameraActor>(&SpawnedActor);
+
+	return CameraActor ? CameraActor->GetCameraComponent() : nullptr;
+}
\ No newline at end of file
diff --git a/Engine/Plugins/Interchange/Runtime/Source/Import/Private/Scene/InterchangeCineCameraActorFactory.cpp b/Engine/Plugins/Interchange/Runtime/Source/Import/Private/Scene/InterchangeCineCameraActorFactory.cpp
deleted file mode 100644
diff --git a/Engine/Plugins/Interchange/Runtime/Source/Import/Public/Scene/InterchangeCameraActorFactory.h b/Engine/Plugins/Interchange/Runtime/Source/Import/Public/Scene/InterchangeCameraActorFactory.h
new file mode 100644
--- /dev/null
+++ b/Engine/Plugins/Interchange/Runtime/Source/Import/Public/Scene/InterchangeCameraActorFactory.h
@@ -0,0 +1,59 @@
+// Copyright Epic Games, Inc. All Rights Reserved.
+
+#pragma once
+
+#include "CoreMinimal.h"
+#include "Scene/InterchangeActorFactory.h"
+#include "UObject/Object.h"
+#include "UObject/ObjectMacros.h"
+
+#include "InterchangeCameraActorFactory.generated.h"
+
+UCLASS(BlueprintType)
+class INTERCHANGEIMPORT_API UInterchangeCineCameraActorFactory : public UInterchangeActorFactory
+{
+	GENERATED_BODY()
+public:
+
+	//////////////////////////////////////////////////////////////////////////
+	// Interchange factory base interface begin
+
+	virtual UClass* GetFactoryClass() const override;
+
+	// Interchange factory base interface end
+	//////////////////////////////////////////////////////////////////////////
+
+	//////////////////////////////////////////////////////////////////////////
+	// Interchange actor factory interface begin
+
+	virtual UObject* ProcessActor(class AActor& SpawnedActor, const UInterchangeActorFactoryNode& FactoryNode, const UInterchangeBaseNodeContainer& NodeContainer) override;
+
+
+	// Interchange actor factory interface end
+	//////////////////////////////////////////////////////////////////////////
+};
+
+UCLASS(BlueprintType)
+class INTERCHANGEIMPORT_API UInterchangeCameraActorFactory : public UInterchangeActorFactory
+{
+	GENERATED_BODY()
+public:
+
+	//////////////////////////////////////////////////////////////////////////
+	// Interchange factory base interface begin
+
+	virtual UClass* GetFactoryClass() const override;
+
+	// Interchange factory base interface end
+	//////////////////////////////////////////////////////////////////////////
+
+	//////////////////////////////////////////////////////////////////////////
+	// Interchange actor factory interface begin
+
+	virtual UObject* ProcessActor(class AActor& SpawnedActor, const UInterchangeActorFactoryNode& FactoryNode, const UInterchangeBaseNodeContainer& NodeContainer) override;
+
+
+	// Interchange actor factory interface end
+	//////////////////////////////////////////////////////////////////////////
+};
+
diff --git a/Engine/Plugins/Interchange/Runtime/Source/Import/Public/Scene/InterchangeCineCameraActorFactory.h b/Engine/Plugins/Interchange/Runtime/Source/Import/Public/Scene/InterchangeCineCameraActorFactory.h
deleted file mode 100644
diff --git a/Engine/Plugins/Interchange/Runtime/Source/Nodes/Public/InterchangeCameraNode.h b/Engine/Plugins/Interchange/Runtime/Source/Nodes/Public/InterchangeCameraNode.h
--- a/Engine/Plugins/Interchange/Runtime/Source/Nodes/Public/InterchangeCameraNode.h
+++ b/Engine/Plugins/Interchange/Runtime/Source/Nodes/Public/InterchangeCameraNode.h
@@ -8,14 +8,14 @@
 #include "InterchangeCameraNode.generated.h"
 
 UCLASS(BlueprintType)
-class INTERCHANGENODES_API UInterchangeCameraNode : public UInterchangeBaseNode
+class INTERCHANGENODES_API UInterchangePhysicalCameraNode : public UInterchangeBaseNode
 {
 	GENERATED_BODY()
 
 public:
 	static FStringView StaticAssetTypeName()
 	{
-		return TEXT("Camera");
+		return TEXT("PhysicalCamera");
 	}
 
 	/**
@@ -23,7 +23,7 @@ public:
 	 */
 	virtual FString GetTypeName() const override
 	{
-		const FString TypeName = TEXT("CameraNode");
+		const FString TypeName = TEXT("PhysicalCameraNode");
 		return TypeName;
 	}
 
@@ -75,10 +75,118 @@ public:
 	{
 		IMPLEMENT_NODE_ATTRIBUTE_SETTER_NODELEGATE(EnableDepthOfField, bool);
 	}
-
 private:
 	const UE::Interchange::FAttributeKey Macro_CustomFocalLengthKey = UE::Interchange::FAttributeKey(TEXT("FocalLength"));
 	const UE::Interchange::FAttributeKey Macro_CustomSensorWidthKey = UE::Interchange::FAttributeKey(TEXT("SensorWidth"));
 	const UE::Interchange::FAttributeKey Macro_CustomSensorHeightKey = UE::Interchange::FAttributeKey(TEXT("SensorHeight"));
 	const UE::Interchange::FAttributeKey Macro_CustomEnableDepthOfFieldKey = UE::Interchange::FAttributeKey(TEXT("EnableDepthOfField"));
 };
+
+UENUM(BlueprintType)
+enum class EInterchangeCameraProjectionType : uint8
+{
+	Perspective,
+	Orthographic
+};
+
+// Primarily used for Ortho Camera
+UCLASS(BlueprintType)
+class INTERCHANGENODES_API UInterchangeStandardCameraNode : public UInterchangeBaseNode
+{
+	GENERATED_BODY()
+
+public:
+	static FStringView StaticAssetTypeName()
+	{
+		return TEXT("StandardCamera");
+	}
+
+	/**
+	 * Return the node type name of the class, we use this when reporting errors
+	 */
+	virtual FString GetTypeName() const override
+	{
+		const FString TypeName = TEXT("StandardCameraNode");
+		return TypeName;
+	}
+
+public:
+
+	UFUNCTION(BlueprintCallable, Category = "Interchange | Node | Camera")
+	bool GetCustomProjectionMode(EInterchangeCameraProjectionType& AttributeValue) const
+	{
+		IMPLEMENT_NODE_ATTRIBUTE_GETTER(ProjectionMode, EInterchangeCameraProjectionType);
+	}
+	UFUNCTION(BlueprintCallable, Category = "Interchange | Node | Camera")
+	bool SetCustomProjectionMode(const EInterchangeCameraProjectionType& AttributeValue)
+	{
+		IMPLEMENT_NODE_ATTRIBUTE_SETTER_NODELEGATE(ProjectionMode, EInterchangeCameraProjectionType);
+	}
+
+	UFUNCTION(BlueprintCallable, Category = "Interchange | Node | Camera")
+	bool GetCustomWidth(float& AttributeValue) const
+	{
+		IMPLEMENT_NODE_ATTRIBUTE_GETTER(Width, float);
+	}
+	UFUNCTION(BlueprintCallable, Category = "Interchange | Node | Camera")
+	bool SetCustomWidth(const float& AttributeValue)
+	{
+		IMPLEMENT_NODE_ATTRIBUTE_SETTER_NODELEGATE(Width, float);
+	}
+
+	UFUNCTION(BlueprintCallable, Category = "Interchange | Node | Camera")
+	bool GetCustomNearClipPlane(float& AttributeValue) const
+	{
+	IMPLEMENT_NODE_ATTRIBUTE_GETTER(NearClipPlane, float);
+	}
+	UFUNCTION(BlueprintCallable, Category = "Interchange | Node | Camera")
+	bool SetCustomNearClipPlane(const float& AttributeValue)
+	{
+		IMPLEMENT_NODE_ATTRIBUTE_SETTER_NODELEGATE(NearClipPlane, float);
+	}
+
+	UFUNCTION(BlueprintCallable, Category = "Interchange | Node | Camera")
+	bool GetCustomFarClipPlane(float& AttributeValue) const
+	{
+		IMPLEMENT_NODE_ATTRIBUTE_GETTER(FarClipPlane, float);
+	}
+	UFUNCTION(BlueprintCallable, Category = "Interchange | Node | Camera")
+	bool SetCustomFarClipPlane(const float& AttributeValue)
+	{
+		IMPLEMENT_NODE_ATTRIBUTE_SETTER_NODELEGATE(FarClipPlane, float);
+	}
+
+	UFUNCTION(BlueprintCallable, Category = "Interchange | Node | Camera")
+	bool GetCustomAspectRatio(float& AttributeValue) const
+	{
+		IMPLEMENT_NODE_ATTRIBUTE_GETTER(AspectRatio, float);
+	}
+	UFUNCTION(BlueprintCallable, Category = "Interchange | Node | Camera")
+	bool SetCustomAspectRatio(const float& AttributeValue)
+	{
+		IMPLEMENT_NODE_ATTRIBUTE_SETTER_NODELEGATE(AspectRatio, float);
+	}
+
+	//Field of View in Degrees.
+	UFUNCTION(BlueprintCallable, Category = "Interchange | Node | Camera")
+	bool GetCustomFieldOfView(float& AttributeValue) const
+	{
+		IMPLEMENT_NODE_ATTRIBUTE_GETTER(FieldOfView, float);
+	}
+	//Field of View in Degrees.
+	UFUNCTION(BlueprintCallable, Category = "Interchange | Node | Camera")
+	bool SetCustomFieldOfView(const float& AttributeValue)
+	{
+		IMPLEMENT_NODE_ATTRIBUTE_SETTER_NODELEGATE(FieldOfView, float);
+	}
+
+private:
+	const UE::Interchange::FAttributeKey Macro_CustomProjectionModeKey = UE::Interchange::FAttributeKey(TEXT("ProjectionMode"));
+	
+	const UE::Interchange::FAttributeKey Macro_CustomWidthKey = UE::Interchange::FAttributeKey(TEXT("Width"));
+	const UE::Interchange::FAttributeKey Macro_CustomNearClipPlaneKey = UE::Interchange::FAttributeKey(TEXT("NearClipPlane"));
+	const UE::Interchange::FAttributeKey Macro_CustomFarClipPlaneKey = UE::Interchange::FAttributeKey(TEXT("FarClipPlane"));
+	const UE::Interchange::FAttributeKey Macro_CustomAspectRatioKey = UE::Interchange::FAttributeKey(TEXT("AspectRatio"));
+
+	const UE::Interchange::FAttributeKey Macro_CustomFieldOfViewKey = UE::Interchange::FAttributeKey(TEXT("FieldOfView"));
+};
\ No newline at end of file
diff --git a/Engine/Plugins/Interchange/Runtime/Source/Parsers/Fbx/Private/FbxCamera.cpp b/Engine/Plugins/Interchange/Runtime/Source/Parsers/Fbx/Private/FbxCamera.cpp
--- a/Engine/Plugins/Interchange/Runtime/Source/Parsers/Fbx/Private/FbxCamera.cpp
+++ b/Engine/Plugins/Interchange/Runtime/Source/Parsers/Fbx/Private/FbxCamera.cpp
@@ -19,7 +19,7 @@ namespace UE
 	{
 		namespace Private
 		{
-			void FillCameraNode(UInterchangeCameraNode* CameraNode, FbxCamera& SourceCamera)
+			void FillCameraNode(UInterchangePhysicalCameraNode* CameraNode, FbxCamera& SourceCamera)
 			{
 				if (!CameraNode)
 				{
@@ -41,9 +41,9 @@ namespace UE
 				CameraNode->SetCustomSensorWidth(FUnitConversion::Convert(SourceCamera.GetApertureWidth(), EUnit::Inches, EUnit::Millimeters));
 			}
 
-			UInterchangeCameraNode* FFbxCamera::CreateCameraNode(UInterchangeBaseNodeContainer& NodeContainer, const FString& NodeUid, const FString& NodeName)
+			UInterchangePhysicalCameraNode* FFbxCamera::CreateCameraNode(UInterchangeBaseNodeContainer& NodeContainer, const FString& NodeUid, const FString& NodeName)
 			{
-				UInterchangeCameraNode* CameraNode = NewObject<UInterchangeCameraNode>(&NodeContainer, NAME_None);
+				UInterchangePhysicalCameraNode* CameraNode = NewObject<UInterchangePhysicalCameraNode>(&NodeContainer, NAME_None);
 				if (!ensure(CameraNode))
 				{
 					UInterchangeResultError_Generic* Message = Parser.AddMessage<UInterchangeResultError_Generic>();
@@ -66,14 +66,14 @@ namespace UE
 
 					if (NodeAttribute && NodeAttribute->GetAttributeType() == FbxNodeAttribute::eCamera)
 					{
-						FString NodeName = Parser.GetFbxHelper()->GetNodeAttributeName(NodeAttribute, UInterchangeCameraNode::StaticAssetTypeName());
-						FString NodeUid = Parser.GetFbxHelper()->GetNodeAttributeUniqueID(NodeAttribute, UInterchangeCameraNode::StaticAssetTypeName());
+						FString NodeName = Parser.GetFbxHelper()->GetNodeAttributeName(NodeAttribute, UInterchangePhysicalCameraNode::StaticAssetTypeName());
+						FString NodeUid = Parser.GetFbxHelper()->GetNodeAttributeUniqueID(NodeAttribute, UInterchangePhysicalCameraNode::StaticAssetTypeName());
 
-						const UInterchangeCameraNode* CameraNode = Cast<const UInterchangeCameraNode>(NodeContainer.GetNode(NodeUid));
+						const UInterchangePhysicalCameraNode* CameraNode = Cast<const UInterchangePhysicalCameraNode>(NodeContainer.GetNode(NodeUid));
 
 						if (!CameraNode)
 						{
-							UInterchangeCameraNode* NewCameraNode = CreateCameraNode(NodeContainer, NodeUid, NodeName);
+							UInterchangePhysicalCameraNode* NewCameraNode = CreateCameraNode(NodeContainer, NodeUid, NodeName);
 							FillCameraNode(NewCameraNode, static_cast<FbxCamera&>(*NodeAttribute));
 						}
 					}
diff --git a/Engine/Plugins/Interchange/Runtime/Source/Parsers/Fbx/Private/FbxCamera.h b/Engine/Plugins/Interchange/Runtime/Source/Parsers/Fbx/Private/FbxCamera.h
--- a/Engine/Plugins/Interchange/Runtime/Source/Parsers/Fbx/Private/FbxCamera.h
+++ b/Engine/Plugins/Interchange/Runtime/Source/Parsers/Fbx/Private/FbxCamera.h
@@ -5,7 +5,7 @@
 #include "FbxInclude.h"
 
 class UInterchangeBaseNodeContainer;
-class UInterchangeCameraNode;
+class UInterchangePhysicalCameraNode;
 
 namespace UE
 {
@@ -23,14 +23,14 @@ namespace UE
 				{}
 
 				/**
-				 * Create a UInterchangeCameraNode and add it to the NodeContainer for all FbxNodeAttributes of type eCamera that the fbx file contains.
+				 * Create a UInterchangePhysicalCameraNode and add it to the NodeContainer for all FbxNodeAttributes of type eCamera that the fbx file contains.
 				 * 
 				 * @note - Any node that already exist in the NodeContainer will not be created or modified.
 				 */
 				void AddAllCameras(FbxScene* SDKScene, UInterchangeBaseNodeContainer& NodeContainer);
 
 			protected:
-				UInterchangeCameraNode* CreateCameraNode(UInterchangeBaseNodeContainer& NodeContainer, const FString& NodeUID, const FString& NodeName);
+				UInterchangePhysicalCameraNode* CreateCameraNode(UInterchangeBaseNodeContainer& NodeContainer, const FString& NodeUID, const FString& NodeName);
 				void AddCamerasRecursively(FbxNode* Node, UInterchangeBaseNodeContainer& NodeContainer);
 
 			private:
diff --git a/Engine/Plugins/Interchange/Runtime/Source/Parsers/Fbx/Private/FbxScene.cpp b/Engine/Plugins/Interchange/Runtime/Source/Parsers/Fbx/Private/FbxScene.cpp
--- a/Engine/Plugins/Interchange/Runtime/Source/Parsers/Fbx/Private/FbxScene.cpp
+++ b/Engine/Plugins/Interchange/Runtime/Source/Parsers/Fbx/Private/FbxScene.cpp
@@ -73,7 +73,7 @@ namespace UE
 
 			void FFbxScene::CreateCameraNodeReference(UInterchangeSceneNode* UnrealSceneNode, FbxNodeAttribute* NodeAttribute, UInterchangeBaseNodeContainer& NodeContainer)
 			{
-				CreateAssetNodeReference(Parser, UnrealSceneNode, NodeAttribute, NodeContainer, UInterchangeCameraNode::StaticAssetTypeName());
+				CreateAssetNodeReference(Parser, UnrealSceneNode, NodeAttribute, NodeContainer, UInterchangePhysicalCameraNode::StaticAssetTypeName());
 			}
 
 			void FFbxScene::CreateLightNodeReference(UInterchangeSceneNode* UnrealSceneNode, FbxNodeAttribute* NodeAttribute, UInterchangeBaseNodeContainer& NodeContainer)
diff --git a/Engine/Plugins/Interchange/Runtime/Source/Pipelines/Private/InterchangeGenericScenesPipeline.cpp b/Engine/Plugins/Interchange/Runtime/Source/Pipelines/Private/InterchangeGenericScenesPipeline.cpp
--- a/Engine/Plugins/Interchange/Runtime/Source/Pipelines/Private/InterchangeGenericScenesPipeline.cpp
+++ b/Engine/Plugins/Interchange/Runtime/Source/Pipelines/Private/InterchangeGenericScenesPipeline.cpp
@@ -4,7 +4,7 @@
 
 #include "InterchangeActorFactoryNode.h"
 #include "InterchangeCameraNode.h"
-#include "InterchangeCineCameraFactoryNode.h"
+#include "InterchangeCameraFactoryNode.h"
 #include "InterchangeCommonPipelineDataFactoryNode.h"
 #include "InterchangeSceneImportAsset.h"
 #include "InterchangeSceneImportAssetFactoryNode.h"
@@ -313,9 +313,26 @@ UInterchangeActorFactoryNode* UInterchangeGenericLevelPipeline::CreateActorFacto
 
 	if(TranslatedAssetNode)
 	{
-		if(TranslatedAssetNode->IsA<UInterchangeCameraNode>())
+		if(TranslatedAssetNode->IsA<UInterchangePhysicalCameraNode>())
 		{
-			return NewObject<UInterchangeCineCameraFactoryNode>(BaseNodeContainer, NAME_None);
+			return NewObject<UInterchangePhysicalCameraFactoryNode>(BaseNodeContainer, NAME_None);
+		}
+		if (TranslatedAssetNode->IsA<UInterchangeStandardCameraNode>())
+		{
+			if (bUsePhysicalInsteadOfStandardPerspectiveCamera)
+			{
+				//in case it has perspective projection we want to use PhysicalCamera (CineCamera) instead:
+				if (const UInterchangeStandardCameraNode* CameraNode = Cast<UInterchangeStandardCameraNode>(TranslatedAssetNode))
+				{
+					EInterchangeCameraProjectionType ProjectionType = EInterchangeCameraProjectionType::Perspective;
+					if (CameraNode->GetCustomProjectionMode(ProjectionType) && ProjectionType == EInterchangeCameraProjectionType::Perspective)
+					{
+						return NewObject<UInterchangePhysicalCameraFactoryNode>(BaseNodeContainer, NAME_None);
+					}
+				}
+			}
+
+			return NewObject<UInterchangeStandardCameraFactoryNode>(BaseNodeContainer, NAME_None);
 		}
 		else if(TranslatedAssetNode->IsA<UInterchangeMeshNode>())
 		{
@@ -508,39 +525,110 @@ void UInterchangeGenericLevelPipeline::SetUpFactoryNode(UInterchangeActorFactory
 			ActorFactoryNode->SetCustomActorClassName(APointLight::StaticClass()->GetPathName());
 		}
 	}
-	else if (const UInterchangeCameraNode* CameraNode = Cast<UInterchangeCameraNode>(TranslatedAssetNode))
+	else if (const UInterchangePhysicalCameraNode* PhysicalCameraNode = Cast<UInterchangePhysicalCameraNode>(TranslatedAssetNode))
 	{
 		ActorFactoryNode->SetCustomActorClassName(ACineCameraActor::StaticClass()->GetPathName());
 		ActorFactoryNode->SetCustomMobility(EComponentMobility::Movable);
 
-		if (UInterchangeCineCameraFactoryNode* CineCameraFactoryNode = Cast<UInterchangeCineCameraFactoryNode>(ActorFactoryNode))
+		if (UInterchangePhysicalCameraFactoryNode* PhysicalCameraFactoryNode = Cast<UInterchangePhysicalCameraFactoryNode>(ActorFactoryNode))
 		{
 			float FocalLength;
-			if (CameraNode->GetCustomFocalLength(FocalLength))
+			if (PhysicalCameraNode->GetCustomFocalLength(FocalLength))
 			{
-				CineCameraFactoryNode->SetCustomFocalLength(FocalLength);
+				PhysicalCameraFactoryNode->SetCustomFocalLength(FocalLength);
 			}
 
 			float SensorHeight;
-			if (CameraNode->GetCustomSensorHeight(SensorHeight))
+			if (PhysicalCameraNode->GetCustomSensorHeight(SensorHeight))
 			{
-				CineCameraFactoryNode->SetCustomSensorHeight(SensorHeight);
+				PhysicalCameraFactoryNode->SetCustomSensorHeight(SensorHeight);
 			}
 
 			float SensorWidth;
-			if (CameraNode->GetCustomSensorWidth(SensorWidth))
+			if (PhysicalCameraNode->GetCustomSensorWidth(SensorWidth))
 			{
-				CineCameraFactoryNode->SetCustomSensorWidth(SensorWidth);
+				PhysicalCameraFactoryNode->SetCustomSensorWidth(SensorWidth);
 			}
 
 			bool bEnableDepthOfField;
-			if (CameraNode->GetCustomEnableDepthOfField(bEnableDepthOfField))
+			if (PhysicalCameraNode->GetCustomEnableDepthOfField(bEnableDepthOfField))
 			{
-				CineCameraFactoryNode->SetCustomFocusMethod(bEnableDepthOfField ? ECameraFocusMethod::Manual : ECameraFocusMethod::DoNotOverride);
+				PhysicalCameraFactoryNode->SetCustomFocusMethod(bEnableDepthOfField ? ECameraFocusMethod::Manual : ECameraFocusMethod::DoNotOverride);
 			}
 			
 		}
 	}
+	else if (const UInterchangeStandardCameraNode* CameraNode = Cast<UInterchangeStandardCameraNode>(TranslatedAssetNode))
+	{
+		EInterchangeCameraProjectionType ProjectionType = EInterchangeCameraProjectionType::Perspective;
+		if (CameraNode->GetCustomProjectionMode(ProjectionType) && bUsePhysicalInsteadOfStandardPerspectiveCamera && ProjectionType == EInterchangeCameraProjectionType::Perspective)
+		{
+			float AspectRatio = 1.0f;
+			CameraNode->GetCustomAspectRatio(AspectRatio);
+
+			const float SensorWidth = 36.f;  // mm
+			float SensorHeight = SensorWidth / AspectRatio;
+
+			float FieldOfView = 90;
+			CameraNode->GetCustomFieldOfView(FieldOfView); //Degrees
+
+			float FocalLength = (SensorHeight) / (2.0 * tan(FMath::DegreesToRadians(FieldOfView) / 2.0));
+
+			ActorFactoryNode->SetCustomActorClassName(ACineCameraActor::StaticClass()->GetPathName());
+			ActorFactoryNode->SetCustomMobility(EComponentMobility::Movable);
+
+			if (UInterchangePhysicalCameraFactoryNode* PhysicalCameraFactoryNode = Cast<UInterchangePhysicalCameraFactoryNode>(ActorFactoryNode))
+			{
+				PhysicalCameraFactoryNode->SetCustomFocalLength(FocalLength);
+				PhysicalCameraFactoryNode->SetCustomSensorHeight(SensorHeight);
+				PhysicalCameraFactoryNode->SetCustomSensorWidth(SensorWidth);
+				PhysicalCameraFactoryNode->SetCustomFocusMethod(ECameraFocusMethod::DoNotOverride);
+			}
+		}
+		else
+		{
+			ActorFactoryNode->SetCustomActorClassName(ACameraActor::StaticClass()->GetPathName());
+			ActorFactoryNode->SetCustomMobility(EComponentMobility::Movable);
+
+			if (UInterchangeStandardCameraFactoryNode* CameraFactoryNode = Cast<UInterchangeStandardCameraFactoryNode>(ActorFactoryNode))
+			{
+				if (CameraNode->GetCustomProjectionMode(ProjectionType))
+				{
+					CameraFactoryNode->SetCustomProjectionMode((ECameraProjectionMode::Type)ProjectionType);
+				}
+
+				float OrthoWidth;
+				if (CameraNode->GetCustomWidth(OrthoWidth))
+				{
+					CameraFactoryNode->SetCustomWidth(OrthoWidth);
+				}
+
+				float OrthoNearClipPlane;
+				if (CameraNode->GetCustomNearClipPlane(OrthoNearClipPlane))
+				{
+					CameraFactoryNode->SetCustomNearClipPlane(OrthoNearClipPlane);
+				}
+
+				float OrthoFarClipPlane;
+				if (CameraNode->GetCustomFarClipPlane(OrthoFarClipPlane))
+				{
+					CameraFactoryNode->SetCustomFarClipPlane(OrthoFarClipPlane);
+				}
+
+				float AspectRatio;
+				if (CameraNode->GetCustomAspectRatio(AspectRatio))
+				{
+					CameraFactoryNode->SetCustomAspectRatio(AspectRatio);
+				}
+
+				float FieldOfView;
+				if (CameraNode->GetCustomFieldOfView(FieldOfView))
+				{
+					CameraFactoryNode->SetCustomFieldOfView(FieldOfView);
+				}
+			}
+		}
+	}
 }
 
 void UInterchangeGenericLevelPipeline::ExecuteSceneVariantSetNodePreImport(const UInterchangeSceneVariantSetsNode& SceneVariantSetNode)
diff --git a/Engine/Plugins/Interchange/Runtime/Source/Pipelines/Public/InterchangeGenericScenesPipeline.h b/Engine/Plugins/Interchange/Runtime/Source/Pipelines/Public/InterchangeGenericScenesPipeline.h
--- a/Engine/Plugins/Interchange/Runtime/Source/Pipelines/Public/InterchangeGenericScenesPipeline.h
+++ b/Engine/Plugins/Interchange/Runtime/Source/Pipelines/Public/InterchangeGenericScenesPipeline.h
@@ -57,6 +57,9 @@ protected:
 	 */
 	virtual void SetUpFactoryNode(UInterchangeActorFactoryNode* ActorFactoryNode, const UInterchangeSceneNode* SceneNode, const UInterchangeBaseNode* TranslatedAssetNode) const;
 	
+	/** Disable this option to not convert Standard(Perspective) to Physical Cameras*/
+	UPROPERTY(EditAnywhere, BlueprintReadWrite, Category = "Common Skeletal Meshes and Animations")
+	bool bUsePhysicalInsteadOfStandardPerspectiveCamera = true;
 protected:
 	UInterchangeBaseNodeContainer* BaseNodeContainer = nullptr;
 #if WITH_EDITORONLY_DATA
