[Interchange] Finalize glTF pipeline

diff --git a/Engine/Config/BaseEngine.ini b/Engine/Config/BaseEngine.ini
--- a/Engine/Config/BaseEngine.ini
+++ b/Engine/Config/BaseEngine.ini
@@ -3311,6 +3311,9 @@ GenericPipelineClass=/Script/InterchangePipelines.InterchangeGenericAssetsPipeli
 [/Script/InterchangePipelines.MaterialXPipelineSettings]
 PredefinedSurfaceShaders=((StandardSurface, "/Interchange/Functions/MX_StandardSurface.MX_StandardSurface"),(StandardSurfaceTransmission, "/Interchange/Functions/MX_TransmissionSurface.MX_TransmissionSurface"),(SurfaceUnlit, "/Interchange/Functions/MX_SurfaceUnlit.MX_SurfaceUnlit"),(UsdPreviewSurface, "/Interchange/Functions/MX_UsdPreviewSurface.MX_UsdPreviewSurface"))
 
+[/Script/InterchangePipelines.GLTFPipelineSettings]
+MaterialParents=(("MI_Default_Opaque", "/Interchange/gltf/MaterialInstances/MI_Default_Opaque.MI_Default_Opaque"), ("MI_Default_Mask", "/Interchange/gltf/MaterialInstances/MI_Default_Mask.MI_Default_Mask"), ("MI_Default_Blend", "/Interchange/gltf/MaterialInstances/MI_Default_Blend.MI_Default_Blend"), ("MI_Unlit_Opaque", "/Interchange/gltf/MaterialInstances/MI_Unlit_Opaque.MI_Unlit_Opaque"), ("MI_Unlit_Mask", "/Interchange/gltf/MaterialInstances/MI_Unlit_Mask.MI_Unlit_Mask"), ("MI_Unlit_Blend", "/Interchange/gltf/MaterialInstances/MI_Unlit_Blend.MI_Unlit_Blend"), ("MI_ClearCoat_Opaque", "/Interchange/gltf/MaterialInstances/MI_ClearCoat_Opaque.MI_ClearCoat_Opaque"), ("MI_ClearCoat_Mask", "/Interchange/gltf/MaterialInstances/MI_ClearCoat_Mask.MI_ClearCoat_Mask"), ("MI_ClearCoat_Blend", "/Interchange/gltf/MaterialInstances/MI_ClearCoat_Blend.MI_ClearCoat_Blend"), ("MI_Sheen_Opaque", "/Interchange/gltf/MaterialInstances/MI_Sheen_Opaque.MI_Sheen_Opaque"), ("MI_Sheen_Mask", "/Interchange/gltf/MaterialInstances/MI_Sheen_Mask.MI_Sheen_Mask"), ("MI_Sheen_Blend", "/Interchange/gltf/MaterialInstances/MI_Sheen_Blend.MI_Sheen_Blend"), ("MI_Transmission", "/Interchange/gltf/MaterialInstances/MI_Transmission.MI_Transmission"), ("MI_SpecularGlossiness_Opaque", "/Interchange/gltf/MaterialInstances/MI_SpecularGlossiness_Opaque.MI_SpecularGlossiness_Opaque"), ("MI_SpecularGlossiness_Mask", "/Interchange/gltf/MaterialInstances/MI_SpecularGlossiness_Mask.MI_SpecularGlossiness_Mask"), ("MI_SpecularGlossiness_Blend", "/Interchange/gltf/MaterialInstances/MI_SpecularGlossiness_Blend.MI_SpecularGlossiness_Blend"), ("MI_Default_Opaque_DS", "/Interchange/gltf/MaterialInstances/MI_Default_Opaque_DS.MI_Default_Opaque_DS"), ("MI_Default_Mask_DS", "/Interchange/gltf/MaterialInstances/MI_Default_Mask_DS.MI_Default_Mask_DS"), ("MI_Default_Blend_DS", "/Interchange/gltf/MaterialInstances/MI_Default_Blend_DS.MI_Default_Blend_DS"), ("MI_Unlit_Opaque_DS", "/Interchange/gltf/MaterialInstances/MI_Unlit_Opaque_DS.MI_Unlit_Opaque_DS"), ("MI_Unlit_Mask_DS", "/Interchange/gltf/MaterialInstances/MI_Unlit_Mask_DS.MI_Unlit_Mask_DS"), ("MI_Unlit_Blend_DS", "/Interchange/gltf/MaterialInstances/MI_Unlit_Blend_DS.MI_Unlit_Blend_DS"), ("MI_ClearCoat_Opaque_DS", "/Interchange/gltf/MaterialInstances/MI_ClearCoat_Opaque_DS.MI_ClearCoat_Opaque_DS"), ("MI_ClearCoat_Mask_DS", "/Interchange/gltf/MaterialInstances/MI_ClearCoat_Mask_DS.MI_ClearCoat_Mask_DS"), ("MI_ClearCoat_Blend_DS", "/Interchange/gltf/MaterialInstances/MI_ClearCoat_Blend_DS.MI_ClearCoat_Blend_DS"), ("MI_Sheen_Opaque_DS", "/Interchange/gltf/MaterialInstances/MI_Sheen_Opaque_DS.MI_Sheen_Opaque_DS"), ("MI_Sheen_Mask_DS", "/Interchange/gltf/MaterialInstances/MI_Sheen_Mask_DS.MI_Sheen_Mask_DS"), ("MI_Sheen_Blend_DS", "/Interchange/gltf/MaterialInstances/MI_Sheen_Blend_DS.MI_Sheen_Blend_DS"), ("MI_Transmission_DS", "/Interchange/gltf/MaterialInstances/MI_Transmission_DS.MI_Transmission_DS"), ("MI_SpecularGlossiness_Opaque_DS", "/Interchange/gltf/MaterialInstances/MI_SpecularGlossiness_Opaque_DS.MI_SpecularGlossiness_Opaque_DS"), ("MI_SpecularGlossiness_Mask_DS", "/Interchange/gltf/MaterialInstances/MI_SpecularGlossiness_Mask_DS.MI_SpecularGlossiness_Mask_DS"), ("MI_SpecularGlossiness_Blend_DS", "/Interchange/gltf/MaterialInstances/MI_SpecularGlossiness_Blend_DS.MI_SpecularGlossiness_Blend_DS"))
+
 [BuildPatchServices]
 ; Enable/Disable mechanism that provides a disk overflow store for patch chunks.
 bEnableDiskOverflowStore=true
\ No newline at end of file
diff --git a/Engine/Plugins/Interchange/Editor/Source/Pipelines/Private/InterchangeEditorPipelinesModule.cpp b/Engine/Plugins/Interchange/Editor/Source/Pipelines/Private/InterchangeEditorPipelinesModule.cpp
--- a/Engine/Plugins/Interchange/Editor/Source/Pipelines/Private/InterchangeEditorPipelinesModule.cpp
+++ b/Engine/Plugins/Interchange/Editor/Source/Pipelines/Private/InterchangeEditorPipelinesModule.cpp
@@ -7,6 +7,8 @@
 #include "InterchangeManager.h"
 #include "InterchangeMaterialXPipeline.h"
 #include "InterchangeMaterialXPipelineCustomizations.h"
+#include "InterchangeglTFPipeline.h"
+#include "InterchangeGLTFPipelineCustomizations.h"
 #include "InterchangePipelineBase.h"
 #include "InterchangePipelineFactories.h"
 #include "Misc/CoreDelegates.h"
@@ -85,6 +87,12 @@ void FInterchangeEditorPipelinesModule::AcquireResources()
 	ClassesToUnregisterOnShutdown.Add(UMaterialXPipelineSettings::StaticClass()->GetFName());
 	PropertyEditorModule.RegisterCustomClassLayout(ClassesToUnregisterOnShutdown.Last(), FOnGetDetailCustomizationInstance::CreateStatic(&FInterchangeMaterialXPipelineSettingsCustomization::MakeInstance));
 
+	ClassesToUnregisterOnShutdown.Add(UInterchangeglTFPipeline::StaticClass()->GetFName());
+	PropertyEditorModule.RegisterCustomClassLayout(ClassesToUnregisterOnShutdown.Last(), FOnGetDetailCustomizationInstance::CreateStatic(&FInterchangeGLTFPipelineCustomization::MakeInstance));
+
+	ClassesToUnregisterOnShutdown.Add(UGLTFPipelineSettings::StaticClass()->GetFName());
+	PropertyEditorModule.RegisterCustomClassLayout(ClassesToUnregisterOnShutdown.Last(), FOnGetDetailCustomizationInstance::CreateStatic(&FInterchangeGLTFPipelineSettingsCustomization::MakeInstance));
+
 	if (!InterchangeEditorPipelineStyle.IsValid())
 	{
 		InterchangeEditorPipelineStyle = MakeShared<FInterchangeEditorPipelineStyle>();
diff --git a/Engine/Plugins/Interchange/Editor/Source/Pipelines/Private/InterchangeGLTFPipelineCustomizations.cpp b/Engine/Plugins/Interchange/Editor/Source/Pipelines/Private/InterchangeGLTFPipelineCustomizations.cpp
new file mode 100644
--- /dev/null
+++ b/Engine/Plugins/Interchange/Editor/Source/Pipelines/Private/InterchangeGLTFPipelineCustomizations.cpp
@@ -0,0 +1,131 @@
+// Copyright Epic Games, Inc. All Rights Reserved.
+
+#include "InterchangeGLTFPipelineCustomizations.h"
+
+//#include "InterchangeMaterialDefinitions.h"
+#include "InterchangeglTFPipeline.h"
+#include "Materials/MaterialFunction.h"
+
+#include "DetailLayoutBuilder.h"
+#include "Modules/ModuleManager.h"
+#include "PropertyCustomizationHelpers.h"
+//#include "DetailWidgetRow.h"
+//#include "IDetailGroup.h"
+//#include "InterchangePipelineBase.h"
+//#include "Nodes/InterchangeBaseNode.h"
+//#include "ScopedTransaction.h"
+//#include "Styling/StyleColors.h"
+//#include "Widgets/Input/SNumericEntryBox.h"
+//#include "Widgets/Layout/SBox.h"
+
+TSharedRef<IDetailCustomization> FInterchangeGLTFPipelineCustomization::MakeInstance()
+{
+	return MakeShareable(new FInterchangeGLTFPipelineCustomization());
+}
+
+TSharedRef<IDetailCustomization> FInterchangeGLTFPipelineSettingsCustomization::MakeInstance()
+{
+	return MakeShareable(new FInterchangeGLTFPipelineSettingsCustomization());
+}
+
+void FInterchangeGLTFPipelineCustomization::CustomizeDetails(IDetailLayoutBuilder& DetailBuilder)
+{
+	TArray<TWeakObjectPtr<UObject>> EditingObjects;
+	DetailBuilder.GetObjectsBeingCustomized(EditingObjects);
+	check(EditingObjects.Num() == 1);
+
+	TWeakObjectPtr<UInterchangeglTFPipeline> Pipeline = Cast<UInterchangeglTFPipeline>(EditingObjects[0].Get());
+
+	if (!ensure(Pipeline.IsValid()))
+	{
+		return;
+	}
+
+	IDetailCategoryBuilder& GLTFCategory = DetailBuilder.EditCategory("GLTF");
+
+	GLTFCategory.SetDisplayName(NSLOCTEXT("InterchangeGLTFPipelineCustomization", "CategoryDisplayName", "GLTF Settings"));
+
+	FPropertyEditorModule& PropertyEditorModule = FModuleManager::GetModuleChecked<FPropertyEditorModule>("PropertyEditor");
+
+	FDetailsViewArgs DetailsViewArgs;
+	DetailsViewArgs.NameAreaSettings = FDetailsViewArgs::HideNameArea;
+	TSharedPtr<IDetailsView> DetailsView = PropertyEditorModule.CreateDetailView(DetailsViewArgs);
+
+	if (Pipeline->GLTFPipelineSettings)
+	{
+		Pipeline->GLTFPipelineSettings->SetMaterialParentsEditible(Pipeline->CanEditPropertiesStates());
+	}
+
+	DetailsView->SetObject(Pipeline->GLTFPipelineSettings);
+
+	GLTFCategory.AddCustomRow(NSLOCTEXT("InterchangeGLTFPipelineCustomization", "GLTFPredefinedMaterialLibrary", "Predefined Material Functions"))
+	[
+		SNew( SVerticalBox )
+		+ SVerticalBox::Slot()
+		[
+			DetailsView.ToSharedRef()
+		]
+	];
+}
+
+void FInterchangeGLTFPipelineSettingsCustomization::CustomizeDetails(IDetailLayoutBuilder& DetailBuilder)
+{
+	TArray<TWeakObjectPtr<UObject>> EditingObjects;
+	DetailBuilder.GetObjectsBeingCustomized(EditingObjects);
+	check(EditingObjects.Num() == 1);
+
+	GLTFPipelineSettings = Cast<UGLTFPipelineSettings>(EditingObjects[0].Get());
+
+	if (!ensure(GLTFPipelineSettings.IsValid()))
+	{
+		return;
+	}
+	
+	TSharedRef< IPropertyHandle > PairingsHandle = DetailBuilder.GetProperty(GET_MEMBER_NAME_CHECKED(UGLTFPipelineSettings, MaterialParents));
+	if (!PairingsHandle->IsValidHandle())
+	{
+		return;
+	}
+
+	DetailBuilder.HideProperty(PairingsHandle);
+
+	IDetailCategoryBuilder& GLTFPredefinedMaterialLibraryCategory = DetailBuilder.EditCategory("GLTFPredefinedMaterialLibrary");
+
+	GLTFPredefinedMaterialLibraryCategory.SetDisplayName(NSLOCTEXT("InterchangeGLTFPipelineCustomization", "CategoryDisplayName", "GLTF Predefined Material Library"));
+
+	uint32 NumChildren = 0;
+	PairingsHandle->GetNumChildren(NumChildren);
+
+	for(uint32 i = 0; i < NumChildren; ++i)
+	{
+		TSharedPtr<IPropertyHandle> ChildPropertyHandle = PairingsHandle->GetChildHandle(i);
+		TSharedPtr<IPropertyHandle> KeyPropertyHandle = ChildPropertyHandle->GetKeyHandle();
+
+		IDetailPropertyRow& GLTFPredefinedMaterials = DetailBuilder.AddPropertyToCategory(ChildPropertyHandle.ToSharedRef());
+		GLTFPredefinedMaterials.ShowPropertyButtons(false);
+		GLTFPredefinedMaterials.IsEnabled(GLTFPipelineSettings->IsMaterialParentsEditible());
+
+		FDetailWidgetRow& DetailWidgetRow = GLTFPredefinedMaterials.CustomWidget();
+		DetailWidgetRow
+		.NameContent()
+		[
+			SNullWidget::NullWidget
+		]
+		.ValueContent()
+		.MinDesiredWidth(1.0f)
+		.MaxDesiredWidth(1.0f)
+		[
+			SNew(SHorizontalBox)
+			+ SHorizontalBox::Slot()
+			.HAlign(HAlign_Fill)
+			[
+				SNew(SObjectPropertyEntryBox)
+				.AllowedClass(UMaterialFunction::StaticClass())
+				.PropertyHandle(ChildPropertyHandle)
+			]
+		];
+	}
+
+	DetailBuilder.HideCategory(GET_MEMBER_NAME_CHECKED(UGLTFPipelineSettings, MaterialParents));
+}
+
diff --git a/Engine/Plugins/Interchange/Editor/Source/Pipelines/Private/InterchangeGLTFPipelineCustomizations.h b/Engine/Plugins/Interchange/Editor/Source/Pipelines/Private/InterchangeGLTFPipelineCustomizations.h
new file mode 100644
--- /dev/null
+++ b/Engine/Plugins/Interchange/Editor/Source/Pipelines/Private/InterchangeGLTFPipelineCustomizations.h
@@ -0,0 +1,41 @@
+// Copyright Epic Games, Inc. All Rights Reserved.
+
+#pragma once
+
+#include "IDetailCustomization.h"
+#include "UObject/WeakObjectPtr.h"
+
+class IDetailLayoutBuilder;
+class UGLTFPipelineSettings;
+struct FAssetData;
+
+class FInterchangeGLTFPipelineSettingsCustomization : public IDetailCustomization
+{
+public:
+
+	/**
+	 * Creates an instance of this class
+	 */
+	static TSharedRef<IDetailCustomization> MakeInstance();
+
+	/** IDetailCustomization interface */
+	virtual void CustomizeDetails(IDetailLayoutBuilder& DetailBuilder) override;
+	/** End IDetailCustomization interface */
+
+private:
+	TWeakObjectPtr<UGLTFPipelineSettings> GLTFPipelineSettings;
+};
+
+class FInterchangeGLTFPipelineCustomization : public IDetailCustomization
+{
+public:
+
+	/**
+	 * Creates an instance of this class
+	 */
+	static TSharedRef<IDetailCustomization> MakeInstance();
+
+	/** IDetailCustomization interface */
+	virtual void CustomizeDetails(IDetailLayoutBuilder& DetailBuilder) override;
+	/** End IDetailCustomization interface */
+};
\ No newline at end of file
diff --git a/Engine/Plugins/Interchange/Runtime/Source/Import/Private/Gltf/InterchangeGLTFMaterialInstance.cpp b/Engine/Plugins/Interchange/Runtime/Source/Import/Private/Gltf/InterchangeGLTFMaterialInstance.cpp
deleted file mode 100644
diff --git a/Engine/Plugins/Interchange/Runtime/Source/Import/Private/Gltf/InterchangeGLTFMaterialInstances.cpp b/Engine/Plugins/Interchange/Runtime/Source/Import/Private/Gltf/InterchangeGLTFMaterialInstances.cpp
new file mode 100644
--- /dev/null
+++ b/Engine/Plugins/Interchange/Runtime/Source/Import/Private/Gltf/InterchangeGLTFMaterialInstances.cpp
@@ -0,0 +1,290 @@
+// Copyright Epic Games, Inc. All Rights Reserved.
+
+#include "Gltf/InterchangeGLTFMaterialInstances.h"
+#include "GLTFMaterial.h"
+#include "GLTFTexture.h"
+
+#include "Nodes/InterchangeBaseNodeContainer.h"
+#include "InterchangeTextureNode.h"
+#include "InterchangeShaderGraphNode.h"
+#include "InterchangeMaterialInstanceNode.h"
+#include "InterchangeShaderGraphNode.h"
+
+#include "Materials/MaterialInstance.h"
+#include "Materials/MaterialInstanceConstant.h"
+#include "Materials/MaterialInstanceDynamic.h"
+
+#include "Async/Async.h"
+
+
+namespace UE::Interchange::GLTFMaterialInstances
+{
+	//Precedence order based on the logic in the InterchangeGenericMaterialPipeline: ClearCoat > Sheen > unlit
+	EShadingModel GetShadingModel(GLTF::FMaterial::EShadingModel GltfShadingModel, bool bHasClearCoat, bool bHasSheen, bool bUnlit, bool bHasTransmission)
+	{
+		if (GltfShadingModel == GLTF::FMaterial::EShadingModel::SpecularGlossiness)
+		{
+			return EShadingModel::SPECULARGLOSSINESS;
+		}
+
+		if (bHasClearCoat)
+		{
+			return EShadingModel::CLEARCOAT;
+		}
+
+		if (bHasSheen)
+		{
+			return EShadingModel::SHEEN;
+		}
+
+		if (bUnlit)
+		{
+			return EShadingModel::UNLIT;
+		}
+
+		if (bHasTransmission)
+		{
+			return EShadingModel::TRANSMISSION;
+		}
+
+		return EShadingModel::DEFAULT;
+	}
+
+	void AddGltfMaterialValuesToShaderGraphNode(const GLTF::FMaterial& GltfMaterial, const TArray<GLTF::FTexture>& Textures, UInterchangeShaderGraphNode* ShaderGraphNode)
+	{
+		FString ParentIdentifier = GetIdentifier(GetShadingModel(GltfMaterial.ShadingModel, GltfMaterial.bHasClearCoat, GltfMaterial.bHasSheen, GltfMaterial.bIsUnlitShadingModel, GltfMaterial.bHasTransmission), EAlphaMode(GltfMaterial.AlphaMode), GltfMaterial.bIsDoubleSided);
+
+		ShaderGraphNode->AddStringAttribute(*(InterchangeGltfMaterialAttributeIdentifier + TEXT("ParentIdentifier")), ParentIdentifier);
+
+		auto SetMap = [=](const GLTF::FTextureMap& TextureMap, const FString& Name)
+		{
+			if (Textures.IsValidIndex(TextureMap.TextureIndex))
+			{
+				ShaderGraphNode->AddStringAttribute(*(InterchangeGltfMaterialAttributeIdentifier + Name), UInterchangeTextureNode::MakeNodeUid(Textures[TextureMap.TextureIndex].UniqueId));
+
+				//TexCoord decision in MaterialInstances work based on a Switch Node
+				//Currently Only supports UV0 and UV1
+				// [0...1) -> UV0
+				// [1...2) -> UV1
+				// else    -> UV0 (defaults to 0)
+				ShaderGraphNode->AddFloatAttribute(*(InterchangeGltfMaterialAttributeIdentifier + Name + TEXT("_TexCoord")), TextureMap.TexCoord);
+			}
+
+			if (TextureMap.bHasTextureTransform)
+			{
+				FVector4f OffsetScale(TextureMap.TextureTransform.Offset[0], TextureMap.TextureTransform.Offset[1], TextureMap.TextureTransform.Scale[0], TextureMap.TextureTransform.Scale[1]);
+				ShaderGraphNode->AddLinearColorAttribute(*(InterchangeGltfMaterialAttributeIdentifier + Name + TEXT("_OffsetScale")), OffsetScale);
+
+				ShaderGraphNode->AddFloatAttribute(*(InterchangeGltfMaterialAttributeIdentifier + Name + TEXT("_Rotation")), TextureMap.TextureTransform.Rotation);
+			}
+		};
+
+		if (GltfMaterial.bIsUnlitShadingModel)
+		{
+			//BaseColorTexture
+			//BaseColorTexture_OffsetScale
+			//BaseColorTExture_Rotation
+			//BaseColorTExture_TexCoord
+			SetMap(GltfMaterial.BaseColor, Parameters::BaseColorTexture);
+
+			//BaseColorFactor
+			ShaderGraphNode->AddLinearColorAttribute(*(InterchangeGltfMaterialAttributeIdentifier + Parameters::BaseColorFactor), GltfMaterial.BaseColorFactor);
+
+			return;
+		}
+
+		if (GltfMaterial.AlphaMode == GLTF::FMaterial::EAlphaMode::Mask)
+		{
+			//AlphaCutoff
+			ShaderGraphNode->AddFloatAttribute(*(InterchangeGltfMaterialAttributeIdentifier + Parameters::AlphaCutoff), GltfMaterial.AlphaCutoff);
+		}
+
+		if (GltfMaterial.ShadingModel == GLTF::FMaterial::EShadingModel::MetallicRoughness)
+		{
+			{
+				//BaseColorTexture
+				//BaseColorTexture_OffsetScale
+				//BaseColorTexture_Rotation
+				//BaseColorTexture_TexCoord
+				SetMap(GltfMaterial.BaseColor, Parameters::BaseColorTexture);
+
+				//BaseColorFactor
+				ShaderGraphNode->AddLinearColorAttribute(*(InterchangeGltfMaterialAttributeIdentifier + Parameters::BaseColorFactor), GltfMaterial.BaseColorFactor);
+			}
+
+			{
+				//MetallicRoughnessTexture
+				//MetallicRoughnessTexture_OffsetScale
+				//MetallicRoughnessTexture_Rotation
+				//MetallicRoughnessTexture_TexCoord
+				SetMap(GltfMaterial.MetallicRoughness.Map, Parameters::MetallicRoughnessTexture);
+
+				//MetallicFactor
+				ShaderGraphNode->AddFloatAttribute(*(InterchangeGltfMaterialAttributeIdentifier + Parameters::MetallicFactor), GltfMaterial.MetallicRoughness.MetallicFactor);
+
+				//RoughnessFactor
+				ShaderGraphNode->AddFloatAttribute(*(InterchangeGltfMaterialAttributeIdentifier + Parameters::RoughnessFactor), GltfMaterial.MetallicRoughness.RoughnessFactor);
+			}
+
+			if (GltfMaterial.bHasSpecular)
+			{
+				//SpecularTexture
+				//SpecularTexture_OffsetScale
+				//SpecularTexture_Rotation
+				//SpecularTexture_TexCoord
+				SetMap(GltfMaterial.Specular.SpecularMap, Parameters::SpecularTexture);
+
+				//SpecularFactor
+				ShaderGraphNode->AddFloatAttribute(*(InterchangeGltfMaterialAttributeIdentifier + Parameters::SpecularFactor), GltfMaterial.Specular.SpecularFactor);
+			}
+		}
+		else if (GltfMaterial.ShadingModel == GLTF::FMaterial::EShadingModel::SpecularGlossiness)
+		{
+			//////////////////////
+			//specular glossiness:
+			//////////////////////
+
+			{
+				//DiffuseTexture
+				//DiffuseTexture_OffsetScale
+				//DiffuseTexture_Rotation
+				//DiffuseTexture_TexCoord
+				SetMap(GltfMaterial.BaseColor, Parameters::DiffuseTexture);
+
+				//DiffuseFactor
+				ShaderGraphNode->AddLinearColorAttribute(*(InterchangeGltfMaterialAttributeIdentifier + Parameters::DiffuseFactor), GltfMaterial.BaseColorFactor);
+			}
+
+			{
+				//SpecularGlossinessTexture
+				//SpecularGlossinessTexture_OffsetScale
+				//SpecularGlossinessTexture_Rotation
+				//SpecularGlossinessTexture_TexCoord
+				SetMap(GltfMaterial.SpecularGlossiness.Map, Parameters::SpecularGlossinessTexture);
+
+				//SpecFactor
+				ShaderGraphNode->AddLinearColorAttribute(*(InterchangeGltfMaterialAttributeIdentifier + Parameters::SpecFactor), FVector4f(GltfMaterial.SpecularGlossiness.SpecularFactor[0],
+					GltfMaterial.SpecularGlossiness.SpecularFactor[1],
+					GltfMaterial.SpecularGlossiness.SpecularFactor[2]));
+
+				//GlossinessFactor
+				ShaderGraphNode->AddFloatAttribute(*(InterchangeGltfMaterialAttributeIdentifier + Parameters::GlossinessFactor), GltfMaterial.SpecularGlossiness.GlossinessFactor);
+			}
+		}
+
+		{
+			//NormalTexture
+			//NormalTexture_OffsetScale
+			//NormalTexture_Rotation
+			//NormalTexture_TexCoord
+			SetMap(GltfMaterial.Normal, Parameters::NormalTexture);
+
+			//NormalScale
+			ShaderGraphNode->AddFloatAttribute(*(InterchangeGltfMaterialAttributeIdentifier + Parameters::NormalScale), GltfMaterial.NormalScale);
+		}
+
+		{
+			//EmissiveTexture
+			//EmissiveTexture_OffsetScale
+			//EmissiveTexture_Rotation
+			//EmissiveTexture_TexCoord
+			SetMap(GltfMaterial.Emissive, Parameters::EmissiveTexture);
+
+			//EmissiveFactor
+			ShaderGraphNode->AddLinearColorAttribute(*(InterchangeGltfMaterialAttributeIdentifier + Parameters::EmissiveFactor), GltfMaterial.EmissiveFactor);
+
+			//EmissiveStrength
+			ShaderGraphNode->AddFloatAttribute(*(InterchangeGltfMaterialAttributeIdentifier + Parameters::EmissiveStrength), GltfMaterial.bHasEmissiveStrength ? GltfMaterial.EmissiveStrength : 1.0f);
+		}
+
+		{
+			//OcclusionTexture
+			//OcclusionTexture_OffsetScale
+			//OcclusionTexture_Rotation
+			//OcclusionTexture_TexCoord
+			SetMap(GltfMaterial.Occlusion, Parameters::OcclusionTexture);
+
+			//OcclusionStrength
+			ShaderGraphNode->AddFloatAttribute(*(InterchangeGltfMaterialAttributeIdentifier + Parameters::OcclusionStrength), GltfMaterial.OcclusionStrength);
+		}
+
+		if (GltfMaterial.bHasIOR)
+		{
+			//ior
+			ShaderGraphNode->AddFloatAttribute(*(InterchangeGltfMaterialAttributeIdentifier + Parameters::IOR), GltfMaterial.IOR);
+		}
+
+		if (GltfMaterial.bHasClearCoat)
+		{
+			{
+				//ClearCoatTexture
+				//ClearCoatTexture_OffsetScale
+				//ClearCoatTexture_Rotation
+				//ClearCoatTexture_TexCoord
+				SetMap(GltfMaterial.ClearCoat.ClearCoatMap, Parameters::ClearCoatTexture);
+
+				//ClearCoatFactor
+				ShaderGraphNode->AddFloatAttribute(*(InterchangeGltfMaterialAttributeIdentifier + Parameters::ClearCoatFactor), GltfMaterial.ClearCoat.ClearCoatFactor);
+			}
+
+			{
+				//ClearCoatRoughnessTexture
+				//ClearCoatRoughnessTexture_OffsetScale
+				//ClearCoatRoughnessTexture_Rotation
+				//ClearCoatRoughnessTexture_TexCoord
+				SetMap(GltfMaterial.ClearCoat.RoughnessMap, Parameters::ClearCoatRoughnessTexture);
+
+				//ClearCoatRoughnessFactor
+				ShaderGraphNode->AddFloatAttribute(*(InterchangeGltfMaterialAttributeIdentifier + Parameters::ClearCoatRoughnessFactor), GltfMaterial.ClearCoat.Roughness);
+			}
+
+			{
+				//ClearCoatNormalTexture
+				//ClearCoatNormalTexture_OffsetScale
+				//ClearCoatNormalTexture_Rotation
+				//ClearCoatNormalTexture_TexCoord
+				SetMap(GltfMaterial.ClearCoat.NormalMap, Parameters::ClearCoatNormalTexture);
+
+				//ClearCoatNormalFactor
+				ShaderGraphNode->AddFloatAttribute(*(InterchangeGltfMaterialAttributeIdentifier + Parameters::ClearCoatNormalScale), GltfMaterial.ClearCoat.NormalMapUVScale);
+			}
+		}
+		else if (GltfMaterial.bHasSheen)
+		{
+			{
+				//SheenColorTexture
+				//SheenColorTexture_OffsetScale
+				//SheenColorTexture_Rotation
+				//SheenColorTexture_TexCoord
+				SetMap(GltfMaterial.Sheen.SheenColorMap, Parameters::SheenColorTexture);
+
+				//SheenColorFactor
+				ShaderGraphNode->AddLinearColorAttribute(*(InterchangeGltfMaterialAttributeIdentifier + Parameters::SheenColorFactor), FVector4f(GltfMaterial.Sheen.SheenColorFactor[0], GltfMaterial.Sheen.SheenColorFactor[1], GltfMaterial.Sheen.SheenColorFactor[2]));
+
+				//SheenRoughnessTexture
+				//SheenRoughnessTexture_OffsetScale
+				//SheenRoughnessTexture_Rotation
+				//SheenRoughnessTexture_TexCoord
+				SetMap(GltfMaterial.Sheen.SheenRoughnessMap, Parameters::SheenRoughnessTexture);
+
+				//SheenRoughnessFactor
+				ShaderGraphNode->AddFloatAttribute(*(InterchangeGltfMaterialAttributeIdentifier + Parameters::SheenRoughnessFactor), GltfMaterial.Sheen.SheenRoughnessFactor);
+			}
+		}
+
+		if (GltfMaterial.bHasTransmission)
+		{
+			{
+				//TransmissionTexture
+				//TransmissionTexture_OffsetScale
+				//TransmissionTexture_Rotation
+				//TransmissionTexture_TexCoord
+				SetMap(GltfMaterial.Transmission.TransmissionMap, Parameters::TransmissionTexture);
+
+				//TransmissionFactor
+				ShaderGraphNode->AddFloatAttribute(*(InterchangeGltfMaterialAttributeIdentifier + Parameters::TransmissionFactor), GltfMaterial.Transmission.TransmissionFactor);
+			}
+		}
+	}
+}
+
diff --git a/Engine/Plugins/Interchange/Runtime/Source/Import/Private/Gltf/InterchangeGltfTranslator.cpp b/Engine/Plugins/Interchange/Runtime/Source/Import/Private/Gltf/InterchangeGltfTranslator.cpp
--- a/Engine/Plugins/Interchange/Runtime/Source/Import/Private/Gltf/InterchangeGltfTranslator.cpp
+++ b/Engine/Plugins/Interchange/Runtime/Source/Import/Private/Gltf/InterchangeGltfTranslator.cpp
@@ -38,6 +38,7 @@
 #include "SkeletalMeshOperations.h"
 
 #include "Gltf/InterchangeGltfPrivate.h"
+#include "Gltf/InterchangeGLTFMaterialInstances.h"
 
 #include "EngineAnalytics.h"
 #include "Engine/RendererSettings.h"
@@ -1134,54 +1135,13 @@ bool UInterchangeGltfTranslator::Translate( UInterchangeBaseNodeContainer& NodeC
 			const FString ShaderGraphNodeUid = UInterchangeShaderGraphNode::MakeNodeUid(GltfMaterial.UniqueId);
 			bool bUseVertexColor = MaterialsUsedOnMeshesWithVertexColor.Contains(ShaderGraphNodeUid);
 
-			if (FApp::IsGame() || GInterchangeGLTFUseMaterialInstancing == 1)
-			{
-				//Only create MaterialInstances
-				if (!MaterialInstanceSettings.ProcessGltfMaterial(NodeContainer, GltfMaterial, GltfAsset.Textures))
-				{
-					UE_LOG(LogInterchangeImport, Warning, TEXT("GLTF MaterialInstance Creation failed to find appropriate ParentMaterial. Material.UniqueID: %s"), *GltfMaterial.UniqueId);
-				}
-			}
-			else if (GInterchangeGLTFUseMaterialInstancing == 2)
-			{
-				//Create ShaderGraphNodes and convert to MaterialInstances
-				//Use the correct UIDs on MaterialInstances so that they get applied onto meshes
-				UInterchangeShaderGraphNode* ShaderGraphNode = UInterchangeShaderGraphNode::Create(&NodeContainer, TEXT("Gltf_MI_AttributeIdentifier_") + GltfMaterial.UniqueId);
-				ShaderGraphNode->SetDisplayLabel(TEXT("Gltf_MI_AttributeIdentifier_") + GltfMaterial.Name);
-
-				HandleGltfMaterial(NodeContainer, GltfMaterial, *ShaderGraphNode, bUseVertexColor);
-
-				if (!MaterialInstanceSettings.AddGltfMaterialValuesToShaderGraphNode(GltfMaterial, GltfAsset.Textures, ShaderGraphNode))
-				{
-					UE_LOG(LogInterchangeImport, Warning, TEXT("GLTF MaterialInstance Creation failed to find appropriate ParentMaterial. Material.UniqueID: %s"), *GltfMaterial.UniqueId);
-				}
-				else
-				{
-					FGLTFMaterialInstanceSettings::CreateMaterialInstanceFromShaderGraphNode(NodeContainer, ShaderGraphNode);
-				}
-			}
-			else if (GInterchangeGLTFUseMaterialInstancing == 3)
-			{
-				//Store Gltf Material Attribute values in ShaderGraphNode but do not create MaterialInstances
-				//Test in preparation for Gltf Pipeline:
-				UInterchangeShaderGraphNode* ShaderGraphNode = UInterchangeShaderGraphNode::Create(&NodeContainer, GltfMaterial.UniqueId);
-				ShaderGraphNode->SetDisplayLabel(GltfMaterial.Name);
-
-				HandleGltfMaterial(NodeContainer, GltfMaterial, *ShaderGraphNode, bUseVertexColor);
+			UInterchangeShaderGraphNode* ShaderGraphNode = UInterchangeShaderGraphNode::Create(&NodeContainer, GltfMaterial.UniqueId);
+			ShaderGraphNode->SetDisplayLabel(GltfMaterial.Name);
 
-				if (!MaterialInstanceSettings.AddGltfMaterialValuesToShaderGraphNode(GltfMaterial, GltfAsset.Textures, ShaderGraphNode))
-				{
-					UE_LOG(LogInterchangeImport, Warning, TEXT("GLTF MaterialInstance Creation failed to find appropriate ParentMaterial. Material.UniqueID: %s"), *GltfMaterial.UniqueId);
-				}
-			}
-			else
-			{
-				//Only create ShaderGraphNodes:
-				UInterchangeShaderGraphNode* ShaderGraphNode = UInterchangeShaderGraphNode::Create(&NodeContainer, GltfMaterial.UniqueId);
-				ShaderGraphNode->SetDisplayLabel(GltfMaterial.Name);
+			HandleGltfMaterial(NodeContainer, GltfMaterial, *ShaderGraphNode, bUseVertexColor);
 
-				HandleGltfMaterial(NodeContainer, GltfMaterial, *ShaderGraphNode, bUseVertexColor);
-			}
+			//Store Gltf Material Attribute values in ShaderGraphNode but do not create MaterialInstances
+			UE::Interchange::GLTFMaterialInstances::AddGltfMaterialValuesToShaderGraphNode(GltfMaterial, GltfAsset.Textures, ShaderGraphNode);
 			
 			++MaterialIndex;
 		}
@@ -2234,15 +2194,6 @@ UInterchangeGltfTranslator::UInterchangeGltfTranslator()
 {
 	if (!HasAllFlags(RF_ClassDefaultObject))
 	{
-		if (FApp::IsGame() || GInterchangeGLTFUseMaterialInstancing)
-		{
-			TArray<FString> MaterialInstanceIssues = MaterialInstanceSettings.ValidateMaterialInstancesAndParameters();
-			for (const FString& MaterialInstanceIssue : MaterialInstanceIssues)
-			{
-				UE_LOG(LogInterchangeImport, Warning, TEXT("%s"), *MaterialInstanceIssue);
-			}
-		}
-
 		bRenderSettingsClearCoatEnableSecondNormal = GetDefault<URendererSettings>()->bClearCoatEnableSecondNormal != 0;
 	}
 }
diff --git a/Engine/Plugins/Interchange/Runtime/Source/Import/Public/Gltf/InterchangeGLTFMaterialInstance.h b/Engine/Plugins/Interchange/Runtime/Source/Import/Public/Gltf/InterchangeGLTFMaterialInstance.h
deleted file mode 100644
diff --git a/Engine/Plugins/Interchange/Runtime/Source/Import/Public/Gltf/InterchangeGLTFMaterialInstances.h b/Engine/Plugins/Interchange/Runtime/Source/Import/Public/Gltf/InterchangeGLTFMaterialInstances.h
new file mode 100644
--- /dev/null
+++ b/Engine/Plugins/Interchange/Runtime/Source/Import/Public/Gltf/InterchangeGLTFMaterialInstances.h
@@ -0,0 +1,158 @@
+// Copyright Epic Games, Inc. All Rights Reserved.
+#pragma once
+
+#include "CoreMinimal.h"
+
+namespace GLTF {
+	struct FMaterial;
+	struct FTexture;
+}
+class UInterchangeShaderGraphNode;
+
+const FString InterchangeGltfMaterialAttributeIdentifier = TEXT("Gltf_MI_AttributeIdentifier_");
+
+#define INTERCHANGE_GLTF_STRINGIFY(x) #x
+#define DECLARE_INTERCHANGE_GLTF_MI_MAP(MapName) const FString MapName##Texture = TEXT(INTERCHANGE_GLTF_STRINGIFY(MapName ## Texture)); \
+												 const FString MapName##Texture_OffsetScale = TEXT(INTERCHANGE_GLTF_STRINGIFY(MapName ## Texture_OffsetScale)); \
+												 const FString MapName##Texture_Rotation = TEXT(INTERCHANGE_GLTF_STRINGIFY(MapName ## Texture_Rotation)); \
+												 const FString MapName##Texture_TexCoord = TEXT(INTERCHANGE_GLTF_STRINGIFY(MapName ## Texture_TexCoord));
+
+namespace UE::Interchange::GLTFMaterialInstances
+{
+	namespace Parameters
+	{
+		//MetalRoughness specific:
+		DECLARE_INTERCHANGE_GLTF_MI_MAP(BaseColor)
+		const FString BaseColorFactor = TEXT("BaseColorFactor");
+
+		DECLARE_INTERCHANGE_GLTF_MI_MAP(MetallicRoughness)
+		const FString MetallicFactor = TEXT("MetallicFactor");
+		const FString RoughnessFactor = TEXT("RoughnessFactor");
+
+		DECLARE_INTERCHANGE_GLTF_MI_MAP(Specular)
+		const FString SpecularFactor = TEXT("SpecularFactor");
+
+
+		//SpecularGlossiness specific
+		DECLARE_INTERCHANGE_GLTF_MI_MAP(Diffuse)
+		const FString DiffuseFactor = TEXT("DiffuseFactor");
+
+		DECLARE_INTERCHANGE_GLTF_MI_MAP(SpecularGlossiness)
+		const FString SpecFactor = TEXT("SpecFactor");
+		const FString GlossinessFactor = TEXT("GlossinessFactor");
+
+
+		//Generic:
+		DECLARE_INTERCHANGE_GLTF_MI_MAP(Normal)
+		const FString NormalScale = TEXT("NormalScale");
+
+		DECLARE_INTERCHANGE_GLTF_MI_MAP(Emissive)
+		const FString EmissiveFactor = TEXT("EmissiveFactor");
+		const FString EmissiveStrength = TEXT("EmissiveStrength");
+
+		DECLARE_INTERCHANGE_GLTF_MI_MAP(Occlusion)
+		const FString OcclusionStrength = TEXT("OcclusionStrength");
+
+		const FString IOR = TEXT("IOR");
+
+		const FString AlphaCutoff = TEXT("AlphaCutoff");
+
+
+		//ClearCoat specific:
+		DECLARE_INTERCHANGE_GLTF_MI_MAP(ClearCoat)
+		const FString ClearCoatFactor = TEXT("ClearCoatFactor");
+
+		DECLARE_INTERCHANGE_GLTF_MI_MAP(ClearCoatRoughness)
+		const FString ClearCoatRoughnessFactor = TEXT("ClearCoatRoughnessFactor");
+
+		DECLARE_INTERCHANGE_GLTF_MI_MAP(ClearCoatNormal)
+		const FString ClearCoatNormalScale = TEXT("ClearCoatNormalScale");
+
+
+		//Sheen specific:
+		DECLARE_INTERCHANGE_GLTF_MI_MAP(SheenColor)
+		const FString SheenColorFactor = TEXT("SheenColorFactor");
+
+		DECLARE_INTERCHANGE_GLTF_MI_MAP(SheenRoughness)
+		const FString SheenRoughnessFactor = TEXT("SheenRoughnessFactor");
+
+
+		//Transmission specific:
+		DECLARE_INTERCHANGE_GLTF_MI_MAP(Transmission)
+		const FString TransmissionFactor = TEXT("TransmissionFactor");
+	}
+	
+
+
+	enum EShadingModel : uint8
+	{
+		DEFAULT = 0, //MetalRoughness
+		UNLIT,
+		CLEARCOAT,
+		SHEEN,
+		TRANSMISSION,
+		SPECULARGLOSSINESS,
+		SHADINGMODELCOUNT
+	};
+
+	enum EAlphaMode : uint8
+	{
+		Opaque = 0,
+		Mask,
+		Blend
+	};
+
+	inline FString ToString(EShadingModel ShadingModel)
+	{
+		switch (ShadingModel)
+		{
+		case DEFAULT:
+			return TEXT("Default");
+		case UNLIT:
+			return TEXT("Unlit");
+		case CLEARCOAT:
+			return TEXT("ClearCoat");
+		case SHEEN:
+			return TEXT("Sheen");
+		case TRANSMISSION:
+			return TEXT("Transmission");
+		case SPECULARGLOSSINESS:
+			return TEXT("SpecularGlossiness");
+		default:
+			return TEXT("Unexpected");
+		}
+	}
+
+	inline FString ToString(EAlphaMode AlphaMode)
+	{
+		switch (AlphaMode)
+		{
+		case EAlphaMode::Opaque:
+			return TEXT("Opaque");
+		case EAlphaMode::Mask:
+			return TEXT("Mask");
+		case EAlphaMode::Blend:
+			return TEXT("Blend");
+		default:
+			return TEXT("Unexpected");
+		}
+	}
+
+	inline FString GetIdentifier(EShadingModel ShadingModel, EAlphaMode AlphaMode, bool bTwoSided)
+	{
+		FString Name = TEXT("MI_");
+
+		Name += ToString(ShadingModel);
+
+		if (ShadingModel != EShadingModel::TRANSMISSION)
+		{
+			Name += TEXT("_") + ToString(AlphaMode);
+		}
+
+		Name += bTwoSided ? TEXT("_DS") : TEXT("");
+
+		return Name;
+	}
+
+	void AddGltfMaterialValuesToShaderGraphNode(const GLTF::FMaterial& GltfMaterial, const TArray<GLTF::FTexture>& Textures, UInterchangeShaderGraphNode* ShaderGraphNode);
+};
diff --git a/Engine/Plugins/Interchange/Runtime/Source/Import/Public/Gltf/InterchangeGltfTranslator.h b/Engine/Plugins/Interchange/Runtime/Source/Import/Public/Gltf/InterchangeGltfTranslator.h
--- a/Engine/Plugins/Interchange/Runtime/Source/Import/Public/Gltf/InterchangeGltfTranslator.h
+++ b/Engine/Plugins/Interchange/Runtime/Source/Import/Public/Gltf/InterchangeGltfTranslator.h
@@ -12,7 +12,6 @@
 #include "Texture/InterchangeTexturePayloadInterface.h"
 #include "UObject/Object.h"
 #include "UObject/ObjectMacros.h"
-#include "Gltf/InterchangeGLTFMaterialInstance.h"
 
 #include "InterchangeGltfTranslator.generated.h"
 
@@ -62,9 +61,6 @@ public:
 	virtual TFuture<TOptional<UE::Interchange::FVariantSetPayloadData>> GetVariantSetPayloadData(const FString& PayloadKey) const override;
 	/* IInterchangeVariantSetPayloadInterface End */
 
-	UPROPERTY(EditAnywhere, BlueprintReadWrite, Category = "GLTF", DisplayName = "Material Settings")
-	FGLTFMaterialInstanceSettings MaterialInstanceSettings;
-
 protected:
 	using FNodeUidMap = TMap<const GLTF::FNode*, FString>;
 
diff --git a/Engine/Plugins/Interchange/Runtime/Source/Pipelines/InterchangePipelines.Build.cs b/Engine/Plugins/Interchange/Runtime/Source/Pipelines/InterchangePipelines.Build.cs
--- a/Engine/Plugins/Interchange/Runtime/Source/Pipelines/InterchangePipelines.Build.cs
+++ b/Engine/Plugins/Interchange/Runtime/Source/Pipelines/InterchangePipelines.Build.cs
@@ -25,6 +25,7 @@ namespace UnrealBuildTool.Rules
 					"CinematicCamera",
 					"InterchangeImport",
 					"DeveloperSettings",
+					"GLTFCore",
 				}
 			);
 
diff --git a/Engine/Plugins/Interchange/Runtime/Source/Pipelines/Private/InterchangeglTFPipeline.cpp b/Engine/Plugins/Interchange/Runtime/Source/Pipelines/Private/InterchangeglTFPipeline.cpp
--- a/Engine/Plugins/Interchange/Runtime/Source/Pipelines/Private/InterchangeglTFPipeline.cpp
+++ b/Engine/Plugins/Interchange/Runtime/Source/Pipelines/Private/InterchangeglTFPipeline.cpp
@@ -5,11 +5,358 @@
 #include "InterchangePipelineLog.h"
 
 #include "InterchangeMeshFactoryNode.h"
+#include "InterchangeShaderGraphNode.h"
+#include "InterchangeMaterialInstanceNode.h"
+#include "InterchangeMaterialFactoryNode.h"
+
+#include "Gltf/InterchangeGLTFMaterialInstances.h"
+
+#include "Materials/MaterialInstanceConstant.h"
+#include "Materials/MaterialInstanceDynamic.h"
+
+#include "Misc/App.h"
+
+const TArray<FString> UGLTFPipelineSettings::ExpectedMaterialInstanceIdentifiers = {TEXT("MI_Default_Opaque"), TEXT("MI_Default_Mask"), TEXT("MI_Default_Blend"), 
+																					TEXT("MI_Unlit_Opaque"), TEXT("MI_Unlit_Mask"), TEXT("MI_Unlit_Blend"), 
+																					TEXT("MI_ClearCoat_Opaque"), TEXT("MI_ClearCoat_Mask"), TEXT("MI_ClearCoat_Blend"),
+																					TEXT("MI_Sheen_Opaque"), TEXT("MI_Sheen_Mask"), TEXT("MI_Sheen_Blend"), 
+																					TEXT("MI_Transmission"), 
+																					TEXT("MI_SpecularGlossiness_Opaque"), TEXT("MI_SpecularGlossiness_Mask"), TEXT("MI_SpecularGlossiness_Blend"), 
+																					TEXT("MI_Default_Opaque_DS"), TEXT("MI_Default_Mask_DS"), TEXT("MI_Default_Blend_DS"), 
+																					TEXT("MI_Unlit_Opaque_DS"), TEXT("MI_Unlit_Mask_DS"), TEXT("MI_Unlit_Blend_DS"), 
+																					TEXT("MI_ClearCoat_Opaque_DS"), TEXT("MI_ClearCoat_Mask_DS"), TEXT("MI_ClearCoat_Blend_DS"), 
+																					TEXT("MI_Sheen_Opaque_DS"), TEXT("MI_Sheen_Mask_DS"), TEXT("MI_Sheen_Blend_DS"), 
+																					TEXT("MI_Transmission_DS"), 
+																					TEXT("MI_SpecularGlossiness_Opaque_DS"), TEXT("MI_SpecularGlossiness_Mask_DS"), TEXT("MI_SpecularGlossiness_Blend_DS")};
+
+TArray<FString> UGLTFPipelineSettings::ValidateMaterialInstancesAndParameters() const
+{
+	TArray<FString> NotCoveredIdentifiersParameters;
+
+	//Check if all Material variations are covered:
+	TArray<FString> ExpectedIdentifiers = ExpectedMaterialInstanceIdentifiers;
+	TArray<FString> IdentifiersUsed;
+	MaterialParents.GetKeys(IdentifiersUsed);
+	for (const FString& Identifier : IdentifiersUsed)
+	{
+		ExpectedIdentifiers.Remove(Identifier);
+	}
+	for (const FString& ExpectedIdentifier : ExpectedIdentifiers)
+	{
+		NotCoveredIdentifiersParameters.Add(TEXT("[") + ExpectedIdentifier + TEXT("]: MaterialInstance not found for Identifier."));
+	}
+
+	for (const TPair<FString, FSoftObjectPath>& MaterialParent : MaterialParents)
+	{
+		TSet<FString> ExpectedParameters = GenerateExpectedParametersList(MaterialParent.Key);
+
+		if (UMaterialInstance* ParentMaterialInstance = Cast<UMaterialInstance>(MaterialParent.Value.TryLoad()))
+		{
+			TArray<FGuid> ParameterIds;
+			TArray<FMaterialParameterInfo> ScalarParameterInfos;
+			TArray<FMaterialParameterInfo> VectorParameterInfos;
+			TArray<FMaterialParameterInfo> TextureParameterInfos;
+			ParentMaterialInstance->GetAllScalarParameterInfo(ScalarParameterInfos, ParameterIds);
+			ParentMaterialInstance->GetAllVectorParameterInfo(VectorParameterInfos, ParameterIds);
+			ParentMaterialInstance->GetAllTextureParameterInfo(TextureParameterInfos, ParameterIds);
+
+			for (const FMaterialParameterInfo& ParameterInfo : ScalarParameterInfos)
+			{
+				ExpectedParameters.Remove(ParameterInfo.Name.ToString());
+			}
+			for (const FMaterialParameterInfo& ParameterInfo : VectorParameterInfos)
+			{
+				ExpectedParameters.Remove(ParameterInfo.Name.ToString());
+			}
+			for (const FMaterialParameterInfo& ParameterInfo : TextureParameterInfos)
+			{
+				ExpectedParameters.Remove(ParameterInfo.Name.ToString());
+			}
+		}
+
+		for (const FString& ExpectedParameter : ExpectedParameters)
+		{
+			NotCoveredIdentifiersParameters.Add(TEXT("[") + MaterialParent.Key + TEXT("]: Does not cover expected parameter: ") + ExpectedParameter + TEXT("."));
+		}
+	}
+
+	return NotCoveredIdentifiersParameters;
+}
+
+TSet<FString> UGLTFPipelineSettings::GenerateExpectedParametersList(const FString& Identifier) const
+{
+	using namespace UE::Interchange::GLTFMaterialInstances;
+
+	TSet<FString> ExpectedParameters;
+
+	if (Identifier.Contains(TEXT("_Unlit")))
+	{
+		ExpectedParameters.Add(Parameters::BaseColorTexture);
+		ExpectedParameters.Add(Parameters::BaseColorTexture_OffsetScale);
+		ExpectedParameters.Add(Parameters::BaseColorTexture_Rotation);
+		ExpectedParameters.Add(Parameters::BaseColorTexture_TexCoord);
+		ExpectedParameters.Add(Parameters::BaseColorFactor);
+
+		return ExpectedParameters;
+	}
+
+	//Generic ones:
+	{
+		ExpectedParameters.Add(Parameters::NormalTexture);
+		ExpectedParameters.Add(Parameters::NormalTexture_OffsetScale);
+		ExpectedParameters.Add(Parameters::NormalTexture_Rotation);
+		ExpectedParameters.Add(Parameters::NormalTexture_TexCoord);
+		ExpectedParameters.Add(Parameters::NormalScale);
+
+		if (!Identifier.Contains(TEXT("Transmission")))
+		{
+			ExpectedParameters.Add(Parameters::EmissiveTexture);
+			ExpectedParameters.Add(Parameters::EmissiveTexture_OffsetScale);
+			ExpectedParameters.Add(Parameters::EmissiveTexture_Rotation);
+			ExpectedParameters.Add(Parameters::EmissiveTexture_TexCoord);
+			ExpectedParameters.Add(Parameters::EmissiveFactor);
+			ExpectedParameters.Add(Parameters::EmissiveStrength);
+		}
+
+		ExpectedParameters.Add(Parameters::OcclusionTexture);
+		ExpectedParameters.Add(Parameters::OcclusionTexture_OffsetScale);
+		ExpectedParameters.Add(Parameters::OcclusionTexture_Rotation);
+		ExpectedParameters.Add(Parameters::OcclusionTexture_TexCoord);
+		ExpectedParameters.Add(Parameters::OcclusionStrength);
+
+		ExpectedParameters.Add(Parameters::IOR);
+	}
+
+	//Based on ShadingModel:
+
+	if (Identifier.Contains(TEXT("Default")))
+	{
+		//MetalRoughness Specific:
+
+		ExpectedParameters.Add(Parameters::BaseColorTexture);
+		ExpectedParameters.Add(Parameters::BaseColorTexture_OffsetScale);
+		ExpectedParameters.Add(Parameters::BaseColorTexture_Rotation);
+		ExpectedParameters.Add(Parameters::BaseColorTexture_TexCoord);
+		ExpectedParameters.Add(Parameters::BaseColorFactor);
+
+		ExpectedParameters.Add(Parameters::MetallicRoughnessTexture);
+		ExpectedParameters.Add(Parameters::MetallicRoughnessTexture_OffsetScale);
+		ExpectedParameters.Add(Parameters::MetallicRoughnessTexture_Rotation);
+		ExpectedParameters.Add(Parameters::MetallicRoughnessTexture_TexCoord);
+		ExpectedParameters.Add(Parameters::MetallicFactor);
+		ExpectedParameters.Add(Parameters::RoughnessFactor);
+
+		ExpectedParameters.Add(Parameters::SpecularTexture);
+		ExpectedParameters.Add(Parameters::SpecularTexture_OffsetScale);
+		ExpectedParameters.Add(Parameters::SpecularTexture_Rotation);
+		ExpectedParameters.Add(Parameters::SpecularTexture_TexCoord);
+		ExpectedParameters.Add(Parameters::SpecularFactor);
+
+	}
+	else if (Identifier.Contains(TEXT("ClearCoat")))
+	{
+		ExpectedParameters.Add(Parameters::ClearCoatTexture);
+		ExpectedParameters.Add(Parameters::ClearCoatTexture_OffsetScale);
+		ExpectedParameters.Add(Parameters::ClearCoatTexture_Rotation);
+		ExpectedParameters.Add(Parameters::ClearCoatTexture_TexCoord);
+		ExpectedParameters.Add(Parameters::ClearCoatFactor);
+
+		ExpectedParameters.Add(Parameters::ClearCoatRoughnessTexture);
+		ExpectedParameters.Add(Parameters::ClearCoatRoughnessTexture_OffsetScale);
+		ExpectedParameters.Add(Parameters::ClearCoatRoughnessTexture_Rotation);
+		ExpectedParameters.Add(Parameters::ClearCoatRoughnessTexture_TexCoord);
+		ExpectedParameters.Add(Parameters::ClearCoatRoughnessFactor);
+
+		ExpectedParameters.Add(Parameters::ClearCoatNormalTexture);
+		ExpectedParameters.Add(Parameters::ClearCoatNormalTexture_OffsetScale);
+		ExpectedParameters.Add(Parameters::ClearCoatNormalTexture_Rotation);
+		ExpectedParameters.Add(Parameters::ClearCoatNormalTexture_TexCoord);
+		ExpectedParameters.Add(Parameters::ClearCoatNormalScale);
+	}
+	else if (Identifier.Contains(TEXT("Sheen")))
+	{
+		ExpectedParameters.Add(Parameters::SheenColorTexture);
+		ExpectedParameters.Add(Parameters::SheenColorTexture_OffsetScale);
+		ExpectedParameters.Add(Parameters::SheenColorTexture_Rotation);
+		ExpectedParameters.Add(Parameters::SheenColorTexture_TexCoord);
+		ExpectedParameters.Add(Parameters::SheenColorFactor);
+
+		ExpectedParameters.Add(Parameters::SheenRoughnessTexture);
+		ExpectedParameters.Add(Parameters::SheenRoughnessTexture_OffsetScale);
+		ExpectedParameters.Add(Parameters::SheenRoughnessTexture_Rotation);
+		ExpectedParameters.Add(Parameters::SheenRoughnessTexture_TexCoord);
+		ExpectedParameters.Add(Parameters::SheenRoughnessFactor);
+	}
+	else if (Identifier.Contains(TEXT("Transmission")))
+	{
+		ExpectedParameters.Add(Parameters::TransmissionTexture);
+		ExpectedParameters.Add(Parameters::TransmissionTexture_OffsetScale);
+		ExpectedParameters.Add(Parameters::TransmissionTexture_Rotation);
+		ExpectedParameters.Add(Parameters::TransmissionTexture_TexCoord);
+		ExpectedParameters.Add(Parameters::TransmissionFactor);
+	}
+	else if (Identifier.Contains(TEXT("SpecularGlossiness")))
+	{
+		ExpectedParameters.Add(Parameters::DiffuseTexture);
+		ExpectedParameters.Add(Parameters::DiffuseTexture_OffsetScale);
+		ExpectedParameters.Add(Parameters::DiffuseTexture_Rotation);
+		ExpectedParameters.Add(Parameters::DiffuseTexture_TexCoord);
+		ExpectedParameters.Add(Parameters::DiffuseFactor);
+
+		ExpectedParameters.Add(Parameters::SpecularGlossinessTexture);
+		ExpectedParameters.Add(Parameters::SpecularGlossinessTexture_OffsetScale);
+		ExpectedParameters.Add(Parameters::SpecularGlossinessTexture_Rotation);
+		ExpectedParameters.Add(Parameters::SpecularGlossinessTexture_TexCoord);
+		ExpectedParameters.Add(Parameters::SpecFactor);
+		ExpectedParameters.Add(Parameters::GlossinessFactor);
+	}
+
+	return ExpectedParameters;
+}
+
+void UGLTFPipelineSettings::BuildMaterialInstance(const UInterchangeShaderGraphNode* ShaderGraphNode, UInterchangeMaterialInstanceFactoryNode* MaterialInstanceFactoryNode)
+{
+	using namespace UE::Interchange::GLTFMaterialInstances;
+
+	TArray<UE::Interchange::FAttributeKey> AttributeKeys;
+	ShaderGraphNode->GetAttributeKeys(AttributeKeys);
+
+	TMap<FString, UE::Interchange::FAttributeKey> GltfAttributeKeys;
+	for (const UE::Interchange::FAttributeKey& AttributeKey : AttributeKeys)
+	{
+		if (AttributeKey.ToString().Contains(InterchangeGltfMaterialAttributeIdentifier))
+		{
+			GltfAttributeKeys.Add(AttributeKey.ToString().Replace(*InterchangeGltfMaterialAttributeIdentifier, TEXT(""), ESearchCase::CaseSensitive), AttributeKey);
+		}
+	}
+	if (GltfAttributeKeys.Num() == 0)
+	{
+		return;
+	}
+
+	FString ParentIdentifier;
+	if (!ShaderGraphNode->GetStringAttribute(*(InterchangeGltfMaterialAttributeIdentifier + TEXT("ParentIdentifier")), ParentIdentifier))
+	{
+		return;
+	}
+	GltfAttributeKeys.Remove(TEXT("ParentIdentifier"));
+
+	FString Parent;
+	if (const FSoftObjectPath* ObjectPath = MaterialParents.Find(ParentIdentifier))
+	{
+		Parent = ObjectPath->GetAssetPathString();
+	}
+	else
+	{
+		UE_LOG(LogInterchangePipeline, Warning, TEXT("[Interchange] Failed to load MaterialParent for ParentIdentifier: %s"), *ParentIdentifier);
+		return;
+	}
+
+	FString MaterialFactoryNodeUid = UInterchangeFactoryBaseNode::BuildFactoryNodeUid(ShaderGraphNode->GetUniqueID());
+	FString MaterialFactoryNodeName = ShaderGraphNode->GetDisplayLabel();
+
+	MaterialInstanceFactoryNode->InitializeNode(MaterialFactoryNodeUid, MaterialFactoryNodeName, EInterchangeNodeContainerType::FactoryData);
+
+	MaterialInstanceFactoryNode->SetCustomParent(Parent);
+
+	const UClass* MaterialClass = FApp::IsGame() ? UMaterialInstanceDynamic::StaticClass() : UMaterialInstanceConstant::StaticClass();
+	MaterialInstanceFactoryNode->SetCustomInstanceClassName(MaterialClass->GetPathName());
+
+	for (const TPair<FString, UE::Interchange::FAttributeKey>& GltfAttributeKey : GltfAttributeKeys)
+	{
+		UE::Interchange::EAttributeTypes AttributeType = ShaderGraphNode->GetAttributeType(GltfAttributeKey.Value);
+
+		FString InputValueKey = UInterchangeShaderPortsAPI::MakeInputValueKey(GltfAttributeKey.Key);
+
+		//we are only using 3 attribute types:
+		switch (AttributeType)
+		{
+		case UE::Interchange::EAttributeTypes::Float:
+		{
+			float Value;
+			if (ShaderGraphNode->GetFloatAttribute(GltfAttributeKey.Value.Key, Value))
+			{
+				MaterialInstanceFactoryNode->AddFloatAttribute(InputValueKey, Value);
+			}
+		}
+		break;
+		case UE::Interchange::EAttributeTypes::LinearColor:
+		{
+			FLinearColor Value;
+			if (ShaderGraphNode->GetLinearColorAttribute(GltfAttributeKey.Value.Key, Value))
+			{
+				MaterialInstanceFactoryNode->AddLinearColorAttribute(InputValueKey, Value);
+			}
+		}
+		break;
+		case UE::Interchange::EAttributeTypes::String:
+		{
+			FString TextureUid;
+			if (ShaderGraphNode->GetStringAttribute(GltfAttributeKey.Value.Key, TextureUid))
+			{
+				FString FactoryTextureUid = UInterchangeFactoryBaseNode::BuildFactoryNodeUid(TextureUid);
+
+				MaterialInstanceFactoryNode->AddStringAttribute(InputValueKey, FactoryTextureUid);
+				MaterialInstanceFactoryNode->AddFactoryDependencyUid(FactoryTextureUid);
+			}
+		}
+		break;
+		default:
+			break;
+		}
+	}
+}
+
+UInterchangeglTFPipeline::UInterchangeglTFPipeline()
+	: GLTFPipelineSettings(UGLTFPipelineSettings::StaticClass()->GetDefaultObject<UGLTFPipelineSettings>())
+{
+}
+
+void UInterchangeglTFPipeline::AdjustSettingsForContext(EInterchangePipelineContext ImportType, TObjectPtr<UObject> ReimportAsset)
+{
+	Super::AdjustSettingsForContext(ImportType, ReimportAsset);
+
+	TArray<FString> MaterialInstanceIssues = GLTFPipelineSettings->ValidateMaterialInstancesAndParameters();
+	for (const FString& MaterialInstanceIssue : MaterialInstanceIssues)
+	{
+		UE_LOG(LogInterchangePipeline, Warning, TEXT("%s"), *MaterialInstanceIssue);
+	}
+}
 
 void UInterchangeglTFPipeline::ExecutePipeline(UInterchangeBaseNodeContainer* NodeContainer, const TArray<UInterchangeSourceData*>& InSourceDatas)
 {
 	Super::ExecutePipeline(NodeContainer, InSourceDatas);
 
-	// TODO: Put logic to convert material graphs to material instances when applicable.
+	if ((FApp::IsGame() || bUseGLTFMaterialInstanceLibrary) && GLTFPipelineSettings)
+	{
+		TMap<FString, const UInterchangeShaderGraphNode*> MaterialFactoryNodeUidsToShaderGraphNodes;
+		NodeContainer->IterateNodesOfType<UInterchangeMaterialFactoryNode>([&MaterialFactoryNodeUidsToShaderGraphNodes, &NodeContainer](const FString& NodeUid, UInterchangeMaterialFactoryNode* MaterialFactoryNode)
+			{
+				TArray<FString> TargetNodeUids;
+				MaterialFactoryNode->GetTargetNodeUids(TargetNodeUids);
+
+				for (const FString& TargetNodeUid : TargetNodeUids)
+				{
+					
+					if (const UInterchangeShaderGraphNode* ShaderGraphNode = Cast<UInterchangeShaderGraphNode>(NodeContainer->GetNode(TargetNodeUid)))
+					{
+						FString ParentIdentifier;
+						if (ShaderGraphNode->GetStringAttribute(*(InterchangeGltfMaterialAttributeIdentifier + TEXT("ParentIdentifier")), ParentIdentifier))
+						{
+							MaterialFactoryNodeUidsToShaderGraphNodes.Add(NodeUid, ShaderGraphNode);
+							break;
+						}
+					}
+				}
+			});
+
+		for (const TPair<FString, const UInterchangeShaderGraphNode*>& ShaderGraphNode : MaterialFactoryNodeUidsToShaderGraphNodes)
+		{
+			UInterchangeMaterialInstanceFactoryNode* MaterialInstanceFactoryNode = NewObject<UInterchangeMaterialInstanceFactoryNode>(NodeContainer);
+			GLTFPipelineSettings->BuildMaterialInstance(ShaderGraphNode.Value, MaterialInstanceFactoryNode);
+
+			NodeContainer->ReplaceNode(ShaderGraphNode.Key, MaterialInstanceFactoryNode);
+		}
+	}
 }
 
diff --git a/Engine/Plugins/Interchange/Runtime/Source/Pipelines/Public/InterchangeglTFPipeline.h b/Engine/Plugins/Interchange/Runtime/Source/Pipelines/Public/InterchangeglTFPipeline.h
--- a/Engine/Plugins/Interchange/Runtime/Source/Pipelines/Public/InterchangeglTFPipeline.h
+++ b/Engine/Plugins/Interchange/Runtime/Source/Pipelines/Public/InterchangeglTFPipeline.h
@@ -8,17 +8,57 @@
 #include "InterchangeSourceData.h"
 #include "Nodes/InterchangeBaseNodeContainer.h"
 
+#include "Engine/DeveloperSettings.h"
+
 #include "InterchangeglTFPipeline.generated.h"
 
+class UInterchangeBaseNodeContainer;
+class UInterchangeMaterialInstanceFactoryNode;
+class UInterchangeShaderGraphNode;
+
+UCLASS(config = Engine, meta = (DisplayName = "Interchange GLTF"))
+class INTERCHANGEPIPELINES_API UGLTFPipelineSettings : public UDeveloperSettings
+{
+	GENERATED_BODY()
+
+public:
+	UPROPERTY(EditAnywhere, config, Category = "GLTFPredefinedMaterialLibrary", meta = (DisplayName = "GLTF Predefined Material Library"))
+	TMap<FString, FSoftObjectPath> MaterialParents;
+
+	TArray<FString> ValidateMaterialInstancesAndParameters() const;
+
+	void BuildMaterialInstance(const UInterchangeShaderGraphNode* ShaderGraphNode, UInterchangeMaterialInstanceFactoryNode* materialInstanceNode);
+
+	bool IsMaterialParentsEditible() { return bMaterialParentsEditible; }
+	void SetMaterialParentsEditible(bool bEditible) { bMaterialParentsEditible = bEditible; }
+private:
+	TSet<FString> GenerateExpectedParametersList(const FString& Identifier) const;
+
+	static const TArray<FString> ExpectedMaterialInstanceIdentifiers; //Default MaterialInstance' identifiers
+
+	//Helper for the Settings Customizer to decide if the ParentMaterials should be editible or not
+	//Should be editible from Project Settings, and should NOT be editible from the Import
+	//It is set by the Pipeline's Customizer.
+	//Equals to Pipeline->CanEditPropertiesStates()
+	bool bMaterialParentsEditible = true;
+};
+
 UCLASS(BlueprintType)
 class INTERCHANGEPIPELINES_API UInterchangeglTFPipeline : public UInterchangePipelineBase
 {
 	GENERATED_BODY()
 
+	UInterchangeglTFPipeline();
+
 public:
-	// TODO: Put properties to manage the creation or not of material instances instead of material graphs.
+	UPROPERTY(EditAnywhere, BlueprintReadWrite, Category = "GLTF")
+	TObjectPtr<UGLTFPipelineSettings> GLTFPipelineSettings;
+
+	UPROPERTY(EditAnywhere, BlueprintReadWrite, Category = "GLTF")
+	bool bUseGLTFMaterialInstanceLibrary = false;
 
 protected:
+	virtual void AdjustSettingsForContext(EInterchangePipelineContext ImportType, TObjectPtr<UObject> ReimportAsset) override;
 	virtual void ExecutePipeline(UInterchangeBaseNodeContainer* BaseNodeContainer, const TArray<UInterchangeSourceData*>& SourceDatas) override;
 
 	virtual bool CanExecuteOnAnyThread(EInterchangePipelineTask PipelineTask) override
diff --git a/Engine/Source/Runtime/Interchange/Core/Private/Nodes/InterchangeBaseNodeContainer.cpp b/Engine/Source/Runtime/Interchange/Core/Private/Nodes/InterchangeBaseNodeContainer.cpp
--- a/Engine/Source/Runtime/Interchange/Core/Private/Nodes/InterchangeBaseNodeContainer.cpp
+++ b/Engine/Source/Runtime/Interchange/Core/Private/Nodes/InterchangeBaseNodeContainer.cpp
@@ -50,6 +50,15 @@ FString UInterchangeBaseNodeContainer::AddNode(UInterchangeBaseNode* Node)
 	return NodeUniqueID;
 }
 
+void UInterchangeBaseNodeContainer::ReplaceNode(const FString& NodeUniqueID, UInterchangeFactoryBaseNode* NewNode)
+{
+	if (GetFactoryNode(NodeUniqueID)) //Check existance and confirm it is FactoryNode
+	{
+		Nodes.Remove(NodeUniqueID);
+		AddNode(NewNode);
+	}
+}
+
 bool UInterchangeBaseNodeContainer::IsNodeUidValid(const FString& NodeUniqueID) const
 {
 	if (NodeUniqueID == UInterchangeBaseNode::InvalidNodeUid())
diff --git a/Engine/Source/Runtime/Interchange/Core/Public/Nodes/InterchangeBaseNodeContainer.h b/Engine/Source/Runtime/Interchange/Core/Public/Nodes/InterchangeBaseNodeContainer.h
--- a/Engine/Source/Runtime/Interchange/Core/Public/Nodes/InterchangeBaseNodeContainer.h
+++ b/Engine/Source/Runtime/Interchange/Core/Public/Nodes/InterchangeBaseNodeContainer.h
@@ -50,6 +50,9 @@ public:
 	UFUNCTION(BlueprintCallable, Category = "Interchange | Node Container")
 	FString AddNode(UInterchangeBaseNode* Node);
 
+	UFUNCTION(BlueprintCallable, Category = "Interchange | Node Container")
+	void ReplaceNode(const FString& NodeUniqueID, UInterchangeFactoryBaseNode* NewNode);
+
 	/** Return true if the node unique ID exist in the container */
 	UFUNCTION(BlueprintCallable, Category = "Interchange | Node Container")
 	bool IsNodeUidValid(const FString& NodeUniqueID) const;
