[Interchange] Removal of GLTF MaterialInstancing unused cvar.

diff --git a/Engine/Plugins/Interchange/Runtime/Source/Import/Private/Gltf/InterchangeGltfTranslator.cpp b/Engine/Plugins/Interchange/Runtime/Source/Import/Private/Gltf/InterchangeGltfTranslator.cpp
--- a/Engine/Plugins/Interchange/Runtime/Source/Import/Private/Gltf/InterchangeGltfTranslator.cpp
+++ b/Engine/Plugins/Interchange/Runtime/Source/Import/Private/Gltf/InterchangeGltfTranslator.cpp
@@ -47,13 +47,6 @@
 
 #define LOCTEXT_NAMESPACE "InterchangeGltfTranslator"
 
-static int GInterchangeGLTFUseMaterialInstancing = 0;
-static FAutoConsoleVariableRef CCvarInterchangeGLTFUseMaterialInstancing(
-	TEXT("Interchange.FeatureFlags.Import.GLTF.UseMaterialInstancing"),
-	GInterchangeGLTFUseMaterialInstancing,
-	TEXT("Whether glTF uses Material Instances."),
-	ECVF_Default);
-
 static const TArray<FString> ImporterSupportedExtensions = {
 	/* Lights */
 	GLTF::ToString(GLTF::EExtension::KHR_LightsPunctual),
