[Interchange] gltf Pipeline material instance usage fix.

diff --git a/Engine/Plugins/Interchange/Runtime/Source/Pipelines/Private/InterchangeglTFPipeline.cpp b/Engine/Plugins/Interchange/Runtime/Source/Pipelines/Private/InterchangeglTFPipeline.cpp
--- a/Engine/Plugins/Interchange/Runtime/Source/Pipelines/Private/InterchangeglTFPipeline.cpp
+++ b/Engine/Plugins/Interchange/Runtime/Source/Pipelines/Private/InterchangeglTFPipeline.cpp
@@ -330,24 +330,33 @@ void UInterchangeglTFPipeline::ExecutePipeline(UInterchangeBaseNodeContainer* No
 	if ((FApp::IsGame() || bUseGLTFMaterialInstanceLibrary) && GLTFPipelineSettings)
 	{
 		TMap<FString, const UInterchangeShaderGraphNode*> MaterialFactoryNodeUidsToShaderGraphNodes;
-		NodeContainer->IterateNodesOfType<UInterchangeMaterialFactoryNode>([&MaterialFactoryNodeUidsToShaderGraphNodes, &NodeContainer](const FString& NodeUid, UInterchangeMaterialFactoryNode* MaterialFactoryNode)
+		auto FindGLTFShaderGraphNode = [&MaterialFactoryNodeUidsToShaderGraphNodes, &NodeContainer](const FString& NodeUid, UInterchangeFactoryBaseNode* /*Material or MaterialInstance*/ FactoryNode)
+		{
+			TArray<FString> TargetNodeUids;
+			FactoryNode->GetTargetNodeUids(TargetNodeUids);
+
+			for (const FString& TargetNodeUid : TargetNodeUids)
 			{
-				TArray<FString> TargetNodeUids;
-				MaterialFactoryNode->GetTargetNodeUids(TargetNodeUids);
 
-				for (const FString& TargetNodeUid : TargetNodeUids)
+				if (const UInterchangeShaderGraphNode* ShaderGraphNode = Cast<UInterchangeShaderGraphNode>(NodeContainer->GetNode(TargetNodeUid)))
 				{
-					
-					if (const UInterchangeShaderGraphNode* ShaderGraphNode = Cast<UInterchangeShaderGraphNode>(NodeContainer->GetNode(TargetNodeUid)))
+					FString ParentIdentifier;
+					if (ShaderGraphNode->GetStringAttribute(*(InterchangeGltfMaterialAttributeIdentifier + TEXT("ParentIdentifier")), ParentIdentifier))
 					{
-						FString ParentIdentifier;
-						if (ShaderGraphNode->GetStringAttribute(*(InterchangeGltfMaterialAttributeIdentifier + TEXT("ParentIdentifier")), ParentIdentifier))
-						{
-							MaterialFactoryNodeUidsToShaderGraphNodes.Add(NodeUid, ShaderGraphNode);
-							break;
-						}
+						MaterialFactoryNodeUidsToShaderGraphNodes.Add(NodeUid, ShaderGraphNode);
+						break;
 					}
 				}
+			}
+		};
+		NodeContainer->IterateNodesOfType<UInterchangeMaterialFactoryNode>([&MaterialFactoryNodeUidsToShaderGraphNodes, &NodeContainer, &FindGLTFShaderGraphNode](const FString& NodeUid, UInterchangeMaterialFactoryNode* MaterialFactoryNode)
+			{
+				FindGLTFShaderGraphNode(NodeUid, MaterialFactoryNode);
+			});
+
+		NodeContainer->IterateNodesOfType<UInterchangeMaterialInstanceFactoryNode>([&MaterialFactoryNodeUidsToShaderGraphNodes, &NodeContainer, &FindGLTFShaderGraphNode](const FString& NodeUid, UInterchangeMaterialInstanceFactoryNode* MaterialInstanceFactoryNode)
+			{
+				FindGLTFShaderGraphNode(NodeUid, MaterialInstanceFactoryNode);
 			});
 
 		for (const TPair<FString, const UInterchangeShaderGraphNode*>& ShaderGraphNode : MaterialFactoryNodeUidsToShaderGraphNodes)
