[Interchange] glTF clean the UI and asset name to respect the glTF naming convention

diff --git a/Engine/Plugins/Interchange/Editor/Source/Pipelines/Private/InterchangeEditorPipelineDetails.cpp b/Engine/Plugins/Interchange/Editor/Source/Pipelines/Private/InterchangeEditorPipelineDetails.cpp
--- a/Engine/Plugins/Interchange/Editor/Source/Pipelines/Private/InterchangeEditorPipelineDetails.cpp
+++ b/Engine/Plugins/Interchange/Editor/Source/Pipelines/Private/InterchangeEditorPipelineDetails.cpp
@@ -279,7 +279,10 @@ void FInterchangePipelineBaseDetailsCustomization::CustomizeDetails(IDetailLayou
 		TMap<FName, IDetailGroup*> SubCategoryGroups;
 
 		const FName CategoryName = CategoryAndProperties.Key;
-		IDetailCategoryBuilder& Category = CachedDetailBuilder->EditCategory(CategoryName);
+		
+		//As prior set displaynames are lost with the call of EditCategory,
+		//we make sure that the Category Name is used as is without auto Capital Letter start and automatic spacing.
+		IDetailCategoryBuilder& Category = CachedDetailBuilder->EditCategory(CategoryName, FText::FromString(CategoryName.ToString()));
 		
 		TArray<TSharedRef<IPropertyHandle>> CategoryProperties;
 
diff --git a/Engine/Plugins/Interchange/Editor/Source/Pipelines/Private/InterchangeEditorPipelinesModule.cpp b/Engine/Plugins/Interchange/Editor/Source/Pipelines/Private/InterchangeEditorPipelinesModule.cpp
--- a/Engine/Plugins/Interchange/Editor/Source/Pipelines/Private/InterchangeEditorPipelinesModule.cpp
+++ b/Engine/Plugins/Interchange/Editor/Source/Pipelines/Private/InterchangeEditorPipelinesModule.cpp
@@ -87,7 +87,7 @@ void FInterchangeEditorPipelinesModule::AcquireResources()
 	ClassesToUnregisterOnShutdown.Add(UMaterialXPipelineSettings::StaticClass()->GetFName());
 	PropertyEditorModule.RegisterCustomClassLayout(ClassesToUnregisterOnShutdown.Last(), FOnGetDetailCustomizationInstance::CreateStatic(&FInterchangeMaterialXPipelineSettingsCustomization::MakeInstance));
 
-	ClassesToUnregisterOnShutdown.Add(UInterchangeglTFPipeline::StaticClass()->GetFName());
+	ClassesToUnregisterOnShutdown.Add(UInterchangeGLTFPipeline::StaticClass()->GetFName());
 	PropertyEditorModule.RegisterCustomClassLayout(ClassesToUnregisterOnShutdown.Last(), FOnGetDetailCustomizationInstance::CreateStatic(&FInterchangeGLTFPipelineCustomization::MakeInstance));
 
 	ClassesToUnregisterOnShutdown.Add(UGLTFPipelineSettings::StaticClass()->GetFName());
diff --git a/Engine/Plugins/Interchange/Editor/Source/Pipelines/Private/InterchangeGLTFPipelineCustomizations.cpp b/Engine/Plugins/Interchange/Editor/Source/Pipelines/Private/InterchangeGLTFPipelineCustomizations.cpp
--- a/Engine/Plugins/Interchange/Editor/Source/Pipelines/Private/InterchangeGLTFPipelineCustomizations.cpp
+++ b/Engine/Plugins/Interchange/Editor/Source/Pipelines/Private/InterchangeGLTFPipelineCustomizations.cpp
@@ -25,16 +25,14 @@ void FInterchangeGLTFPipelineCustomization::CustomizeDetails(IDetailLayoutBuilde
 	DetailBuilder.GetObjectsBeingCustomized(EditingObjects);
 	check(EditingObjects.Num() == 1);
 
-	TWeakObjectPtr<UInterchangeglTFPipeline> Pipeline = Cast<UInterchangeglTFPipeline>(EditingObjects[0].Get());
+	TWeakObjectPtr<UInterchangeGLTFPipeline> Pipeline = Cast<UInterchangeGLTFPipeline>(EditingObjects[0].Get());
 
 	if (!ensure(Pipeline.IsValid()))
 	{
 		return;
 	}
 
-	IDetailCategoryBuilder& GLTFCategory = DetailBuilder.EditCategory("GLTF");
-
-	GLTFCategory.SetDisplayName(NSLOCTEXT("InterchangeGLTFPipelineCustomization", "CategoryDisplayName", "GLTF Settings"));
+	IDetailCategoryBuilder& GLTFCategory = DetailBuilder.EditCategory("glTF");
 
 	FPropertyEditorModule& PropertyEditorModule = FModuleManager::GetModuleChecked<FPropertyEditorModule>("PropertyEditor");
 
@@ -80,10 +78,6 @@ void FInterchangeGLTFPipelineSettingsCustomization::CustomizeDetails(IDetailLayo
 
 	DetailBuilder.HideProperty(PairingsHandle);
 
-	IDetailCategoryBuilder& GLTFPredefinedMaterialLibraryCategory = DetailBuilder.EditCategory("GLTFPredefinedMaterialLibrary");
-
-	GLTFPredefinedMaterialLibraryCategory.SetDisplayName(NSLOCTEXT("InterchangeGLTFPipelineSettingsCustomization", "CategoryDisplayName", "GLTF Predefined Material Library"));
-
 	uint32 NumChildren = 0;
 	PairingsHandle->GetNumChildren(NumChildren);
 
@@ -117,6 +111,8 @@ void FInterchangeGLTFPipelineSettingsCustomization::CustomizeDetails(IDetailLayo
 		];
 	}
 
+	IDetailCategoryBuilder& GLTFPredefinedMaterialLibraryCategory = DetailBuilder.EditCategory("PredefinedglTFMaterialLibrary", FText::FromString("Predefined glTF Material Library"));
+
 	DetailBuilder.HideCategory(GET_MEMBER_NAME_CHECKED(UGLTFPipelineSettings, MaterialParents));
 }
 
diff --git a/Engine/Plugins/Interchange/Runtime/Source/Import/Private/Gltf/InterchangeGltfTranslator.cpp b/Engine/Plugins/Interchange/Runtime/Source/Import/Private/Gltf/InterchangeGltfTranslator.cpp
--- a/Engine/Plugins/Interchange/Runtime/Source/Import/Private/Gltf/InterchangeGltfTranslator.cpp
+++ b/Engine/Plugins/Interchange/Runtime/Source/Import/Private/Gltf/InterchangeGltfTranslator.cpp
@@ -205,7 +205,7 @@ namespace UE::Interchange::Gltf::Private
 	};
 }
 
-void UInterchangeGltfTranslator::HandleGltfNode( UInterchangeBaseNodeContainer& NodeContainer, const GLTF::FNode& GltfNode, const FString& ParentNodeUid, const int32 NodeIndex, 
+void UInterchangeGLTFTranslator::HandleGltfNode( UInterchangeBaseNodeContainer& NodeContainer, const GLTF::FNode& GltfNode, const FString& ParentNodeUid, const int32 NodeIndex, 
 	bool &bHasVariants, TArray<int32>& SkinnedMeshNodes, TSet<int>& UnusedMeshIndices ) const
 {
 	using namespace UE::Interchange::Gltf::Private;
@@ -351,7 +351,7 @@ void UInterchangeGltfTranslator::HandleGltfNode( UInterchangeBaseNodeContainer&
 	}
 }
 
-void UInterchangeGltfTranslator::HandleGltfMaterialParameter( UInterchangeBaseNodeContainer& NodeContainer, const GLTF::FTextureMap& TextureMap, UInterchangeShaderNode& ShaderNode,
+void UInterchangeGLTFTranslator::HandleGltfMaterialParameter( UInterchangeBaseNodeContainer& NodeContainer, const GLTF::FTextureMap& TextureMap, UInterchangeShaderNode& ShaderNode,
 		const FString& MapName, const TVariant< FLinearColor, float >& MapFactor, const FString& OutputChannel, const bool bInverse, const bool bIsNormal, const bool bUseVertexColor) const
 {
 	using namespace UE::Interchange::Materials;
@@ -503,7 +503,7 @@ void UInterchangeGltfTranslator::HandleGltfMaterialParameter( UInterchangeBaseNo
 	}
 }
 
-void UInterchangeGltfTranslator::HandleGltfMaterial( UInterchangeBaseNodeContainer& NodeContainer, const GLTF::FMaterial& GltfMaterial, UInterchangeShaderGraphNode& ShaderGraphNode, bool bUseVertexColor) const
+void UInterchangeGLTFTranslator::HandleGltfMaterial( UInterchangeBaseNodeContainer& NodeContainer, const GLTF::FMaterial& GltfMaterial, UInterchangeShaderGraphNode& ShaderGraphNode, bool bUseVertexColor) const
 {
 	using namespace UE::Interchange::Materials;
 
@@ -682,7 +682,7 @@ void UInterchangeGltfTranslator::HandleGltfMaterial( UInterchangeBaseNodeContain
 	}
 }
 
-void UInterchangeGltfTranslator::HandleGltfClearCoat( UInterchangeBaseNodeContainer& NodeContainer, const GLTF::FMaterial& GltfMaterial, UInterchangeShaderGraphNode& ShaderGraphNode, const bool bSwapNormalAndClearCoatNormal) const
+void UInterchangeGLTFTranslator::HandleGltfClearCoat( UInterchangeBaseNodeContainer& NodeContainer, const GLTF::FMaterial& GltfMaterial, UInterchangeShaderGraphNode& ShaderGraphNode, const bool bSwapNormalAndClearCoatNormal) const
 {
 	using namespace UE::Interchange::Materials;
 
@@ -728,7 +728,7 @@ void UInterchangeGltfTranslator::HandleGltfClearCoat( UInterchangeBaseNodeContai
 	}
 }
 
-void UInterchangeGltfTranslator::HandleGltfSheen( UInterchangeBaseNodeContainer& NodeContainer, const GLTF::FMaterial& GltfMaterial, UInterchangeShaderGraphNode& ShaderGraphNode ) const
+void UInterchangeGLTFTranslator::HandleGltfSheen( UInterchangeBaseNodeContainer& NodeContainer, const GLTF::FMaterial& GltfMaterial, UInterchangeShaderGraphNode& ShaderGraphNode ) const
 {
 	using namespace UE::Interchange::Materials;
 
@@ -766,7 +766,7 @@ void UInterchangeGltfTranslator::HandleGltfSheen( UInterchangeBaseNodeContainer&
  * When a transmission material isn't fully opaque, we reduce the transmission color by the opacity to mimic GLTF's BTDF.
  * Ideally, this would be better represented by blending a default lit alpha blended material with a thin translucent material based on GLTF's opacity.
  */
-void UInterchangeGltfTranslator::HandleGltfTransmission( UInterchangeBaseNodeContainer& NodeContainer, const GLTF::FMaterial& GltfMaterial, UInterchangeShaderGraphNode& ShaderGraphNode ) const
+void UInterchangeGLTFTranslator::HandleGltfTransmission( UInterchangeBaseNodeContainer& NodeContainer, const GLTF::FMaterial& GltfMaterial, UInterchangeShaderGraphNode& ShaderGraphNode ) const
 {
 	using namespace UE::Interchange::Materials;
 
@@ -901,7 +901,7 @@ void UInterchangeGltfTranslator::HandleGltfTransmission( UInterchangeBaseNodeCon
 	}
 }
 
-void UInterchangeGltfTranslator::HandleGltfTextureTransform( UInterchangeBaseNodeContainer& NodeContainer, const GLTF::FTextureTransform& TextureTransform, const int32 TexCoordIndex, UInterchangeShaderNode& ShaderNode ) const
+void UInterchangeGLTFTranslator::HandleGltfTextureTransform( UInterchangeBaseNodeContainer& NodeContainer, const GLTF::FTextureTransform& TextureTransform, const int32 TexCoordIndex, UInterchangeShaderNode& ShaderNode ) const
 {
 	using namespace UE::Interchange::Materials;
 
@@ -947,18 +947,18 @@ void UInterchangeGltfTranslator::HandleGltfTextureTransform( UInterchangeBaseNod
 
 }
 
-EInterchangeTranslatorType UInterchangeGltfTranslator::GetTranslatorType() const
+EInterchangeTranslatorType UInterchangeGLTFTranslator::GetTranslatorType() const
 {
 	return EInterchangeTranslatorType::Scenes;
 }
 
-EInterchangeTranslatorAssetType UInterchangeGltfTranslator::GetSupportedAssetTypes() const
+EInterchangeTranslatorAssetType UInterchangeGLTFTranslator::GetSupportedAssetTypes() const
 {
 	//gltf translator support Meshes and Materials
 	return EInterchangeTranslatorAssetType::Materials | EInterchangeTranslatorAssetType::Meshes | EInterchangeTranslatorAssetType::Animations;
 }
 
-TArray<FString> UInterchangeGltfTranslator::GetSupportedFormats() const
+TArray<FString> UInterchangeGLTFTranslator::GetSupportedFormats() const
 {
 	TArray<FString> GltfExtensions;
 
@@ -969,7 +969,7 @@ TArray<FString> UInterchangeGltfTranslator::GetSupportedFormats() const
 	return GltfExtensions;
 }
 
-bool UInterchangeGltfTranslator::Translate( UInterchangeBaseNodeContainer& NodeContainer ) const
+bool UInterchangeGLTFTranslator::Translate( UInterchangeBaseNodeContainer& NodeContainer ) const
 {
 	using namespace UE::Interchange::Gltf::Private;
 
@@ -984,7 +984,7 @@ bool UInterchangeGltfTranslator::Translate( UInterchangeBaseNodeContainer& NodeC
 
 	const bool bLoadImageData = false;
 	const bool bLoadMetaData = false;
-	GltfFileReader.ReadFile( FilePath, bLoadImageData, bLoadMetaData, const_cast< UInterchangeGltfTranslator* >( this )->GltfAsset );
+	GltfFileReader.ReadFile( FilePath, bLoadImageData, bLoadMetaData, const_cast< UInterchangeGLTFTranslator* >( this )->GltfAsset );
 
 	const FString FileName = GltfAsset.Name;
 
@@ -1039,7 +1039,7 @@ bool UInterchangeGltfTranslator::Translate( UInterchangeBaseNodeContainer& NodeC
 		return false;
 	}
 
-	ScaleNodeTranslations(const_cast<UInterchangeGltfTranslator*>(this)->GltfAsset.Nodes, GltfUnitConversionMultiplier);
+	ScaleNodeTranslations(const_cast<UInterchangeGLTFTranslator*>(this)->GltfAsset.Nodes, GltfUnitConversionMultiplier);
 
 	// Textures
 	{
@@ -1325,7 +1325,7 @@ bool UInterchangeGltfTranslator::Translate( UInterchangeBaseNodeContainer& NodeC
 	return true;
 }
 
-TOptional< UE::Interchange::FImportImage > UInterchangeGltfTranslator::GetTexturePayloadData(const FString& PayloadKey, TOptional<FString>& AlternateTexturePath) const
+TOptional< UE::Interchange::FImportImage > UInterchangeGLTFTranslator::GetTexturePayloadData(const FString& PayloadKey, TOptional<FString>& AlternateTexturePath) const
 {
 	int32 TextureIndex = 0;
 	LexFromString( TextureIndex, *PayloadKey);
@@ -1371,7 +1371,7 @@ TOptional< UE::Interchange::FImportImage > UInterchangeGltfTranslator::GetTextur
 	}
 }
 
-TFuture<TOptional<UE::Interchange::FAnimationPayloadData>> UInterchangeGltfTranslator::GetAnimationPayloadData(const FInterchangeAnimationPayLoadKey& PayLoadKey, const double BakeFrequency, const double RangeStartSecond, const double RangeStopSecond) const
+TFuture<TOptional<UE::Interchange::FAnimationPayloadData>> UInterchangeGLTFTranslator::GetAnimationPayloadData(const FInterchangeAnimationPayLoadKey& PayLoadKey, const double BakeFrequency, const double RangeStartSecond, const double RangeStopSecond) const
 {
 	return Async(EAsyncExecution::TaskGraph, [this, PayLoadKey, BakeFrequency, RangeStartSecond, RangeStopSecond]
 		{
@@ -1413,7 +1413,7 @@ TFuture<TOptional<UE::Interchange::FAnimationPayloadData>> UInterchangeGltfTrans
 	);
 }
 
-void UInterchangeGltfTranslator::HandleGltfAnimation(UInterchangeBaseNodeContainer& NodeContainer, int32 AnimationIndex) const
+void UInterchangeGLTFTranslator::HandleGltfAnimation(UInterchangeBaseNodeContainer& NodeContainer, int32 AnimationIndex) const
 {
 	const GLTF::FAnimation& GltfAnimation = GltfAsset.Animations[AnimationIndex];
 
@@ -1714,7 +1714,7 @@ void UInterchangeGltfTranslator::HandleGltfAnimation(UInterchangeBaseNodeContain
 	NodeContainer.AddNode(TrackSetNode);
 }
 
-void UInterchangeGltfTranslator::SetTextureSRGB(UInterchangeBaseNodeContainer& NodeContainer, const GLTF::FTextureMap& TextureMap, bool bSRGB) const
+void UInterchangeGLTFTranslator::SetTextureSRGB(UInterchangeBaseNodeContainer& NodeContainer, const GLTF::FTextureMap& TextureMap, bool bSRGB) const
 {
 	if (GltfAsset.Textures.IsValidIndex(TextureMap.TextureIndex))
 	{
@@ -1725,7 +1725,7 @@ void UInterchangeGltfTranslator::SetTextureSRGB(UInterchangeBaseNodeContainer& N
 		}
 	}
 }
-void UInterchangeGltfTranslator::SetTextureFlipGreenChannel(UInterchangeBaseNodeContainer& NodeContainer, const GLTF::FTextureMap& TextureMap) const
+void UInterchangeGLTFTranslator::SetTextureFlipGreenChannel(UInterchangeBaseNodeContainer& NodeContainer, const GLTF::FTextureMap& TextureMap) const
 {
 	if (GltfAsset.Textures.IsValidIndex(TextureMap.TextureIndex))
 	{
@@ -1737,7 +1737,7 @@ void UInterchangeGltfTranslator::SetTextureFlipGreenChannel(UInterchangeBaseNode
 	}
 }
 
-TFuture<TOptional<UE::Interchange::FVariantSetPayloadData>> UInterchangeGltfTranslator::GetVariantSetPayloadData(const FString& PayloadKey) const
+TFuture<TOptional<UE::Interchange::FVariantSetPayloadData>> UInterchangeGLTFTranslator::GetVariantSetPayloadData(const FString& PayloadKey) const
 {
 	using namespace UE::Interchange;
 
@@ -1778,7 +1778,7 @@ TFuture<TOptional<UE::Interchange::FVariantSetPayloadData>> UInterchangeGltfTran
 		);
 }
 
-void UInterchangeGltfTranslator::HandleGltfVariants(UInterchangeBaseNodeContainer& NodeContainer, const FString& FileName) const
+void UInterchangeGLTFTranslator::HandleGltfVariants(UInterchangeBaseNodeContainer& NodeContainer, const FString& FileName) const
 {
 	UInterchangeVariantSetNode* VariantSetNode = nullptr;
 	VariantSetNode = NewObject<UInterchangeVariantSetNode>(&NodeContainer);
@@ -1861,7 +1861,7 @@ void UInterchangeGltfTranslator::HandleGltfVariants(UInterchangeBaseNodeContaine
 	SceneVariantSetsNode->AddCustomVariantSetUid(VariantSetNodeUid);
 }
 
-bool UInterchangeGltfTranslator::GetVariantSetPayloadData(UE::Interchange::FVariantSetPayloadData& PayloadData) const
+bool UInterchangeGLTFTranslator::GetVariantSetPayloadData(UE::Interchange::FVariantSetPayloadData& PayloadData) const
 {
 	using namespace UE;
 
@@ -1955,7 +1955,7 @@ bool UInterchangeGltfTranslator::GetVariantSetPayloadData(UE::Interchange::FVari
 	return true;
 }
 
-TFuture< TOptional< UE::Interchange::FMeshPayloadData > > UInterchangeGltfTranslator::GetMeshPayloadData(const FInterchangeMeshPayLoadKey& PayLoadKey) const
+TFuture< TOptional< UE::Interchange::FMeshPayloadData > > UInterchangeGLTFTranslator::GetMeshPayloadData(const FInterchangeMeshPayLoadKey& PayLoadKey) const
 {
 	return Async(EAsyncExecution::TaskGraph, [this, PayLoadKey]
 		{
@@ -1989,7 +1989,7 @@ TFuture< TOptional< UE::Interchange::FMeshPayloadData > > UInterchangeGltfTransl
 		});
 }
 
-void UInterchangeGltfTranslator::HandleGltfSkeletons(UInterchangeBaseNodeContainer& NodeContainer, const FString& SceneNodeUid, const TArray<int32>& SkinnedMeshNodes, TSet<int>& UnusedMeshIndices) const
+void UInterchangeGLTFTranslator::HandleGltfSkeletons(UInterchangeBaseNodeContainer& NodeContainer, const FString& SceneNodeUid, const TArray<int32>& SkinnedMeshNodes, TSet<int>& UnusedMeshIndices) const
 {
 	TMap<int32, TMap<int32, TArray<int32>>> MeshIndexToRootJointGroupedSkinnedMeshNodesMap;
 
@@ -2085,7 +2085,7 @@ void UInterchangeGltfTranslator::HandleGltfSkeletons(UInterchangeBaseNodeContain
 	}
 }
 
-UInterchangeMeshNode* UInterchangeGltfTranslator::HandleGltfMesh(UInterchangeBaseNodeContainer& NodeContainer,
+UInterchangeMeshNode* UInterchangeGLTFTranslator::HandleGltfMesh(UInterchangeBaseNodeContainer& NodeContainer,
 	const GLTF::FMesh& GltfMesh, int MeshIndex,
 	TSet<int>& UnusedMeshIndices,
 	const FString& SkeletalName/*If set it creates the mesh even if it was already created (for Skeletals)*/,
@@ -2183,7 +2183,7 @@ UInterchangeMeshNode* UInterchangeGltfTranslator::HandleGltfMesh(UInterchangeBas
 	return MeshNode;
 }
 
-UInterchangeGltfTranslator::UInterchangeGltfTranslator()
+UInterchangeGLTFTranslator::UInterchangeGLTFTranslator()
 {
 	if (!HasAllFlags(RF_ClassDefaultObject))
 	{
diff --git a/Engine/Plugins/Interchange/Runtime/Source/Import/Private/InterchangeImportModule.cpp b/Engine/Plugins/Interchange/Runtime/Source/Import/Private/InterchangeImportModule.cpp
--- a/Engine/Plugins/Interchange/Runtime/Source/Import/Private/InterchangeImportModule.cpp
+++ b/Engine/Plugins/Interchange/Runtime/Source/Import/Private/InterchangeImportModule.cpp
@@ -58,7 +58,7 @@ void FInterchangeImportModule::StartupModule()
 		//Register the translators
 		//Scenes
 		InterchangeManager.RegisterTranslator(UInterchangeFbxTranslator::StaticClass());
-		InterchangeManager.RegisterTranslator(UInterchangeGltfTranslator::StaticClass());
+		InterchangeManager.RegisterTranslator(UInterchangeGLTFTranslator::StaticClass());
 		InterchangeManager.RegisterTranslator(UInterchangeOBJTranslator::StaticClass());
 
 		//Materials
diff --git a/Engine/Plugins/Interchange/Runtime/Source/Import/Public/Gltf/InterchangeGltfTranslator.h b/Engine/Plugins/Interchange/Runtime/Source/Import/Public/Gltf/InterchangeGltfTranslator.h
--- a/Engine/Plugins/Interchange/Runtime/Source/Import/Public/Gltf/InterchangeGltfTranslator.h
+++ b/Engine/Plugins/Interchange/Runtime/Source/Import/Public/Gltf/InterchangeGltfTranslator.h
@@ -20,10 +20,9 @@ class UInterchangeShaderNode;
 class UInterchangeVariantSetNode;
 class UInterchangeMeshNode;
 
-/* Gltf translator class support import of texture, material, static mesh, skeletal mesh, */
-
+/* glTF translator class support import of texture, material, static mesh, skeletal mesh, */
 UCLASS(BlueprintType)
-class UInterchangeGltfTranslator : public UInterchangeTranslatorBase
+class UInterchangeGLTFTranslator : public UInterchangeTranslatorBase
 	, public IInterchangeMeshPayloadInterface
 	, public IInterchangeTexturePayloadInterface
 	, public IInterchangeAnimationPayloadInterface
@@ -32,7 +31,7 @@ class UInterchangeGltfTranslator : public UInterchangeTranslatorBase
 	GENERATED_BODY()
 
 public:
-	UInterchangeGltfTranslator();
+	UInterchangeGLTFTranslator();
 
 	/** Begin UInterchangeTranslatorBase API*/
 	virtual EInterchangeTranslatorType GetTranslatorType() const override;
diff --git a/Engine/Plugins/Interchange/Runtime/Source/Pipelines/Private/InterchangeglTFPipeline.cpp b/Engine/Plugins/Interchange/Runtime/Source/Pipelines/Private/InterchangeglTFPipeline.cpp
--- a/Engine/Plugins/Interchange/Runtime/Source/Pipelines/Private/InterchangeglTFPipeline.cpp
+++ b/Engine/Plugins/Interchange/Runtime/Source/Pipelines/Private/InterchangeglTFPipeline.cpp
@@ -307,12 +307,12 @@ void UGLTFPipelineSettings::BuildMaterialInstance(const UInterchangeShaderGraphN
 	}
 }
 
-UInterchangeglTFPipeline::UInterchangeglTFPipeline()
+UInterchangeGLTFPipeline::UInterchangeGLTFPipeline()
 	: GLTFPipelineSettings(UGLTFPipelineSettings::StaticClass()->GetDefaultObject<UGLTFPipelineSettings>())
 {
 }
 
-void UInterchangeglTFPipeline::AdjustSettingsForContext(EInterchangePipelineContext ImportType, TObjectPtr<UObject> ReimportAsset)
+void UInterchangeGLTFPipeline::AdjustSettingsForContext(EInterchangePipelineContext ImportType, TObjectPtr<UObject> ReimportAsset)
 {
 	Super::AdjustSettingsForContext(ImportType, ReimportAsset);
 
@@ -323,7 +323,7 @@ void UInterchangeglTFPipeline::AdjustSettingsForContext(EInterchangePipelineCont
 	}
 }
 
-void UInterchangeglTFPipeline::ExecutePipeline(UInterchangeBaseNodeContainer* NodeContainer, const TArray<UInterchangeSourceData*>& InSourceDatas)
+void UInterchangeGLTFPipeline::ExecutePipeline(UInterchangeBaseNodeContainer* NodeContainer, const TArray<UInterchangeSourceData*>& InSourceDatas)
 {
 	Super::ExecutePipeline(NodeContainer, InSourceDatas);
 
diff --git a/Engine/Plugins/Interchange/Runtime/Source/Pipelines/Public/InterchangeglTFPipeline.h b/Engine/Plugins/Interchange/Runtime/Source/Pipelines/Public/InterchangeglTFPipeline.h
--- a/Engine/Plugins/Interchange/Runtime/Source/Pipelines/Public/InterchangeglTFPipeline.h
+++ b/Engine/Plugins/Interchange/Runtime/Source/Pipelines/Public/InterchangeglTFPipeline.h
@@ -16,13 +16,13 @@ class UInterchangeBaseNodeContainer;
 class UInterchangeMaterialInstanceFactoryNode;
 class UInterchangeShaderGraphNode;
 
-UCLASS(config = Engine, meta = (DisplayName = "Interchange GLTF"))
+UCLASS(config = Engine, meta = (DisplayName = "Interchange glTF"))
 class INTERCHANGEPIPELINES_API UGLTFPipelineSettings : public UDeveloperSettings
 {
 	GENERATED_BODY()
 
 public:
-	UPROPERTY(EditAnywhere, config, Category = "GLTFPredefinedMaterialLibrary", meta = (DisplayName = "GLTF Predefined Material Library"))
+	UPROPERTY(EditAnywhere, config, Category = "PredefinedglTFMaterialLibrary", meta = (DisplayName = "Predefined glTF Material Library"))
 	TMap<FString, FSoftObjectPath> MaterialParents;
 
 	TArray<FString> ValidateMaterialInstancesAndParameters() const;
@@ -44,17 +44,17 @@ private:
 };
 
 UCLASS(BlueprintType)
-class INTERCHANGEPIPELINES_API UInterchangeglTFPipeline : public UInterchangePipelineBase
+class INTERCHANGEPIPELINES_API UInterchangeGLTFPipeline : public UInterchangePipelineBase
 {
 	GENERATED_BODY()
 
-	UInterchangeglTFPipeline();
+	UInterchangeGLTFPipeline();
 
 public:
-	UPROPERTY(EditAnywhere, BlueprintReadWrite, Category = "GLTF")
+	UPROPERTY(EditAnywhere, BlueprintReadWrite, Category = "glTF", meta = (DisplayName = "glTF Pipeline Settings"))
 	TObjectPtr<UGLTFPipelineSettings> GLTFPipelineSettings;
 
-	UPROPERTY(EditAnywhere, BlueprintReadWrite, Category = "GLTF")
+	UPROPERTY(EditAnywhere, BlueprintReadWrite, Category = "glTF", meta = (DisplayName = "Use glTF Material Instance Library"))
 	bool bUseGLTFMaterialInstanceLibrary = false;
 
 protected:
