[Interchange - GLTF] Revisit Morph Targets for the GLTF translator

diff --git a/Engine/Plugins/Interchange/Runtime/Source/Import/Private/Gltf/InterchangeGltfTranslator.cpp b/Engine/Plugins/Interchange/Runtime/Source/Import/Private/Gltf/InterchangeGltfTranslator.cpp
--- a/Engine/Plugins/Interchange/Runtime/Source/Import/Private/Gltf/InterchangeGltfTranslator.cpp
+++ b/Engine/Plugins/Interchange/Runtime/Source/Import/Private/Gltf/InterchangeGltfTranslator.cpp
@@ -269,7 +269,7 @@ void UInterchangeGLTFTranslator::HandleGltfNode( UInterchangeBaseNodeContainer&
 				UInterchangeMeshNode* MeshNode = HandleGltfMesh(NodeContainer, GltfAsset.Meshes[GltfNode.MeshIndex], GltfNode.MeshIndex, UnusedMeshIndices);
 
 				InterchangeSceneNode->SetCustomAssetInstanceUid( MeshNode->GetUniqueID() );
-				if (MeshNode->IsSkinnedMesh())
+				if (GltfAsset.Meshes[GltfNode.MeshIndex].MorphTargetNames.Num() > 0)
 				{
 					const TArray<FString>& MorphTargetNames = GltfAsset.Meshes[GltfNode.MeshIndex].MorphTargetNames;
 					int32 MorphTargetNamesCount = MorphTargetNames.Num();
@@ -286,13 +286,6 @@ void UInterchangeGLTFTranslator::HandleGltfNode( UInterchangeBaseNodeContainer&
 					{
 						UE_LOG(LogInterchangeImport, Warning, TEXT("GLTF Node [%] Import Warning. Gltf Node's MorphTargetNames count is missmatched against MorphTargetWeights count."), *GltfNode.UniqueId);
 					}
-					
-					//Interchange/UE handles Morph Targets in skeletalMeshes:
-					InterchangeSceneNode->AddSpecializedType(UE::Interchange::FSceneNodeStaticData::GetJointSpecializeTypeString());
-					if (MeshNode->GetSkeletonDependeciesCount() == 0)
-					{
-						MeshNode->SetSkeletonDependencyUid(InterchangeSceneNode->GetUniqueID());
-					}
 				}
 
 				if (!bHasVariants && GltfAsset.Variants.Num() > 0)
@@ -2134,8 +2127,6 @@ UInterchangeMeshNode* UInterchangeGLTFTranslator::HandleGltfMesh(UInterchangeBas
 	//Generate Morph Target Meshes:
 	if (GltfMesh.MorphTargetNames.Num() > 0)
 	{
-		MeshNode->SetSkinnedMesh(true);
-
 		for (int32 MorphTargetIndex = 0; MorphTargetIndex < GltfMesh.MorphTargetNames.Num(); MorphTargetIndex++)
 		{
 			//check if MorphTarget mesh was already created or not:
diff --git a/Engine/Plugins/Interchange/Runtime/Source/Pipelines/Private/InterchangeGenericScenesPipeline.cpp b/Engine/Plugins/Interchange/Runtime/Source/Pipelines/Private/InterchangeGenericScenesPipeline.cpp
--- a/Engine/Plugins/Interchange/Runtime/Source/Pipelines/Private/InterchangeGenericScenesPipeline.cpp
+++ b/Engine/Plugins/Interchange/Runtime/Source/Pipelines/Private/InterchangeGenericScenesPipeline.cpp
@@ -19,6 +19,7 @@
 #include "InterchangeVariantSetNode.h"
 #include "Nodes/InterchangeUserDefinedAttribute.h"
 #include "InterchangeSkeletonFactoryNode.h"
+#include "InterchangeSkeletalMeshFactoryNode.h"
 
 #include "Animation/SkeletalMeshActor.h"
 #include "CineCameraActor.h"
@@ -189,38 +190,27 @@ void UInterchangeGenericLevelPipeline::ExecuteSceneNodePreImport(const FTransfor
 	}
 
 	const UInterchangeBaseNode* TranslatedAssetNode = nullptr;
-	bool bRootJointNode = SceneNode->IsSpecializedTypeContains(UE::Interchange::FSceneNodeStaticData::GetJointSpecializeTypeString());
+	bool bRootJointNode = false;
 	FString SkeletalMeshFactoryNodeUid;
 
-	if (bRootJointNode)
+	FString SkeletonFactoryNodeUid = UInterchangeFactoryBaseNode::BuildFactoryNodeUid(SceneNode->GetUniqueID());
+	const UInterchangeSkeletonFactoryNode* SkeletonFactoryNode = Cast<UInterchangeSkeletonFactoryNode>(BaseNodeContainer->GetFactoryNode(SkeletonFactoryNodeUid));
+	if (SkeletonFactoryNode)
 	{
-		FString SkeletonFactoryNodeUid = UInterchangeFactoryBaseNode::BuildFactoryNodeUid(SceneNode->GetUniqueID());
-		const UInterchangeSkeletonFactoryNode* SkeletonFactoryNode = Cast<UInterchangeSkeletonFactoryNode>(BaseNodeContainer->GetFactoryNode(SkeletonFactoryNodeUid));
-		if (SkeletonFactoryNode)
+		if (SkeletonFactoryNode->GetCustomSkeletalMeshFactoryNodeUid(SkeletalMeshFactoryNodeUid))
 		{
-			if (SkeletonFactoryNode->GetCustomSkeletalMeshFactoryNodeUid(SkeletalMeshFactoryNodeUid))
+			if (const UInterchangeFactoryBaseNode* SkeletalMeshFactoryNode = BaseNodeContainer->GetFactoryNode(SkeletalMeshFactoryNodeUid))
 			{
-				if (const UInterchangeFactoryBaseNode* SkeletalMeshFactoryNode = BaseNodeContainer->GetFactoryNode(SkeletalMeshFactoryNodeUid))
-				{
-					TArray<FString> NodeUids;
-					SkeletalMeshFactoryNode->GetTargetNodeUids(NodeUids);
+				TArray<FString> NodeUids;
+				SkeletalMeshFactoryNode->GetTargetNodeUids(NodeUids);
 
-					if (NodeUids.Num() > 0)
-					{
-						TranslatedAssetNode = BaseNodeContainer->GetNode(NodeUids[0]);
-					}
-					else
-					{
-						TranslatedAssetNode = nullptr;
-					}
+				if (NodeUids.Num() > 0)
+				{
+					TranslatedAssetNode = BaseNodeContainer->GetNode(NodeUids[0]);
+					bRootJointNode = true;
 				}
 			}
 		}
-		else
-		{
-			//If it does not have a SkeletonFactoryNode with the given Uid, then it is not a root joint node:
-			bRootJointNode = false;
-		}
 	}
 	
 	if (!bRootJointNode)
@@ -368,14 +358,19 @@ void UInterchangeGenericLevelPipeline::SetUpFactoryNode(UInterchangeActorFactory
 
 	if (const UInterchangeMeshNode* MeshNode = Cast<UInterchangeMeshNode>(TranslatedAssetNode))
 	{
-		if (MeshNode->IsSkinnedMesh())
+		TArray<FString> TargetNodeUids;
+		ActorFactoryNode->GetTargetNodeUids(TargetNodeUids);
+		bool bSkeletal = false;
+		if (TargetNodeUids.Num() > 0)
 		{
-			bool bRootJointNode = SceneNode->IsSpecializedTypeContains(UE::Interchange::FSceneNodeStaticData::GetJointSpecializeTypeString());
-			if (!bRootJointNode)
+			if (const UInterchangeSkeletalMeshFactoryNode* SkeletalMeshFactoryNode = Cast<UInterchangeSkeletalMeshFactoryNode>(BaseNodeContainer->GetFactoryNode(TargetNodeUids[0])))
 			{
-				return;
+				bSkeletal = true;
 			}
+		}
 
+		if (bSkeletal)
+		{
 			ActorFactoryNode->SetCustomActorClassName(ASkeletalMeshActor::StaticClass()->GetPathName());
 			ActorFactoryNode->SetCustomMobility(EComponentMobility::Movable);
 		}
diff --git a/Engine/Plugins/Interchange/Runtime/Source/Pipelines/Private/InterchangeGenericSkeletalMeshPipeline.cpp b/Engine/Plugins/Interchange/Runtime/Source/Pipelines/Private/InterchangeGenericSkeletalMeshPipeline.cpp
--- a/Engine/Plugins/Interchange/Runtime/Source/Pipelines/Private/InterchangeGenericSkeletalMeshPipeline.cpp
+++ b/Engine/Plugins/Interchange/Runtime/Source/Pipelines/Private/InterchangeGenericSkeletalMeshPipeline.cpp
@@ -174,7 +174,7 @@ void UInterchangeGenericMeshPipeline::ExecutePreImportPipelineSkeletalMesh()
 			return bFoundInstances;
 		};
 
-		PipelineMeshesUtilities->GetCombinedSkinnedMeshInstances(BaseNodeContainer, MeshUidsPerSkeletonRootUid, bConvertStaticMeshToSkeletalMesh);
+		PipelineMeshesUtilities->GetCombinedSkinnedMeshInstances(BaseNodeContainer, MeshUidsPerSkeletonRootUid, bConvertStaticMeshToSkeletalMesh, bConvertStaticsWithMorphTargetsToSkeletals);
 		bool bUseMeshInstance = true;
 		bool bFoundMeshes = CreatePerSkeletonRootUidCombinedSkinnedMesh(bUseMeshInstance);
 
@@ -208,8 +208,15 @@ void UInterchangeGenericMeshPipeline::ExecutePreImportPipelineSkeletalMesh()
 				SkeletonRootUid = (bUseInstanceMesh ? PipelineMeshesUtilities->GetMeshInstanceSkeletonRootUid(MeshUid) : PipelineMeshesUtilities->GetMeshGeometrySkeletonRootUid(MeshUid));
 				if (SkeletonRootUid.IsEmpty())
 				{
-					//Log an error
-					continue;
+					if (bConvertStaticsWithMorphTargetsToSkeletals)
+					{
+						SkeletonRootUid = MeshUid;
+					}
+					else
+					{
+						//Log an error
+						continue;
+					}
 				}
 				UInterchangeSkeletonFactoryNode* SkeletonFactoryNode = CreateSkeletonFactoryNode(SkeletonRootUid);
 				if (bUseInstanceMesh)
@@ -245,7 +252,7 @@ void UInterchangeGenericMeshPipeline::ExecutePreImportPipelineSkeletalMesh()
 			return bFoundInstances;
 		};
 		
-		PipelineMeshesUtilities->GetAllSkinnedMeshInstance(MeshUids, bConvertStaticMeshToSkeletalMesh);
+		PipelineMeshesUtilities->GetAllSkinnedMeshInstance(MeshUids, bConvertStaticMeshToSkeletalMesh, bConvertStaticsWithMorphTargetsToSkeletals);
 		bool bUseMeshInstance = true;
 		bool bFoundMeshes = CreatePerSkeletonRootUidSkinnedMesh(bUseMeshInstance);
 
diff --git a/Engine/Plugins/Interchange/Runtime/Source/Pipelines/Private/InterchangePipelineMeshesUtilities.cpp b/Engine/Plugins/Interchange/Runtime/Source/Pipelines/Private/InterchangePipelineMeshesUtilities.cpp
--- a/Engine/Plugins/Interchange/Runtime/Source/Pipelines/Private/InterchangePipelineMeshesUtilities.cpp
+++ b/Engine/Plugins/Interchange/Runtime/Source/Pipelines/Private/InterchangePipelineMeshesUtilities.cpp
@@ -232,6 +232,7 @@ UInterchangePipelineMeshesUtilities* UInterchangePipelineMeshesUtilities::Create
 									MeshInstance.ReferencingMeshGeometryUids.Add(MeshUid);
 									MeshInstance.bReferenceSkinnedMesh |= MeshGeometry.MeshNode->IsSkinnedMesh();
 									MeshInstance.bReferenceMorphTarget |= MeshGeometry.MeshNode->IsMorphTarget();
+									MeshInstance.bHasMorphTargets |= MeshGeometry.MeshNode->GetMorphTargetDependeciesCount() > 0;
 								}
 								else
 								{
@@ -244,6 +245,7 @@ UInterchangePipelineMeshesUtilities* UInterchangePipelineMeshesUtilities::Create
 									MeshInstance.ReferencingMeshGeometryUids.Add(MeshUid);
 									MeshInstance.bReferenceSkinnedMesh |= MeshGeometry.MeshNode->IsSkinnedMesh();
 									MeshInstance.bReferenceMorphTarget |= MeshGeometry.MeshNode->IsMorphTarget();
+									MeshInstance.bHasMorphTargets |= MeshGeometry.MeshNode->GetMorphTargetDependeciesCount() > 0;
 								}
 							}
 						}
@@ -335,13 +337,13 @@ void UInterchangePipelineMeshesUtilities::IterateAllMeshInstance(TFunctionRef<vo
 	}
 }
 
-void UInterchangePipelineMeshesUtilities::GetAllSkinnedMeshInstance(TArray<FString>& MeshInstanceUids, const bool bConvertStaticMeshToSkeletalMesh) const
+void UInterchangePipelineMeshesUtilities::GetAllSkinnedMeshInstance(TArray<FString>& MeshInstanceUids, const bool bConvertStaticMeshToSkeletalMesh, const bool bConvertStaticsWithMorphTargetsToSkeletals) const
 {
 	MeshInstanceUids.Empty(MeshInstancesPerMeshInstanceUid.Num());
 	for (const TPair<FString, FInterchangeMeshInstance>& MeshInstanceUidAndMeshInstance : MeshInstancesPerMeshInstanceUid)
 	{
 		const FInterchangeMeshInstance& MeshInstance = MeshInstanceUidAndMeshInstance.Value;
-		if((bConvertStaticMeshToSkeletalMesh || MeshInstance.bReferenceSkinnedMesh) && !MeshInstance.bReferenceMorphTarget)
+		if((bConvertStaticMeshToSkeletalMesh || MeshInstance.bReferenceSkinnedMesh || (bConvertStaticsWithMorphTargetsToSkeletals && MeshInstance.bHasMorphTargets)) && !MeshInstance.bReferenceMorphTarget)
 		{
 			MeshInstanceUids.Add(MeshInstance.MeshInstanceUid);
 		}
@@ -510,18 +512,18 @@ void UInterchangePipelineMeshesUtilities::IterateAllMeshInstanceUsingMeshGeometr
 	}
 }
 
-void UInterchangePipelineMeshesUtilities::GetCombinedSkinnedMeshInstances(UInterchangeBaseNodeContainer* BaseNodeContainer, TMap<FString, TArray<FString>>& OutMeshInstanceUidsPerSkeletonRootUid, const bool bConvertStaticMeshToSkeletalMesh) const
+void UInterchangePipelineMeshesUtilities::GetCombinedSkinnedMeshInstances(UInterchangeBaseNodeContainer* BaseNodeContainer, TMap<FString, TArray<FString>>& OutMeshInstanceUidsPerSkeletonRootUid, const bool bConvertStaticMeshToSkeletalMesh, const bool bConvertStaticsWithMorphTargetsToSkeletals) const
 {
 	check(BaseNodeContainer);
 
 	for (const TPair<FString, FInterchangeMeshInstance>& MeshInstanceUidAndMeshInstance : MeshInstancesPerMeshInstanceUid)
 	{
 		const FInterchangeMeshInstance& MeshInstance = MeshInstanceUidAndMeshInstance.Value;
-		if (!bConvertStaticMeshToSkeletalMesh && !MeshInstance.bReferenceSkinnedMesh)
+		if (!(bConvertStaticMeshToSkeletalMesh || bConvertStaticsWithMorphTargetsToSkeletals) && !MeshInstance.bReferenceSkinnedMesh)
 		{
 			continue;
 		}
-		bool bIsStaticMesh = bConvertStaticMeshToSkeletalMesh && !MeshInstance.bReferenceSkinnedMesh;
+		bool bIsStaticMesh = (bConvertStaticMeshToSkeletalMesh || (bConvertStaticsWithMorphTargetsToSkeletals && MeshInstance.bHasMorphTargets)) && !MeshInstance.bReferenceSkinnedMesh;
 		//Find the root skeleton for this MeshInstance
 		FString SkeletonRootUid;
 		for (const FString& MeshGeometryUid : MeshInstance.ReferencingMeshGeometryUids)
@@ -542,17 +544,21 @@ void UInterchangePipelineMeshesUtilities::GetCombinedSkinnedMeshInstances(UInter
 			else if (bIsStaticMesh)
 			{
 				//Create a joint from the instance node (the scene node pointing on the mesh).
-				//Since we are dealing with rigid mesh and combine we will use the outermost valid parent
 				SkeletonRootUid = MeshInstance.MeshInstanceUid;
-				if (const UInterchangeSceneNode* SceneNode = Cast<UInterchangeSceneNode>(BaseNodeContainer->GetNode(MeshInstance.MeshInstanceUid)))
+
+				if (!bConvertStaticsWithMorphTargetsToSkeletals)
 				{
-					FString ParentUid = SceneNode->GetParentUid();
-					while (!ParentUid.Equals(UInterchangeBaseNode::InvalidNodeUid()))
+					//Since we are dealing with rigid mesh and combine we will use the outermost valid parent
+					if (const UInterchangeSceneNode* SceneNode = Cast<UInterchangeSceneNode>(BaseNodeContainer->GetNode(MeshInstance.MeshInstanceUid)))
 					{
-						if (const UInterchangeSceneNode* ParentNode = Cast<UInterchangeSceneNode>(BaseNodeContainer->GetNode(ParentUid)))
+						FString ParentUid = SceneNode->GetParentUid();
+						while (!ParentUid.Equals(UInterchangeBaseNode::InvalidNodeUid()))
 						{
-							SkeletonRootUid = ParentUid;
-							ParentUid = ParentNode->GetParentUid();
+							if (const UInterchangeSceneNode* ParentNode = Cast<UInterchangeSceneNode>(BaseNodeContainer->GetNode(ParentUid)))
+							{
+								SkeletonRootUid = ParentUid;
+								ParentUid = ParentNode->GetParentUid();
+							}
 						}
 					}
 				}
diff --git a/Engine/Plugins/Interchange/Runtime/Source/Pipelines/Public/InterchangeGenericMeshPipeline.h b/Engine/Plugins/Interchange/Runtime/Source/Pipelines/Public/InterchangeGenericMeshPipeline.h
--- a/Engine/Plugins/Interchange/Runtime/Source/Pipelines/Public/InterchangeGenericMeshPipeline.h
+++ b/Engine/Plugins/Interchange/Runtime/Source/Pipelines/Public/InterchangeGenericMeshPipeline.h
@@ -44,6 +44,10 @@ public:
 	UPROPERTY(EditAnywhere, BlueprintReadWrite, Category = "Static Meshes")
 	bool bImportStaticMeshes = true;
 
+	/** Allow to convert static mesh using morph target to skeletal mesh" */
+	UPROPERTY(EditAnywhere, BlueprintReadWrite, Category = "Static Meshes")
+	bool bConvertStaticsWithMorphTargetsToSkeletals = false;
+
 	/** If enabled, all translated static mesh nodes will be imported as a single static mesh. */
 	UPROPERTY(EditAnywhere, BlueprintReadWrite, Category = "Static Meshes")
 	bool bCombineStaticMeshes = false;
diff --git a/Engine/Plugins/Interchange/Runtime/Source/Pipelines/Public/InterchangePipelineMeshesUtilities.h b/Engine/Plugins/Interchange/Runtime/Source/Pipelines/Public/InterchangePipelineMeshesUtilities.h
--- a/Engine/Plugins/Interchange/Runtime/Source/Pipelines/Public/InterchangePipelineMeshesUtilities.h
+++ b/Engine/Plugins/Interchange/Runtime/Source/Pipelines/Public/InterchangePipelineMeshesUtilities.h
@@ -46,6 +46,7 @@ struct FInterchangeMeshInstance
 		LodGroupNode = nullptr;
 		bReferenceSkinnedMesh = false;
 		bReferenceMorphTarget = false;
+		bHasMorphTargets = false;
 	}
 	/**
 	 * This ID represent either 1: a lod group scene node uid or 2: a mesh scene node uid.
@@ -65,6 +66,9 @@ struct FInterchangeMeshInstance
 	UPROPERTY(EditAnywhere, Category = "Interchange | Pipeline | MeshesInstance")
 	bool bReferenceMorphTarget;
 
+	UPROPERTY(EditAnywhere, Category = "Interchange | Pipeline | MeshesInstance")
+	bool bHasMorphTargets;
+
 	/**
 	 * Each scene node here represent a mesh scene node. Only if we represent a lod group we can have more then 1 mesh scene node for a specific lod index.
 	 */
@@ -144,7 +148,7 @@ public:
 	* Get all skinned mesh instance unique ids.
 	*/
 	UFUNCTION(BlueprintCallable, Category = "Interchange | Pipeline | Meshes")
-	void GetAllSkinnedMeshInstance(TArray<FString>& MeshInstanceUids, const bool bConvertStaticMeshToSkeletalMesh) const;
+	void GetAllSkinnedMeshInstance(TArray<FString>& MeshInstanceUids, const bool bConvertStaticMeshToSkeletalMesh, const bool bConvertStaticsWithMorphTargetsToSkeletals) const;
 
 	/**
 	* Iterate all skinned mesh instance.
@@ -245,7 +249,7 @@ public:
 	* Return a list of skinned FInterchangeMeshInstance uid that can be combined together.
 	* We cannot create a skinned mesh with multiple skeleton root node, This function return combined MeshInstance per skeleton roots
 	*/
-	void GetCombinedSkinnedMeshInstances(UInterchangeBaseNodeContainer* BaseNodeContainer, TMap<FString, TArray<FString>>& OutMeshInstanceUidsPerSkeletonRootUid, const bool bConvertStaticMeshToSkeletalMesh) const;
+	void GetCombinedSkinnedMeshInstances(UInterchangeBaseNodeContainer* BaseNodeContainer, TMap<FString, TArray<FString>>& OutMeshInstanceUidsPerSkeletonRootUid, const bool bConvertStaticMeshToSkeletalMesh, const bool bConvertStaticsWithMorphTargetsToSkeletals) const;
 	
 	/**
 	* Return a list of skinned FInterchangeMeshGeometry uid that can be combined together.
