[Interchange] GLTF Importer not processing Bind Poses correctly

diff --git a/Engine/Plugins/Interchange/Runtime/Source/Import/Private/Gltf/InterchangeGltfTranslator.cpp b/Engine/Plugins/Interchange/Runtime/Source/Import/Private/Gltf/InterchangeGltfTranslator.cpp
--- a/Engine/Plugins/Interchange/Runtime/Source/Import/Private/Gltf/InterchangeGltfTranslator.cpp
+++ b/Engine/Plugins/Interchange/Runtime/Source/Import/Private/Gltf/InterchangeGltfTranslator.cpp
@@ -129,6 +129,7 @@ namespace UE::Interchange::Gltf::Private
 		for (GLTF::FNode& Node : Nodes)
 		{
 			Node.Transform.SetTranslation(Node.Transform.GetTranslation() * Scale);
+			Node.LocalBindPose.SetTranslation(Node.LocalBindPose.GetTranslation() * Scale);
 		}
 	}
 
@@ -259,6 +260,10 @@ void UInterchangeGLTFTranslator::HandleGltfNode( UInterchangeBaseNodeContainer&
 		case GLTF::FNode::EType::Joint:
 		{
 			InterchangeSceneNode->AddSpecializedType(UE::Interchange::FSceneNodeStaticData::GetJointSpecializeTypeString());
+			if (GltfNode.bHasLocalBindPose)
+			{
+				InterchangeSceneNode->SetCustomBindPoseLocalTransform(&NodeContainer, GltfNode.LocalBindPose);
+			}
 			break;
 		}
 
diff --git a/Engine/Plugins/Interchange/Runtime/Source/Parsers/GLTFCore/Private/GLTF/GLTFReader.cpp b/Engine/Plugins/Interchange/Runtime/Source/Parsers/GLTFCore/Private/GLTF/GLTFReader.cpp
--- a/Engine/Plugins/Interchange/Runtime/Source/Parsers/GLTFCore/Private/GLTF/GLTFReader.cpp
+++ b/Engine/Plugins/Interchange/Runtime/Source/Parsers/GLTFCore/Private/GLTF/GLTFReader.cpp
@@ -973,6 +973,11 @@ namespace GLTF
 		}
 
 		SetupNodesType();
+
+		GenerateInverseBindPosesPerSkinIndices();
+		GenerateLocalBindPosesPerSkinIndices();
+		SetLocalBindPosesForJoints();
+
 		ExtensionsHandler->SetupAssetExtensions(*JsonRoot);
 		BuildRootJoints();
 	}
@@ -1029,6 +1034,85 @@ namespace GLTF
 		}
 	}
 
+	void FFileReader::GenerateInverseBindPosesPerSkinIndices() const
+	{
+		for (size_t SkinIndex = 0; SkinIndex < Asset->Skins.Num(); SkinIndex++)
+		{
+			const FSkinInfo& Skin = Asset->Skins[SkinIndex];
+			if (Skin.InverseBindMatrices.Count == Skin.Joints.Num())
+			{
+				for (size_t JointCounter = 0; JointCounter < Skin.Joints.Num(); JointCounter++)
+				{
+					FMatrix InverseBindMatrix = Skin.InverseBindMatrices.GetMat4(JointCounter);
+					InverseBindMatrix = GLTF::ConvertMat(InverseBindMatrix);
+
+					FTransform InverseBindMatrixTransform;
+					InverseBindMatrixTransform.SetFromMatrix(InverseBindMatrix);
+					InverseBindMatrixTransform.SetRotation(GLTF::ConvertQuat(InverseBindMatrixTransform.GetRotation()));
+					InverseBindMatrixTransform.SetTranslation(GLTF::ConvertVec3(InverseBindMatrixTransform.GetTranslation()));
+					InverseBindMatrixTransform.SetScale3D(GLTF::ConvertVec3(InverseBindMatrixTransform.GetScale3D()));
+
+					Asset->Nodes[Skin.Joints[JointCounter]].SkinIndexToGlobalInverseBindTransform.Add(SkinIndex, InverseBindMatrixTransform);
+				}
+			}
+		}
+	}
+	void FFileReader::GenerateLocalBindPosesPerSkinIndices() const
+	{
+		for (size_t SkinIndex = 0; SkinIndex < Asset->Skins.Num(); SkinIndex++)
+		{
+			const FSkinInfo& Skin = Asset->Skins[SkinIndex];
+			if (Skin.InverseBindMatrices.Count == Skin.Joints.Num())
+			{
+				for (size_t JointCounter = 0; JointCounter < Skin.Joints.Num(); JointCounter++)
+				{
+					FNode& CurrentNode = Asset->Nodes[Skin.Joints[JointCounter]];
+					if (CurrentNode.ParentIndex != INDEX_NONE)
+					{
+						if (Asset->Nodes[CurrentNode.ParentIndex].Type == FNode::EType::Joint
+							&& Asset->Nodes[CurrentNode.ParentIndex].SkinIndexToGlobalInverseBindTransform.Contains(SkinIndex))
+						{
+							//LocalBindPose; //bind pose would be CurrentNode.GlobalInverseBindTransform.Inverse() * ParentNode.GlobalInverseBindTransform
+							FTransform ParentGlobalInverseBindTransform = Asset->Nodes[CurrentNode.ParentIndex].SkinIndexToGlobalInverseBindTransform[SkinIndex];
+							FTransform LocalBindPose = CurrentNode.SkinIndexToGlobalInverseBindTransform[SkinIndex].Inverse() * ParentGlobalInverseBindTransform;
+
+							CurrentNode.SkinIndexToLocalBindPose.Add(SkinIndex, LocalBindPose);
+						}
+					}
+				}
+			}
+		}
+	}
+	void FFileReader::SetLocalBindPosesForJoints() const
+	{
+		for (size_t SkinIndex = 0; SkinIndex < Asset->Skins.Num(); SkinIndex++)
+		{
+			const FSkinInfo& Skin = Asset->Skins[SkinIndex];
+			if (Skin.InverseBindMatrices.Count == Skin.Joints.Num())
+			{
+				for (size_t JointCounter = 0; JointCounter < Skin.Joints.Num(); JointCounter++)
+				{
+					FNode& CurrentNode = Asset->Nodes[Skin.Joints[JointCounter]];
+
+					if (CurrentNode.bHasLocalBindPose && CurrentNode.SkinIndexToLocalBindPose.Contains(SkinIndex))
+					{
+						if (!CurrentNode.LocalBindPose.Equals(CurrentNode.SkinIndexToLocalBindPose[SkinIndex]))
+						{
+							Messages.Emplace(EMessageSeverity::Error, FString::Printf(TEXT("The same Joint [%s] is used in multiple Skins with multiple different InverseBindMatrix values, which is currently not supported."), *CurrentNode.Name));
+						}
+					}
+
+					if (!CurrentNode.bHasLocalBindPose
+						&& CurrentNode.SkinIndexToLocalBindPose.Contains(SkinIndex))
+					{
+						CurrentNode.bHasLocalBindPose = true;
+						CurrentNode.LocalBindPose = CurrentNode.SkinIndexToLocalBindPose[SkinIndex];
+					}
+				}
+			}
+		}
+	}
+
 	void FFileReader::BuildParentIndices(int32 ParentNodeIndex, int32 CurrentNodeIndex) const
 	{
 		if (!Asset->Nodes.IsValidIndex(CurrentNodeIndex))
diff --git a/Engine/Plugins/Interchange/Runtime/Source/Parsers/GLTFCore/Public/GLTFNode.h b/Engine/Plugins/Interchange/Runtime/Source/Parsers/GLTFCore/Public/GLTFNode.h
--- a/Engine/Plugins/Interchange/Runtime/Source/Parsers/GLTFCore/Public/GLTFNode.h
+++ b/Engine/Plugins/Interchange/Runtime/Source/Parsers/GLTFCore/Public/GLTFNode.h
@@ -44,6 +44,13 @@ namespace GLTF
 
 		FString UniqueId; //will be generated in FAsset::GenerateNames
 
+		TMap<int, FTransform> SkinIndexToGlobalInverseBindTransform;
+		TMap<int, FTransform> SkinIndexToLocalBindPose; //bind pose would be CurrentNode.GlobalInverseBindTransform.Inverse() * ParentNode.GlobalInverseBindTransform
+		bool bHasLocalBindPose;
+		FTransform LocalBindPose;	// First Skin that's using the joint will fill the LocalBindPose.
+									//	Edge case Scenario which is currently not supported:
+									//	Where multiple skins use the same Joint. Currently expected bad outcome if the different skins have different inversebindmatrices on the joint.
+
 		FNode()
 		    : Type(EType::None)
 		    , MeshIndex(INDEX_NONE)
@@ -52,6 +59,7 @@ namespace GLTF
 		    , LightIndex(INDEX_NONE)
 			, ParentIndex(INDEX_NONE)
 			, RootJointIndex(INDEX_NONE)
+			, bHasLocalBindPose(false)
 		{
 		}
 	};
diff --git a/Engine/Plugins/Interchange/Runtime/Source/Parsers/GLTFCore/Public/GLTFReader.h b/Engine/Plugins/Interchange/Runtime/Source/Parsers/GLTFCore/Public/GLTFReader.h
--- a/Engine/Plugins/Interchange/Runtime/Source/Parsers/GLTFCore/Public/GLTFReader.h
+++ b/Engine/Plugins/Interchange/Runtime/Source/Parsers/GLTFCore/Public/GLTFReader.h
@@ -62,6 +62,10 @@ namespace GLTF
 		void SetupObjects(uint32 ObjectCount, const TCHAR* FieldName, SetupFunc Func) const;
 		void SetupNodesType() const;
 
+		void GenerateInverseBindPosesPerSkinIndices() const; //Per node
+		void GenerateLocalBindPosesPerSkinIndices() const; //Per node
+		void SetLocalBindPosesForJoints() const; //Per node
+
 		void BuildParentIndices(int32 ParentNodeIndex, int32 CurrentNodeIndex) const;
 		int32 FindRootJointIndex(int32 CurrentNodeIndex) const;
 		void BuildRootJoints() const;
