[Interchange] Redirector for the UInterchangeCineCameraFactoryNode to UInterchangePhyiscalCameraFactoryNode

diff --git a/Engine/Plugins/Interchange/Runtime/Config/BaseInterchange.ini b/Engine/Plugins/Interchange/Runtime/Config/BaseInterchange.ini
--- a/Engine/Plugins/Interchange/Runtime/Config/BaseInterchange.ini
+++ b/Engine/Plugins/Interchange/Runtime/Config/BaseInterchange.ini
@@ -1,2 +1,5 @@
 [CoreRedirects]
 +EnumRedirects=(OldName="/Script/InterchangePipelines.EInterchangeMaterialImportOption",ValueChanges=(("DoNotImport","ImportAsMaterials"))
+
+; Rename Interchange CineCamera to PhysicalCamera FactoryNode
++ClassRedirects=(OldName="/Script/InterchangeFactoryNodes.InterchangeCineCameraFactoryNode", NewName="/Script/InterchangeFactoryNodes.InterchangePhysicalCameraFactoryNode")
\ No newline at end of file
