[Interchange] AnimSequenceFactory refactoring check to ensure where possible

diff --git a/Engine/Plugins/Interchange/Runtime/Source/Import/Private/Animation/InterchangeAnimSequenceFactory.cpp b/Engine/Plugins/Interchange/Runtime/Source/Import/Private/Animation/InterchangeAnimSequenceFactory.cpp
--- a/Engine/Plugins/Interchange/Runtime/Source/Import/Private/Animation/InterchangeAnimSequenceFactory.cpp
+++ b/Engine/Plugins/Interchange/Runtime/Source/Import/Private/Animation/InterchangeAnimSequenceFactory.cpp
@@ -392,20 +392,28 @@ namespace UE::Interchange::Private
 				TArray<float> TimeKeys;
 				TimeKeys.Reserve(BakeKeyCountForAnimationPayload);
 
-				check(AnimationTransformPayload.Transforms.Num() == BakeKeyCountForAnimationPayload);
+				if (!ensure(AnimationTransformPayload.Transforms.Num() == BakeKeyCountForAnimationPayload))
+				{
+					FString PayloadKey = PayloadKeys[AnimationPayload.Key->GetUniqueID()].UniqueId;
+					UE_LOG(LogInterchangeImport, Warning, TEXT("Animation Payload [%s] has unexpected number of Baked Transforms."), *PayloadKey);
+					break;
+				}
 
 				if (AnimationTransformPayload.Type == EInterchangeAnimationPayLoadType::BAKED)
 				{
 					//Everything should match Key count, sample rate and range
-					check(FMath::IsNearlyEqual(AnimationTransformPayload.BakeFrequency, SampleRate, UE_DOUBLE_KINDA_SMALL_NUMBER));
-					check(FMath::IsNearlyEqual(AnimationTransformPayload.RangeStartTime, RangeStart, UE_DOUBLE_KINDA_SMALL_NUMBER));
-					check(FMath::IsNearlyEqual(AnimationTransformPayload.RangeEndTime, RangeEnd, UE_DOUBLE_KINDA_SMALL_NUMBER));
+					if (!(ensure(FMath::IsNearlyEqual(AnimationTransformPayload.BakeFrequency, SampleRate, UE_DOUBLE_KINDA_SMALL_NUMBER)) &&
+						  ensure(FMath::IsNearlyEqual(AnimationTransformPayload.RangeStartTime, RangeStart, UE_DOUBLE_KINDA_SMALL_NUMBER)) &&
+						  ensure(FMath::IsNearlyEqual(AnimationTransformPayload.RangeEndTime, RangeEnd, UE_DOUBLE_KINDA_SMALL_NUMBER))))
+					{
+						FString PayloadKey = PayloadKeys[AnimationPayload.Key->GetUniqueID()].UniqueId;
+						UE_LOG(LogInterchangeImport, Warning, TEXT("Animation Payload [%s] 's BakeFrequency, RangeStartTime and RangeEndTime does not equal with the provided one."), *PayloadKey);
+					}
 				}
 
 				double CurrentTime = 0;
 				for (int32 BakeIndex = 0; BakeIndex < BakeKeyCountForAnimationPayload; BakeIndex++, CurrentTime += BakeInterval)
 				{
-					check(AnimationTransformPayload.Transforms.IsValidIndex(BakeIndex));
 					FTransform3f AnimKeyTransform = FTransform3f(AnimationTransformPayload.Transforms[BakeIndex]);
 					if (bApplyGlobalOffset)
 					{
@@ -442,10 +450,15 @@ namespace UE::Interchange::Private
 				}
 
 				//Make sure we create the correct amount of keys
-				check(RawTrack.ScaleKeys.Num() == BakeKeyCountForAnimationPayload
+				if (!ensure(RawTrack.ScaleKeys.Num() == BakeKeyCountForAnimationPayload
 					&& RawTrack.PosKeys.Num() == BakeKeyCountForAnimationPayload
 					&& RawTrack.RotKeys.Num() == BakeKeyCountForAnimationPayload
-					&& TimeKeys.Num() == BakeKeyCountForAnimationPayload);
+					&& TimeKeys.Num() == BakeKeyCountForAnimationPayload))
+				{
+					FString PayloadKey = PayloadKeys[AnimationPayload.Key->GetUniqueID()].UniqueId;
+					UE_LOG(LogInterchangeImport, Warning, TEXT("Animation Payload [%s] has unexpected number of animation keys. Animation will be incorrect."), *PayloadKey);
+					continue;
+				}
 
 				//add new track
 				Controller.AddBoneCurve(BoneName, bShouldTransact);
