[Interchange] glTF specular and IOR removal for SpecularGlossiness for MaterialInstances.
Also fixing Lighting Mode setup for Translucent materials. (Surface TranlucencyVolume)

diff --git a/Engine/Plugins/Interchange/Runtime/Source/Pipelines/Private/InterchangeglTFPipeline.cpp b/Engine/Plugins/Interchange/Runtime/Source/Pipelines/Private/InterchangeglTFPipeline.cpp
--- a/Engine/Plugins/Interchange/Runtime/Source/Pipelines/Private/InterchangeglTFPipeline.cpp
+++ b/Engine/Plugins/Interchange/Runtime/Source/Pipelines/Private/InterchangeglTFPipeline.cpp
@@ -124,7 +124,16 @@ TSet<FString> UGLTFPipelineSettings::GenerateExpectedParametersList(const FStrin
 		ExpectedParameters.Add(Parameters::OcclusionTexture_TexCoord);
 		ExpectedParameters.Add(Parameters::OcclusionStrength);
 
-		ExpectedParameters.Add(Parameters::IOR);
+		if (!Identifier.Contains(TEXT("SpecularGlossiness")))
+		{
+			ExpectedParameters.Add(Parameters::IOR);
+
+			ExpectedParameters.Add(Parameters::SpecularTexture);
+			ExpectedParameters.Add(Parameters::SpecularTexture_OffsetScale);
+			ExpectedParameters.Add(Parameters::SpecularTexture_Rotation);
+			ExpectedParameters.Add(Parameters::SpecularTexture_TexCoord);
+			ExpectedParameters.Add(Parameters::SpecularFactor);
+		}
 	}
 
 	//Based on ShadingModel:
@@ -145,13 +154,6 @@ TSet<FString> UGLTFPipelineSettings::GenerateExpectedParametersList(const FStrin
 		ExpectedParameters.Add(Parameters::MetallicRoughnessTexture_TexCoord);
 		ExpectedParameters.Add(Parameters::MetallicFactor);
 		ExpectedParameters.Add(Parameters::RoughnessFactor);
-
-		ExpectedParameters.Add(Parameters::SpecularTexture);
-		ExpectedParameters.Add(Parameters::SpecularTexture_OffsetScale);
-		ExpectedParameters.Add(Parameters::SpecularTexture_Rotation);
-		ExpectedParameters.Add(Parameters::SpecularTexture_TexCoord);
-		ExpectedParameters.Add(Parameters::SpecularFactor);
-
 	}
 	else if (Identifier.Contains(TEXT("ClearCoat")))
 	{
