Hiding MaterialX settings and making them not editable in the pipeline window

diff --git a/Engine/Plugins/Interchange/Editor/Source/Pipelines/Private/InterchangeMaterialXPipelineCustomizations.cpp b/Engine/Plugins/Interchange/Editor/Source/Pipelines/Private/InterchangeMaterialXPipelineCustomizations.cpp
--- a/Engine/Plugins/Interchange/Editor/Source/Pipelines/Private/InterchangeMaterialXPipelineCustomizations.cpp
+++ b/Engine/Plugins/Interchange/Editor/Source/Pipelines/Private/InterchangeMaterialXPipelineCustomizations.cpp
@@ -9,14 +9,8 @@
 #include "DetailLayoutBuilder.h"
 #include "Modules/ModuleManager.h"
 #include "PropertyCustomizationHelpers.h"
-//#include "DetailWidgetRow.h"
-//#include "IDetailGroup.h"
-//#include "InterchangePipelineBase.h"
-//#include "Nodes/InterchangeBaseNode.h"
-//#include "ScopedTransaction.h"
-//#include "Styling/StyleColors.h"
-//#include "Widgets/Input/SNumericEntryBox.h"
-//#include "Widgets/Layout/SBox.h"
+#include "Widgets/Text/STextBlock.h"
+#include "Widgets/Input/SEditableText.h"
 
 TSharedRef<IDetailCustomization> FInterchangeMaterialXPipelineCustomization::MakeInstance()
 {
@@ -53,13 +47,14 @@ void FInterchangeMaterialXPipelineCustomization::CustomizeDetails(IDetailLayoutB
 
 	DetailsView->SetObject(Pipeline->MaterialXSettings);
 
-	MaterialXCategory.AddCustomRow(NSLOCTEXT("InterchangeMaterialXPipelineCustomization", "MaterialXPredefined", "Predefined Material Functions"))
+	MaterialXCategory.AddCustomRow(NSLOCTEXT("InterchangeMaterialXPipelineMaterialSubstitution::Message::Row", "MaterialXPredefinedMaterialSubstitutionMessageRow", "MaterialX Pipeline Material Substitution Message"))
+	.NameContent()
+	.VAlign(VAlign_Center)
 	[
-		SNew( SVerticalBox )
-		+ SVerticalBox::Slot()
-		[
-			DetailsView.ToSharedRef()
-		]
+		SNew(STextBlock)
+		.Font(IDetailLayoutBuilder::GetDetailFont())
+		.Text(NSLOCTEXT("InterchangeMaterialXPipelineCustomization::Message", "MaterialXPredefinedMaterialSubstitutionMessage", "See 'Project Settings > Engine > Interchange MaterialX' to edit settings."))
+		.AutoWrapText(true)
 	];
 }
 
@@ -88,7 +83,8 @@ void FInterchangeMaterialXPipelineSettingsCustomization::CustomizeDetails(IDetai
 
 	IDetailCategoryBuilder& MaterialXPredefinedCategory = DetailBuilder.EditCategory("MaterialXPredefined");
 
-	MaterialXPredefinedCategory.SetDisplayName(NSLOCTEXT("InterchangeMaterialXPipelineSettingsCustomization", "CategoryDisplayName", "MaterialX Predefined Surface Shaders"));
+	FText DisplayName{ NSLOCTEXT("InterchangeMaterialXPipelineSettingsCustomization", "CategoryDisplayName", "MaterialX Predefined Surface Shaders") };
+	MaterialXPredefinedCategory.SetDisplayName(DisplayName);
 
 	uint32 NumChildren = 0;
 	PairingsHandle->GetNumChildren(NumChildren);
@@ -98,38 +94,40 @@ void FInterchangeMaterialXPipelineSettingsCustomization::CustomizeDetails(IDetai
 		TSharedPtr<IPropertyHandle> ChildPropertyHandle = PairingsHandle->GetChildHandle(i);
 		TSharedPtr<IPropertyHandle> KeyPropertyHandle = ChildPropertyHandle->GetKeyHandle();
 
+		UEnum* Enum = CastField<FEnumProperty>(KeyPropertyHandle->GetProperty())->GetEnum();
+
+		FText EnumString;
+		KeyPropertyHandle->GetValueAsDisplayText(EnumString);
+
 		OnShouldFilterAssetFunc OnShouldFilterAsset = &FInterchangeMaterialXPipelineSettingsCustomization::OnShouldFilterAsset;
 
-		if(KeyPropertyHandle.IsValid())
+		uint8 EnumValue = 0;
+		KeyPropertyHandle->GetValue(EnumValue);
+		if(EnumString.ToString() == TEXT("Standard Surface"))
 		{
-			FString Str;
-			KeyPropertyHandle->GetValueAsDisplayString(Str);
-			if(Str == TEXT("Standard Surface"))
-			{
-				OnShouldFilterAsset = &FInterchangeMaterialXPipelineSettingsCustomization::OnShouldFilterAssetStandardSurface;
-			}
-			else if(Str == TEXT("Standard Surface Transmission"))
-			{
-				OnShouldFilterAsset = &FInterchangeMaterialXPipelineSettingsCustomization::OnShouldFilterAssetStandardSurfaceTransmission;
-			}
-			else if(Str == TEXT("Surface Unlit"))
-			{
-				OnShouldFilterAsset = &FInterchangeMaterialXPipelineSettingsCustomization::OnShouldFilterAssetSurfaceUnlit;
-			}
-			else if(Str == TEXT("Usd Preview Surface"))
-			{
-				OnShouldFilterAsset = &FInterchangeMaterialXPipelineSettingsCustomization::OnShouldFilterAssetUsdPreviewSurface;
-			}
+			OnShouldFilterAsset = &FInterchangeMaterialXPipelineSettingsCustomization::OnShouldFilterAssetStandardSurface;
+		}
+		else if(EnumString.ToString() == TEXT("Standard Surface Transmission"))
+		{
+			OnShouldFilterAsset = &FInterchangeMaterialXPipelineSettingsCustomization::OnShouldFilterAssetStandardSurfaceTransmission;
+		}
+		else if(EnumString.ToString() == TEXT("Surface Unlit"))
+		{
+			OnShouldFilterAsset = &FInterchangeMaterialXPipelineSettingsCustomization::OnShouldFilterAssetSurfaceUnlit;
+		}
+		else if(EnumString.ToString() == TEXT("Usd Preview Surface"))
+		{
+			OnShouldFilterAsset = &FInterchangeMaterialXPipelineSettingsCustomization::OnShouldFilterAssetUsdPreviewSurface;
 		}
 
-		IDetailPropertyRow& SurfaceShaderTypeRow = DetailBuilder.AddPropertyToCategory(ChildPropertyHandle.ToSharedRef());
-		SurfaceShaderTypeRow.ShowPropertyButtons(false);
-
-		FDetailWidgetRow& DetailWidgetRow = SurfaceShaderTypeRow.CustomWidget();
+		FDetailWidgetRow& DetailWidgetRow = MaterialXPredefinedCategory.AddCustomRow(DisplayName);
 		DetailWidgetRow
 		.NameContent()
 		[
-			SNullWidget::NullWidget
+			SNew(SEditableText)
+			.Text(EnumString)
+			.IsReadOnly(true)
+			.ToolTipText(Enum->GetToolTipTextByIndex(EnumValue))
 		]
 		.ValueContent()
 		.MinDesiredWidth(1.0f)
diff --git a/Engine/Plugins/Interchange/Runtime/Source/Pipelines/Private/InterchangeMaterialXPipeline.cpp b/Engine/Plugins/Interchange/Runtime/Source/Pipelines/Private/InterchangeMaterialXPipeline.cpp
--- a/Engine/Plugins/Interchange/Runtime/Source/Pipelines/Private/InterchangeMaterialXPipeline.cpp
+++ b/Engine/Plugins/Interchange/Runtime/Source/Pipelines/Private/InterchangeMaterialXPipeline.cpp
@@ -12,275 +12,6 @@
 #include "Materials/MaterialFunction.h"
 #include "Misc/PackageName.h"
 
-#if WITH_EDITOR
-namespace UE::Interchange::InterchangeMaterialXPipeline::Private
-{
-	namespace StandardSurface
-	{
-		const FName Name = TEXT("standard_surface");
-
-		namespace Parameters
-		{
-			const FName Base = TEXT("base");
-			const FName BaseColor = TEXT("base_color");
-			const FName DiffuseRoughness = TEXT("diffuse_roughness");
-			const FName Metalness = TEXT("metalness");
-			const FName Specular = TEXT("specular");
-			const FName SpecularColor = TEXT("specular_color");
-			const FName SpecularRoughness = TEXT("specular_roughness");
-			const FName SpecularIOR = TEXT("specular_IOR");
-			const FName SpecularAnisotropy = TEXT("specular_anisotropy");
-			const FName SpecularRotation = TEXT("specular_rotation");
-			const FName Transmission = TEXT("transmission");
-			const FName TransmissionColor = TEXT("transmission_color");
-			const FName TransmissionDepth = TEXT("transmission_depth");
-			const FName TransmissionScatter = TEXT("transmission_scatter");
-			const FName TransmissionScatterAnisotropy = TEXT("transmission_scatter_anisotropy");
-			const FName TransmissionDispersion = TEXT("transmission_dispersion");
-			const FName TransmissionExtraRoughness = TEXT("transmission_extra_roughness");
-			const FName Subsurface = TEXT("subsurface");
-			const FName SubsurfaceColor = TEXT("subsurface_color");
-			const FName SubsurfaceRadius = TEXT("subsurface_radius");
-			const FName SubsurfaceScale = TEXT("subsurface_scale");
-			const FName SubsurfaceAnisotropy = TEXT("subsurface_anisotropy");
-			const FName Sheen = TEXT("sheen");
-			const FName SheenColor = TEXT("sheen_color");
-			const FName SheenRoughness = TEXT("sheen_roughness");
-			const FName Coat = TEXT("coat");
-			const FName CoatColor = TEXT("coat_color");
-			const FName CoatRoughness = TEXT("coat_roughness");
-			const FName CoatAnisotropy = TEXT("coat_anisotropy");
-			const FName CoatRotation = TEXT("coat_rotation");
-			const FName CoatIOR = TEXT("coat_IOR");
-			const FName CoatNormal = TEXT("coat_normal");
-			const FName CoatAffectColor = TEXT("coat_affect_color");
-			const FName CoatAffectRoughness = TEXT("coat_affect_roughness");
-			const FName ThinFilmThickness = TEXT("thin_film_thickness");
-			const FName ThinFilmIOR = TEXT("thin_film_IOR");
-			const FName Emission = TEXT("emission");
-			const FName EmissionColor = TEXT("emission_color");
-			const FName Opacity = TEXT("opacity");
-			const FName ThinWalled = TEXT("thin_walled");
-			const FName Normal = TEXT("normal");
-			const FName Tangent = TEXT("tangent");
-		}
-	}
-
-	namespace SurfaceUnlit
-	{
-		const FName Name = TEXT("surface_unlit");
-
-		namespace Parameters
-		{
-			const FName Emission = TEXT("emission");
-			const FName EmissionColor = TEXT("emission_color");
-			const FName Transmission = TEXT("transmission");
-			const FName TransmissionColor = TEXT("transmission_color");
-			const FName Opacity = TEXT("opacity");
-		}
-	}
-
-	TSet<FName> StandardSurfaceInputs
-	{
-		StandardSurface::Parameters::Base,
-		StandardSurface::Parameters::BaseColor,
-		StandardSurface::Parameters::DiffuseRoughness,
-		StandardSurface::Parameters::Metalness,
-		StandardSurface::Parameters::Specular,
-		StandardSurface::Parameters::SpecularRoughness,
-		StandardSurface::Parameters::SpecularIOR,
-		StandardSurface::Parameters::SpecularAnisotropy,
-		StandardSurface::Parameters::SpecularRotation,
-		StandardSurface::Parameters::Subsurface,
-		StandardSurface::Parameters::SubsurfaceColor,
-		StandardSurface::Parameters::SubsurfaceRadius,
-		StandardSurface::Parameters::SubsurfaceScale,
-		StandardSurface::Parameters::Sheen,
-		StandardSurface::Parameters::SheenColor,
-		StandardSurface::Parameters::SheenRoughness,
-		StandardSurface::Parameters::Coat,
-		StandardSurface::Parameters::CoatColor,
-		StandardSurface::Parameters::CoatRoughness,
-		StandardSurface::Parameters::CoatNormal,
-		StandardSurface::Parameters::ThinFilmThickness,
-		StandardSurface::Parameters::Emission,
-		StandardSurface::Parameters::EmissionColor,
-		StandardSurface::Parameters::Normal,
-		StandardSurface::Parameters::Tangent
-	};
-
-	TSet<FName> StandardSurfaceOutputs
-	{
-		TEXT("Base Color"), // MX_StandardSurface has BaseColor with a whitespace, this should be fixed in further release
-		UE::Interchange::Materials::PBR::Parameters::Metallic,
-		UE::Interchange::Materials::PBR::Parameters::Specular,
-		UE::Interchange::Materials::PBR::Parameters::Roughness,
-		UE::Interchange::Materials::PBR::Parameters::Anisotropy,
-		UE::Interchange::Materials::PBR::Parameters::EmissiveColor,
-		UE::Interchange::Materials::PBR::Parameters::Opacity,
-		UE::Interchange::Materials::PBR::Parameters::Normal,
-		UE::Interchange::Materials::PBR::Parameters::Tangent,
-		UE::Interchange::Materials::Sheen::Parameters::SheenRoughness,
-		UE::Interchange::Materials::Sheen::Parameters::SheenColor,
-		UE::Interchange::Materials::Subsurface::Parameters::SubsurfaceColor,
-		UE::Interchange::Materials::ClearCoat::Parameters::ClearCoat,
-		UE::Interchange::Materials::ClearCoat::Parameters::ClearCoatRoughness,
-		UE::Interchange::Materials::ClearCoat::Parameters::ClearCoatNormal
-	};
-
-	TSet<FName> TransmissionSurfaceInputs
-	{
-		StandardSurface::Parameters::Base,
-		StandardSurface::Parameters::BaseColor,
-		StandardSurface::Parameters::DiffuseRoughness,
-		StandardSurface::Parameters::Metalness,
-		StandardSurface::Parameters::Specular,
-		StandardSurface::Parameters::SpecularRoughness,
-		StandardSurface::Parameters::SpecularIOR,
-		StandardSurface::Parameters::SpecularAnisotropy,
-		StandardSurface::Parameters::SpecularRotation,
-		StandardSurface::Parameters::Transmission,
-		StandardSurface::Parameters::TransmissionColor,
-		StandardSurface::Parameters::TransmissionDepth,
-		StandardSurface::Parameters::TransmissionScatter,
-		StandardSurface::Parameters::TransmissionScatterAnisotropy,
-		StandardSurface::Parameters::TransmissionDispersion,
-		StandardSurface::Parameters::TransmissionExtraRoughness,
-		StandardSurface::Parameters::Subsurface,
-		StandardSurface::Parameters::SubsurfaceColor,
-		StandardSurface::Parameters::SubsurfaceRadius,
-		StandardSurface::Parameters::SubsurfaceScale,
-		StandardSurface::Parameters::Sheen,
-		StandardSurface::Parameters::SheenColor,
-		StandardSurface::Parameters::SheenRoughness,
-		StandardSurface::Parameters::Coat,
-		StandardSurface::Parameters::CoatColor,
-		StandardSurface::Parameters::CoatRoughness,
-		StandardSurface::Parameters::CoatNormal,
-		StandardSurface::Parameters::ThinFilmThickness,
-		StandardSurface::Parameters::Emission,
-		StandardSurface::Parameters::EmissionColor,
-		StandardSurface::Parameters::Normal,
-	};
-
-	TSet<FName> TransmissionSurfaceOutputs
-	{
-		UE::Interchange::Materials::PBR::Parameters::BaseColor,
-		UE::Interchange::Materials::PBR::Parameters::Metallic,
-		UE::Interchange::Materials::PBR::Parameters::Specular,
-		UE::Interchange::Materials::PBR::Parameters::Roughness,
-		UE::Interchange::Materials::PBR::Parameters::Anisotropy,
-		UE::Interchange::Materials::PBR::Parameters::EmissiveColor,
-		UE::Interchange::Materials::PBR::Parameters::Opacity,
-		UE::Interchange::Materials::PBR::Parameters::Normal,
-		UE::Interchange::Materials::PBR::Parameters::Tangent,
-		UE::Interchange::Materials::PBR::Parameters::Refraction,
-		UE::Interchange::Materials::ThinTranslucent::Parameters::TransmissionColor
-	};
-
-	TSet<FName> SurfaceUnlitInputs
-	{
-		SurfaceUnlit::Parameters::Emission,
-		SurfaceUnlit::Parameters::EmissionColor,
-		SurfaceUnlit::Parameters::Transmission,
-		SurfaceUnlit::Parameters::TransmissionColor,
-		SurfaceUnlit::Parameters::Opacity
-	};
-
-	TSet<FName> SurfaceUnlitOutputs
-	{
-		UE::Interchange::Materials::Common::Parameters::EmissiveColor,
-		UE::Interchange::Materials::Common::Parameters::Opacity,
-	};
-
-	bool ShouldFilterAssets(UMaterialFunction* Asset, const TSet<FName>& Inputs, const TSet<FName>& Outputs)
-	{
-		int32 InputMatches = 0;
-		int32 OutputMatches = 0;
-
-		if (Asset != nullptr)
-		{
-			TArray<FFunctionExpressionInput> ExpressionInputs;
-			TArray<FFunctionExpressionOutput> ExpressionOutputs;
-			Asset->GetInputsAndOutputs(ExpressionInputs, ExpressionOutputs);
-
-			for (const FFunctionExpressionInput& ExpressionInput : ExpressionInputs)
-			{
-				if (Inputs.Find(ExpressionInput.Input.InputName))
-				{
-					InputMatches++;
-				}
-			}
-
-			for (const FFunctionExpressionOutput& ExpressionOutput : ExpressionOutputs)
-			{
-				if (Outputs.Find(ExpressionOutput.Output.OutputName))
-				{
-					OutputMatches++;
-				}
-			}
-		}
-
-		// we allow at least one input of the same name, but we should have exactly the same outputs
-		return !(InputMatches > 0 && OutputMatches == Outputs.Num());
-	}
-
-	bool AreRequiredPackagesLoaded(TMap<EInterchangeMaterialXShaders, FSoftObjectPath>& SurfaceShaderMapping)
-	{
-		auto ArePackagesLoaded = [&](const TMap<EInterchangeMaterialXShaders, FSoftObjectPath>& ObjectPaths)
-		{
-			bool bAllLoaded = true;
-
-			for (const TPair<EInterchangeMaterialXShaders, FSoftObjectPath>& Pair : ObjectPaths)
-			{
-				const FSoftObjectPath& ObjectPath = Pair.Get<1>();
-
-				if (!ObjectPath.ResolveObject())
-				{
-					FString PackagePath = ObjectPath.GetLongPackageName();
-					if (FPackageName::DoesPackageExist(PackagePath))
-					{
-						UObject* Asset = ObjectPath.TryLoad();
-						if (!Asset)
-						{
-							UE_LOG(LogInterchangePipeline, Warning, TEXT("Couldn't load %s"), *PackagePath);
-							bAllLoaded = false;
-						}
-#if WITH_EDITOR
-						else
-						{
-							if (Pair.Get<EInterchangeMaterialXShaders>() == EInterchangeMaterialXShaders::StandardSurface)
-							{
-								bAllLoaded = !ShouldFilterAssets(Cast<UMaterialFunction>(Asset), StandardSurfaceInputs, StandardSurfaceOutputs);
-							}
-							else if (Pair.Get<EInterchangeMaterialXShaders>() == EInterchangeMaterialXShaders::StandardSurfaceTransmission)
-							{
-								bAllLoaded = !ShouldFilterAssets(Cast<UMaterialFunction>(Asset), TransmissionSurfaceInputs, TransmissionSurfaceOutputs);
-							}
-							else if (Pair.Get<EInterchangeMaterialXShaders>() == EInterchangeMaterialXShaders::SurfaceUnlit)
-							{
-								bAllLoaded = !ShouldFilterAssets(Cast<UMaterialFunction>(Asset), SurfaceUnlitInputs, SurfaceUnlitOutputs);
-							}
-
-						}
-#endif // WITH_EDITOR
-					}
-					else
-					{
-						UE_LOG(LogInterchangePipeline, Warning, TEXT("Couldn't find %s"), *PackagePath);
-						bAllLoaded = false;
-					}
-				}
-			}
-
-			return bAllLoaded;
-		};
-
-		return ArePackagesLoaded(SurfaceShaderMapping);
-	}
-}
-#endif
-
 UInterchangeMaterialXPipeline::UInterchangeMaterialXPipeline()
 	: MaterialXSettings(UMaterialXPipelineSettings::StaticClass()->GetDefaultObject< UMaterialXPipelineSettings>())
 {
diff --git a/Engine/Plugins/Interchange/Runtime/Source/Pipelines/Public/InterchangeMaterialXPipeline.h b/Engine/Plugins/Interchange/Runtime/Source/Pipelines/Public/InterchangeMaterialXPipeline.h
--- a/Engine/Plugins/Interchange/Runtime/Source/Pipelines/Public/InterchangeMaterialXPipeline.h
+++ b/Engine/Plugins/Interchange/Runtime/Source/Pipelines/Public/InterchangeMaterialXPipeline.h
@@ -83,7 +83,7 @@ class INTERCHANGEPIPELINES_API UInterchangeMaterialXPipeline : public UInterchan
 	UInterchangeMaterialXPipeline();
 
 public:
-	UPROPERTY(VisibleAnywhere, BlueprintReadWrite, Category = "MaterialX")
+	
 	TObjectPtr<UMaterialXPipelineSettings> MaterialXSettings;
 
 protected:
