[Interchange] Improve quality of material created by OBJ translator

Emissive - from color "Ke" and texture "map_Ke")
Roughness - from Specular Exponent "Ns", from Pbr roughness "Pr", roughness texture "map_Pr"
Opacity - from opacity scalar "d", texture "map_d", transparency "Tr", transparency texture "map_Tr"
Specular - from specular texture "map_Ks"
Normal - from normal texture "norm"
Bump - bump texture "bump" and "map_bump", bump multiplier from "-bm" map option
Metallic - from metallic "Pm" and texture "map_Pm"
IOR - from ior scalar"Ni"
TransmittanceColor - from transmission filter"Tf"
Clear Coat - from clear coat thickness "Pc" and cc roughness "Pcr"
Anisotropy - from anisotropy "aniso"
Illumination modes - from "illum"

All textures with corresponding color/scalar are modulated by that constant. E.g. Emissive color resulted is "map_Ke"*"Ke"

Texture file names are now correctly retrieved when texture options are present.

diff --git a/Engine/Plugins/Interchange/Runtime/Source/Import/Private/Mesh/InterchangeOBJTranslator.cpp b/Engine/Plugins/Interchange/Runtime/Source/Import/Private/Mesh/InterchangeOBJTranslator.cpp
--- a/Engine/Plugins/Interchange/Runtime/Source/Import/Private/Mesh/InterchangeOBJTranslator.cpp
+++ b/Engine/Plugins/Interchange/Runtime/Source/Import/Private/Mesh/InterchangeOBJTranslator.cpp
@@ -26,6 +26,7 @@
 #include "UObject/GCObjectScopeGuard.h"
 #include "UVMapSettings.h"
 
+
 #include UE_INLINE_GENERATED_CPP_BY_NAME(InterchangeOBJTranslator)
 
 static bool GInterchangeEnableOBJImport = true;
@@ -98,20 +99,50 @@ struct FObjData
 	/** All the groups defined by the .obj */
 	TMap<FString, FGroupData> Groups;
 
+	struct FTextureParameter
+	{
+		FString Path;
+		float BumpMultiplier = 1.0f;  // Default to 1.0f, can be negative
+		bool IsEmpty() const
+		{
+			return Path.IsEmpty();
+		}
+	};
+
 	/** Defines the data for a material */
 	struct FMaterialData
 	{
-		// These are the defaults for unspecified .mtl attributes
-		int32 IlluminationModel = 0;
-		FVector3f DiffuseColor = FVector3f(0.8f, 0.8f, 0.8f);
-		FVector3f AmbientColor = FVector3f(0.2f, 0.2f, 0.2f);
-		FVector3f SpecularColor = FVector3f(1.0f, 1.0f, 1.0f);
-		FString DiffuseTexture;
-		FString AmbientTexture;
-		FString SpecularTexture;
-		FVector3f TransmissionFilter = FVector3f(1.0f, 1.0f, 1.0f);
-		float RefractiveIndex = 1.0f;
-		float Transparency = 1.0f;
+		// Unspecified properties designated by negative values
+		int32 IlluminationModel = 2;  // Default illumination model to Phong for non-Pbr materials
+		FVector3f DiffuseColor = -FVector3f::One();
+		FVector3f AmbientColor = -FVector3f::One();
+		FVector3f SpecularColor = -FVector3f::One();
+		FVector3f EmissiveColor = -FVector3f::One();
+		FTextureParameter DiffuseTexture;
+		FTextureParameter AmbientTexture;
+		FTextureParameter SpecularTexture;
+		FTextureParameter EmissiveTexture;
+		FTextureParameter OpacityTexture;
+		FTextureParameter TransparencyTexture;
+		FTextureParameter SpecularExponentTexture;
+		FVector3f TransmissionFilter = -FVector3f::One();
+		float IndexOfRefraction = -1.f;
+		float Opacity = -1.f;
+		float Transparency = -1.f;
+
+		float SpecularExponent = -1.f;
+
+		float Roughness = -1.f;
+		FTextureParameter RoughnessTexture;
+
+		float Metallic = -1.f;
+		FTextureParameter MetallicTexture;
+		FTextureParameter BumpmapTexture;
+		FTextureParameter NormalmapTexture;
+
+		float ClearCoatThickness;
+		float ClearCoatRoughness;
+		float Anisotropy;
 	};
 
 	/** All the materials defined by the .obj */
@@ -437,7 +468,7 @@ FMeshDescription FObjData::MakeMeshDescriptionForGroup(const FString& GroupName)
 namespace UE {
 namespace Interchange {
 
-namespace ObjParser
+namespace ObjParserUtils
 {
 	using FFunctionType = TFunction<bool(FObjData&, FStringView)>;
 	using FKeywordMap = TMap<FString, FFunctionType>;
@@ -740,7 +771,7 @@ namespace ObjParser
 		FFileHelper::LoadFileToStringWithLineVisitor(Filename,
 			[&ObjData, &KeywordMap, &bSuccess](FStringView Line)
 			{
-				bool bResult = ObjParser::ParseLine(ObjData, Line, KeywordMap);
+				bool bResult = ObjParserUtils::ParseLine(ObjData, Line, KeywordMap);
 				if (bResult == false)
 				{
 					bSuccess = false;
@@ -796,7 +827,7 @@ namespace ObjParser
 
 		if (!Line.IsEmpty())
 		{
-			INTERCHANGE_OBJ_TRANSLATOR_LOG_ONCE(LogTemp, Warning, TEXT("Interchange Obj translator: Unexpected extra arguments on illum keyword"));
+			INTERCHANGE_OBJ_TRANSLATOR_LOG_ONCE(LogTemp, Warning, TEXT("Interchange Obj translator: Unexpected extra arguments on material property"));
 		}
 
 		return true;
@@ -808,7 +839,7 @@ namespace ObjParser
 	{
 		if (Line.IsEmpty())
 		{
-			INTERCHANGE_OBJ_TRANSLATOR_LOG_ONCE(LogTemp, Warning, TEXT("Interchange Obj translator: Missing texture filename on map_Kd keyword"));
+			INTERCHANGE_OBJ_TRANSLATOR_LOG_ONCE(LogTemp, Warning, TEXT("Interchange Obj translator: Expecting more parameters(including texture filename) on a map keyword"));
 		}
 
 		if (FObjData::FMaterialData* MaterialData = ObjData.Materials.Find(ObjData.MaterialBeingDefined))
@@ -824,6 +855,52 @@ namespace ObjParser
 		return true;
 	}
 
+	template <FObjData::FTextureParameter FObjData::FMaterialData::* Property>
+	static bool ParseTextureStatement(FObjData& ObjData, FStringView Line)
+	{
+		if (Line.IsEmpty())
+		{
+			INTERCHANGE_OBJ_TRANSLATOR_LOG_ONCE(LogTemp, Warning, TEXT("Interchange Obj translator: Expecting more parameters(including texture filename) on a map keyword"));
+		}
+
+		if (FObjData::FMaterialData* MaterialData = ObjData.Materials.Find(ObjData.MaterialBeingDefined))
+		{
+			FObjData::FTextureParameter& Parameter = (MaterialData->*Property);
+
+			FStringView FileName;
+
+			FStringView LastToken;
+			while (true)
+			{
+				FStringView Token = GetToken(Line);
+				if (Token.IsEmpty())
+				{
+					FileName = LastToken;
+					break;
+				}
+
+				LastToken = Token;
+
+				if (Token == TEXT("-bm"))
+				{
+					float BumpMultiplier;
+					if (GetFloat(Line, BumpMultiplier))
+					{
+						Parameter.BumpMultiplier = BumpMultiplier;
+					}
+				}
+			}
+
+			Parameter.Path = FPaths::GetPath(ObjData.ObjFilename) / FString(FileName);
+		}
+		else
+		{
+			INTERCHANGE_OBJ_TRANSLATOR_LOG_ONCE(LogTemp, Warning, TEXT("Interchange Obj translator: Missing newmtl keyword"));
+			return false;
+		}
+
+		return true;
+	}
 
 	static bool ParseMaterialLib(FObjData& ObjData, FStringView Line)
 	{
@@ -834,12 +911,34 @@ namespace ObjParser
 			{ TEXT("Kd"),      ParseMaterialProperty<FVector3f, &FObjData::FMaterialData::DiffuseColor> },
 			{ TEXT("Ka"),      ParseMaterialProperty<FVector3f, &FObjData::FMaterialData::AmbientColor> },
 			{ TEXT("Ks"),      ParseMaterialProperty<FVector3f, &FObjData::FMaterialData::SpecularColor> },
+			{ TEXT("Ke"),      ParseMaterialProperty<FVector3f, &FObjData::FMaterialData::EmissiveColor> },
 			{ TEXT("Tf"),      ParseMaterialProperty<FVector3f, &FObjData::FMaterialData::TransmissionFilter> },
-			{ TEXT("Ni"),      ParseMaterialProperty<float,     &FObjData::FMaterialData::RefractiveIndex> },
-			{ TEXT("d"),       ParseMaterialProperty<float,     &FObjData::FMaterialData::Transparency> },
-			{ TEXT("map_Kd"),  ParseMaterialPath<&FObjData::FMaterialData::DiffuseTexture> },
-			{ TEXT("map_Ka"),  ParseMaterialPath<&FObjData::FMaterialData::AmbientTexture> },
-			{ TEXT("map_Ks"),  ParseMaterialPath<&FObjData::FMaterialData::SpecularTexture> }
+			{ TEXT("Ni"),      ParseMaterialProperty<float,     &FObjData::FMaterialData::IndexOfRefraction> },
+			{ TEXT("Ns"),      ParseMaterialProperty<float,     &FObjData::FMaterialData::SpecularExponent> },
+			{ TEXT("d"),       ParseMaterialProperty<float,     &FObjData::FMaterialData::Opacity> },
+			{ TEXT("Tr"),      ParseMaterialProperty<float,     &FObjData::FMaterialData::Transparency> },
+			{ TEXT("map_Kd"),  ParseTextureStatement<&FObjData::FMaterialData::DiffuseTexture> },
+			{ TEXT("map_Ka"),  ParseTextureStatement<&FObjData::FMaterialData::AmbientTexture> },
+			{ TEXT("map_Ks"),  ParseTextureStatement<&FObjData::FMaterialData::SpecularTexture> },
+			{ TEXT("map_Ke"),  ParseTextureStatement<&FObjData::FMaterialData::EmissiveTexture> },
+			{ TEXT("map_d"),   ParseTextureStatement<&FObjData::FMaterialData::OpacityTexture> },
+			{ TEXT("map_Tr"),  ParseTextureStatement<&FObjData::FMaterialData::TransparencyTexture> },
+			{ TEXT("map_Ns"),  ParseTextureStatement<&FObjData::FMaterialData::SpecularExponentTexture> },
+			{ TEXT("map_bump"),ParseTextureStatement<&FObjData::FMaterialData::BumpmapTexture> },
+			{ TEXT("bump"),	   ParseTextureStatement<&FObjData::FMaterialData::BumpmapTexture> },
+
+			{ TEXT("Pr"),      ParseMaterialProperty<float,     &FObjData::FMaterialData::Roughness> },
+			{ TEXT("map_Pr"),  ParseTextureStatement<&FObjData::FMaterialData::RoughnessTexture> },
+
+			{ TEXT("Pm"),      ParseMaterialProperty<float,     &FObjData::FMaterialData::Metallic> },
+			{ TEXT("map_Pm"),  ParseTextureStatement<&FObjData::FMaterialData::MetallicTexture> },
+
+			{ TEXT("Pc"),      ParseMaterialProperty<float,     &FObjData::FMaterialData::ClearCoatThickness> },
+			{ TEXT("Pcr"),     ParseMaterialProperty<float,     &FObjData::FMaterialData::ClearCoatRoughness> },
+			{ TEXT("aniso"),   ParseMaterialProperty<float,     &FObjData::FMaterialData::Anisotropy> },
+
+			{ TEXT("norm"),    ParseTextureStatement<&FObjData::FMaterialData::NormalmapTexture> },
+
 		};
 
 		// ZBrush doesn't surround filenames with spaces in quotes, so just take the entire remainder of the line as the filename,
@@ -871,41 +970,226 @@ namespace ObjParser
 
 		return bSuccess;
 	}
+} // namespace ObjParserUtils
 
+namespace ObjTranslatorUtils
+{
+	static bool IsColorInitialized(const FVector3f& Color)
+	{
+		return Color.GetMin() >= 0.f;
+	};
 
-} // namespace ObjParser
-} // namespace Interchange
-} // namespace UE
+	static bool IsScalarInitialized(const float& Scalar)
+	{
+		return Scalar >= 0.f;
+	};
 
+	/**
+	 * Build and return a UID name for a mesh node.
+	 * @todo: move this to be a static method on the MeshNode class itself?
+	 */
+	static FString MakeMeshNodeUid(const FString& Name)
+	{
+		return TEXT("\\Mesh\\") + (Name.IsEmpty() ? FString(TEXT("Null")) : Name);
+	}
 
-/**
- * Build and return a UID name for a mesh node.
- * @todo: move this to be a static method on the MeshNode class itself?
- */
-static FString MakeMeshNodeUid(const FString& Name)
-{
-	return TEXT("\\Mesh\\") + (Name.IsEmpty() ? FString(TEXT("Null")) : Name);
-}
 
+	/**
+	 * Build and return a UID name for a shader graph node.
+	 * @todo: move this to be a static method on the ShaderGraphNode class itself?
+	 */
+	static FString MakeShaderGraphNodeName(const FString& Name)
+	{
+		return TEXT("MAT_") + (Name.IsEmpty() ? FString(TEXT("Null")) : Name);
+	}
 
-/**
- * Build and return a UID name for a shader graph node.
- * @todo: move this to be a static method on the ShaderGraphNode class itself?
- */
-static FString MakeShaderGraphNodeName(const FString& Name)
-{
-	return TEXT("MAT_") + (Name.IsEmpty() ? FString(TEXT("Null")) : Name);
-}
 
+	/**
+	 * Build and return a UID name for a texture node.
+	 * @todo: move this to be a static method on the TextureNode class itself?
+	 */
+	static FString MakeTextureNodeName(const FString& Name)
+	{
+		return TEXT("TEX_") + (Name.IsEmpty() ? FString(TEXT("Null")) : Name);
+	}
 
-/**
- * Build and return a UID name for a texture node.
- * @todo: move this to be a static method on the TextureNode class itself?
- */
-static FString MakeTextureNodeName(const FString& Name)
-{
-	return TEXT("TEX_") + (Name.IsEmpty() ? FString(TEXT("Null")) : Name);
-}
+
+	static const UInterchangeTexture2DNode* CreateTexture2DNode(UInterchangeBaseNodeContainer& BaseNodeContainer, FString TexturePath, FString& InOutTextureName) 
+	{
+		FString TextureName = InOutTextureName;
+		for (int32 Index = 0;;Index ++)
+		{
+			FString NormalizedTexturePath(TexturePath);
+			FPaths::NormalizeFilename(NormalizedTexturePath);
+
+			FString TextureNodeUid = UInterchangeTextureNode::MakeNodeUid(TextureName);
+
+			UInterchangeTexture2DNode* TextureNode = const_cast<UInterchangeTexture2DNode*>(Cast<const UInterchangeTexture2DNode>(BaseNodeContainer.GetNode(TextureNodeUid)));
+			if (!TextureNode)
+			{
+				TextureNode = UInterchangeTexture2DNode::Create(&BaseNodeContainer, TextureName);
+				TextureNode->SetPayLoadKey(NormalizedTexturePath);
+
+				InOutTextureName = TextureName;
+				return TextureNode;
+			}
+
+			// Check if existing texture node was created for the same image file 
+			if (NormalizedTexturePath == TextureNode->GetPayLoadKey())
+			{
+				InOutTextureName = TextureName;
+				return TextureNode;
+			}
+
+			// Generate different texture name(and therefore NodeUid)
+			TextureName  = InOutTextureName + FString::FromInt(Index);
+		}
+	}
+
+	static UInterchangeShaderNode* CreateMaterialTextureSampleNode(UInterchangeBaseNodeContainer& BaseNodeContainer, UInterchangeShaderGraphNode* ShaderGraphNode, const FObjData::FTextureParameter& TextureParameter) 
+	{
+		if (!TextureParameter.IsEmpty())
+		{
+			// @todo: better way to handle textures than implementing the TexturePayload interface in the OBJTranslator
+			// and creating a temporary texture translator there.
+
+			using namespace UE::Interchange::Materials::Standard::Nodes;
+
+			FString TextureName = MakeTextureNodeName(FPaths::GetBaseFilename(TextureParameter.Path));
+
+			if (const UInterchangeTexture2DNode* TextureNode = CreateTexture2DNode(BaseNodeContainer, TextureParameter.Path, TextureName))
+			{
+				UInterchangeShaderNode* TextureSampleShader = UInterchangeShaderNode::Create(&BaseNodeContainer, TextureName, ShaderGraphNode->GetUniqueID());
+				TextureSampleShader->SetCustomShaderType(TextureSample::Name.ToString());
+				TextureSampleShader->AddStringAttribute(UInterchangeShaderPortsAPI::MakeInputValueKey(TextureSample::Inputs::Texture.ToString()), TextureNode->GetUniqueID());
+				return TextureSampleShader;
+			}
+		}
+		return nullptr;
+	}
+
+	// multiply texture by color, handling cases when any of the arguments is missing
+	static bool AddTexturedColoredInput(UInterchangeBaseNodeContainer& BaseNodeContainer, UInterchangeShaderGraphNode* ShaderGraphNode, const FString& InputType, const FVector3f& Color, const FObjData::FTextureParameter& TexturePath) 
+	{
+		UInterchangeShaderNode* TextureSampleShader = CreateMaterialTextureSampleNode(BaseNodeContainer, ShaderGraphNode, TexturePath);
+
+		if (TextureSampleShader && IsColorInitialized(Color))
+		{
+			FString MapName = Materials::Phong::Parameters::SpecularColor.ToString();
+
+			using namespace UE::Interchange::Materials::Standard::Nodes;
+
+			const FString MultiplierNodeName = MapName + TEXT("Multiply");
+			UInterchangeShaderNode* MultiplierNode = UInterchangeShaderNode::Create(&BaseNodeContainer, MultiplierNodeName, ShaderGraphNode->GetUniqueID());
+			MultiplierNode->SetCustomShaderType(Multiply::Name.ToString());
+
+			UInterchangeShaderPortsAPI::ConnectDefaultOuputToInput(MultiplierNode, Multiply::Inputs::A.ToString(), TextureSampleShader->GetUniqueID());
+			MultiplierNode->AddLinearColorAttribute(UInterchangeShaderPortsAPI::MakeInputValueKey( Multiply::Inputs::B.ToString() ), Color);
+
+			UInterchangeShaderPortsAPI::ConnectDefaultOuputToInput(ShaderGraphNode, InputType, MultiplierNode->GetUniqueID());
+		}
+		else if (TextureSampleShader)
+		{
+			UInterchangeShaderPortsAPI::ConnectDefaultOuputToInput(ShaderGraphNode, InputType, TextureSampleShader->GetUniqueID());
+		}
+		else if (IsColorInitialized(Color))
+		{
+			ShaderGraphNode->AddLinearColorAttribute(UInterchangeShaderPortsAPI::MakeInputValueKey(InputType), FLinearColor(Color));
+		}
+		else
+		{
+			return false;
+		}
+		return true;
+	}
+
+	static UInterchangeShaderNode* MakeTexturedWeighted(UInterchangeBaseNodeContainer& BaseNodeContainer, UInterchangeShaderGraphNode* ShaderGraphNode, FString MapName, const float& Weight, UInterchangeShaderNode* TextureSampleShader) 
+	{
+		if (!TextureSampleShader)
+		{
+			return nullptr;
+		}
+
+		if (IsScalarInitialized(Weight))
+		{
+			using namespace UE::Interchange::Materials::Standard::Nodes;
+
+			const FString MultiplierNodeName = MapName + TEXT("Multiply");
+			UInterchangeShaderNode* MultiplierNode = UInterchangeShaderNode::Create(&BaseNodeContainer, MultiplierNodeName, ShaderGraphNode->GetUniqueID());
+			MultiplierNode->SetCustomShaderType(Multiply::Name.ToString());
+
+			UInterchangeShaderPortsAPI::ConnectDefaultOuputToInput(MultiplierNode, Multiply::Inputs::A.ToString(), TextureSampleShader->GetUniqueID());
+			MultiplierNode->AddFloatAttribute(UInterchangeShaderPortsAPI::MakeInputValueKey( Multiply::Inputs::B.ToString() ), Weight);
+			return MultiplierNode;
+		}
+		return TextureSampleShader;
+	}
+
+	static bool AddTexturedWeightedInput(UInterchangeBaseNodeContainer& BaseNodeContainer, UInterchangeShaderGraphNode* ShaderGraphNode, const FString& InputType, const float& Weight, UInterchangeShaderNode* TextureSampleShader) 
+	{
+		if (UInterchangeShaderNode* TextureSampleWeighted = MakeTexturedWeighted(BaseNodeContainer, ShaderGraphNode, InputType, Weight, TextureSampleShader))
+		{
+			UInterchangeShaderPortsAPI::ConnectDefaultOuputToInput(ShaderGraphNode, InputType, TextureSampleWeighted->GetUniqueID());
+		}
+		else if (IsScalarInitialized(Weight))
+		{
+			ShaderGraphNode->AddFloatAttribute(UInterchangeShaderPortsAPI::MakeInputValueKey(InputType), Weight);
+		}
+		else
+		{
+			return false;
+		}
+		return true;
+	}
+
+	// multiply texture by scalar, handling cases when any of the arguments is missing
+	static bool AddTexturedWeightedInput(UInterchangeBaseNodeContainer& BaseNodeContainer, UInterchangeShaderGraphNode* ShaderGraphNode, const FString& InputType, const float& Weight, const FObjData::FTextureParameter& TexturePath) 
+	{
+		UInterchangeShaderNode* TextureSampleShader = CreateMaterialTextureSampleNode(BaseNodeContainer, ShaderGraphNode, TexturePath);
+
+		return AddTexturedWeightedInput(BaseNodeContainer, ShaderGraphNode, InputType, Weight, TextureSampleShader);
+	}
+
+	static bool HandlePbrAttributes(UInterchangeBaseNodeContainer& BaseNodeContainer, const FObjData::FMaterialData& MaterialData, UInterchangeShaderGraphNode* ShaderGraphNode)
+	{
+		bool bPbr = !MaterialData.RoughnessTexture.IsEmpty() || IsScalarInitialized(MaterialData.Roughness)
+		         || !MaterialData.MetallicTexture.IsEmpty()  || IsScalarInitialized(MaterialData.Metallic);
+
+		if (!bPbr)
+		{
+			return false;
+		}
+
+		AddTexturedColoredInput(BaseNodeContainer, ShaderGraphNode, Materials::PBR::Parameters::BaseColor.ToString(), MaterialData.DiffuseColor, MaterialData.DiffuseTexture);
+		AddTexturedWeightedInput(BaseNodeContainer, ShaderGraphNode, Materials::PBR::Parameters::Metallic.ToString(), MaterialData.Metallic, MaterialData.MetallicTexture);
+
+		const float SpecularValue = (MaterialData.SpecularColor.X + MaterialData.SpecularColor.Y + MaterialData.SpecularColor.Z) / 3;
+		AddTexturedWeightedInput(BaseNodeContainer, ShaderGraphNode, Materials::PBR::Parameters::Specular.ToString(), SpecularValue, MaterialData.SpecularTexture);
+
+		// Roughness (Glossiness, Shininess, Specular Power/Exponent in non-PBR terms)
+		AddTexturedWeightedInput(BaseNodeContainer, ShaderGraphNode, Materials::PBR::Parameters::Roughness.ToString(), MaterialData.Roughness, MaterialData.RoughnessTexture);
+
+		if (IsScalarInitialized(MaterialData.ClearCoatThickness) && !FMath::IsNearlyZero(MaterialData.ClearCoatThickness))
+		{
+			ShaderGraphNode->AddFloatAttribute(UInterchangeShaderPortsAPI::MakeInputValueKey(Materials::ClearCoat::Parameters::ClearCoat.ToString()), MaterialData.ClearCoatThickness);
+
+			if (IsScalarInitialized(MaterialData.ClearCoatRoughness))
+			{
+				ShaderGraphNode->AddFloatAttribute(UInterchangeShaderPortsAPI::MakeInputValueKey(Materials::ClearCoat::Parameters::ClearCoatRoughness.ToString()), MaterialData.ClearCoatRoughness);
+			}
+		}
+
+		if (IsScalarInitialized(MaterialData.Anisotropy)  && !FMath::IsNearlyZero(MaterialData.Anisotropy))
+		{
+			ShaderGraphNode->AddFloatAttribute(UInterchangeShaderPortsAPI::MakeInputValueKey(Materials::PBR::Parameters::Anisotropy.ToString()), MaterialData.Anisotropy);
+		}
+
+		return true;
+	}
+}  // namespace ObjTranslatorUtils
+
+} // namespace Interchange
+} // namespace UE
 
 
 UInterchangeOBJTranslator::UInterchangeOBJTranslator()
@@ -937,7 +1221,6 @@ EInterchangeTranslatorAssetType UInterchangeOBJTranslator::GetSupportedAssetType
 	return EInterchangeTranslatorAssetType::Materials | EInterchangeTranslatorAssetType::Meshes;
 }
 
-
 bool UInterchangeOBJTranslator::Translate(UInterchangeBaseNodeContainer& BaseNodeContainer) const
 {
 	// Get filename from SourceData
@@ -964,20 +1247,21 @@ bool UInterchangeOBJTranslator::Translate(UInterchangeBaseNodeContainer& BaseNod
 	// Do we maybe want to decouple translators from payloads (which arguably makes more sense).
 
 	using namespace UE::Interchange;
+	using namespace ObjTranslatorUtils;
 
-	static ObjParser::FKeywordMap ObjKeywordMap =
+	static ObjParserUtils::FKeywordMap ObjKeywordMap =
 	{
-		{ TEXT("o"),      ObjParser::ParseObjectName },
-		{ TEXT("v"),      ObjParser::ParseVertexPosition },
-		{ TEXT("vt"),     ObjParser::ParseTextureCoordinate },
-		{ TEXT("vn"),     ObjParser::ParseNormalVector },
-		{ TEXT("f"),      ObjParser::ParseFace },
-		{ TEXT("g"),      ObjParser::ParseGroup },
-		{ TEXT("mtllib"), ObjParser::ParseMaterialLib },
-		{ TEXT("usemtl"), ObjParser::ParseUseMaterial }
+		{ TEXT("o"),      ObjParserUtils::ParseObjectName },
+		{ TEXT("v"),      ObjParserUtils::ParseVertexPosition },
+		{ TEXT("vt"),     ObjParserUtils::ParseTextureCoordinate },
+		{ TEXT("vn"),     ObjParserUtils::ParseNormalVector },
+		{ TEXT("f"),      ObjParserUtils::ParseFace },
+		{ TEXT("g"),      ObjParserUtils::ParseGroup },
+		{ TEXT("mtllib"), ObjParserUtils::ParseMaterialLib },
+		{ TEXT("usemtl"), ObjParserUtils::ParseUseMaterial }
 	};
 
-	bool bSuccess = ObjParser::ParseFile(*ObjDataPtr.Get(), *Filename, ObjKeywordMap);
+	bool bSuccess = ObjParserUtils::ParseFile(*ObjDataPtr.Get(), *Filename, ObjKeywordMap);
 
 	if (bSuccess)
 	{
@@ -1025,7 +1309,6 @@ bool UInterchangeOBJTranslator::Translate(UInterchangeBaseNodeContainer& BaseNod
 		}
 
 		// Add material nodes to the container
-
 		for (const TPair<FString, FObjData::FMaterialData>& Material : ObjDataPtr->Materials)
 		{
 			const FString& MaterialName = Material.Key;
@@ -1036,18 +1319,207 @@ bool UInterchangeOBJTranslator::Translate(UInterchangeBaseNodeContainer& BaseNod
 			ShaderGraphNode->InitializeNode(NodeUid, MaterialName, EInterchangeNodeContainerType::TranslatedAsset);
 			BaseNodeContainer.AddNode(ShaderGraphNode);
 
-			if (MaterialData.IlluminationModel >= 2)
+			bool bUnlit = false;
+			bool bSpecular = true;
+			bool bReflectionMetal = false;
+			bool bReflectionFresnel = false;
+			bool bTransmission = false;
+			bool bRefraction = false;
+
+			if (!HandlePbrAttributes(BaseNodeContainer, MaterialData, ShaderGraphNode))
 			{
-				AddMaterialNodes(BaseNodeContainer, ShaderGraphNode, Materials::Phong::Parameters::SpecularColor.ToString(), MaterialData.SpecularColor, MaterialData.SpecularTexture);
+				// It's not a Pbr material - in this case parse the illumination model
+
+				// What is relevant to us about the illumination model in OBJ spec
+				// 0 - Constant color
+				// 1 - Lambertian diffuse (i.e. no specular)
+				// 2 - Phong (i.e. includes specular)
+				// 3 - Has 'reflection map'(non-fresnel, metal)
+				// 4 - "Simulating Glass" - transparency with strong reflection(non-fresnel, metal)
+				// 5 - 3 but Fresnel reflection(non-metal)
+				// 6 - Reflection and refraction, Tf transmission is used
+				// 7 - Reflection and refraction, Tf transmission is used(just Fresnel)
+				// 8 - as 3 
+				// 9 - as 4
+
+				switch (MaterialData.IlluminationModel)
+				{
+					case 0:
+					{
+						bUnlit = true;
+						bSpecular = false;
+					}
+					break;
+
+					case 1:
+					{
+						bSpecular = false;
+					}
+					break;
+
+					case 2:
+					{
+					}
+					break;
+
+					case 3:
+					case 4:
+					case 8:
+					case 9:
+					{
+						bReflectionMetal = true;
+					}
+					break;
+
+					case 5:
+					{
+						bReflectionFresnel = true;
+					}
+					break;
+
+					case 6:
+					{
+						bReflectionMetal = true;
+						bTransmission = true;
+						bRefraction = true;
+					}
+					break;
+
+					case 7:
+					{
+						bReflectionFresnel = true;
+						bTransmission = true;
+						bRefraction = true;
+					}
+					break;
+
+
+					default: ;
+				}
+
+				if(bUnlit)
+				{
+					AddTexturedColoredInput(BaseNodeContainer, ShaderGraphNode, Materials::Unlit::Parameters::UnlitColor.ToString(), MaterialData.DiffuseColor, MaterialData.DiffuseTexture);
+				}
+				else
+				{
+					AddTexturedColoredInput(BaseNodeContainer, ShaderGraphNode, Materials::Phong::Parameters::DiffuseColor.ToString(), (MaterialData.DiffuseColor.GetMin() < 0) ? FVector3f(0.8f, 0.8f, 0.8f) : MaterialData.DiffuseColor, MaterialData.DiffuseTexture);
+
+					if (bSpecular)
+					{
+						AddTexturedColoredInput(BaseNodeContainer, ShaderGraphNode, Materials::Phong::Parameters::SpecularColor.ToString(), MaterialData.SpecularColor, MaterialData.SpecularTexture);
+
+						// OBJ Specular Exponent('Ns') is in range 0 to 1000(for sharpest highlight) and Interchange Shininess is up to 100
+						const float SpecularExponentToShininessFactor = 0.1f;
+						if (UInterchangeShaderNode* SpecularExponentTextureSampleShader = CreateMaterialTextureSampleNode(BaseNodeContainer, ShaderGraphNode, MaterialData.SpecularExponentTexture))
+						{
+							using namespace UE::Interchange::Materials::Standard::Nodes;
+
+							float SpecularExponentScale = (IsScalarInitialized(MaterialData.SpecularExponent) ? MaterialData.SpecularExponent : 1.0f) * SpecularExponentToShininessFactor;
+
+							FString MapName = Materials::Phong::Parameters::Shininess.ToString();
+							const FString MultiplierNodeName = MapName + TEXT("_Multiply");
+							UInterchangeShaderNode* MultiplierNode = UInterchangeShaderNode::Create(&BaseNodeContainer, MultiplierNodeName, ShaderGraphNode->GetUniqueID());
+							MultiplierNode->SetCustomShaderType(Multiply::Name.ToString());
+
+							UInterchangeShaderPortsAPI::ConnectDefaultOuputToInput(MultiplierNode, Multiply::Inputs::A.ToString(), SpecularExponentTextureSampleShader->GetUniqueID());
+							MultiplierNode->AddFloatAttribute(UInterchangeShaderPortsAPI::MakeInputValueKey( Multiply::Inputs::B.ToString() ), SpecularExponentScale);
+
+							UInterchangeShaderPortsAPI::ConnectDefaultOuputToInput(ShaderGraphNode, Materials::Phong::Parameters::Shininess.ToString(), MultiplierNode->GetUniqueID());
+						}
+						else if (IsScalarInitialized(MaterialData.SpecularExponent))
+						{
+							ShaderGraphNode->AddFloatAttribute(UInterchangeShaderPortsAPI::MakeInputValueKey(Materials::Phong::Parameters::Shininess.ToString()), MaterialData.SpecularExponent * SpecularExponentToShininessFactor);
+						}
+					}
+
+					if (bReflectionMetal)
+					{
+						ShaderGraphNode->AddFloatAttribute(UInterchangeShaderPortsAPI::MakeInputValueKey(Materials::PBR::Parameters::Metallic.ToString()), 1.0);
+					}
+				}
+			}
+
+			// Common parameters
+
+			// Opacity
+			{
+				FString OpacityInputType = Materials::Common::Parameters::Opacity.ToString();
+
+				// Transparency texture replaces Transparency scalar
+				if (UInterchangeShaderNode* TransparencySample = CreateMaterialTextureSampleNode(BaseNodeContainer, ShaderGraphNode, MaterialData.TransparencyTexture))
+				{
+					// Invert transparency to get Opacity
+					FString MapName = OpacityInputType;
+					const FString OneMinusNodeName = MapName + TEXT("OneMinus");
+					UInterchangeShaderNode* OneMinusNode = UInterchangeShaderNode::Create( &BaseNodeContainer, OneMinusNodeName, ShaderGraphNode->GetUniqueID() );
+					OneMinusNode->SetCustomShaderType(Materials::Standard::Nodes::OneMinus::Name.ToString());
+					UInterchangeShaderPortsAPI::ConnectDefaultOuputToInput(OneMinusNode, Materials::Standard::Nodes::OneMinus::Inputs::Input.ToString(), TransparencySample->GetUniqueID());
+					UInterchangeShaderPortsAPI::ConnectDefaultOuputToInput(ShaderGraphNode, OpacityInputType, OneMinusNode->GetUniqueID());
+				}
+				// Opacity texture is weighted by scalar opacity("d")
+				else if (!AddTexturedWeightedInput(BaseNodeContainer, ShaderGraphNode, OpacityInputType, MaterialData.Opacity, MaterialData.OpacityTexture))
+				{
+					// When no textures are present compute opacity value depending on which is defined - opacity or transparency
+					// OBJ has two ways to define Opacity/Transparency - "d"(for opacity) and "Tr"(for "inverse opacity" - transparency)
+					float OpacityScalar = IsScalarInitialized(MaterialData.Opacity) ? MaterialData.Opacity : (IsScalarInitialized(MaterialData.Transparency) ? 1.f - MaterialData.Transparency : -1.f);
+					if (IsScalarInitialized(OpacityScalar))
+					{
+						ShaderGraphNode->AddFloatAttribute(UInterchangeShaderPortsAPI::MakeInputValueKey(OpacityInputType), OpacityScalar);
+					}
+				}
+			}
+
+			if (bUnlit)
+			{
+				continue;
+			}
+
+			// Emissive
+			AddTexturedColoredInput(BaseNodeContainer, ShaderGraphNode, Materials::Common::Parameters::EmissiveColor.ToString(), MaterialData.EmissiveColor, MaterialData.EmissiveTexture);
+				
+			// Ior
+			if (IsScalarInitialized(MaterialData.IndexOfRefraction) || bRefraction)
+			{
+				const float DefaultIndexOfRefraction = 1.52f;  // Use glass index of refraction by default
+				float IndexOfRefraction = IsScalarInitialized(MaterialData.IndexOfRefraction) ? MaterialData.IndexOfRefraction : DefaultIndexOfRefraction;
+				ShaderGraphNode->AddFloatAttribute(UInterchangeShaderPortsAPI::MakeInputValueKey(Materials::Common::Parameters::IndexOfRefraction.ToString()), IndexOfRefraction);
 			}
 
-			if (MaterialData.IlluminationModel >= 1)
+			if ((IsColorInitialized(MaterialData.TransmissionFilter) && !(MaterialData.TransmissionFilter - FVector3f::One()).IsNearlyZero()) || bTransmission)
 			{
-				AddMaterialNodes(BaseNodeContainer, ShaderGraphNode, Materials::Phong::Parameters::EmissiveColor.ToString(), MaterialData.AmbientColor, MaterialData.AmbientTexture);
+				FLinearColor TransmissionColor = IsColorInitialized(MaterialData.TransmissionFilter) ? MaterialData.TransmissionFilter : FLinearColor::White;
+				ShaderGraphNode->AddLinearColorAttribute(UInterchangeShaderPortsAPI::MakeInputValueKey(Materials::ThinTranslucent::Parameters::TransmissionColor.ToString()), TransmissionColor);
 			}
 
-			// Always add diffuse colour
-			AddMaterialNodes(BaseNodeContainer, ShaderGraphNode, Materials::Phong::Parameters::DiffuseColor.ToString(), MaterialData.DiffuseColor, MaterialData.DiffuseTexture);
+			// Normal/bump
+			if (UInterchangeShaderNode* NormalmapTextureSampleShader = CreateMaterialTextureSampleNode(BaseNodeContainer, ShaderGraphNode, MaterialData.NormalmapTexture))
+			{
+				UInterchangeShaderPortsAPI::ConnectDefaultOuputToInput(ShaderGraphNode, Materials::Common::Parameters::Normal.ToString(), NormalmapTextureSampleShader->GetUniqueID());
+			}
+			if (!MaterialData.BumpmapTexture.IsEmpty())
+			{
+				using namespace Materials::Standard::Nodes;
+
+				FString TexturePath = MaterialData.BumpmapTexture.Path;
+				FString TextureName = MakeTextureNodeName(FPaths::GetBaseFilename(TexturePath));
+
+				if (const UInterchangeTexture2DNode* TextureNode = CreateTexture2DNode(BaseNodeContainer, TexturePath, TextureName))
+				{
+					// NormalFromHeightmap needs TextureObject(not just a sample as it takes multiple samples from it)
+					UInterchangeShaderNode* TextureObjectNode = UInterchangeShaderNode::Create(&BaseNodeContainer, TextureName, ShaderGraphNode->GetUniqueID());
+					TextureObjectNode->SetCustomShaderType(TextureObject::Name.ToString());
+					TextureObjectNode->AddStringAttribute(UInterchangeShaderPortsAPI::MakeInputValueKey(TextureObject::Inputs::Texture.ToString()), TextureNode->GetUniqueID());
+
+					UInterchangeShaderNode* HeightMapNode = UInterchangeShaderNode::Create( &BaseNodeContainer, NormalFromHeightMap::Name.ToString(), ShaderGraphNode->GetUniqueID() );
+					HeightMapNode->SetCustomShaderType(NormalFromHeightMap::Name.ToString());
+
+					UInterchangeShaderPortsAPI::ConnectDefaultOuputToInput(HeightMapNode, NormalFromHeightMap::Inputs::HeightMap.ToString(), TextureObjectNode->GetUniqueID());
+					HeightMapNode->AddFloatAttribute(UInterchangeShaderPortsAPI::MakeInputValueKey(NormalFromHeightMap::Inputs::Intensity.ToString()), MaterialData.BumpmapTexture.BumpMultiplier);
+
+					UInterchangeShaderPortsAPI::ConnectDefaultOuputToInput(ShaderGraphNode, Materials::Common::Parameters::Normal.ToString(), HeightMapNode->GetUniqueID());
+				}
+			}
 		}
 	}
 
@@ -1101,39 +1573,3 @@ TOptional<UE::Interchange::FImportImage> UInterchangeOBJTranslator::GetTexturePa
 }
 
 
-
-void UInterchangeOBJTranslator::AddMaterialNodes(UInterchangeBaseNodeContainer& BaseNodeContainer, UInterchangeShaderGraphNode* ShaderGraphNode, const FString& InputType, const FVector3f& Color, const FString& TexturePath) const
-{
-	if (!TexturePath.IsEmpty())
-	{
-		// @todo: better way to handle textures than implementing the TexturePayload interface in the OBJTranslator
-		// and creating a temporary texture translator there.
-
-		using namespace UE::Interchange::Materials::Standard::Nodes;
-
-		const FString TextureName = MakeTextureNodeName(FPaths::GetBaseFilename(TexturePath));
-
-		UInterchangeShaderNode* TextureSampleShader = UInterchangeShaderNode::Create(&BaseNodeContainer, TextureName, ShaderGraphNode->GetUniqueID());
-		TextureSampleShader->SetCustomShaderType(TextureSample::Name.ToString());
-
-		FString TextureNodeUid = UInterchangeTextureNode::MakeNodeUid(TextureName);
-		const UInterchangeTexture2DNode* TextureNode = Cast<const UInterchangeTexture2DNode>(BaseNodeContainer.GetNode(TextureNodeUid));
-		if (!TextureNode)
-		{
-			UInterchangeTexture2DNode* NewTextureNode = UInterchangeTexture2DNode::Create(&BaseNodeContainer, TextureName);
-
-			FString NormalizedTexturePath(TexturePath);
-			FPaths::NormalizeFilename(NormalizedTexturePath);
-			NewTextureNode->SetPayLoadKey(NormalizedTexturePath);
-		}
-
-		TextureSampleShader->AddStringAttribute(UInterchangeShaderPortsAPI::MakeInputValueKey(TextureSample::Inputs::Texture.ToString()), TextureNodeUid);
-
-		UInterchangeShaderPortsAPI::ConnectDefaultOuputToInput(ShaderGraphNode, InputType, TextureSampleShader->GetUniqueID());
-	}
-	else
-	{
-		ShaderGraphNode->AddLinearColorAttribute(UInterchangeShaderPortsAPI::MakeInputValueKey(InputType), FLinearColor(Color));
-	}
-
-}
diff --git a/Engine/Plugins/Interchange/Runtime/Source/Import/Public/Mesh/InterchangeOBJTranslator.h b/Engine/Plugins/Interchange/Runtime/Source/Import/Public/Mesh/InterchangeOBJTranslator.h
--- a/Engine/Plugins/Interchange/Runtime/Source/Import/Public/Mesh/InterchangeOBJTranslator.h
+++ b/Engine/Plugins/Interchange/Runtime/Source/Import/Public/Mesh/InterchangeOBJTranslator.h
@@ -3,8 +3,6 @@
 #pragma once
 
 #include "CoreMinimal.h"
-#include "UObject/Object.h"
-#include "UObject/ObjectMacros.h"
 #include "InterchangeTranslatorBase.h"
 #include "Nodes/InterchangeBaseNodeContainer.h"
 #include "Mesh/InterchangeMeshPayload.h"
@@ -32,6 +30,7 @@ public:
 
 	virtual TArray<FString> GetSupportedFormats() const override;
 	virtual EInterchangeTranslatorAssetType GetSupportedAssetTypes() const override;
+
 	/**
 	 * Translate the associated source data into a node hold by the specified nodes container.
 	 *
@@ -66,11 +65,8 @@ public:
 	 * @return a PayloadData containing the imported data. The TOptional will not be set if there is an error.
 	 */
 	virtual TOptional<UE::Interchange::FImportImage> GetTexturePayloadData(const FString& PayloadKey, TOptional<FString>& AlternateTexturePath) const override;
-
 	/* IInterchangeTexturePayloadInterface End */
 
 private:
-	void AddMaterialNodes(UInterchangeBaseNodeContainer& BaseNodeContainer, UInterchangeShaderGraphNode* ShaderGraphNode, const FString& InputType, const FVector3f& Color, const FString& TexturePath) const;
-
 	TPimplPtr<FObjData> ObjDataPtr;
 };
