Fixing base color washed out when importing fbx

diff --git a/Engine/Plugins/Interchange/Runtime/Source/Parsers/Fbx/Private/FbxMaterial.cpp b/Engine/Plugins/Interchange/Runtime/Source/Parsers/Fbx/Private/FbxMaterial.cpp
--- a/Engine/Plugins/Interchange/Runtime/Source/Parsers/Fbx/Private/FbxMaterial.cpp
+++ b/Engine/Plugins/Interchange/Runtime/Source/Parsers/Fbx/Private/FbxMaterial.cpp
@@ -197,7 +197,7 @@ namespace UE
 								const float Factor = (float)((FbxSurfaceLambert*)SurfaceMaterial)->DiffuseFactor.Get();
 								TVariant<FLinearColor, float> DefaultValue;
 								DefaultValue.Set<FLinearColor>(FLinearColor::Black);
-								ConvertPropertyToShaderNode(NodeContainer, ShaderGraphNode, CurrentProperty, Factor, Phong::Parameters::DiffuseColor, DefaultValue);
+								ConvertPropertyToShaderNode(NodeContainer, ShaderGraphNode, CurrentProperty, Factor, PBR::Parameters::BaseColor, DefaultValue);
 							}
 							else if (PropertyName == FbxSurfaceMaterial::sSpecular && SurfaceMaterial->Is<FbxSurfacePhong>())
 							{
