[Interchange - glTF] Add support for KHR_mesh_quantization in glTF importer

diff --git a/Engine/Plugins/Interchange/Runtime/Source/Import/Private/Gltf/InterchangeGltfTranslator.cpp b/Engine/Plugins/Interchange/Runtime/Source/Import/Private/Gltf/InterchangeGltfTranslator.cpp
--- a/Engine/Plugins/Interchange/Runtime/Source/Import/Private/Gltf/InterchangeGltfTranslator.cpp
+++ b/Engine/Plugins/Interchange/Runtime/Source/Import/Private/Gltf/InterchangeGltfTranslator.cpp
@@ -65,7 +65,10 @@ static const TArray<FString> ImporterSupportedExtensions = {
 	GLTF::ToString(GLTF::EExtension::MSFT_PackingOcclusionRoughnessMetallic),
 	GLTF::ToString(GLTF::EExtension::MSFT_PackingNormalRoughnessMetallic),
 	/* Textures */
-	GLTF::ToString(GLTF::EExtension::KHR_TextureTransform)
+	GLTF::ToString(GLTF::EExtension::KHR_TextureTransform),
+
+	/* Mesh */
+	GLTF::ToString(GLTF::EExtension::KHR_MeshQuantization)
 };
 
 namespace UE::Interchange::Gltf::Private
diff --git a/Engine/Plugins/Interchange/Runtime/Source/Parsers/GLTFCore/Private/GLTF/GLTFAccessor.cpp b/Engine/Plugins/Interchange/Runtime/Source/Parsers/GLTFCore/Private/GLTF/GLTFAccessor.cpp
--- a/Engine/Plugins/Interchange/Runtime/Source/Parsers/GLTFCore/Private/GLTF/GLTFAccessor.cpp
+++ b/Engine/Plugins/Interchange/Runtime/Source/Parsers/GLTFCore/Private/GLTF/GLTFAccessor.cpp
@@ -28,27 +28,154 @@ namespace GLTF
 		{
 			ReturnType Res;
 
-			// convert to 0..1
-			if (ComponentType == FAccessor::EComponentType::U8)
+			//U32 and F32 cannot be normalized.
+			switch (ComponentType)
 			{
-				const uint8*    P = static_cast<const uint8*>(Pointer);
-				constexpr float S = 1.0f / 255.0f;
-				for (uint32 Index = 0; Index < Count; ++Index)
+			case GLTF::FAccessor::EComponentType::S8:
+				{
+					const int8* P = static_cast<const int8*>(Pointer);
+					for (uint32 Index = 0; Index < Count; ++Index)
+					{
+						Res[Index] = FMath::Max((P[Index] / 127.0f), -1.0f);
+					}
+					break;
+				}
+			case GLTF::FAccessor::EComponentType::U8:
+				{
+					const uint8* P = static_cast<const uint8*>(Pointer);
+					for (uint32 Index = 0; Index < Count; ++Index)
+					{
+						Res[Index] = P[Index] / 255.0f;
+					}
+					break;
+				}
+			case GLTF::FAccessor::EComponentType::S16:
 				{
-					Res[Index] = P[Index] * S;
+					const int16* P = static_cast<const int16*>(Pointer);
+					for (uint32 Index = 0; Index < Count; ++Index)
+					{
+						Res[Index] = FMath::Max(P[Index] / 32767.0f, -1.0f);
+					}
+					break;
 				}
+			case GLTF::FAccessor::EComponentType::U16:
+				{
+					const uint16* P = static_cast<const uint16*>(Pointer);
+					for (uint32 Index = 0; Index < Count; ++Index)
+					{
+						Res[Index] = P[Index] / 65535.0f;
+					}
+					break;
+				}
+			default:
+				Res = ReturnType();
+				ensure(false);
+				break;
 			}
-			else if (ComponentType == FAccessor::EComponentType::U16)
+
+			return Res;
+		}
+
+		float GetNormalizedFloat(FAccessor::EComponentType ComponentType, const void* Pointer)
+		{
+			switch (ComponentType)
 			{
-				const uint16*   P = static_cast<const uint16*>(Pointer);
-				constexpr float S = 1.0f / 65535.0f;
+			case GLTF::FAccessor::EComponentType::S8:
+				{
+					const int8* P = static_cast<const int8*>(Pointer);
+					return FMath::Max(P[0] / 127.0f, -1.0f);
+				}
+			case GLTF::FAccessor::EComponentType::U8:
+				{
+					const uint8* P = static_cast<const uint8*>(Pointer);
+					return P[0] / 255.0f;
+				}
+			case GLTF::FAccessor::EComponentType::S16:
+				{
+					const int16* P = static_cast<const int16*>(Pointer);
+					return FMath::Max(P[0] / 32767.0f, -1.0f);
+				}
+			case GLTF::FAccessor::EComponentType::U16:
+				{
+					const uint16* P = static_cast<const uint16*>(Pointer);
+					return P[0] / 65535.0f;
+				}
+			default:
+				ensure(false);
+				break;
+			}
+
+			return 0.0f;
+		}
+
+		template <class ReturnType, uint32 Count>
+		ReturnType GetNonNormalized(FAccessor::EComponentType ComponentType, const void* Pointer)
+		{
+			ReturnType Res;
+
+			switch (ComponentType)
+			{
+			case GLTF::FAccessor::EComponentType::S8:
+				{
+					const int8* P = static_cast<const int8*>(Pointer);
+					for (uint32 Index = 0; Index < Count; ++Index)
+					{
+						Res[Index] = P[Index];
+					}
+					break;
+				}
+			case GLTF::FAccessor::EComponentType::U8:
+				{
+					const uint8* P = static_cast<const uint8*>(Pointer);
+					for (uint32 Index = 0; Index < Count; ++Index)
+					{
+						Res[Index] = P[Index];
+					}
+					break;
+				}
+			case GLTF::FAccessor::EComponentType::S16:
+				{
+					const int16* P = static_cast<const int16*>(Pointer);
+					for (uint32 Index = 0; Index < Count; ++Index)
+					{
+						Res[Index] = P[Index];
+					}
+					break;
+				}
+			case GLTF::FAccessor::EComponentType::U16:
+				{
+					const uint16* P = static_cast<const uint16*>(Pointer);
+					for (uint32 Index = 0; Index < Count; ++Index)
+					{
+						Res[Index] = P[Index];
+					}
+					break;
+				}
+			case GLTF::FAccessor::EComponentType::U32:
+			{
+				//This is unexpected, U32 is only supported on Indices,
+				//and indices are acquired via GetUnsignedIntArray function.
+				const uint32* P = static_cast<const uint32*>(Pointer);
 				for (uint32 Index = 0; Index < Count; ++Index)
 				{
-					Res[Index] = P[Index] * S;
+					Res[Index] = P[Index];
 				}
+				break;
 			}
-			else
+			case FAccessor::EComponentType::F32:
+				{
+					const float* P = static_cast<const float*>(Pointer);
+					for (uint32 Index = 0; Index < Count; ++Index)
+					{
+						Res[Index] = P[Index];
+					}
+					break;
+				}
+			default:
+				Res = ReturnType();
 				ensure(false);
+				break;
+			}
 
 			return Res;
 		}
@@ -80,40 +207,54 @@ namespace GLTF
 			}
 		}
 
-		template <typename DstT, uint32 ElementCount>
-		void CopyNormalized(DstT* Dst, const void* Src, FAccessor::EComponentType ComponentType, uint32 Count)
+		void CopyNormalizedFloat(float* Dst, const void* Src, FAccessor::EComponentType ComponentType, uint32 Count)
 		{
-			// convert to 0..1
-			if (ComponentType == FAccessor::EComponentType::U8)
+			switch (ComponentType)
 			{
-				const uint8*    P = static_cast<const uint8*>(Src);
-				constexpr float S = 1.0f / 255.0f;
-				for (uint32 Index = 0; Index < Count; ++Index)
+			case GLTF::FAccessor::EComponentType::S8:
 				{
-					DstT& VecDst = *Dst++;
-					for (uint32 J = 0; J < ElementCount; ++J)
+					const int8* P = static_cast<const int8*>(Src);
+					for (uint32 Index = 0; Index < Count; ++Index, P++)
 					{
-						VecDst[J] = P[J] * S;
+						float* VecDst = Dst++;
+						*VecDst = FMath::Max((*P) / 127.0f, -1.0f);
 					}
-					P += ElementCount;
+					break;
 				}
-			}
-			else if (ComponentType == FAccessor::EComponentType::U16)
-			{
-				const uint16*   P = static_cast<const uint16*>(Src);
-				constexpr float S = 1.0f / 65535.0f;
-				for (uint32 Index = 0; Index < Count; ++Index)
+			case GLTF::FAccessor::EComponentType::U8:
 				{
-					DstT& VecDst = *Dst++;
-					for (uint32 J = 0; J < ElementCount; ++J)
+					const uint8* P = static_cast<const uint8*>(Src);
+					for (uint32 Index = 0; Index < Count; ++Index, P++)
 					{
-						VecDst[J] = P[J] * S;
+						float* VecDst = Dst++;
+						*VecDst = (*P) / 255.0f;
 					}
-					P += ElementCount;
+					break;
 				}
-			}
-			else
+			case GLTF::FAccessor::EComponentType::S16:
+				{
+					const int16* P = static_cast<const int16*>(Src);
+					for (uint32 Index = 0; Index < Count; ++Index, P++)
+					{
+						float* VecDst = Dst++;
+						*VecDst = FMath::Max((*P) / 32767.0f, -1.0f);
+					}
+					break;
+				}
+			case GLTF::FAccessor::EComponentType::U16:
+				{
+					const uint16* P = static_cast<const uint16*>(Src);
+					for (uint32 Index = 0; Index < Count; ++Index, P++)
+					{
+						float* VecDst = Dst++;
+						*VecDst = (*P) / 65535.0f;
+					}
+					break;
+				}
+			default:
 				ensure(false);
+				break;
+			}
 		}
 
 		// Copy data items that don't need conversion/expansion(i.e. Vec3 to Vec3, uint8 to uint8(not uint16)
@@ -128,39 +269,35 @@ namespace GLTF
 				if (Accessor.ComponentType == FAccessor::EComponentType::F32)
 				{
 					memcpy(Buffer, Src, Accessor.Count * sizeof(ItemType));
+					return;
 				}
-				else if (Accessor.Normalized)
+			}
+
+			if (Accessor.bNormalized)
+			{
+				for (uint32 Index = 0; Index < Accessor.Count; ++Index)
 				{
-					CopyNormalized<ItemType, ItemElementCount>(Buffer, Src, Accessor.ComponentType, Accessor.Count);
+					const void* Pointer = Accessor.DataAt(Index);
+					Buffer[Index] = GetNormalized<ItemType, ItemElementCount>(Accessor.ComponentType, Pointer);
 				}
 			}
 			else
 			{
-				if (Accessor.ComponentType == FAccessor::EComponentType::F32)
+				for (uint32 Index = 0; Index < Accessor.Count; ++Index)
 				{
-					for (uint32 Index = 0; Index < Accessor.Count; ++Index)
-					{
-						const void* Pointer = Accessor.DataAt(Index);
-						Buffer[Index] = *static_cast<const ItemType*>(Pointer);
-					}
-				}
-				else if (Accessor.Normalized)
-				{
-					for (uint32 Index = 0; Index < Accessor.Count; ++Index)
-					{
-						const void* Pointer = Accessor.DataAt(Index);
-						Buffer[Index] = GetNormalized<ItemType, ItemElementCount>(Accessor.ComponentType, Pointer);
-					}
+					const void* Pointer = Accessor.DataAt(Index);
+					Buffer[Index] = GetNonNormalized<ItemType, ItemElementCount>(Accessor.ComponentType, Pointer);
 				}
 			}
 		}
 	}
 
-	FAccessor::FAccessor(uint32 InCount, EType InType, EComponentType InCompType, bool InNormalized)
+	FAccessor::FAccessor(uint32 InCount, EType InType, EComponentType InCompType, bool bInNormalized)
 	    : Count(InCount)
 	    , Type(InType)
 	    , ComponentType(InCompType)
-	    , Normalized(InNormalized)
+	    , bNormalized(bInNormalized)
+		, bQuantized(false)
 	{
 	}
 
@@ -232,11 +369,148 @@ namespace GLTF
 		}
 	}
 
-	//
+	bool FAccessor::CheckAccessorTypeForDataType(EDataType DataType, bool bMorphTargetProperty) const
+	{
+		switch (DataType)
+		{
+		case GLTF::FAccessor::EDataType::Position:
+		case GLTF::FAccessor::EDataType::Normal:
+			return Type == FAccessor::EType::Vec3;
+		case GLTF::FAccessor::EDataType::Tangent:
+			return (!bMorphTargetProperty && Type == FAccessor::EType::Vec4) || (bMorphTargetProperty && Type == FAccessor::EType::Vec3);
+		case GLTF::FAccessor::EDataType::Texcoord:
+			return Type == FAccessor::EType::Vec2;
+		case GLTF::FAccessor::EDataType::Color:
+			return Type == FAccessor::EType::Vec3 || Type == FAccessor::EType::Vec4;
+		case GLTF::FAccessor::EDataType::Joints:
+		case GLTF::FAccessor::EDataType::Weights:
+			return Type == FAccessor::EType::Vec4;
+		default:
+			ensure(false);
+			break;
+		}
+
+		return false;
+	}
+
+	bool FAccessor::CheckNonQuantizedComponentTypeForDataType(EDataType DataType, bool bMorphTargetProperty) const
+	{
+		if (!bMorphTargetProperty)
+		{
+			//Non Morph Target:
+			switch (DataType)
+			{
+			case GLTF::FValidAccessor::EDataType::Position:
+			case GLTF::FValidAccessor::EDataType::Normal:
+			case GLTF::FValidAccessor::EDataType::Tangent:
+				return !bNormalized && ComponentType == FAccessor::EComponentType::F32;
+			case GLTF::FValidAccessor::EDataType::Texcoord:
+			case GLTF::FValidAccessor::EDataType::Color:
+				return (!bNormalized && ComponentType == FAccessor::EComponentType::F32) 
+					|| (bNormalized && (ComponentType == FAccessor::EComponentType::U8 || ComponentType == FAccessor::EComponentType::U16));
+			case GLTF::FValidAccessor::EDataType::Joints:
+				return (!bNormalized && (ComponentType == FAccessor::EComponentType::U8 || ComponentType == FAccessor::EComponentType::U16));
+			case GLTF::FValidAccessor::EDataType::Weights:
+				return (!bNormalized && ComponentType == FAccessor::EComponentType::F32)
+					|| (bNormalized && (ComponentType == FAccessor::EComponentType::U8 || ComponentType == FAccessor::EComponentType::U16));
+			default:
+				ensure(false);
+				break;
+			}
+		}
+		else
+		{
+			//Morph Target:
+			switch (DataType)
+			{
+			case GLTF::FValidAccessor::EDataType::Position:
+			case GLTF::FValidAccessor::EDataType::Normal:
+			case GLTF::FValidAccessor::EDataType::Tangent:
+				return !bNormalized && ComponentType == FAccessor::EComponentType::F32;
+			case GLTF::FValidAccessor::EDataType::Texcoord:
+			case GLTF::FValidAccessor::EDataType::Color:
+				return (!bNormalized && ComponentType == FAccessor::EComponentType::F32)
+					|| (bNormalized && (ComponentType == FAccessor::EComponentType::S8 || ComponentType == FAccessor::EComponentType::S16 || ComponentType == FAccessor::EComponentType::U8 || ComponentType == FAccessor::EComponentType::U16));
+			default:
+				ensure(false);
+				break;
+			}
+		}
+
+		return false;
+	}
+
+	bool FAccessor::CheckQuantizedComponentTypeForDataType(EDataType DataType, bool bMorphTargetProperty) const
+	{
+		if (!bMorphTargetProperty)
+		{
+			//Non Morph Target:
+			switch (DataType)
+			{
+			case GLTF::FValidAccessor::EDataType::Position:
+				return ComponentType == FAccessor::EComponentType::S8 || ComponentType == FAccessor::EComponentType::U8 || ComponentType == FAccessor::EComponentType::S16 || ComponentType == FAccessor::EComponentType::U16;
+			case GLTF::FValidAccessor::EDataType::Normal:
+			case GLTF::FValidAccessor::EDataType::Tangent:
+				return bNormalized && (ComponentType == FAccessor::EComponentType::S8 || ComponentType == FAccessor::EComponentType::S16);
+			case GLTF::FValidAccessor::EDataType::Texcoord:
+				return (ComponentType == FAccessor::EComponentType::S8 || ComponentType == FAccessor::EComponentType::S16)
+					|| (!bNormalized && (ComponentType == FAccessor::EComponentType::U8 || ComponentType == FAccessor::EComponentType::U16));
+			default:
+				ensure(false);
+				break;
+			}
+		}
+		else
+		{
+			//Morph Target:
+			switch (DataType)
+			{
+			case GLTF::FValidAccessor::EDataType::Position:
+				return ComponentType == FAccessor::EComponentType::S8 || ComponentType == FAccessor::EComponentType::S16;
+			case GLTF::FValidAccessor::EDataType::Normal:
+			case GLTF::FValidAccessor::EDataType::Tangent:
+				return bNormalized && (ComponentType == FAccessor::EComponentType::S8 || ComponentType == FAccessor::EComponentType::S16);
+			case GLTF::FValidAccessor::EDataType::Texcoord:
+				return !bNormalized && (ComponentType == FAccessor::EComponentType::S8 || ComponentType == FAccessor::EComponentType::S16);
+			default:
+				ensure(false);
+				break;
+			}
+		}
 
+		return false;
+	}
+
+	bool FAccessor::IsValidDataType(EDataType DataType, bool bMorphTargetProperty) const
+	{
+		//Check Accessor Type restrictions:
+		if (!CheckAccessorTypeForDataType(DataType, bMorphTargetProperty))
+		{
+			return false;
+		}
+
+		//Check Component Type restrictions (non-quantized) first:
+		if (CheckNonQuantizedComponentTypeForDataType(DataType, bMorphTargetProperty))
+		{
+			return true;
+		}
+
+		//Check quantized Component Type restrictions:
+		if (bQuantized)
+		{
+			if (CheckQuantizedComponentTypeForDataType(DataType, bMorphTargetProperty))
+			{
+				return true;
+			}
+		}
+		
+		return false;
+	}
+
+	//
 	FValidAccessor::FValidAccessor(FBufferView& InBufferView, uint32 InOffset, uint32 InCount, EType InType, EComponentType InCompType,
-	                               bool InNormalized)
-	    : FAccessor(InCount, InType, InCompType, InNormalized)
+	                               bool bInNormalized)
+	    : FAccessor(InCount, InType, InCompType, bInNormalized)
 	    , BufferView(InBufferView)
 	    , ByteOffset(InOffset)
 	    , ElementSize(GetElementSize(Type, ComponentType))
@@ -275,7 +549,7 @@ namespace GLTF
 
 		uint8 TypeInt = static_cast<uint8>(Type);
 		uint8 ComponentInt = static_cast<uint8>(ComponentType);
-		uint8 NormalizedInt = static_cast<uint8>(Normalized);
+		uint8 NormalizedInt = static_cast<uint8>(bNormalized);
 
 		MD5.Update(&TypeInt, sizeof(TypeInt));
 		MD5.Update(&ComponentInt, sizeof(ComponentInt));
@@ -288,11 +562,11 @@ namespace GLTF
 
 	uint32 FValidAccessor::GetUnsignedInt(uint32 Index) const
 	{
-		// should be Scalar, not Normalized, unsigned integer (8, 16 or 32 bit)
+		// should be Scalar, not bNormalized, unsigned integer (8, 16 or 32 bit)
 
 		if (Index < Count)
 		{
-			if (Type == EType::Scalar && !Normalized)
+			if (Type == EType::Scalar && !bNormalized)
 			{
 				const uint8* ValuePtr = DataAt(Index);
 				switch (ComponentType)
@@ -315,11 +589,11 @@ namespace GLTF
 
 	void FValidAccessor::GetUnsignedInt16x4(uint32 Index, uint16 Values[4]) const
 	{
-		// should be Vec4, not Normalized, unsigned integer (8 or 16 bit)
+		// should be Vec4, not bNormalized, unsigned integer (8 or 16 bit)
 
 		if (Index < Count)
 		{
-			if (Type == EType::Vec4 && !Normalized)
+			if (Type == EType::Vec4 && !bNormalized)
 			{
 				const void* ValuePtr = DataAt(Index);
 				switch (ComponentType)
@@ -350,7 +624,7 @@ namespace GLTF
 
 		if (Index < Count)
 		{
-			if (Type == EType::Scalar && !Normalized)
+			if (Type == EType::Scalar && !bNormalized)
 			{
 				const uint8* ValuePtr = DataAt(Index);
 				switch (ComponentType)
@@ -386,7 +660,7 @@ namespace GLTF
 					// copy float vec2 directly from buffer
 					return FVector2D(*reinterpret_cast<const FVector2f*>(Pointer));
 				}
-				else if (Normalized)
+				else if (bNormalized)
 				{
 					return FVector2D(GetNormalized<FVector2f, 2>(ComponentType, Pointer));
 				}
@@ -416,7 +690,7 @@ namespace GLTF
 					// copy float vec3 directly from buffer
 					return FVector(*reinterpret_cast<const FVector3f*>(Pointer));
 				}
-				else if (Normalized)
+				else if (bNormalized)
 				{
 					return FVector(GetNormalized<FVector3f, 3>(ComponentType, Pointer));
 				}
@@ -446,7 +720,7 @@ namespace GLTF
 					// copy float vec4 directly from buffer
 					return FVector4(*reinterpret_cast<const FVector4f*>(Pointer));
 				}
-				else if (Normalized)
+				else if (bNormalized)
 				{
 					return FVector4(GetNormalized<FVector4f, 4>(ComponentType, Pointer));
 				}
@@ -479,7 +753,7 @@ namespace GLTF
 
 	void FValidAccessor::GetUnsignedIntArray(uint32* Buffer) const
 	{
-		if (Type == EType::Scalar && !Normalized)
+		if (Type == EType::Scalar && !bNormalized)
 		{
 			const uint8* Src = DataAt(0);
 			switch (ComponentType)
@@ -503,16 +777,37 @@ namespace GLTF
 
 	void FValidAccessor::GetFloatArray(float* Buffer) const
 	{
-		if (Type == EType::Scalar && !Normalized)
+		if (Type == EType::Scalar)
 		{
-			const uint8* Src = DataAt(0);
-			switch (ComponentType)
+			if (!bNormalized)
 			{
+				const uint8* Src = DataAt(0);
+				switch (ComponentType)
+				{
 				case EComponentType::F32:
 					memcpy(Buffer, Src, Count * sizeof(float));
 					return;
 				default:
 					break;
+				}
+			}
+			else
+			{
+				// Stride equals item size => use simpler copy
+				if ((ByteStride == 0) || ByteStride == sizeof(float))
+				{
+					const void* Src = DataAt(0);
+					CopyNormalizedFloat(Buffer, Src, ComponentType, Count);
+				}
+				else
+				{
+					for (uint32 Index = 0; Index < Count; ++Index)
+					{
+						const void* Pointer = DataAt(Index);
+						Buffer[Index] = GetNormalizedFloat(ComponentType, Pointer);
+					}
+				}
+				return;
 			}
 		}
 
diff --git a/Engine/Plugins/Interchange/Runtime/Source/Parsers/GLTFCore/Private/GLTF/GLTFAsset.cpp b/Engine/Plugins/Interchange/Runtime/Source/Parsers/GLTFCore/Private/GLTF/GLTFAsset.cpp
--- a/Engine/Plugins/Interchange/Runtime/Source/Parsers/GLTFCore/Private/GLTF/GLTFAsset.cpp
+++ b/Engine/Plugins/Interchange/Runtime/Source/Parsers/GLTFCore/Private/GLTF/GLTFAsset.cpp
@@ -400,6 +400,8 @@ namespace GLTF
 				return TEXT("MSFT_packing_normalRoughnessMetallic");
 			case GLTF::EExtension::MSFT_PackingOcclusionRoughnessMetallic:
 				return TEXT("MSFT_packing_occlusionRoughnessMetallic");
+			case GLTF::EExtension::KHR_MeshQuantization:
+				return TEXT("KHR_mesh_quantization");
 			case GLTF::EExtension::Count:
 				ensure(false);
 			default:
diff --git a/Engine/Plugins/Interchange/Runtime/Source/Parsers/GLTFCore/Private/GLTF/GLTFMesh.cpp b/Engine/Plugins/Interchange/Runtime/Source/Parsers/GLTFCore/Private/GLTF/GLTFMesh.cpp
--- a/Engine/Plugins/Interchange/Runtime/Source/Parsers/GLTFCore/Private/GLTF/GLTFMesh.cpp
+++ b/Engine/Plugins/Interchange/Runtime/Source/Parsers/GLTFCore/Private/GLTF/GLTFMesh.cpp
@@ -6,18 +6,6 @@
 
 namespace GLTF
 {
-	namespace
-	{
-		// convertible to 0.0 .. 1.0 factor
-		// (colors, tex coords, weights, etc.)
-		bool IsConvertibleToNormalizedFloat(const FAccessor& Attrib)
-		{
-			return Attrib.ComponentType == FAccessor::EComponentType::F32 ||
-				(Attrib.Normalized &&
-					(Attrib.ComponentType == FAccessor::EComponentType::U8 || Attrib.ComponentType == FAccessor::EComponentType::U16));
-		}
-	}
-
 	FMorphTarget::FMorphTarget(const FAccessor& InPositionDisplacements, const FAccessor& InNormalDisplacements,
 		const FAccessor& InTangentDisplacements, const FAccessor& InTexCoord0Displacements, const FAccessor& InTexCoord1Displacements, const FAccessor& InColor0Deltas)
 		: PositionDisplacements(InPositionDisplacements)
@@ -31,7 +19,62 @@ namespace GLTF
 
 	bool FMorphTarget::IsValid() const
 	{
-		return HasPositionDisplacements() || HasNormalDisplacements() || HasTangentDisplacements() || HasTexCoordDisplacements(0) || HasTexCoordDisplacements(1) || HasColorDeltas();
+		bool bHasData = false;
+		const bool bMorphTarget = true;
+
+		if (HasPositionDisplacements())
+		{
+			bHasData = true;
+			if (!PositionDisplacements.IsValidDataType(FAccessor::EDataType::Position, bMorphTarget))
+			{
+				return false;
+			}
+		}
+		
+		if (HasNormalDisplacements())
+		{
+			bHasData = true;
+			if (!NormalDisplacements.IsValidDataType(FAccessor::EDataType::Normal, bMorphTarget))
+			{
+				return false;
+			}
+		}
+		
+		if (HasTangentDisplacements())
+		{
+			bHasData = true;
+			if (!TangentDisplacements.IsValidDataType(FAccessor::EDataType::Tangent, bMorphTarget))
+			{
+				return false;
+			}
+		}
+		
+		if (HasTexCoordDisplacements(0))
+		{
+			bHasData = true;
+			if (!TexCoord0Displacements.IsValidDataType(FAccessor::EDataType::Texcoord, bMorphTarget))
+			{
+				return false;
+			}
+		}
+		if (HasTexCoordDisplacements(1))
+		{
+			bHasData = true;
+			if (!TexCoord1Displacements.IsValidDataType(FAccessor::EDataType::Texcoord, bMorphTarget))
+			{
+				return false;
+			}
+		}
+
+		if (HasColorDeltas())
+		{
+			bHasData = true;
+			if (Color0Deltas.IsValidDataType(FAccessor::EDataType::Color, bMorphTarget))
+			{
+				return false;
+			}
+		}
+		return bHasData;
 	}
 	FMD5Hash FMorphTarget::GetHash() const
 	{
@@ -385,19 +428,15 @@ namespace GLTF
 	bool FPrimitive::IsValidPrivate() const
 	{
 		// make sure all semantic attributes meet the spec
+		const bool bMorphTarget = false;
 
-		if (!Position.IsValid())
+		if (!(Position.IsValid() && Position.IsValidDataType(FAccessor::EDataType::Position, bMorphTarget)))
 		{
 			return false;
 		}
 
 		const uint32 VertexCount = Position.Count;
 
-		if (Position.Type != FAccessor::EType::Vec3 || Position.ComponentType != FAccessor::EComponentType::F32)
-		{
-			return false;
-		}
-
 		if (Normal.IsValid())
 		{
 			if (Normal.Count != VertexCount)
@@ -405,7 +444,7 @@ namespace GLTF
 				return false;
 			}
 
-			if (Normal.Type != FAccessor::EType::Vec3 || Normal.ComponentType != FAccessor::EComponentType::F32)
+			if (!Normal.IsValidDataType(FAccessor::EDataType::Normal, bMorphTarget))
 			{
 				return false;
 			}
@@ -418,7 +457,7 @@ namespace GLTF
 				return false;
 			}
 
-			if (Tangent.Type != FAccessor::EType::Vec4 || Tangent.ComponentType != FAccessor::EComponentType::F32)
+			if (!Tangent.IsValidDataType(FAccessor::EDataType::Tangent, bMorphTarget))
 			{
 				return false;
 			}
@@ -434,7 +473,7 @@ namespace GLTF
 					return false;
 				}
 
-				if (TexCoord->Type != FAccessor::EType::Vec2 || !IsConvertibleToNormalizedFloat(*TexCoord))
+				if (!TexCoord->IsValidDataType(FAccessor::EDataType::Texcoord, bMorphTarget))
 				{
 					return false;
 				}
@@ -448,7 +487,7 @@ namespace GLTF
 				return false;
 			}
 
-			if (!(Color0.Type == FAccessor::EType::Vec3 || Color0.Type == FAccessor::EType::Vec4) || !IsConvertibleToNormalizedFloat(Color0))
+			if (!Color0.IsValidDataType(FAccessor::EDataType::Color, bMorphTarget))
 			{
 				return false;
 			}
diff --git a/Engine/Plugins/Interchange/Runtime/Source/Parsers/GLTFCore/Private/GLTF/GLTFReader.cpp b/Engine/Plugins/Interchange/Runtime/Source/Parsers/GLTFCore/Private/GLTF/GLTFReader.cpp
--- a/Engine/Plugins/Interchange/Runtime/Source/Parsers/GLTFCore/Private/GLTF/GLTFReader.cpp
+++ b/Engine/Plugins/Interchange/Runtime/Source/Parsers/GLTFCore/Private/GLTF/GLTFReader.cpp
@@ -84,6 +84,19 @@ namespace GLTF
 			return FBase64::GetDecodedDataSize(EncodedData);
 		}
 
+		FAccessor& AccessorAtIndex(TArray<FValidAccessor>& Accessors, int32 Index)
+		{
+			if (Accessors.IsValidIndex(Index))
+			{
+				return Accessors[Index];
+			}
+			else
+			{
+				static FVoidAccessor Void;
+				return Void;
+			}
+		}
+
 		const FAccessor& AccessorAtIndex(const TArray<FValidAccessor>& Accessors, int32 Index)
 		{
 			if (Accessors.IsValidIndex(Index))
@@ -220,20 +233,26 @@ namespace GLTF
 		}
 	}
 
-	void FFileReader::SetupMorphTarget(const FJsonObject& Object, GLTF::FPrimitive& Primitive) const
+	void FFileReader::SetupMorphTarget(const FJsonObject& Object, GLTF::FPrimitive& Primitive, const bool bMeshQuantized) const
 	{
 		const TArray<FValidAccessor>& A = Asset->Accessors;
-		const FAccessor& Position = AccessorAtIndex(A, GetIndex(Object, TEXT("POSITION")));
-		const FAccessor& Normal = AccessorAtIndex(A, GetIndex(Object, TEXT("NORMAL")));
-		const FAccessor& Tangent = AccessorAtIndex(A, GetIndex(Object, TEXT("TANGENT")));
-		const FAccessor& TexCoord0 = AccessorAtIndex(A, GetIndex(Object, TEXT("TEXCOORD_0")));
-		const FAccessor& TexCoord1 = AccessorAtIndex(A, GetIndex(Object, TEXT("TEXCOORD_1")));
+		FAccessor& Position = AccessorAtIndex(Asset->Accessors, GetIndex(Object, TEXT("POSITION")));
+		FAccessor& Normal = AccessorAtIndex(Asset->Accessors, GetIndex(Object, TEXT("NORMAL")));
+		FAccessor& Tangent = AccessorAtIndex(Asset->Accessors, GetIndex(Object, TEXT("TANGENT")));
+		FAccessor& TexCoord0 = AccessorAtIndex(Asset->Accessors, GetIndex(Object, TEXT("TEXCOORD_0")));
+		FAccessor& TexCoord1 = AccessorAtIndex(Asset->Accessors, GetIndex(Object, TEXT("TEXCOORD_1")));
 		const FAccessor& Color0 = AccessorAtIndex(A, GetIndex(Object, TEXT("COLOR_0")));
 
+		Position.bQuantized = bMeshQuantized;
+		Normal.bQuantized = bMeshQuantized;
+		Tangent.bQuantized = bMeshQuantized;
+		TexCoord0.bQuantized = bMeshQuantized;
+		TexCoord1.bQuantized = bMeshQuantized;
+
 		Primitive.MorphTargets.Emplace(Position, Normal, Tangent, TexCoord0, TexCoord1, Color0);
 	}
 
-	void FFileReader::SetupPrimitive(const FJsonObject& Object, FMesh& Mesh) const
+	void FFileReader::SetupPrimitive(const FJsonObject& Object, FMesh& Mesh, const bool bMeshQuantized) const
 	{
 		const FPrimitive::EMode       Mode = PrimitiveModeFromNumber(GetUnsignedInt(Object, TEXT("mode"), (uint32)FPrimitive::EMode::Triangles));
 		if (Mode == FPrimitive::EMode::Unknown)
@@ -248,15 +267,21 @@ namespace GLTF
 
 		// the only required attribute is POSITION
 		const FJsonObject& Attributes = *Object.GetObjectField("attributes");
-		const FAccessor&   Position   = AccessorAtIndex(A, GetIndex(Attributes, TEXT("POSITION")));
-		const FAccessor&   Normal     = AccessorAtIndex(A, GetIndex(Attributes, TEXT("NORMAL")));
-		const FAccessor&   Tangent    = AccessorAtIndex(A, GetIndex(Attributes, TEXT("TANGENT")));
-		const FAccessor&   TexCoord0  = AccessorAtIndex(A, GetIndex(Attributes, TEXT("TEXCOORD_0")));
-		const FAccessor&   TexCoord1  = AccessorAtIndex(A, GetIndex(Attributes, TEXT("TEXCOORD_1")));
+		FAccessor&   Position   = AccessorAtIndex(Asset->Accessors, GetIndex(Attributes, TEXT("POSITION")));
+		FAccessor&   Normal     = AccessorAtIndex(Asset->Accessors, GetIndex(Attributes, TEXT("NORMAL")));
+		FAccessor&   Tangent    = AccessorAtIndex(Asset->Accessors, GetIndex(Attributes, TEXT("TANGENT")));
+		FAccessor&   TexCoord0  = AccessorAtIndex(Asset->Accessors, GetIndex(Attributes, TEXT("TEXCOORD_0")));
+		FAccessor&   TexCoord1  = AccessorAtIndex(Asset->Accessors, GetIndex(Attributes, TEXT("TEXCOORD_1")));
 		const FAccessor&   Color0     = AccessorAtIndex(A, GetIndex(Attributes, TEXT("COLOR_0")));
 		const FAccessor&   Joints0    = AccessorAtIndex(A, GetIndex(Attributes, TEXT("JOINTS_0")));
 		const FAccessor&   Weights0   = AccessorAtIndex(A, GetIndex(Attributes, TEXT("WEIGHTS_0")));
 
+		Position.bQuantized = bMeshQuantized;
+		Normal.bQuantized = bMeshQuantized;
+		Tangent.bQuantized = bMeshQuantized;
+		TexCoord0.bQuantized = bMeshQuantized;
+		TexCoord1.bQuantized = bMeshQuantized;
+
 		Mesh.Primitives.Emplace(Mode, MaterialIndex, Indices, Position, Normal, Tangent, TexCoord0, TexCoord1, Color0, Joints0, Weights0);
 
 		//Morph Targets:
@@ -266,14 +291,14 @@ namespace GLTF
 			for (TSharedPtr<FJsonValue> Value : MorphTargets)
 			{
 				const FJsonObject& MorphTargetObject = *Value->AsObject();
-				SetupMorphTarget(MorphTargetObject, Mesh.Primitives.Last());
+				SetupMorphTarget(MorphTargetObject, Mesh.Primitives.Last(), bMeshQuantized);
 			}
 		}
 
 		ExtensionsHandler->SetupPrimitiveExtensions(Object, Mesh.Primitives.Last());
 	}
 
-	void FFileReader::SetupMesh(const FJsonObject& Object) const
+	void FFileReader::SetupMesh(const FJsonObject& Object, const bool bMeshQuantized) const
 	{
 		Asset->Meshes.Emplace();
 		FMesh& Mesh = Asset->Meshes.Last();
@@ -287,7 +312,7 @@ namespace GLTF
 		for (TSharedPtr<FJsonValue> Value : PrimArray)
 		{
 			const FJsonObject& PrimObject = *Value->AsObject();
-			SetupPrimitive(PrimObject, Mesh);
+			SetupPrimitive(PrimObject, Mesh, bMeshQuantized);
 
 			if (NumberOfMorphTargets == -1)
 			{
@@ -925,11 +950,37 @@ namespace GLTF
 		Asset = &OutAsset;
 		ExtensionsHandler->SetAsset(OutAsset);
 
+		/**
+		 * According to gltf specification checking only the top-level extensionsRequired property in order to decide
+		 * if the model import is supported should be sufficient as the documentation states:
+		 * "All glTF extensions required to load and/or render an asset MUST be listed in the top-level extensionsRequired array"
+		 */
+		bool bMeshQuantized = false;
+		const FString KHR_MeshQuantizationString = ToString(GLTF::EExtension::KHR_MeshQuantization);
+		const TArray<TSharedPtr<FJsonValue>>* ExtensionsRequired;
+		if (JsonRoot->TryGetArrayField(TEXT("extensionsRequired"), ExtensionsRequired))
+		{
+			if (ensure(ExtensionsRequired))
+			{
+				OutAsset.ExtensionsRequired.Reserve(OutAsset.ExtensionsRequired.Num() + ExtensionsRequired->Num());
+				for (const TSharedPtr<FJsonValue>& Extension : *ExtensionsRequired)
+				{
+					FString ExtensionString = Extension->AsString();
+					OutAsset.ExtensionsRequired.Add(ExtensionString);
+
+					if (ExtensionString == KHR_MeshQuantizationString)
+					{
+						bMeshQuantized = true;
+					}
+				}
+			}
+		}
+
 		SetupObjects(BufferCount, TEXT("buffers"), [this, InResourcesPath](const FJsonObject& Object) { SetupBuffer(Object, InResourcesPath); });
 		SetupObjects(BufferViewCount, TEXT("bufferViews"), [this](const FJsonObject& Object) { SetupBufferView(Object); });
 		SetupObjects(AccessorCount, TEXT("accessors"), [this](const FJsonObject& Object) { SetupAccessor(Object); });
 
-		SetupObjects(MeshCount, TEXT("meshes"), [this](const FJsonObject& Object) { SetupMesh(Object); });
+		SetupObjects(MeshCount, TEXT("meshes"), [this, &bMeshQuantized](const FJsonObject& Object) { SetupMesh(Object, bMeshQuantized); });
 		SetupObjects(NodeCount, TEXT("nodes"), [this](const FJsonObject& Object) { SetupNode(Object); });
 		SetupObjects(SceneCount, TEXT("scenes"), [this](const FJsonObject& Object) { SetupScene(Object); });
 		SetupObjects(CameraCount, TEXT("cameras"), [this](const FJsonObject& Object) { SetupCamera(Object); });
@@ -943,24 +994,6 @@ namespace GLTF
 		SetupObjects(TextureCount, TEXT("textures"), [this](const FJsonObject& Object) { SetupTexture(Object); });
 		SetupObjects(MaterialCount, TEXT("materials"), [this](const FJsonObject& Object) { SetupMaterial(Object); });
 
-		/**
-		 * Returns true if all required extensions are supported.
-		 * According to gltf specification checking only the top-level extensionsRequired property in order to decide
-		 * if the model import is supported should be sufficient as the documentation states:
-		 * "All glTF extensions required to load and/or render an asset MUST be listed in the top-level extensionsRequired array"
-		 */
-		const TArray<TSharedPtr<FJsonValue>>* ExtensionsRequired;
-		if (JsonRoot->TryGetArrayField(TEXT("extensionsRequired"), ExtensionsRequired))
-		{
-			if (ensure(ExtensionsRequired))
-			{
-				OutAsset.ExtensionsRequired.Reserve(OutAsset.ExtensionsRequired.Num() + ExtensionsRequired->Num());
-				for (const TSharedPtr<FJsonValue>& Extension : *ExtensionsRequired)
-				{
-					OutAsset.ExtensionsRequired.Add(Extension->AsString());
-				}
-			}
-		}
 
 		const TArray<TSharedPtr<FJsonValue>>* ExtensionsUsed;
 		if (JsonRoot->TryGetArrayField(TEXT("extensionsUsed"), ExtensionsUsed))
diff --git a/Engine/Plugins/Interchange/Runtime/Source/Parsers/GLTFCore/Public/GLTFAccessor.h b/Engine/Plugins/Interchange/Runtime/Source/Parsers/GLTFCore/Public/GLTFAccessor.h
--- a/Engine/Plugins/Interchange/Runtime/Source/Parsers/GLTFCore/Public/GLTFAccessor.h
+++ b/Engine/Plugins/Interchange/Runtime/Source/Parsers/GLTFCore/Public/GLTFAccessor.h
@@ -91,9 +91,10 @@ namespace GLTF
 		const uint32         Count;
 		const EType          Type;
 		const EComponentType ComponentType;
-		const bool           Normalized;
+		const bool           bNormalized;
+		bool                 bQuantized;
 
-		FAccessor(uint32 InCount, EType InType, EComponentType InComponentType, bool InNormalized);
+		FAccessor(uint32 InCount, EType InType, EComponentType InComponentType, bool bInNormalized);
 
 		virtual bool IsValid() const = 0;
 		virtual FMD5Hash GetHash() const = 0;
@@ -125,13 +126,29 @@ namespace GLTF
 		void         GetQuatArray(FVector4f* Buffer) const;
 		void         GetMat4Array(TArray<FMatrix44f>& Buffer) const;
 		virtual void GetMat4Array(FMatrix44f* Buffer) const;
+
+		enum class EDataType
+		{
+			Position,
+			Normal,
+			Tangent,
+			Texcoord,
+			Color,
+			Joints,
+			Weights
+		};
+		bool IsValidDataType(EDataType DataType, bool bMorphTargetProperty) const;
+		bool CheckAccessorTypeForDataType(EDataType DataType, bool bMorphTargetProperty) const;
+		bool CheckNonQuantizedComponentTypeForDataType(EDataType DataType, bool bMorphTargetProperty) const;
+		bool CheckQuantizedComponentTypeForDataType(EDataType DataType, bool bMorphTargetProperty) const;
 	};
 
 	struct GLTFCORE_API FValidAccessor final : FAccessor
 	{
-		FValidAccessor(FBufferView& InBufferView, uint32 InOffset, uint32 InCount, EType InType, EComponentType InCompType, bool InNormalized);
+		FValidAccessor(FBufferView& InBufferView, uint32 InOffset, uint32 InCount, EType InType, EComponentType InCompType, bool bInNormalized);
 
 		bool IsValid() const override;
+
 		FMD5Hash GetHash() const override;
 
 		uint32 GetUnsignedInt(uint32 Index) const override;
diff --git a/Engine/Plugins/Interchange/Runtime/Source/Parsers/GLTFCore/Public/GLTFAsset.h b/Engine/Plugins/Interchange/Runtime/Source/Parsers/GLTFCore/Public/GLTFAsset.h
--- a/Engine/Plugins/Interchange/Runtime/Source/Parsers/GLTFCore/Public/GLTFAsset.h
+++ b/Engine/Plugins/Interchange/Runtime/Source/Parsers/GLTFCore/Public/GLTFAsset.h
@@ -45,6 +45,7 @@ namespace GLTF
 		MSFT_TextureDDS,
 		MSFT_PackingNormalRoughnessMetallic,
 		MSFT_PackingOcclusionRoughnessMetallic,
+		KHR_MeshQuantization,
 		Count
 	};
 
diff --git a/Engine/Plugins/Interchange/Runtime/Source/Parsers/GLTFCore/Public/GLTFReader.h b/Engine/Plugins/Interchange/Runtime/Source/Parsers/GLTFCore/Public/GLTFReader.h
--- a/Engine/Plugins/Interchange/Runtime/Source/Parsers/GLTFCore/Public/GLTFReader.h
+++ b/Engine/Plugins/Interchange/Runtime/Source/Parsers/GLTFCore/Public/GLTFReader.h
@@ -43,9 +43,9 @@ namespace GLTF
 		void SetupBuffer(const FJsonObject& Object, const FString& Path);
 		void SetupBufferView(const FJsonObject& Object) const;
 		void SetupAccessor(const FJsonObject& Object) const;
-		void SetupMorphTarget(const FJsonObject& Object, GLTF::FPrimitive& Primitive) const;
-		void SetupPrimitive(const FJsonObject& Object, GLTF::FMesh& Mesh) const;
-		void SetupMesh(const FJsonObject& Object) const;
+		void SetupMorphTarget(const FJsonObject& Object, GLTF::FPrimitive& Primitive, const bool bMeshQuantized) const;
+		void SetupPrimitive(const FJsonObject& Object, GLTF::FMesh& Mesh, const bool bMeshQuantized) const;
+		void SetupMesh(const FJsonObject& Object, const bool bMeshQuantized) const;
 
 		void SetupScene(const FJsonObject& Object) const;
 		void SetupNode(const FJsonObject& Object) const;
