Interchange: Pipeline now have an optional "asset name" property use when importing 1 asset into the content browser.

#rnx

diff --git a/Engine/Plugins/Interchange/Editor/Source/Pipelines/Private/InterchangeEditorPipelineDetails.cpp b/Engine/Plugins/Interchange/Editor/Source/Pipelines/Private/InterchangeEditorPipelineDetails.cpp
--- a/Engine/Plugins/Interchange/Editor/Source/Pipelines/Private/InterchangeEditorPipelineDetails.cpp
+++ b/Engine/Plugins/Interchange/Editor/Source/Pipelines/Private/InterchangeEditorPipelineDetails.cpp
@@ -295,7 +295,7 @@ void FInterchangePipelineBaseDetailsCustomization::CustomizeDetails(IDetailLayou
 			{
 				continue;
 			}
-			const FName PropertyName = PropertyHandle.Get().GetProperty() ? PropertyHandle.Get().GetProperty()->GetFName() : NAME_None;
+			const FName PropertyName = PropertyPtr ? PropertyPtr->GetFName() : NAME_None;
 			if (PropertyName == UInterchangePipelineBase::GetPropertiesStatesPropertyName())
 			{
 				CachedDetailBuilder->HideProperty(PropertyHandle);
@@ -306,6 +306,18 @@ void FInterchangePipelineBaseDetailsCustomization::CustomizeDetails(IDetailLayou
 			{
 				continue;
 			}
+
+			//If the property UObject owner class is not the same type as the main pipeline class, it mean
+			//we deal with a sub pipeline and we must hide property that have the "StandAlonePipelineProperty" meta data
+			if(UClass* InterchangePipelinePropertyClass = Cast<UClass>(PropertyPtr->GetOwnerUObject()))
+			{
+				if (InterchangePipeline->GetClass() != InterchangePipelinePropertyClass && PropertyHandle->GetBoolMetaData(FName("StandAlonePipelineProperty")))
+				{
+					CachedDetailBuilder->HideProperty(PropertyHandle);
+					continue;
+				}
+			}
+			
 			FName PropertyPath = FName(PropertyPtr->GetPathName());
 			CachedDetailBuilder->HideProperty(PropertyHandle);
 
diff --git a/Engine/Plugins/Interchange/Runtime/Source/Pipelines/Private/InterchangeGenericAssetsPipeline.cpp b/Engine/Plugins/Interchange/Runtime/Source/Pipelines/Private/InterchangeGenericAssetsPipeline.cpp
--- a/Engine/Plugins/Interchange/Runtime/Source/Pipelines/Private/InterchangeGenericAssetsPipeline.cpp
+++ b/Engine/Plugins/Interchange/Runtime/Source/Pipelines/Private/InterchangeGenericAssetsPipeline.cpp
@@ -553,16 +553,17 @@ void UInterchangeGenericAssetsPipeline::ImplementUseSourceNameForAssetOption()
 	//If we import only one mesh, we want to rename the mesh using the file name.
 	const int32 MeshesImportedNodeCount = SkeletalMeshNodeUids.Num() + StaticMeshNodeUids.Num();
 
+	FString OverrideAssetName = IsStandAlonePipeline() ? AssetName : FString();
 	//SkeletalMesh it must always be run even if there is no rename option, skeleton and physics asset will be rename properly
-	MeshPipeline->ImplementUseSourceNameForAssetOptionSkeletalMesh(MeshesImportedNodeCount, bUseSourceNameForAsset, AssetName);
+	MeshPipeline->ImplementUseSourceNameForAssetOptionSkeletalMesh(MeshesImportedNodeCount, bUseSourceNameForAsset, OverrideAssetName);
 
-	if (bUseSourceNameForAsset || !AssetName.IsEmpty())
+	if (bUseSourceNameForAsset || !OverrideAssetName.IsEmpty())
 	{
 		//StaticMesh
 		if (MeshesImportedNodeCount == 1 && StaticMeshNodeUids.Num() > 0)
 		{
 			UInterchangeStaticMeshFactoryNode* StaticMeshNode = Cast<UInterchangeStaticMeshFactoryNode>(BaseNodeContainer->GetFactoryNode(StaticMeshNodeUids[0]));
-			const FString DisplayLabelName = AssetName.IsEmpty() ? FPaths::GetBaseFilename(SourceDatas[0]->GetFilename()) : AssetName;
+			const FString DisplayLabelName = OverrideAssetName.IsEmpty() ? FPaths::GetBaseFilename(SourceDatas[0]->GetFilename()) : OverrideAssetName;
 			StaticMeshNode->SetDisplayLabel(DisplayLabelName);
 		}
 
@@ -570,7 +571,7 @@ void UInterchangeGenericAssetsPipeline::ImplementUseSourceNameForAssetOption()
 		if (AnimSequenceNodeUids.Num() == 1)
 		{
 			UInterchangeAnimSequenceFactoryNode* AnimSequenceNode = Cast<UInterchangeAnimSequenceFactoryNode>(BaseNodeContainer->GetFactoryNode(AnimSequenceNodeUids[0]));
-			const FString DisplayLabelName = (AssetName.IsEmpty() ? FPaths::GetBaseFilename(SourceDatas[0]->GetFilename()) : AssetName) + TEXT("_Anim");
+			const FString DisplayLabelName = (OverrideAssetName.IsEmpty() ? FPaths::GetBaseFilename(SourceDatas[0]->GetFilename()) : OverrideAssetName) + TEXT("_Anim");
 			AnimSequenceNode->SetDisplayLabel(DisplayLabelName);
 		}
 	}
diff --git a/Engine/Plugins/Interchange/Runtime/Source/Pipelines/Private/InterchangeGenericMaterialPipeline.cpp b/Engine/Plugins/Interchange/Runtime/Source/Pipelines/Private/InterchangeGenericMaterialPipeline.cpp
--- a/Engine/Plugins/Interchange/Runtime/Source/Pipelines/Private/InterchangeGenericMaterialPipeline.cpp
+++ b/Engine/Plugins/Interchange/Runtime/Source/Pipelines/Private/InterchangeGenericMaterialPipeline.cpp
@@ -457,6 +457,21 @@ void UInterchangeGenericMaterialPipeline::ExecutePipeline(UInterchangeBaseNodeCo
 			}
 		}
 	}
+
+	if (IsStandAlonePipeline() && !AssetName.IsEmpty())
+	{
+		TArray<UInterchangeBaseMaterialFactoryNode*> BaseMaterialNodes;
+		BaseNodeContainer->IterateNodesOfType<UInterchangeBaseMaterialFactoryNode>([&BaseMaterialNodes](const FString& NodeUid, UInterchangeBaseMaterialFactoryNode* MaterialNode)
+			{
+				BaseMaterialNodes.Add(MaterialNode);
+			});
+
+		if(BaseMaterialNodes.Num() == 1)
+		{
+			BaseMaterialNodes[0]->SetAssetName(AssetName);
+			BaseMaterialNodes[0]->SetDisplayLabel(AssetName);
+		}
+	}
 }
 
 void UInterchangeGenericMaterialPipeline::ExecutePostFactoryPipeline(const UInterchangeBaseNodeContainer* InBaseNodeContainer, const FString& NodeKey, UObject* CreatedAsset, bool bIsAReimport)
diff --git a/Engine/Plugins/Interchange/Runtime/Source/Pipelines/Private/InterchangeGenericTexturePipeline.cpp b/Engine/Plugins/Interchange/Runtime/Source/Pipelines/Private/InterchangeGenericTexturePipeline.cpp
--- a/Engine/Plugins/Interchange/Runtime/Source/Pipelines/Private/InterchangeGenericTexturePipeline.cpp
+++ b/Engine/Plugins/Interchange/Runtime/Source/Pipelines/Private/InterchangeGenericTexturePipeline.cpp
@@ -167,9 +167,17 @@ void UInterchangeGenericTexturePipeline::ExecutePipeline(UInterchangeBaseNodeCon
 
 	if (bImportTextures)
 	{
-		for (const UInterchangeTextureNode* TextureNode : TextureNodes)
+		UInterchangeTextureFactoryNode* TextureFactoryNode = nullptr;
+		for (UInterchangeTextureNode* TextureNode : TextureNodes)
 		{
-			HandleCreationOfTextureFactoryNode(TextureNode);
+			TextureFactoryNode = HandleCreationOfTextureFactoryNode(TextureNode);
+		}
+		//If we have a valid override name
+		const bool bOverrideAssetName = TextureNodes.Num() == 1 && IsStandAlonePipeline() && !AssetName.IsEmpty();
+		if (TextureFactoryNode && bOverrideAssetName)
+		{
+			TextureFactoryNode->SetAssetName(AssetName);
+			TextureFactoryNode->SetDisplayLabel(AssetName);
 		}
 	}
 }
diff --git a/Engine/Plugins/Interchange/Runtime/Source/Pipelines/Public/InterchangeGenericAssetsPipeline.h b/Engine/Plugins/Interchange/Runtime/Source/Pipelines/Public/InterchangeGenericAssetsPipeline.h
--- a/Engine/Plugins/Interchange/Runtime/Source/Pipelines/Public/InterchangeGenericAssetsPipeline.h
+++ b/Engine/Plugins/Interchange/Runtime/Source/Pipelines/Public/InterchangeGenericAssetsPipeline.h
@@ -41,7 +41,7 @@ public:
 	bool bUseSourceNameForAsset = true;
 
 	/** If not empty, and there is only one asset and one source data, we will name the asset with this string. */
-	UPROPERTY(EditAnywhere, BlueprintReadWrite, Category = "Common")
+	UPROPERTY(EditAnywhere, BlueprintReadWrite, Category = "Common", meta = (StandAlonePipelineProperty = "True", AlwaysResetToDefault = "True"))
 	FString AssetName;
 
 	/** Translation offset applied to meshes and animations. */
diff --git a/Engine/Plugins/Interchange/Runtime/Source/Pipelines/Public/InterchangeGenericMaterialPipeline.h b/Engine/Plugins/Interchange/Runtime/Source/Pipelines/Public/InterchangeGenericMaterialPipeline.h
--- a/Engine/Plugins/Interchange/Runtime/Source/Pipelines/Public/InterchangeGenericMaterialPipeline.h
+++ b/Engine/Plugins/Interchange/Runtime/Source/Pipelines/Public/InterchangeGenericMaterialPipeline.h
@@ -43,6 +43,10 @@ public:
 	UPROPERTY(EditAnywhere, BlueprintReadWrite, Category = "Materials")
 	bool bImportMaterials = true;
 
+	/** If not empty, and there is only one asset and one source data, we will name the asset with this string. */
+	UPROPERTY(EditAnywhere, BlueprintReadWrite, Category = "Materials", meta = (StandAlonePipelineProperty = "True", AlwaysResetToDefault = "True"))
+	FString AssetName;
+
 	UPROPERTY(EditAnywhere, BlueprintReadWrite, Category = "Materials", Meta=(EditCondition="bImportMaterials"))
 	EInterchangeMaterialImportOption MaterialImport = EInterchangeMaterialImportOption::ImportAsMaterials;
 
diff --git a/Engine/Plugins/Interchange/Runtime/Source/Pipelines/Public/InterchangeGenericTexturePipeline.h b/Engine/Plugins/Interchange/Runtime/Source/Pipelines/Public/InterchangeGenericTexturePipeline.h
--- a/Engine/Plugins/Interchange/Runtime/Source/Pipelines/Public/InterchangeGenericTexturePipeline.h
+++ b/Engine/Plugins/Interchange/Runtime/Source/Pipelines/Public/InterchangeGenericTexturePipeline.h
@@ -23,6 +23,10 @@ public:
 	UPROPERTY(EditAnywhere, BlueprintReadWrite, Category = "Textures")
 	bool bImportTextures = true;
 
+	/** If not empty, and there is only one asset and one source data, we will name the asset with this string. */
+	UPROPERTY(EditAnywhere, BlueprintReadWrite, Category = "Textures", meta=(StandAlonePipelineProperty = "True", AlwaysResetToDefault = "True"))
+	FString AssetName;
+
 #if WITH_EDITORONLY_DATA
 	/** 
 	 * If enable, after a new import a test will be run to see if the texture is a normal map.
@@ -83,5 +87,3 @@ private:
 	/** Texture factory assets nodes */
 	TArray<UInterchangeTextureFactoryNode*> TextureFactoryNodes;
 };
-
-
diff --git a/Engine/Source/Developer/AssetTools/Private/AssetTools.cpp b/Engine/Source/Developer/AssetTools/Private/AssetTools.cpp
--- a/Engine/Source/Developer/AssetTools/Private/AssetTools.cpp
+++ b/Engine/Source/Developer/AssetTools/Private/AssetTools.cpp
@@ -3300,6 +3300,10 @@ TArray<UObject*> UAssetToolsImpl::ImportAssetsInternal(const TArray<FString>& Fi
 				ImportAssetParameters.bIsAutomated = bAutomatedImport;
 				ImportAssetParameters.bFollowRedirectors = UE::AssetTools::Private::bFollowRedirectorsWhenImporting;
 				ImportAssetParameters.ReimportAsset = nullptr;
+				if (Params.AssetImportTask && !Params.AssetImportTask->DestinationName.IsEmpty())
+				{
+					UE_LOG(LogAssetTools, Warning, TEXT("Interchange do not use the DestinationName parameter. Users must specify the destination name in a pipeline through the Options parameter."));
+				}
 
 
 				TFunction<void(UE::Interchange::FImportResult&)> AppendImportResult =
diff --git a/Engine/Source/Runtime/Interchange/Core/Private/InterchangePipelineBase.cpp b/Engine/Source/Runtime/Interchange/Core/Private/InterchangePipelineBase.cpp
--- a/Engine/Source/Runtime/Interchange/Core/Private/InterchangePipelineBase.cpp
+++ b/Engine/Source/Runtime/Interchange/Core/Private/InterchangePipelineBase.cpp
@@ -121,6 +121,13 @@ void UInterchangePipelineBase::LoadSettingsInternal(const FName PipelineStackNam
 		{
 			continue;
 		}
+#if WITH_EDITOR
+		if (Property->GetBoolMetaData(FName("AlwaysResetToDefault")))
+		{
+			//Stand alone pipeline property
+			continue;
+		}
+#endif //WITH_EDITOR
 
 		if (const FInterchangePipelinePropertyStates* PropertyStates = ParentPropertiesStates.Find(PropertyPath))
 		{
diff --git a/Engine/Source/Runtime/Interchange/Core/Public/InterchangePipelineBase.h b/Engine/Source/Runtime/Interchange/Core/Public/InterchangePipelineBase.h
--- a/Engine/Source/Runtime/Interchange/Core/Public/InterchangePipelineBase.h
+++ b/Engine/Source/Runtime/Interchange/Core/Public/InterchangePipelineBase.h
@@ -312,6 +312,17 @@ public:
 	}
 #endif //WITH_EDITOR
 
+	/*
+	 * Return true if the pipeline is not a sub pipeline. Stand Alone Pipeline can have some extra properties.
+	 * The customize detail of the pipeline will call this function to show or hide those property.
+	 * The pipeline itself can call the function to know if a property must be use or not.
+	 * Example: the override asset name property
+	 */
+	bool IsStandAlonePipeline()
+	{
+		return this == GetMostPipelineOuter();
+	};
+
 protected:
 
 	UE_DEPRECATED(5.2, "This function is replace by ExecutePipeline.")
