[Interchange] Mesh force consistency fixes.
Animation creation fixes for skeletal mesh forcing.
StaticMesh creation fixes for force mesh type settings.

diff --git a/Engine/Plugins/Interchange/Runtime/Source/Pipelines/Private/InterchangeGenericStaticMeshPipeline.cpp b/Engine/Plugins/Interchange/Runtime/Source/Pipelines/Private/InterchangeGenericStaticMeshPipeline.cpp
--- a/Engine/Plugins/Interchange/Runtime/Source/Pipelines/Private/InterchangeGenericStaticMeshPipeline.cpp
+++ b/Engine/Plugins/Interchange/Runtime/Source/Pipelines/Private/InterchangeGenericStaticMeshPipeline.cpp
@@ -168,7 +168,7 @@ void UInterchangeGenericMeshPipeline::ExecutePreImportPipelineStaticMesh()
 			{
 				// If baking transforms, get all the static mesh instance nodes, and group them by LOD
 				TArray<FString> MeshUids;
-				PipelineMeshesUtilities->GetAllStaticMeshInstance(MeshUids, bConvertSkeletalMeshToStaticMesh);
+				PipelineMeshesUtilities->GetAllStaticMeshInstance(MeshUids, bConvertSkeletalMeshToStaticMesh, bConvertStaticsWithMorphTargetsToSkeletals);
 
 				TMap<int32, TArray<FString>> MeshUidsPerLodIndex;
 
@@ -229,7 +229,7 @@ void UInterchangeGenericMeshPipeline::ExecutePreImportPipelineStaticMesh()
 			bool bFoundMeshes = false;
 			{
 				TArray<FString> MeshUids;
-				PipelineMeshesUtilities->GetAllStaticMeshInstance(MeshUids, bConvertSkeletalMeshToStaticMesh);
+				PipelineMeshesUtilities->GetAllStaticMeshInstance(MeshUids, bConvertSkeletalMeshToStaticMesh, bConvertStaticsWithMorphTargetsToSkeletals);
 
 				// Work out which meshes are collision meshes which correspond to another mesh
 				TMap<FString, TArray<FString>> MeshToCollisionMeshMap;
diff --git a/Engine/Plugins/Interchange/Runtime/Source/Pipelines/Private/InterchangePipelineMeshesUtilities.cpp b/Engine/Plugins/Interchange/Runtime/Source/Pipelines/Private/InterchangePipelineMeshesUtilities.cpp
--- a/Engine/Plugins/Interchange/Runtime/Source/Pipelines/Private/InterchangePipelineMeshesUtilities.cpp
+++ b/Engine/Plugins/Interchange/Runtime/Source/Pipelines/Private/InterchangePipelineMeshesUtilities.cpp
@@ -353,37 +353,37 @@ void UInterchangePipelineMeshesUtilities::GetAllSkinnedMeshInstance(TArray<FStri
 	}
 }
 
-void UInterchangePipelineMeshesUtilities::IterateAllSkinnedMeshInstance(TFunctionRef<void(const FInterchangeMeshInstance&)> IterationLambda, const bool bConvertStaticMeshToSkeletalMesh) const
+void UInterchangePipelineMeshesUtilities::IterateAllSkinnedMeshInstance(TFunctionRef<void(const FInterchangeMeshInstance&)> IterationLambda, const bool bConvertStaticMeshToSkeletalMesh, const bool bConvertStaticsWithMorphTargetsToSkeletals) const
 {
 	for (const TPair<FString, FInterchangeMeshInstance>& MeshInstanceUidAndMeshInstance : MeshInstancesPerMeshInstanceUid)
 	{
 		const FInterchangeMeshInstance& MeshInstance = MeshInstanceUidAndMeshInstance.Value;
-		if ((bConvertStaticMeshToSkeletalMesh || MeshInstance.bReferenceSkinnedMesh) && !MeshInstance.bReferenceMorphTarget)
+		if ((bConvertStaticMeshToSkeletalMesh || MeshInstance.bReferenceSkinnedMesh || (bConvertStaticsWithMorphTargetsToSkeletals && MeshInstance.bHasMorphTargets)) && !MeshInstance.bReferenceMorphTarget)
 		{
 			IterationLambda(MeshInstance);
 		}
 	}
 }
 
-void UInterchangePipelineMeshesUtilities::GetAllStaticMeshInstance(TArray<FString>& MeshInstanceUids, const bool bConvertSkeletalMeshToStaticMesh) const
+void UInterchangePipelineMeshesUtilities::GetAllStaticMeshInstance(TArray<FString>& MeshInstanceUids, const bool bConvertSkeletalMeshToStaticMesh, const bool bConvertStaticsWithMorphTargetsToSkeletals) const
 {
 	MeshInstanceUids.Empty(MeshInstancesPerMeshInstanceUid.Num());
 	for (const TPair<FString, FInterchangeMeshInstance>& MeshInstanceUidAndMeshInstance : MeshInstancesPerMeshInstanceUid)
 	{
 		const FInterchangeMeshInstance& MeshInstance = MeshInstanceUidAndMeshInstance.Value;
-		if ((bConvertSkeletalMeshToStaticMesh || !MeshInstance.bReferenceSkinnedMesh) && !MeshInstance.bReferenceMorphTarget)
+		if ((bConvertSkeletalMeshToStaticMesh || !MeshInstance.bReferenceSkinnedMesh) && !MeshInstance.bReferenceMorphTarget && !(bConvertStaticsWithMorphTargetsToSkeletals && MeshInstance.bHasMorphTargets))
 		{
 			MeshInstanceUids.Add(MeshInstance.MeshInstanceUid);
 		}
 	}
 }
 
-void UInterchangePipelineMeshesUtilities::IterateAllStaticMeshInstance(TFunctionRef<void(const FInterchangeMeshInstance&)> IterationLambda, const bool bConvertSkeletalMeshToStaticMesh) const
+void UInterchangePipelineMeshesUtilities::IterateAllStaticMeshInstance(TFunctionRef<void(const FInterchangeMeshInstance&)> IterationLambda, const bool bConvertSkeletalMeshToStaticMesh, const bool bConvertStaticsWithMorphTargetsToSkeletals) const
 {
 	for (const TPair<FString, FInterchangeMeshInstance>& MeshInstanceUidAndMeshInstance : MeshInstancesPerMeshInstanceUid)
 	{
 		const FInterchangeMeshInstance& MeshInstance = MeshInstanceUidAndMeshInstance.Value;
-		if ((bConvertSkeletalMeshToStaticMesh || !MeshInstance.bReferenceSkinnedMesh) && !MeshInstance.bReferenceMorphTarget)
+		if ((bConvertSkeletalMeshToStaticMesh || !MeshInstance.bReferenceSkinnedMesh) && !MeshInstance.bReferenceMorphTarget && !(bConvertStaticsWithMorphTargetsToSkeletals && MeshInstance.bHasMorphTargets))
 		{
 			IterationLambda(MeshInstance);
 		}
@@ -549,7 +549,7 @@ void UInterchangePipelineMeshesUtilities::GetCombinedSkinnedMeshInstances(UInter
 				//Create a joint from the instance node (the scene node pointing on the mesh).
 				SkeletonRootUid = MeshInstance.MeshInstanceUid;
 
-				if (!bConvertStaticsWithMorphTargetsToSkeletals)
+				if (!MeshInstance.bHasMorphTargets)
 				{
 					//Since we are dealing with rigid mesh and combine we will use the outermost valid parent
 					if (const UInterchangeSceneNode* SceneNode = Cast<UInterchangeSceneNode>(BaseNodeContainer->GetNode(MeshInstance.MeshInstanceUid)))
diff --git a/Engine/Plugins/Interchange/Runtime/Source/Pipelines/Public/InterchangePipelineMeshesUtilities.h b/Engine/Plugins/Interchange/Runtime/Source/Pipelines/Public/InterchangePipelineMeshesUtilities.h
--- a/Engine/Plugins/Interchange/Runtime/Source/Pipelines/Public/InterchangePipelineMeshesUtilities.h
+++ b/Engine/Plugins/Interchange/Runtime/Source/Pipelines/Public/InterchangePipelineMeshesUtilities.h
@@ -153,18 +153,18 @@ public:
 	/**
 	* Iterate all skinned mesh instance.
 	*/
-	void IterateAllSkinnedMeshInstance(TFunctionRef<void(const FInterchangeMeshInstance&)> IterationLambda, const bool bConvertStaticMeshToSkeletalMesh) const;
+	void IterateAllSkinnedMeshInstance(TFunctionRef<void(const FInterchangeMeshInstance&)> IterationLambda, const bool bConvertStaticMeshToSkeletalMesh, const bool bConvertStaticsWithMorphTargetsToSkeletals) const;
 
 	/**
 	* Get all static mesh instance unique ids.
 	*/
 	UFUNCTION(BlueprintCallable, Category = "Interchange | Pipeline | Meshes")
-	void GetAllStaticMeshInstance(TArray<FString>& MeshInstanceUids, const bool bConvertSkeletalMeshToStaticMesh) const;
+	void GetAllStaticMeshInstance(TArray<FString>& MeshInstanceUids, const bool bConvertSkeletalMeshToStaticMesh, const bool bConvertStaticsWithMorphTargetsToSkeletals) const;
 
 	/**
 	* Iterate all static mesh instance.
 	*/
-	void IterateAllStaticMeshInstance(TFunctionRef<void(const FInterchangeMeshInstance&)> IterationLambda, const bool bConvertSkeletalMeshToStaticMesh) const;
+	void IterateAllStaticMeshInstance(TFunctionRef<void(const FInterchangeMeshInstance&)> IterationLambda, const bool bConvertSkeletalMeshToStaticMesh, const bool bConvertStaticsWithMorphTargetsToSkeletals) const;
 
 	/**
 	* Get all mesh geometry unique ids.
