Change the AssetImportTask interchange options to use FSoftObjectPath instead of TObjectPtr to allow adding python and blueprint pipeline when filling the options.

#rnx

diff --git a/Engine/Plugins/Interchange/Runtime/Source/Pipelines/Public/InterchangeGenericAssetsPipeline.h b/Engine/Plugins/Interchange/Runtime/Source/Pipelines/Public/InterchangeGenericAssetsPipeline.h
--- a/Engine/Plugins/Interchange/Runtime/Source/Pipelines/Public/InterchangeGenericAssetsPipeline.h
+++ b/Engine/Plugins/Interchange/Runtime/Source/Pipelines/Public/InterchangeGenericAssetsPipeline.h
@@ -57,23 +57,23 @@ public:
 	float ImportOffsetUniformScale = 1.0f;
 
 	//////	COMMON_MESHES_CATEGORY Properties //////
-	UPROPERTY(VisibleAnywhere, Instanced, Category = "Common Meshes")
+	UPROPERTY(VisibleAnywhere, BlueprintReadWrite, Instanced, Category = "Common Meshes")
 	TObjectPtr<UInterchangeGenericCommonMeshesProperties> CommonMeshesProperties;
 		
 	//////  COMMON_SKELETAL_ANIMATIONS_CATEGORY //////
-	UPROPERTY(VisibleAnywhere, Instanced, Category = "Common Skeletal Meshes and Animations")
+	UPROPERTY(VisibleAnywhere, BlueprintReadWrite, Instanced, Category = "Common Skeletal Meshes and Animations")
 	TObjectPtr<UInterchangeGenericCommonSkeletalMeshesAndAnimationsProperties> CommonSkeletalMeshesAndAnimationsProperties;
 
 	//////	MESHES_CATEGORY Properties //////
-	UPROPERTY(VisibleAnywhere, Instanced, Category = "Meshes")
+	UPROPERTY(VisibleAnywhere, BlueprintReadWrite, Instanced, Category = "Meshes")
 	TObjectPtr<UInterchangeGenericMeshPipeline> MeshPipeline;
 
 	//////	ANIMATIONS_CATEGORY Properties //////
-	UPROPERTY(VisibleAnywhere, Instanced, Category = "Animation")
+	UPROPERTY(VisibleAnywhere, BlueprintReadWrite, Instanced, Category = "Animation")
 	TObjectPtr<UInterchangeGenericAnimationPipeline> AnimationPipeline;
 
 	//////	MATERIALS_CATEGORY Properties //////
-	UPROPERTY(VisibleAnywhere, Instanced, Category = "Materials")
+	UPROPERTY(VisibleAnywhere, BlueprintReadWrite, Instanced, Category = "Materials")
 	TObjectPtr<UInterchangeGenericMaterialPipeline> MaterialPipeline;
 
 	virtual void PreDialogCleanup(const FName PipelineStackName) override;
diff --git a/Engine/Plugins/Interchange/Runtime/Source/Pipelines/Public/InterchangeGenericMaterialPipeline.h b/Engine/Plugins/Interchange/Runtime/Source/Pipelines/Public/InterchangeGenericMaterialPipeline.h
--- a/Engine/Plugins/Interchange/Runtime/Source/Pipelines/Public/InterchangeGenericMaterialPipeline.h
+++ b/Engine/Plugins/Interchange/Runtime/Source/Pipelines/Public/InterchangeGenericMaterialPipeline.h
@@ -54,7 +54,7 @@ public:
 	UPROPERTY(EditAnywhere, BlueprintReadWrite, Category = "Materials", Meta= (EditCondition="bImportMaterials && MaterialImport==EInterchangeMaterialImportOption::ImportAsMaterialInstances", AllowedClasses="/Script/Engine.MaterialInterface"))
 	FSoftObjectPath ParentMaterial;
 
-	UPROPERTY(VisibleAnywhere, Instanced, Category = "Textures")
+	UPROPERTY(VisibleAnywhere, BlueprintReadWrite, Instanced, Category = "Textures")
 	TObjectPtr<UInterchangeGenericTexturePipeline> TexturePipeline;
 
 	/** BEGIN UInterchangePipelineBase overrides */
diff --git a/Engine/Source/Developer/AssetTools/Private/AssetTools.cpp b/Engine/Source/Developer/AssetTools/Private/AssetTools.cpp
--- a/Engine/Source/Developer/AssetTools/Private/AssetTools.cpp
+++ b/Engine/Source/Developer/AssetTools/Private/AssetTools.cpp
@@ -3392,9 +3392,9 @@ TArray<UObject*> UAssetToolsImpl::ImportAssetsInternal(const TArray<FString>& Fi
 						//Add the override pipelines if we have a task that specify valid interchange override pipelines
 						if (UInterchangePipelineStackOverride* PipelineStackOverride = Cast<UInterchangePipelineStackOverride>(Params.AssetImportTask->Options))
 						{
-							for (TObjectPtr<UInterchangePipelineBase> OverridePipeline : PipelineStackOverride->OverridePipelines)
+							for (const FSoftObjectPath& OverridePipelinePath : PipelineStackOverride->OverridePipelines)
 							{
-								ImportAssetParameters.OverridePipelines.Add(OverridePipeline);
+								ImportAssetParameters.OverridePipelines.Add(OverridePipelinePath);
 							}
 						}
 					}
diff --git a/Engine/Source/Runtime/Interchange/Engine/Private/InterchangeManager.cpp b/Engine/Source/Runtime/Interchange/Engine/Private/InterchangeManager.cpp
--- a/Engine/Source/Runtime/Interchange/Engine/Private/InterchangeManager.cpp
+++ b/Engine/Source/Runtime/Interchange/Engine/Private/InterchangeManager.cpp
@@ -775,6 +775,21 @@ UInterchangePipelineBase* UE::Interchange::GeneratePipelineInstance(const FSoftO
 	return GeneratedPipeline;
 }
 
+void UInterchangePipelineStackOverride::AddPythonPipeline(UInterchangePythonPipelineBase* PipelineBase)
+{
+	OverridePipelines.Add(PipelineBase);
+}
+
+void UInterchangePipelineStackOverride::AddBlueprintPipeline(UInterchangeBlueprintPipelineBase* PipelineBase)
+{
+	OverridePipelines.Add(PipelineBase);
+}
+
+void UInterchangePipelineStackOverride::AddPipeline(UInterchangePipelineBase* PipelineBase)
+{
+	OverridePipelines.Add(PipelineBase);
+}
+
 UInterchangeManager& UInterchangeManager::GetInterchangeManager()
 {
 	static TStrongObjectPtr<UInterchangeManager> InterchangeManager = nullptr;
@@ -1271,6 +1286,15 @@ UInterchangeManager::ImportInternal(const FString& ContentPath, const UInterchan
 		return EarlyExit();
 	}
 	
+	{
+		const UInterchangeTranslatorBase* Translator = GetTranslatorForSourceData(SourceData);
+		if(!Translator)
+		{
+			UE_LOG(LogInterchangeEngine, Error, TEXT("Cannot import file, the source data is not supported. See if you can enable for interchange the extension [%s]"), *FPaths::GetExtension(SourceData->GetFilename()));
+			return EarlyExit();
+		}
+	}
+
 	if (FEngineAnalytics::IsAvailable())
 	{
 		Attribs.Add(FAnalyticsEventAttribute(TEXT("SourceExtension"), FPaths::GetExtension(SourceData->GetFilename())));
diff --git a/Engine/Source/Runtime/Interchange/Engine/Private/Tasks/InterchangeTaskParsing.cpp b/Engine/Source/Runtime/Interchange/Engine/Private/Tasks/InterchangeTaskParsing.cpp
--- a/Engine/Source/Runtime/Interchange/Engine/Private/Tasks/InterchangeTaskParsing.cpp
+++ b/Engine/Source/Runtime/Interchange/Engine/Private/Tasks/InterchangeTaskParsing.cpp
@@ -327,6 +327,11 @@ void UE::Interchange::FTaskParsing::DoTask(ENamedThreads::Type CurrentThread, co
 	{
 		for(TObjectPtr<UInterchangeTranslatorBase> Translator : AsyncHelper->Translators)
 		{
+			//We need a valid translator to add a result message
+			if (!Translator)
+			{
+				continue;
+			}
 			UInterchangeResultError_Generic* ErrorResult = NewObject<UInterchangeResultError_Generic>(GetTransientPackage(), UInterchangeResultError_Generic::StaticClass());
 			ErrorResult->Text = NSLOCTEXT("InterchangeTaskParsing", "TranslatorNoAssetImported", "There was no asset to import in the provided source data (found no mesh and no animation).");
 			ErrorResult->SourceAssetName = Translator->GetSourceData()->GetFilename();
diff --git a/Engine/Source/Runtime/Interchange/Engine/Public/InterchangeManager.h b/Engine/Source/Runtime/Interchange/Engine/Public/InterchangeManager.h
--- a/Engine/Source/Runtime/Interchange/Engine/Public/InterchangeManager.h
+++ b/Engine/Source/Runtime/Interchange/Engine/Public/InterchangeManager.h
@@ -28,9 +28,11 @@
 #include "InterchangeManager.generated.h"
 
 class FAsyncTaskNotification;
+class UInterchangeBlueprintPipelineBase;
 class UInterchangeFactoryBase;
 class UInterchangeFactoryBaseNode;
 class UInterchangePipelineBase;
+class UInterchangePythonPipelineBase;
 
 /** Some utilities delegates for the automation of interchange */
 DECLARE_DYNAMIC_DELEGATE_OneParam(FOnObjectImportDoneDynamic, UObject*, Object);
@@ -306,7 +308,16 @@ class INTERCHANGEENGINE_API UInterchangePipelineStackOverride : public UObject
 public:
 
 	UPROPERTY(EditAnywhere, BlueprintReadWrite, Category = "Interchange|ImportAsset")
-		TArray<TObjectPtr<UInterchangePipelineBase>> OverridePipelines;
+	TArray<FSoftObjectPath> OverridePipelines;
+
+	UFUNCTION(BlueprintCallable, Category = "Interchange | PipelineStackOverride")
+	void AddPythonPipeline(UInterchangePythonPipelineBase* PipelineBase);
+
+	UFUNCTION(BlueprintCallable, Category = "Interchange | PipelineStackOverride")
+	void AddBlueprintPipeline(UInterchangeBlueprintPipelineBase* PipelineBase);
+
+	UFUNCTION(BlueprintCallable, Category = "Interchange | PipelineStackOverride")
+	void AddPipeline(UInterchangePipelineBase* PipelineBase);
 };
 
 USTRUCT(BlueprintType)
