[Interchange - glTF] Fix support for importing sparse accessors

diff --git a/Engine/Plugins/Interchange/Runtime/Source/Parsers/GLTFCore/Private/GLTF/GLTFAccessor.cpp b/Engine/Plugins/Interchange/Runtime/Source/Parsers/GLTFCore/Private/GLTF/GLTFAccessor.cpp
--- a/Engine/Plugins/Interchange/Runtime/Source/Parsers/GLTFCore/Private/GLTF/GLTFAccessor.cpp
+++ b/Engine/Plugins/Interchange/Runtime/Source/Parsers/GLTFCore/Private/GLTF/GLTFAccessor.cpp
@@ -257,47 +257,262 @@ namespace GLTF
 			}
 		}
 
-		// Copy data items that don't need conversion/expansion(i.e. Vec3 to Vec3, uint8 to uint8(not uint16)
-		// but including normalized from fixed-point types
 		template<typename ItemType, uint32 ItemElementCount>
-		void CopyWithoutConversion(const FValidAccessor& Accessor, ItemType* Buffer)
+		void CopyWithoutConversion(const uint32 ByteStride, const FAccessor::EComponentType ComponentType, bool bNormalized, const uint32 Count, const FBufferView& BufferView, const uint32 ByteOffset, ItemType* Buffer)
 		{
 			// Stride equals item size => use simpler copy
-			if ((Accessor.ByteStride == 0) || Accessor.ByteStride == sizeof(ItemType))
+			if ((ByteStride == 0) || ByteStride == sizeof(ItemType))
 			{
-				const void* Src = Accessor.DataAt(0);
-				if (Accessor.ComponentType == FAccessor::EComponentType::F32)
+				const void* Src = BufferView.DataAt((0 * ByteStride) + ByteOffset);
+				if (ComponentType == FAccessor::EComponentType::F32)
 				{
-					memcpy(Buffer, Src, Accessor.Count * sizeof(ItemType));
+					memcpy(Buffer, Src, Count * sizeof(ItemType));
 					return;
 				}
 			}
 
-			if (Accessor.bNormalized)
+			if (bNormalized)
+			{
+				for (uint32 Index = 0; Index < Count; ++Index)
+				{
+					const void* Pointer = BufferView.DataAt((Index * ByteStride) + ByteOffset);
+					Buffer[Index] = GetNormalized<ItemType, ItemElementCount>(ComponentType, Pointer);
+				}
+			}
+			else
+			{
+				for (uint32 Index = 0; Index < Count; ++Index)
+				{
+					const void* Pointer = BufferView.DataAt((Index * ByteStride) + ByteOffset);
+					Buffer[Index] = GetNonNormalized<ItemType, ItemElementCount>(ComponentType, Pointer);
+				}
+			}
+		}
+
+		// Copy data items that don't need conversion/expansion(i.e. Vec3 to Vec3, uint8 to uint8(not uint16)
+		// but including normalized from fixed-point types
+		template<typename ItemType, uint32 ItemElementCount>
+		void CopyWithoutConversion(const FValidAccessor& Accessor, ItemType* Buffer)
+		{
+			CopyWithoutConversion<ItemType, ItemElementCount>(Accessor.ByteStride, Accessor.ComponentType, Accessor.bNormalized, Accessor.Count, Accessor.BufferView, Accessor.ByteOffset, Buffer);
+		}
+
+		bool FillFloatArray(bool bNormalized, const FBufferView& BufferView, const uint32 ByteOffset, FAccessor::EComponentType ComponentType, const uint32 ByteStride, const uint32 Count, float* Buffer)
+		{
+			if (!bNormalized)
 			{
-				for (uint32 Index = 0; Index < Accessor.Count; ++Index)
+				const uint8* Src = BufferView.DataAt((0 * ByteStride) + ByteOffset);
+				switch (ComponentType)
 				{
-					const void* Pointer = Accessor.DataAt(Index);
-					Buffer[Index] = GetNormalized<ItemType, ItemElementCount>(Accessor.ComponentType, Pointer);
+				case FAccessor::EComponentType::F32:
+					memcpy(Buffer, Src, Count * sizeof(float));
+					return true;
+				default:
+					break;
 				}
 			}
 			else
 			{
-				for (uint32 Index = 0; Index < Accessor.Count; ++Index)
+				// Stride equals item size => use simpler copy
+				if ((ByteStride == 0) || ByteStride == sizeof(float))
 				{
-					const void* Pointer = Accessor.DataAt(Index);
-					Buffer[Index] = GetNonNormalized<ItemType, ItemElementCount>(Accessor.ComponentType, Pointer);
+					const uint8* Src = BufferView.DataAt((0 * ByteStride) + ByteOffset);
+					CopyNormalizedFloat(Buffer, Src, ComponentType, Count);
 				}
+				else
+				{
+					for (uint32 Index = 0; Index < Count; ++Index)
+					{
+						const uint8* Pointer = BufferView.DataAt((Index * ByteStride) + ByteOffset);
+						Buffer[Index] = GetNormalizedFloat(ComponentType, Pointer);
+					}
+				}
+				return true;
+			}
+
+			return false;
+		}
+
+		bool FillUnsignedIntArray(const FBufferView& BufferView, const uint32 ByteOffset, FAccessor::EComponentType ComponentType, const uint32 Count, uint32* Buffer)
+		{
+			const uint8* Src = BufferView.DataAt(ByteOffset);
+			switch (ComponentType)
+			{
+				case FAccessor::EComponentType::U8:
+					Copy(Buffer, reinterpret_cast<const uint8*>(Src), Count);
+					return true;
+				case FAccessor::EComponentType::U16:
+					Copy(Buffer, reinterpret_cast<const uint16*>(Src), Count);
+					return true;
+				case FAccessor::EComponentType::U32:
+					memcpy(Buffer, Src, Count * sizeof(uint32));
+					return true;
+				default:
+					return false;
+			}
+		}
+
+		TArray<uint32> AcquireUnsignedIntArray(const uint32 Count, const FBufferView& BufferView, const uint32 ByteOffset, FAccessor::EComponentType ComponentType)
+		{
+			TArray<uint32> ValueBuffer;
+			ValueBuffer.SetNumUninitialized(Count, false);
+
+			FillUnsignedIntArray(BufferView, ByteOffset, ComponentType, Count, ValueBuffer.GetData());
+
+			return ValueBuffer;
+		}
+
+		bool AcquireUnsignedInt(uint32 Index, const FBufferView& BufferView, const uint32 ByteOffset, const uint32 ByteStride, FAccessor::EComponentType ComponentType, uint32& Buffer)
+		{
+			Buffer = 0;
+
+			const uint8* ValuePtr = BufferView.DataAt((Index * ByteStride) + ByteOffset);
+			switch (ComponentType)
+			{
+				case FAccessor::EComponentType::U8:
+					Buffer = *ValuePtr;
+					return true;
+				case FAccessor::EComponentType::U16:
+					Buffer = *reinterpret_cast<const uint16*>(ValuePtr);
+					return true;
+				case FAccessor::EComponentType::U32:
+					Buffer = *reinterpret_cast<const uint32*>(ValuePtr);
+					return true;
+				default:
+					return false;
+			}
+		}
+
+		bool AcquireUnsignedInt16x4(uint32 Index, const FBufferView& BufferView, const uint32 ByteOffset, const uint32 ByteStride, FAccessor::EComponentType ComponentType, uint16 Buffer[4])
+		{
+			Buffer[0] = Buffer[1] = Buffer[2] = Buffer[3] = 0;
+
+			const uint8* ValuePtr = BufferView.DataAt((Index * ByteStride) + ByteOffset);
+			switch (ComponentType)
+			{
+				case FAccessor::EComponentType::U8:
+					for (int i = 0; i < 4; ++i)
+					{
+						Buffer[i] = ((uint8*)ValuePtr)[i];
+					}
+					return true;
+				case FAccessor::EComponentType::U16:
+					for (int i = 0; i < 4; ++i)
+					{
+						Buffer[i] = ((uint16*)ValuePtr)[i];
+					}
+					return true;
+				default:
+					return false;
+			}
+		}
+
+		bool AcquireFloat(uint32 Index, const FBufferView& BufferView, const uint32 ByteOffset, const uint32 ByteStride, FAccessor::EComponentType ComponentType, float& Buffer)
+		{
+			Buffer = 0;
+
+			const uint8* ValuePtr = BufferView.DataAt((Index * ByteStride) + ByteOffset);
+			switch (ComponentType)
+			{
+				case FAccessor::EComponentType::F32:
+					Buffer = *ValuePtr;
+					return true;
+				default:
+					return false;
+			}
+		}
+
+		bool AcquireVec2(uint32 Index, const FBufferView& BufferView, const uint32 ByteOffset, const uint32 ByteStride, FAccessor::EComponentType ComponentType, bool bNormalized, FVector2D& Buffer)
+		{
+			Buffer = FVector2D();
+
+			const uint8* Pointer = BufferView.DataAt((Index * ByteStride) + ByteOffset);
+			if (ComponentType == FAccessor::EComponentType::F32)
+			{
+				// copy float vec2 directly from buffer
+				Buffer = FVector2D(*reinterpret_cast<const FVector2f*>(Pointer));
+				return true;
 			}
+			if (bNormalized)
+			{
+				Buffer = FVector2D(GetNormalized<FVector2f, 2>(ComponentType, Pointer));
+				return true;
+			}
+
+			return false;
+		}
+
+		bool AcquireVec3(uint32 Index, const FBufferView& BufferView, const uint32 ByteOffset, const uint32 ByteStride, FAccessor::EComponentType ComponentType, bool bNormalized, FVector& Buffer)
+		{
+			Buffer = FVector();
+
+			const uint8* Pointer = BufferView.DataAt((Index * ByteStride) + ByteOffset);
+			if (ComponentType == FAccessor::EComponentType::F32)
+			{
+				// copy float vec3 directly from buffer
+				Buffer = FVector(*reinterpret_cast<const FVector3f*>(Pointer));
+				return true;
+			}
+			if (bNormalized)
+			{
+				Buffer = FVector(GetNormalized<FVector3f, 3>(ComponentType, Pointer));
+				return true;
+			}
+
+			return false;
+		}
+
+		bool AcquireVec4(uint32 Index, const FBufferView& BufferView, const uint32 ByteOffset, const uint32 ByteStride, FAccessor::EComponentType ComponentType, bool bNormalized, FVector4& Buffer)
+		{
+			Buffer = FVector4();
+
+			const uint8* Pointer = BufferView.DataAt((Index * ByteStride) + ByteOffset);
+			if (ComponentType == FAccessor::EComponentType::F32)
+			{
+				// copy float vec4 directly from buffer
+				Buffer = FVector4(*reinterpret_cast<const FVector4f*>(Pointer));
+				return true;
+			}
+			if (bNormalized)
+			{
+				Buffer = FVector4(GetNormalized<FVector4f, 4>(ComponentType, Pointer));
+				return true;
+			}
+
+			return false;
+		}
+
+		bool AcquireMat4(uint32 Index, const FBufferView& BufferView, const uint32 ByteOffset, const uint32 ByteStride, FAccessor::EComponentType ComponentType, FMatrix& Buffer)
+		{
+			Buffer = FMatrix();
+
+			if (ComponentType == FAccessor::EComponentType::F32)
+			{
+				const uint8* Pointer = BufferView.DataAt((Index * ByteStride) + ByteOffset);
+				Buffer = GetMatrix(Pointer);
+				return true;
+			}
+
+			return false;
 		}
 	}
+	
+	FAccessor::FSparse::FIndices::FIndices(uint32 InCount, const FBufferView& InBufferView, uint32 InByteOffset, EComponentType InComponentType)
+		: Count(InCount)
+		, BufferView(InBufferView)
+		, ByteOffset(InByteOffset)
+		, ComponentType(InComponentType)
+	{
+	}
+		
 
-	FAccessor::FAccessor(uint32 InCount, EType InType, EComponentType InCompType, bool bInNormalized)
+	FAccessor::FAccessor(uint32 InCount, EType InType, EComponentType InCompType, bool bInNormalized, const FSparse& InSparse)
 	    : Count(InCount)
 	    , Type(InType)
 	    , ComponentType(InCompType)
 	    , bNormalized(bInNormalized)
 		, bQuantized(false)
+		, Sparse(InSparse)
 	{
 	}
 
@@ -503,14 +718,207 @@ namespace GLTF
 				return true;
 			}
 		}
-		
+
 		return false;
 	}
 
+	//Sparse related helpers:
+	void FValidAccessor::UpdateUnsignedIntWithSparse(uint32 Index, uint32& Data) const
+	{
+		if (Sparse.bHasSparse)
+		{
+			TArray<uint32> IndicesData = AcquireUnsignedIntArray(Sparse.Count, Sparse.Indices.BufferView, Sparse.Indices.ByteOffset, Sparse.Indices.ComponentType);
+			int32 SparseArrayIndex = IndicesData.IndexOfByKey(Index);
+			if (SparseArrayIndex != INDEX_NONE)
+			{
+				uint32 Buffer;
+				if (AcquireUnsignedInt(Index, Sparse.Values.BufferView, Sparse.Values.ByteOffset, ByteStride, ComponentType, Buffer))
+				{
+					Data = Buffer;
+				}
+			}
+		}
+	}
+	void FValidAccessor::UpdateUnsignedInt16x4WithSparse(uint32 Index, uint16 Data[4]) const
+	{
+		if (Sparse.bHasSparse)
+		{
+			TArray<uint32> IndicesData = AcquireUnsignedIntArray(Sparse.Count, Sparse.Indices.BufferView, Sparse.Indices.ByteOffset, Sparse.Indices.ComponentType);
+			int32 SparseArrayIndex = IndicesData.IndexOfByKey(Index);
+			if (SparseArrayIndex != INDEX_NONE)
+			{
+				uint16 Buffer[4];
+				if (AcquireUnsignedInt16x4(Index, Sparse.Values.BufferView, Sparse.Values.ByteOffset, ByteStride, ComponentType, Buffer))
+				{
+					Data[0] = Buffer[0];
+					Data[1] = Buffer[1];
+					Data[2] = Buffer[2];
+					Data[3] = Buffer[3];
+				}
+			}
+		}
+	}
+
+	void FValidAccessor::UpdateFloatWithSparse(uint32 Index, float& Data) const
+	{
+		if (Sparse.bHasSparse)
+		{
+			TArray<uint32> IndicesData = AcquireUnsignedIntArray(Sparse.Count, Sparse.Indices.BufferView, Sparse.Indices.ByteOffset, Sparse.Indices.ComponentType);
+			int32 SparseArrayIndex = IndicesData.IndexOfByKey(Index);
+			if (SparseArrayIndex != INDEX_NONE)
+			{
+				float Buffer;
+				if (AcquireFloat(Index, Sparse.Values.BufferView, Sparse.Values.ByteOffset, ByteStride, ComponentType, Buffer))
+				{
+					Data = Buffer;
+				}
+			}
+		}
+	}
+	void FValidAccessor::UpdateVec2WithSparse(uint32 Index, FVector2D& Data) const
+	{
+		if (Sparse.bHasSparse)
+		{
+			TArray<uint32> IndicesData = AcquireUnsignedIntArray(Sparse.Count, Sparse.Indices.BufferView, Sparse.Indices.ByteOffset, Sparse.Indices.ComponentType);
+			int32 SparseArrayIndex = IndicesData.IndexOfByKey(Index);
+			if (SparseArrayIndex != INDEX_NONE)
+			{
+				FVector2D Buffer;
+				if (AcquireVec2(Index, Sparse.Values.BufferView, Sparse.Values.ByteOffset, ByteStride, ComponentType, bNormalized, Buffer))
+				{
+					Data = Buffer;
+				}
+			}
+		}
+	}
+	void FValidAccessor::UpdateVec3WithSparse(uint32 Index, FVector& Data) const
+	{
+		if (Sparse.bHasSparse)
+		{
+			TArray<uint32> IndicesData = AcquireUnsignedIntArray(Sparse.Count, Sparse.Indices.BufferView, Sparse.Indices.ByteOffset, Sparse.Indices.ComponentType);
+			int32 SparseArrayIndex = IndicesData.IndexOfByKey(Index);
+			if (SparseArrayIndex != INDEX_NONE)
+			{
+				FVector Buffer;
+				if (AcquireVec3(Index, Sparse.Values.BufferView, Sparse.Values.ByteOffset, ByteStride, ComponentType, bNormalized, Buffer))
+				{
+					Data = Buffer;
+				}
+			}
+		}
+	}
+	void FValidAccessor::UpdateVec4WithSparse(uint32 Index, FVector4& Data) const
+	{
+		if (Sparse.bHasSparse)
+		{
+			TArray<uint32> IndicesData = AcquireUnsignedIntArray(Sparse.Count, Sparse.Indices.BufferView, Sparse.Indices.ByteOffset, Sparse.Indices.ComponentType);
+			int32 SparseArrayIndex = IndicesData.IndexOfByKey(Index);
+			if (SparseArrayIndex != INDEX_NONE)
+			{
+				FVector4 Buffer;
+				if (AcquireVec4(Index, Sparse.Values.BufferView, Sparse.Values.ByteOffset, ByteStride, ComponentType, bNormalized, Buffer))
+				{
+					Data = Buffer;
+				}
+			}
+		}
+	}
+
+	void FValidAccessor::UpdateMat4WithSparse(uint32 Index, FMatrix& Data) const
+	{
+		if (Sparse.bHasSparse)
+		{
+			TArray<uint32> IndicesData = AcquireUnsignedIntArray(Sparse.Count, Sparse.Indices.BufferView, Sparse.Indices.ByteOffset, Sparse.Indices.ComponentType);
+			int32 SparseArrayIndex = IndicesData.IndexOfByKey(Index);
+			if (SparseArrayIndex != INDEX_NONE)
+			{
+				FMatrix Buffer;
+				if (AcquireMat4(Index, Sparse.Values.BufferView, Sparse.Values.ByteOffset, ByteStride, ComponentType, Buffer))
+				{
+					Data = Buffer;
+				}
+			}
+		}
+	}
+
+	void FValidAccessor::UpdateFloatArrayWithSparse(float* Data) const
+	{
+		if (Sparse.bHasSparse && Sparse.Count)
+		{
+			TArray<uint32> IndicesData = AcquireUnsignedIntArray(Sparse.Count, Sparse.Indices.BufferView, Sparse.Indices.ByteOffset, Sparse.Indices.ComponentType);
+
+			TArray<float> ValueBuffer;
+			ValueBuffer.SetNumUninitialized(Sparse.Count, false);
+
+			if (FillFloatArray(bNormalized, Sparse.Values.BufferView, Sparse.Values.ByteOffset, ComponentType, ByteStride, Sparse.Count, ValueBuffer.GetData()))
+			{
+				for (size_t Index = 0; Index < Sparse.Count; Index++)
+				{
+					const uint32 SparseIndex = IndicesData[Index];
+					const float SparseValue = ValueBuffer[Index];
+
+					if (SparseIndex < Count)
+					{
+						Data[SparseIndex] = SparseValue;
+					}
+				}
+			}
+		}
+	}
+	void FValidAccessor::UpdateUnsignedIntArrayWithSparse(uint32* Data) const
+	{
+		if (Sparse.bHasSparse && Sparse.Count)
+		{
+			TArray<uint32> IndicesData = AcquireUnsignedIntArray(Sparse.Count, Sparse.Indices.BufferView, Sparse.Indices.ByteOffset, Sparse.Indices.ComponentType);
+
+			TArray<uint32> ValueBuffer;
+			ValueBuffer.SetNumUninitialized(Sparse.Count, false);
+
+			if (FillUnsignedIntArray(Sparse.Values.BufferView, Sparse.Values.ByteOffset, ComponentType, Count, ValueBuffer.GetData()))
+			{
+				for (size_t Index = 0; Index < Sparse.Count; Index++)
+				{
+					const uint32 SparseIndex = IndicesData[Index];
+					const float SparseValue = ValueBuffer[Index];
+
+					if (SparseIndex < Count)
+					{
+						Data[SparseIndex] = SparseValue;
+					}
+				}
+			}
+		}
+	}
+
+	template<typename ItemType, uint32 ItemElementCount>
+	void FValidAccessor::UpdateArrayWithSparse(ItemType* Data) const
+	{
+		if (Sparse.bHasSparse && Sparse.Count)
+		{
+			TArray<uint32> IndicesData = AcquireUnsignedIntArray(Sparse.Count, Sparse.Indices.BufferView, Sparse.Indices.ByteOffset, Sparse.Indices.ComponentType);
+
+			TArray<ItemType> ValueBuffer;
+			ValueBuffer.SetNumUninitialized(Sparse.Count, false);
+
+			CopyWithoutConversion<ItemType, ItemElementCount>(ByteStride, ComponentType, bNormalized, Sparse.Count, Sparse.Values.BufferView, Sparse.Values.ByteOffset, ValueBuffer.GetData());
+
+			for (size_t Index = 0; Index < Sparse.Count; Index++)
+			{
+				const uint32 SparseIndex = IndicesData[Index];
+				const ItemType SparseValue = ValueBuffer[Index];
+
+				if (SparseIndex < Count)
+				{
+					Data[SparseIndex] = SparseValue;
+				}
+			}
+		}
+	}
+
 	//
 	FValidAccessor::FValidAccessor(FBufferView& InBufferView, uint32 InOffset, uint32 InCount, EType InType, EComponentType InCompType,
-	                               bool bInNormalized)
-	    : FAccessor(InCount, InType, InCompType, bInNormalized)
+	                               bool bInNormalized, const FSparse& InSparse)
+	    : FAccessor(InCount, InType, InCompType, bInNormalized, InSparse)
 	    , BufferView(InBufferView)
 	    , ByteOffset(InOffset)
 	    , ElementSize(GetElementSize(Type, ComponentType))
@@ -568,17 +976,11 @@ namespace GLTF
 		{
 			if (Type == EType::Scalar && !bNormalized)
 			{
-				const uint8* ValuePtr = DataAt(Index);
-				switch (ComponentType)
+				uint32 Buffer;
+				if (AcquireUnsignedInt(Index, BufferView, ByteOffset, ByteStride, ComponentType, Buffer))
 				{
-					case EComponentType::U8:
-						return *ValuePtr;
-					case EComponentType::U16:
-						return *reinterpret_cast<const uint16*>(ValuePtr);
-					case EComponentType::U32:
-						return *reinterpret_cast<const uint32*>(ValuePtr);
-					default:
-						break;
+					UpdateUnsignedIntWithSparse(Index, Buffer);
+					return Buffer;
 				}
 			}
 		}
@@ -595,23 +997,10 @@ namespace GLTF
 		{
 			if (Type == EType::Vec4 && !bNormalized)
 			{
-				const void* ValuePtr = DataAt(Index);
-				switch (ComponentType)
+				if (AcquireUnsignedInt16x4(Index, BufferView, ByteOffset, ByteStride, ComponentType, Values))
 				{
-					case EComponentType::U8:
-						for (int i = 0; i < 4; ++i)
-						{
-							Values[i] = ((uint8*)ValuePtr)[i];
-						}
-						return;
-					case EComponentType::U16:
-						for (int i = 0; i < 4; ++i)
-						{
-							Values[i] = ((uint16*)ValuePtr)[i];
-						}
-						return;
-					default:
-						break;
+					UpdateUnsignedInt16x4WithSparse(Index, Values);
+					return;
 				}
 			}
 		}
@@ -626,13 +1015,11 @@ namespace GLTF
 		{
 			if (Type == EType::Scalar && !bNormalized)
 			{
-				const uint8* ValuePtr = DataAt(Index);
-				switch (ComponentType)
+				float Buffer;
+				if (AcquireFloat(Index, BufferView, ByteOffset, ByteStride, ComponentType, Buffer))
 				{
-					case EComponentType::F32:
-						return *ValuePtr;
-					default:
-						break;
+					UpdateFloatWithSparse(Index, Buffer);
+					return Buffer;
 				}
 			}
 		}
@@ -653,16 +1040,11 @@ namespace GLTF
 		{
 			if (Type == EType::Vec2)  // strict format match, unlike GPU shader fetch
 			{
-				const void* Pointer = DataAt(Index);
-
-				if (ComponentType == EComponentType::F32)
+				FVector2D Buffer;
+				if (AcquireVec2(Index, BufferView, ByteOffset, ByteStride, ComponentType, bNormalized, Buffer))
 				{
-					// copy float vec2 directly from buffer
-					return FVector2D(*reinterpret_cast<const FVector2f*>(Pointer));
-				}
-				else if (bNormalized)
-				{
-					return FVector2D(GetNormalized<FVector2f, 2>(ComponentType, Pointer));
+					UpdateVec2WithSparse(Index, Buffer);
+					return Buffer;
 				}
 			}
 		}
@@ -683,16 +1065,11 @@ namespace GLTF
 		{
 			if (Type == EType::Vec3)  // strict format match, unlike GPU shader fetch
 			{
-				const void* Pointer = DataAt(Index);
-
-				if (ComponentType == EComponentType::F32)
+				FVector Buffer;
+				if (AcquireVec3(Index, BufferView, ByteOffset, ByteStride, ComponentType, bNormalized, Buffer))
 				{
-					// copy float vec3 directly from buffer
-					return FVector(*reinterpret_cast<const FVector3f*>(Pointer));
-				}
-				else if (bNormalized)
-				{
-					return FVector(GetNormalized<FVector3f, 3>(ComponentType, Pointer));
+					UpdateVec3WithSparse(Index, Buffer);
+					return Buffer;
 				}
 			}
 		}
@@ -713,16 +1090,11 @@ namespace GLTF
 		{
 			if (Type == EType::Vec4)  // strict format match, unlike GPU shader fetch
 			{
-				const void* Pointer = DataAt(Index);
-
-				if (ComponentType == EComponentType::F32)
+				FVector4 Buffer;
+				if (AcquireVec4(Index, BufferView, ByteOffset, ByteStride, ComponentType, bNormalized, Buffer))
 				{
-					// copy float vec4 directly from buffer
-					return FVector4(*reinterpret_cast<const FVector4f*>(Pointer));
-				}
-				else if (bNormalized)
-				{
-					return FVector4(GetNormalized<FVector4f, 4>(ComponentType, Pointer));
+					UpdateVec4WithSparse(Index, Buffer);
+					return Buffer;
 				}
 			}
 		}
@@ -739,10 +1111,11 @@ namespace GLTF
 		{
 			if (Type == EType::Mat4)  // strict format match, unlike GPU shader fetch
 			{
-				if (ComponentType == EComponentType::F32)
+				FMatrix Buffer;
+				if (AcquireMat4(Index, BufferView, ByteOffset, ByteStride, ComponentType, Buffer))
 				{
-					const void* Pointer = DataAt(Index);
-					return GetMatrix(Pointer);
+					UpdateMat4WithSparse(Index, Buffer);
+					return Buffer;
 				}
 			}
 		}
@@ -755,20 +1128,10 @@ namespace GLTF
 	{
 		if (Type == EType::Scalar && !bNormalized)
 		{
-			const uint8* Src = DataAt(0);
-			switch (ComponentType)
+			if (FillUnsignedIntArray(BufferView, ByteOffset, ComponentType, Count, Buffer))
 			{
-				case EComponentType::U8:
-					Copy(Buffer, reinterpret_cast<const uint8*>(Src), Count);
-					return;
-				case EComponentType::U16:
-					Copy(Buffer, reinterpret_cast<const uint16*>(Src), Count);
-					return;
-				case EComponentType::U32:
-					memcpy(Buffer, Src, Count * sizeof(uint32));
-					return;
-				default:
-					break;
+				UpdateUnsignedIntArrayWithSparse(Buffer);
+				return;
 			}
 		}
 
@@ -779,34 +1142,9 @@ namespace GLTF
 	{
 		if (Type == EType::Scalar)
 		{
-			if (!bNormalized)
+			if (FillFloatArray(bNormalized, BufferView, ByteOffset, ComponentType, ByteStride, Count, Buffer))
 			{
-				const uint8* Src = DataAt(0);
-				switch (ComponentType)
-				{
-				case EComponentType::F32:
-					memcpy(Buffer, Src, Count * sizeof(float));
-					return;
-				default:
-					break;
-				}
-			}
-			else
-			{
-				// Stride equals item size => use simpler copy
-				if ((ByteStride == 0) || ByteStride == sizeof(float))
-				{
-					const void* Src = DataAt(0);
-					CopyNormalizedFloat(Buffer, Src, ComponentType, Count);
-				}
-				else
-				{
-					for (uint32 Index = 0; Index < Count; ++Index)
-					{
-						const void* Pointer = DataAt(Index);
-						Buffer[Index] = GetNormalizedFloat(ComponentType, Pointer);
-					}
-				}
+				UpdateFloatArrayWithSparse(Buffer);
 				return;
 			}
 		}
@@ -821,6 +1159,7 @@ namespace GLTF
 			return;
 		}
 		CopyWithoutConversion<FVector2f, 2>(*this, Buffer);
+		UpdateArrayWithSparse<FVector2f, 2>(Buffer);
 	}
 
 	void FValidAccessor::GetVec3Array(FVector3f* Buffer) const
@@ -830,6 +1169,7 @@ namespace GLTF
 			return;
 		}
 		CopyWithoutConversion<FVector3f, 3>(*this, Buffer);
+		UpdateArrayWithSparse<FVector3f, 3>(Buffer);
 	}
 
 	void FValidAccessor::GetVec4Array(FVector4f* Buffer) const
@@ -839,6 +1179,7 @@ namespace GLTF
 			return;
 		}
 		CopyWithoutConversion<FVector4f, 4>(*this, Buffer);
+		UpdateArrayWithSparse<FVector4f, 4>(Buffer);
 	}
 
 	void FValidAccessor::GetMat4Array(FMatrix44f* Buffer) const
diff --git a/Engine/Plugins/Interchange/Runtime/Source/Parsers/GLTFCore/Private/GLTF/GLTFReader.cpp b/Engine/Plugins/Interchange/Runtime/Source/Parsers/GLTFCore/Private/GLTF/GLTFReader.cpp
--- a/Engine/Plugins/Interchange/Runtime/Source/Parsers/GLTFCore/Private/GLTF/GLTFReader.cpp
+++ b/Engine/Plugins/Interchange/Runtime/Source/Parsers/GLTFCore/Private/GLTF/GLTFReader.cpp
@@ -228,7 +228,33 @@ namespace GLTF
 			const uint32                    Count      = GetUnsignedInt(Object, TEXT("count"), 0);
 			const FAccessor::EType          Type       = AccessorTypeFromString(Object.GetStringField("type"));
 			const bool                      Normalized = GetBool(Object, TEXT("normalized"));
-			Asset->Accessors.Emplace(Asset->BufferViews[BufferViewIdx], ByteOffset, Count, Type, CompType, Normalized);
+
+			//Sparse:
+			if (Object.HasField("sparse"))
+			{
+				const FJsonObject& SparseObject = *Object.GetObjectField(TEXT("sparse"));
+
+				const uint32 SparseCount = GetUnsignedInt(SparseObject, TEXT("count"), 0);
+
+				const FJsonObject& IndicesObject = *SparseObject.GetObjectField(TEXT("indices"));
+				const uint32 IndicesBufferViewIdx = GetUnsignedInt(IndicesObject, TEXT("bufferView"), BufferViewCount);
+				const int32 IndicesByteOffset = GetUnsignedInt(IndicesObject, TEXT("byteOffset"), 0);
+				const FAccessor::EComponentType IndicesCompType = ComponentTypeFromNumber(GetUnsignedInt(IndicesObject, TEXT("componentType"), 0));
+
+				const FJsonObject& ValuesObject = *SparseObject.GetObjectField(TEXT("values"));
+				const uint32 ValuesBufferViewIdx = GetUnsignedInt(ValuesObject, TEXT("bufferView"), BufferViewCount);
+				const int32 ValuesByteOffset = GetUnsignedInt(ValuesObject, TEXT("byteOffset"), 0);
+
+				Asset->Accessors.Emplace(Asset->BufferViews[BufferViewIdx], ByteOffset, Count, Type, CompType, Normalized, 
+					FAccessor::FSparse(SparseCount, 
+						Asset->BufferViews[IndicesBufferViewIdx], IndicesByteOffset, IndicesCompType,
+						Asset->BufferViews[ValuesBufferViewIdx], ValuesByteOffset));
+			}
+			else
+			{
+				Asset->Accessors.Emplace(Asset->BufferViews[BufferViewIdx], ByteOffset, Count, Type, CompType, Normalized, FAccessor::FSparse());
+			}
+
 			ExtensionsHandler->SetupAccessorExtensions(Object, Asset->Accessors.Last());
 		}
 	}
diff --git a/Engine/Plugins/Interchange/Runtime/Source/Parsers/GLTFCore/Public/GLTFAccessor.h b/Engine/Plugins/Interchange/Runtime/Source/Parsers/GLTFCore/Public/GLTFAccessor.h
--- a/Engine/Plugins/Interchange/Runtime/Source/Parsers/GLTFCore/Public/GLTFAccessor.h
+++ b/Engine/Plugins/Interchange/Runtime/Source/Parsers/GLTFCore/Public/GLTFAccessor.h
@@ -59,6 +59,20 @@ namespace GLTF
 		}
 	};
 
+	struct GLTFCORE_API FVoidBufferView final : FBufferView
+	{
+		FVoidBufferView()
+			: FBufferView(FBuffer(0), 0, 0, 0)
+		{
+
+		}
+		static FVoidBufferView& GetVoidBufferView()
+		{
+			static FVoidBufferView Void;
+			return Void;
+		}
+	};
+
 	struct GLTFCORE_API FAccessor
 	{
 		// accessor stores the data but has no usage semantics
@@ -88,13 +102,79 @@ namespace GLTF
 			Count
 		};
 
+		struct GLTFCORE_API FSparse
+		{
+			bool                     bHasSparse;
+
+			const uint32              Count;
+
+			//Indices:
+			struct FIndices
+			{
+				const int32          Count; //Helper for creating cache, equals to FSparse.Count
+
+				const FBufferView&   BufferView;
+				const uint32         ByteOffset;
+				const EComponentType ComponentType;
+
+				FIndices()
+					: Count(0)
+					, BufferView(FVoidBufferView::GetVoidBufferView())
+					, ByteOffset(0)
+					, ComponentType(EComponentType::None)
+				{
+				}
+
+				FIndices(uint32 InCount, const FBufferView& InBufferView, uint32 InByteOffset, EComponentType InComponentType);
+			} Indices;
+			
+			//Values:
+			struct FValues
+			{
+				const FBufferView&   BufferView;
+				const uint32         ByteOffset;
+				
+				FValues()
+					: BufferView(FVoidBufferView::GetVoidBufferView())
+					, ByteOffset(0)
+				{
+				}
+
+				FValues(const FBufferView& InBufferView, uint32 InByteOffset)
+					: BufferView(InBufferView)
+					, ByteOffset(InByteOffset)
+				{
+				}
+			} Values;
+			
+			FSparse()
+				: bHasSparse(false)
+				, Count(0)
+				, Indices(FIndices())
+				, Values(FValues())
+			{
+			}
+
+			FSparse(uint32 InCount,
+				const FBufferView& InIndicesBufferView, uint32 InIndicesByteOffset, EComponentType InIndicesComponentType,
+				const FBufferView& InValuesBufferView, uint32 InValuesByteOffset)
+				: bHasSparse(true)
+				, Count(InCount)
+				, Indices(InCount, InIndicesBufferView, InIndicesByteOffset, InIndicesComponentType)
+				, Values(InValuesBufferView, InValuesByteOffset)
+			{
+			}
+		};
+
 		const uint32         Count;
 		const EType          Type;
 		const EComponentType ComponentType;
 		const bool           bNormalized;
 		bool                 bQuantized;
+		const FSparse        Sparse;
 
-		FAccessor(uint32 InCount, EType InType, EComponentType InComponentType, bool bInNormalized);
+
+		FAccessor(uint32 InCount, EType InType, EComponentType InComponentType, bool bInNormalized, const FSparse& InSparse);
 
 		virtual bool IsValid() const = 0;
 		virtual FMD5Hash GetHash() const = 0;
@@ -145,7 +225,7 @@ namespace GLTF
 
 	struct GLTFCORE_API FValidAccessor final : FAccessor
 	{
-		FValidAccessor(FBufferView& InBufferView, uint32 InOffset, uint32 InCount, EType InType, EComponentType InCompType, bool bInNormalized);
+		FValidAccessor(FBufferView& InBufferView, uint32 InOffset, uint32 InCount, EType InType, EComponentType InCompType, bool bInNormalized, const FSparse& InSparse);
 
 		bool IsValid() const override;
 
@@ -174,12 +254,29 @@ namespace GLTF
 		const uint32       ByteOffset;
 		const uint32       ElementSize;
 		const uint32	   ByteStride;
+
+		//Sparse related helpers:
+		void UpdateUnsignedIntWithSparse(uint32 Index, uint32& Data) const;
+		void UpdateUnsignedInt16x4WithSparse(uint32 Index, uint16 Data[4]) const;
+
+		void UpdateFloatWithSparse(uint32 Index, float& Data) const;
+		void UpdateVec2WithSparse(uint32 Index, FVector2D& Data) const;
+		void UpdateVec3WithSparse(uint32 Index, FVector& Data) const;
+		void UpdateVec4WithSparse(uint32 Index, FVector4& Data) const;
+
+		void UpdateMat4WithSparse(uint32 Index, FMatrix& Data) const;
+
+		void UpdateFloatArrayWithSparse(float* Data) const;
+		void UpdateUnsignedIntArrayWithSparse(uint32* Data) const;
+
+		template<typename ItemType, uint32 ItemElementCount>
+		void UpdateArrayWithSparse(ItemType* Data) const;
 	};
 
 	struct GLTFCORE_API FVoidAccessor final : FAccessor
 	{
 		FVoidAccessor()
-		    : FAccessor(0, EType::Scalar, EComponentType::S8, false)
+		    : FAccessor(0, EType::Scalar, EComponentType::S8, false, FSparse())
 		{
 		}
 
