Fix geometry import only

#rnx

diff --git a/Engine/Plugins/Interchange/Runtime/Source/Import/Private/Mesh/InterchangeSkeletalMeshFactory.cpp b/Engine/Plugins/Interchange/Runtime/Source/Import/Private/Mesh/InterchangeSkeletalMeshFactory.cpp
--- a/Engine/Plugins/Interchange/Runtime/Source/Import/Private/Mesh/InterchangeSkeletalMeshFactory.cpp
+++ b/Engine/Plugins/Interchange/Runtime/Source/Import/Private/Mesh/InterchangeSkeletalMeshFactory.cpp
@@ -1110,7 +1110,7 @@ UInterchangeFactoryBase::FImportAssetResult UInterchangeSkeletalMeshFactory::Imp
 	}
 
 	ImportAssetObjectData.bIsReImport = bIsReImport;
-	ImportAssetObjectData.bApplySkinningOnly = bApplySkinningOnly;
+	ImportAssetObjectData.bApplyGeometryOnly = bApplyGeometryOnly;
 	
 	for (int32 LodIndex = 0; LodIndex < LodCount; ++LodIndex)
 	{
@@ -1199,7 +1199,16 @@ UInterchangeFactoryBase::FImportAssetResult UInterchangeSkeletalMeshFactory::Imp
 		bool bUseTimeZeroAsBindPose = false;
 		SkeletonNode->GetCustomUseTimeZeroForBindPose(bUseTimeZeroAsBindPose);
 		bool bDiffPose = false;
-		UE::Interchange::Private::FSkeletonHelper::ProcessImportMeshSkeleton(SkeletonReference, SkeletalMesh->GetRefSkeleton(), SkeletonDepth, Arguments.NodeContainer, RootJointNodeId, RefBonesBinary, bUseTimeZeroAsBindPose, bDiffPose);
+		if(bApplyGeometryOnly)
+		{
+			//Do not alter the skeletal mesh reference skeleton when importing geometry only
+			FReferenceSkeleton RefSkeleton;
+			UE::Interchange::Private::FSkeletonHelper::ProcessImportMeshSkeleton(SkeletonReference, RefSkeleton, SkeletonDepth, Arguments.NodeContainer, RootJointNodeId, RefBonesBinary, bUseTimeZeroAsBindPose, bDiffPose);
+		}
+		else
+		{
+			UE::Interchange::Private::FSkeletonHelper::ProcessImportMeshSkeleton(SkeletonReference, SkeletalMesh->GetRefSkeleton(), SkeletonDepth, Arguments.NodeContainer, RootJointNodeId, RefBonesBinary, bUseTimeZeroAsBindPose, bDiffPose);
+		}
 		if (bSpecifiedSkeleton && !SkeletonReference->IsCompatibleMesh(SkeletalMesh))
 		{
 			UE_LOG(LogInterchangeImport, Warning, TEXT("The skeleton %s is incompatible with the imported skeletalmesh asset %s"), *SkeletonReference->GetName(), *Arguments.AssetName);
@@ -1518,7 +1527,7 @@ UInterchangeFactoryBase::FImportAssetResult UInterchangeSkeletalMeshFactory::End
 
 	if (USkeleton* SkeletonReference = ImportAssetObjectData.SkeletonReference)
 	{
-		if ((!ImportAssetObjectData.bApplySkinningOnly || !ImportAssetObjectData.bIsReImport) && !SkeletonReference->MergeAllBonesToBoneTree(SkeletalMesh))
+		if ((!ImportAssetObjectData.bApplyGeometryOnly || !ImportAssetObjectData.bIsReImport) && !SkeletonReference->MergeAllBonesToBoneTree(SkeletalMesh))
 		{
 			TUniqueFunction<bool()> RecreateSkeleton = [this, WeakSkeletalMesh = TWeakObjectPtr<USkeletalMesh>(SkeletalMesh), WeakSkeleton = TWeakObjectPtr<USkeleton>(SkeletonReference)]()
 			{
diff --git a/Engine/Plugins/Interchange/Runtime/Source/Import/Public/Mesh/InterchangeSkeletalMeshFactory.h b/Engine/Plugins/Interchange/Runtime/Source/Import/Public/Mesh/InterchangeSkeletalMeshFactory.h
--- a/Engine/Plugins/Interchange/Runtime/Source/Import/Public/Mesh/InterchangeSkeletalMeshFactory.h
+++ b/Engine/Plugins/Interchange/Runtime/Source/Import/Public/Mesh/InterchangeSkeletalMeshFactory.h
@@ -31,7 +31,7 @@ public:
 	{
 		bool bIsReImport = false;
 		USkeleton* SkeletonReference = nullptr;
-		bool bApplySkinningOnly = false;
+		bool bApplyGeometryOnly = false;
 		TArray<FImportAssetObjectLODData> LodDatas;
 		bool IsValid() const;
 	};
diff --git a/Engine/Plugins/Interchange/Runtime/Source/Pipelines/Private/InterchangeGenericMeshPipeline.cpp b/Engine/Plugins/Interchange/Runtime/Source/Pipelines/Private/InterchangeGenericMeshPipeline.cpp
--- a/Engine/Plugins/Interchange/Runtime/Source/Pipelines/Private/InterchangeGenericMeshPipeline.cpp
+++ b/Engine/Plugins/Interchange/Runtime/Source/Pipelines/Private/InterchangeGenericMeshPipeline.cpp
@@ -57,6 +57,10 @@ void UInterchangeGenericMeshPipeline::AdjustSettingsForContext(EInterchangePipel
 			CommonSkeletalMeshesAndAnimationsProperties->Skeleton = SkeletalMesh->GetSkeleton();
 			bImportStaticMeshes = false;
 			HideCategories.Add(StaticMeshesCategory);
+			if(SkeletalMeshImportContentType == EInterchangeSkeletalMeshContentType::Geometry)
+			{
+				CommonMeshesProperties->ForceAllMeshAsType = EInterchangeForceMeshType::IFMT_SkeletalMesh;
+			}
 		}
 		else if (UStaticMesh* StaticMesh = Cast<UStaticMesh>(ReimportAsset))
 		{
@@ -135,7 +139,7 @@ void UInterchangeGenericMeshPipeline::ExecutePipeline(UInterchangeBaseNodeContai
 		UE_LOG(LogInterchangePipeline, Warning, TEXT("UInterchangeGenericMeshPipeline: Cannot execute pre-import pipeline because InBaseNodeContrainer is null"));
 		return;
 	}
-
+	
 	BaseNodeContainer = InBaseNodeContainer;
 	SourceDatas.Empty(InSourceDatas.Num());
 	for (const UInterchangeSourceData* SourceData : InSourceDatas)
