[Interchange] Morph Target handling fixes

diff --git a/Engine/Plugins/Interchange/Runtime/Source/Parsers/Fbx/Private/FbxScene.cpp b/Engine/Plugins/Interchange/Runtime/Source/Parsers/Fbx/Private/FbxScene.cpp
--- a/Engine/Plugins/Interchange/Runtime/Source/Parsers/Fbx/Private/FbxScene.cpp
+++ b/Engine/Plugins/Interchange/Runtime/Source/Parsers/Fbx/Private/FbxScene.cpp
@@ -701,55 +701,74 @@ namespace UE
 						continue;
 					}
 					
-					FString SkeletonUid;
+					TSet<FString> SkeletonUids;
 					if (MorphTargetAnimationBuildingData.InterchangeMeshNode->IsSkinnedMesh())
 					{
-						//Find the root joint for this MeshGeometry
-						FString JointNodeUid;
-						FString ParentNodeUid;
-						MorphTargetAnimationBuildingData.InterchangeMeshNode->GetSkeletonDependency(0, JointNodeUid);
-						ParentNodeUid = JointNodeUid;
-
-						while (!JointNodeUid.Equals(UInterchangeBaseNode::InvalidNodeUid()))
+						//Find the root joint(s) for this MeshGeometry
+						TArray<FString> SkeletonDependencies;
+						MorphTargetAnimationBuildingData.InterchangeMeshNode->GetSkeletonDependencies(SkeletonDependencies);
+						for (const FString& SkeletonDependency : SkeletonDependencies)
 						{
-							if (const UInterchangeSceneNode* Node = Cast< UInterchangeSceneNode >(NodeContainer.GetNode(ParentNodeUid)))
+							FString JointNodeUid = SkeletonDependency;
+							FString ParentNodeUid = SkeletonDependency;
+
+							while (!JointNodeUid.Equals(UInterchangeBaseNode::InvalidNodeUid()))
 							{
-								if (Node->IsSpecializedTypeContains(FSceneNodeStaticData::GetJointSpecializeTypeString()))
+								if (const UInterchangeSceneNode* Node = Cast< UInterchangeSceneNode >(NodeContainer.GetNode(ParentNodeUid)))
 								{
-									JointNodeUid = ParentNodeUid;
-									ParentNodeUid = Node->GetParentUid();
+									if (Node->IsSpecializedTypeContains(FSceneNodeStaticData::GetJointSpecializeTypeString()))
+									{
+										JointNodeUid = ParentNodeUid;
+										ParentNodeUid = Node->GetParentUid();
+									}
+									else
+									{
+										break;
+									}
 								}
 								else
 								{
 									break;
 								}
 							}
-							else
+
+							if (!JointNodeUid.Equals(UInterchangeBaseNode::InvalidNodeUid()))
 							{
-								break;
+								SkeletonUids.Add(JointNodeUid);
 							}
 						}
-
-						if (!JointNodeUid.Equals(UInterchangeBaseNode::InvalidNodeUid()))
-						{
-							SkeletonUid = JointNodeUid;
-						}
 					}
 					else
 					{
-						//If it is not skinned then the only way to get this animation to show up is by mesh type forcing to Skeleton,
-						//where the root node id is going to be the Skeleton id:
-						FbxNode* RootNode = SDKScene->GetRootNode();
-						SkeletonUid = Parser.GetFbxHelper()->GetFbxNodeHierarchyName(RootNode);
+						//Find MeshInstances: where CustomAssetInstanceUid == MeshNode->GetUniqueID
+						// For every occurance create a morphtarget entry with given MeshNode->GetUniqueID 
+						NodeContainer.IterateNodesOfType<UInterchangeSceneNode>([&](const FString& NodeUid, UInterchangeSceneNode* SceneNode)
+							{
+								FString AssetInstanceUid;
+								if (SceneNode->GetCustomAssetInstanceUid(AssetInstanceUid) && AssetInstanceUid == MorphTargetAnimationBuildingData.InterchangeMeshNode->GetUniqueID())
+								{
+									SkeletonUids.Add(SceneNode->GetUniqueID());
+								}
+							});
+						
+						if (SkeletonUids.Num() == 0)
+						{
+							//If it is not skinned and does not have an instantation, then it is presumed to get used on the RootNode level.
+							FbxNode* RootNode = SDKScene->GetRootNode();
+							SkeletonUids.Add(Parser.GetFbxHelper()->GetFbxNodeHierarchyName(RootNode));
+						}						
 					}
 
-					if (const UInterchangeSceneNode* SkeletonNode = Cast< UInterchangeSceneNode >(NodeContainer.GetNode(SkeletonUid)))
+					for (const FString& SkeletonUid : SkeletonUids)
 					{
-						//For the given skeleton:
-						TMap<int32, TArray<FMorphTargetAnimationBuildingData>>& MorphTargetAnimationPerAnimationIndex = MorphTargetAnimationsBuildingDataGrouped.FindOrAdd(SkeletonNode);
-						//For the given skeleton and animationindex:
-						TArray<FMorphTargetAnimationBuildingData>& MorphTargetAnimations = MorphTargetAnimationPerAnimationIndex.FindOrAdd(MorphTargetAnimationBuildingData.AnimationIndex);
-						MorphTargetAnimations.Add(MorphTargetAnimationBuildingData);
+						if (const UInterchangeSceneNode* SkeletonNode = Cast< UInterchangeSceneNode >(NodeContainer.GetNode(SkeletonUid)))
+						{
+							//For the given skeleton:
+							TMap<int32, TArray<FMorphTargetAnimationBuildingData>>& MorphTargetAnimationPerAnimationIndex = MorphTargetAnimationsBuildingDataGrouped.FindOrAdd(SkeletonNode);
+							//For the given skeleton and animationindex:
+							TArray<FMorphTargetAnimationBuildingData>& MorphTargetAnimations = MorphTargetAnimationPerAnimationIndex.FindOrAdd(MorphTargetAnimationBuildingData.AnimationIndex);
+							MorphTargetAnimations.Add(MorphTargetAnimationBuildingData);
+						}
 					}
 				}
 				
diff --git a/Engine/Plugins/Interchange/Runtime/Source/Pipelines/Private/InterchangeGenericStaticMeshPipeline.cpp b/Engine/Plugins/Interchange/Runtime/Source/Pipelines/Private/InterchangeGenericStaticMeshPipeline.cpp
--- a/Engine/Plugins/Interchange/Runtime/Source/Pipelines/Private/InterchangeGenericStaticMeshPipeline.cpp
+++ b/Engine/Plugins/Interchange/Runtime/Source/Pipelines/Private/InterchangeGenericStaticMeshPipeline.cpp
@@ -201,7 +201,7 @@ void UInterchangeGenericMeshPipeline::ExecutePreImportPipelineStaticMesh()
 			{
 				// If we haven't yet managed to build a factory node, look at static mesh geometry directly.
 				TArray<FString> MeshUids;
-				PipelineMeshesUtilities->GetAllStaticMeshGeometry(MeshUids, bConvertSkeletalMeshToStaticMesh);
+				PipelineMeshesUtilities->GetAllStaticMeshGeometry(MeshUids, bConvertSkeletalMeshToStaticMesh, bConvertStaticsWithMorphTargetsToSkeletals);
 
 				TMap<int32, TArray<FString>> MeshUidsPerLodIndex;
 
@@ -282,7 +282,7 @@ void UInterchangeGenericMeshPipeline::ExecutePreImportPipelineStaticMesh()
 			if (!bFoundMeshes)
 			{
 				TArray<FString> MeshUids;
-				PipelineMeshesUtilities->GetAllStaticMeshGeometry(MeshUids, bConvertSkeletalMeshToStaticMesh);
+				PipelineMeshesUtilities->GetAllStaticMeshGeometry(MeshUids, bConvertSkeletalMeshToStaticMesh, bConvertStaticsWithMorphTargetsToSkeletals);
 
 				// Work out which meshes are collision meshes which correspond to another mesh
 				TMap<FString, TArray<FString>> MeshToCollisionMeshMap;
diff --git a/Engine/Plugins/Interchange/Runtime/Source/Pipelines/Private/InterchangePipelineMeshesUtilities.cpp b/Engine/Plugins/Interchange/Runtime/Source/Pipelines/Private/InterchangePipelineMeshesUtilities.cpp
--- a/Engine/Plugins/Interchange/Runtime/Source/Pipelines/Private/InterchangePipelineMeshesUtilities.cpp
+++ b/Engine/Plugins/Interchange/Runtime/Source/Pipelines/Private/InterchangePipelineMeshesUtilities.cpp
@@ -429,25 +429,25 @@ void UInterchangePipelineMeshesUtilities::IterateAllSkinnedMeshGeometry(TFunctio
 	}
 }
 
-void UInterchangePipelineMeshesUtilities::GetAllStaticMeshGeometry(TArray<FString>& MeshGeometryUids, const bool bConvertSkeletalMeshToStaticMesh) const
+void UInterchangePipelineMeshesUtilities::GetAllStaticMeshGeometry(TArray<FString>& MeshGeometryUids, const bool bConvertSkeletalMeshToStaticMesh, const bool bConvertStaticsWithMorphTargetsToSkeletals) const
 {
 	MeshGeometryUids.Empty(MeshGeometriesPerMeshUid.Num());
 	for (const TPair<FString, FInterchangeMeshGeometry>& MeshGeometryUidAndMeshGeometry : MeshGeometriesPerMeshUid)
 	{
 		const FInterchangeMeshGeometry& MeshGeometry = MeshGeometryUidAndMeshGeometry.Value;
-		if ((bConvertSkeletalMeshToStaticMesh || !MeshGeometry.MeshNode->IsSkinnedMesh()) && !MeshGeometry.MeshNode->IsMorphTarget())
+		if ((bConvertSkeletalMeshToStaticMesh || !MeshGeometry.MeshNode->IsSkinnedMesh()) && !MeshGeometry.MeshNode->IsMorphTarget() && !(bConvertStaticsWithMorphTargetsToSkeletals && (MeshGeometry.MeshNode->GetMorphTargetDependeciesCount() > 0)))
 		{
 			MeshGeometryUids.Add(MeshGeometry.MeshUid);
 		}
 	}
 }
 
-void UInterchangePipelineMeshesUtilities::IterateAllStaticMeshGeometry(TFunctionRef<void(const FInterchangeMeshGeometry&)> IterationLambda, const bool bConvertSkeletalMeshToStaticMesh) const
+void UInterchangePipelineMeshesUtilities::IterateAllStaticMeshGeometry(TFunctionRef<void(const FInterchangeMeshGeometry&)> IterationLambda, const bool bConvertSkeletalMeshToStaticMesh, const bool bConvertStaticsWithMorphTargetsToSkeletals) const
 {
 	for (const TPair<FString, FInterchangeMeshGeometry>& MeshGeometryUidAndMeshGeometry : MeshGeometriesPerMeshUid)
 	{
 		const FInterchangeMeshGeometry& MeshGeometry = MeshGeometryUidAndMeshGeometry.Value;
-		if ((bConvertSkeletalMeshToStaticMesh || !MeshGeometry.MeshNode->IsSkinnedMesh()) && !MeshGeometry.MeshNode->IsMorphTarget())
+		if ((bConvertSkeletalMeshToStaticMesh || !MeshGeometry.MeshNode->IsSkinnedMesh()) && !MeshGeometry.MeshNode->IsMorphTarget() && !(bConvertStaticsWithMorphTargetsToSkeletals && (MeshGeometry.MeshNode->GetMorphTargetDependeciesCount() > 0)))
 		{
 			IterationLambda(MeshGeometry);
 		}
diff --git a/Engine/Plugins/Interchange/Runtime/Source/Pipelines/Public/InterchangePipelineMeshesUtilities.h b/Engine/Plugins/Interchange/Runtime/Source/Pipelines/Public/InterchangePipelineMeshesUtilities.h
--- a/Engine/Plugins/Interchange/Runtime/Source/Pipelines/Public/InterchangePipelineMeshesUtilities.h
+++ b/Engine/Plugins/Interchange/Runtime/Source/Pipelines/Public/InterchangePipelineMeshesUtilities.h
@@ -192,12 +192,12 @@ public:
 	* Get all static mesh geometry unique ids.
 	*/
 	UFUNCTION(BlueprintCallable, Category = "Interchange | Pipeline | Meshes")
-	void GetAllStaticMeshGeometry(TArray<FString>& MeshGeometryUids, const bool bConvertSkeletalMeshToStaticMesh) const;
+	void GetAllStaticMeshGeometry(TArray<FString>& MeshGeometryUids, const bool bConvertSkeletalMeshToStaticMesh, const bool bConvertStaticsWithMorphTargetsToSkeletals) const;
 
 	/**
 	* Iterate all static mesh geometry.
 	*/
-	void IterateAllStaticMeshGeometry(TFunctionRef<void(const FInterchangeMeshGeometry&)> IterationLambda, const bool bConvertSkeletalMeshToStaticMesh) const;
+	void IterateAllStaticMeshGeometry(TFunctionRef<void(const FInterchangeMeshGeometry&)> IterationLambda, const bool bConvertSkeletalMeshToStaticMesh, const bool bConvertStaticsWithMorphTargetsToSkeletals) const;
 
 	/**
 	* Get all not instanced mesh geometry unique ids.
