Fixed level re-import logic when asset has been deleted and reimported as empty
Fixed level re-import logic when reimporting using another file

diff --git a/Engine/Plugins/Interchange/Runtime/Source/Import/Private/Scene/InterchangeSceneImportAssetFactory.cpp b/Engine/Plugins/Interchange/Runtime/Source/Import/Private/Scene/InterchangeSceneImportAssetFactory.cpp
--- a/Engine/Plugins/Interchange/Runtime/Source/Import/Private/Scene/InterchangeSceneImportAssetFactory.cpp
+++ b/Engine/Plugins/Interchange/Runtime/Source/Import/Private/Scene/InterchangeSceneImportAssetFactory.cpp
@@ -100,44 +100,16 @@ UInterchangeFactoryBase::FImportAssetResult UInterchangeSceneImportAssetFactory:
 	{
 		UE_LOG(LogInterchangeImport, Warning, TEXT("Could not create InterchangeSceneImportAsset asset %s"), *Arguments.AssetName);
 	}
-
-	FImportAssetResult Result;
-	Result.ImportedObject = SceneImportAsset;
-	
-	return Result;
-}
-
-UObject* UInterchangeSceneImportAssetFactory::ImportAssetObject_Async(const FImportAssetObjectParams& Arguments)
-{
-	using namespace UE::Interchange::Private::InterchangeSceneImportAssetFactory;
-
-	UClass* TargetClass = GetFactoryClass();
-
-	const UInterchangeSceneImportAssetFactoryNode* FactoryNode = GetFactoryNode(Arguments, TargetClass);
-	if (!FactoryNode)
-	{
-		return nullptr;
-	}
-
-	UInterchangeSceneImportAsset* SceneImportAsset = Cast<UInterchangeSceneImportAsset>(FindOrCreateAsset(Arguments, TargetClass));
-
-	if (!ensure(SceneImportAsset))
+	//Getting the file Hash will cache it into the source data
+	else if (ensure(Arguments.SourceData))
 	{
-		UInterchangeResultError_Generic* Message = AddMessage<UInterchangeResultError_Generic>();
-		Message->SourceAssetName = Arguments.SourceData->GetFilename();
-		Message->DestinationAssetName = Arguments.AssetName;
-		Message->AssetType = UInterchangeSceneImportAsset::StaticClass();
-		Message->Text = FText::Format(LOCTEXT("CreateAssetFailed", "Could not create nor find SceneImportAsset asset {0}."), FText::FromString(Arguments.AssetName));
-		return nullptr;
+		Arguments.SourceData->GetFileContentHash();
 	}
 
-	/** Apply all FactoryNode custom attributes to the level sequence asset */
-	FactoryNode->ApplyAllCustomAttributeToObject(SceneImportAsset);
-
-	//Getting the file Hash will cache it into the source data
-	Arguments.SourceData->GetFileContentHash();
+	FImportAssetResult Result;
+	Result.ImportedObject = SceneImportAsset;
 
-	return SceneImportAsset;
+	return Result;
 }
 
 void UInterchangeSceneImportAssetFactory::SetupObject_GameThread(const FSetupObjectParams& Arguments)
@@ -145,27 +117,26 @@ void UInterchangeSceneImportAssetFactory::SetupObject_GameThread(const FSetupObj
 	check(IsInGameThread());
 	Super::SetupObject_GameThread(Arguments);
 
+	UInterchangeSceneImportAsset* SceneImportAsset = Cast<UInterchangeSceneImportAsset>(Arguments.ImportedObject);
+	UInterchangeSceneImportAssetFactoryNode* FactoryNode = Cast<UInterchangeSceneImportAssetFactoryNode>(Arguments.FactoryNode);
+	if (ensure(FactoryNode && SceneImportAsset && Arguments.SourceData))
+	{
+		/** Apply all FactoryNode custom attributes to the level sequence asset */
+		FactoryNode->ApplyAllCustomAttributeToObject(SceneImportAsset);
+
 #if WITH_EDITORONLY_DATA
-	using namespace UE::Interchange;
+		using namespace UE::Interchange;
 
-	if (ensure(Arguments.ImportedObject && Arguments.SourceData))
-	{
-		// We must call the Update of the asset source file in the main thread because UAssetImportData::Update execute some delegate we do not control
-		UInterchangeSceneImportAsset* SceneImportAsset = Cast<UInterchangeSceneImportAsset>(Arguments.ImportedObject);
-		if (ensure(SceneImportAsset))
+		if (Arguments.bIsReimport)
 		{
-			if (Arguments.bIsReimport)
-			{
-				ensure(SceneImportAsset->AssetImportData);
-				SceneImportAsset->UpdateSceneObjects();
-			}
-
-			FFactoryCommon::FUpdateImportAssetDataParameters Parameters(SceneImportAsset, SceneImportAsset->AssetImportData, Arguments.SourceData, Arguments.NodeUniqueID, Arguments.NodeContainer, Arguments.OriginalPipelines);
-			SceneImportAsset->AssetImportData = Cast<UInterchangeAssetImportData>(FFactoryCommon::UpdateImportAssetData(Parameters));
+			ensure(SceneImportAsset->AssetImportData);
+			SceneImportAsset->UpdateSceneObjects();
 		}
 
-	}
+		FFactoryCommon::FUpdateImportAssetDataParameters Parameters(SceneImportAsset, SceneImportAsset->AssetImportData, Arguments.SourceData, Arguments.NodeUniqueID, Arguments.NodeContainer, Arguments.OriginalPipelines);
+		SceneImportAsset->AssetImportData = Cast<UInterchangeAssetImportData>(FFactoryCommon::UpdateImportAssetData(Parameters));
 #endif
+	}
 }
 
 void UInterchangeSceneImportAssetFactory::FinalizeObject_GameThread(const FSetupObjectParams& Arguments)
@@ -173,11 +144,15 @@ void UInterchangeSceneImportAssetFactory::FinalizeObject_GameThread(const FSetup
 	check(IsInGameThread());
 	Super::FinalizeObject_GameThread(Arguments);
 
-#if WITH_EDITOR
-	if (UInterchangeSceneImportAsset* SceneImportAsset = Cast<UInterchangeSceneImportAsset>(Arguments.ImportedObject))
+	UInterchangeSceneImportAsset* SceneImportAsset = Cast<UInterchangeSceneImportAsset>(Arguments.ImportedObject);
+	if (!ensure(SceneImportAsset && Arguments.FactoryNode))
 	{
-		SceneImportAsset->RegisterWorldRenameCallbacks();
+		return;
 	}
+
+
+#if WITH_EDITOR
+	SceneImportAsset->RegisterWorldRenameCallbacks();
 #endif
 }
 
diff --git a/Engine/Plugins/Interchange/Runtime/Source/Import/Public/Scene/InterchangeSceneImportAssetFactory.h b/Engine/Plugins/Interchange/Runtime/Source/Import/Public/Scene/InterchangeSceneImportAssetFactory.h
--- a/Engine/Plugins/Interchange/Runtime/Source/Import/Public/Scene/InterchangeSceneImportAssetFactory.h
+++ b/Engine/Plugins/Interchange/Runtime/Source/Import/Public/Scene/InterchangeSceneImportAssetFactory.h
@@ -20,7 +20,6 @@ public:
 
 	virtual UClass* GetFactoryClass() const override;
 	virtual FImportAssetResult BeginImportAsset_GameThread(const FImportAssetObjectParams& Arguments) override;
-	virtual UObject* ImportAssetObject_Async(const FImportAssetObjectParams& Arguments) override;
 	virtual void SetupObject_GameThread(const FSetupObjectParams& Arguments) override;
 	virtual void FinalizeObject_GameThread(const FSetupObjectParams& Arguments) override;
 	virtual bool GetSourceFilenames(const UObject* Object, TArray<FString>& OutSourceFilenames) const override;
diff --git a/Engine/Plugins/Interchange/Runtime/Source/Pipelines/Private/InterchangeGenericScenesPipeline.cpp b/Engine/Plugins/Interchange/Runtime/Source/Pipelines/Private/InterchangeGenericScenesPipeline.cpp
--- a/Engine/Plugins/Interchange/Runtime/Source/Pipelines/Private/InterchangeGenericScenesPipeline.cpp
+++ b/Engine/Plugins/Interchange/Runtime/Source/Pipelines/Private/InterchangeGenericScenesPipeline.cpp
@@ -33,7 +33,7 @@
 
 namespace UE::Interchange::Private
 {
-	//Either a non-joint, or a rootjoint can be a parent (only those get FactoryNodes) :
+	//Either a non-joint, or a root joint can be a parent (only those get FactoryNodes) :
 	FString FindFactoryParentSceneNodeUid(UInterchangeBaseNodeContainer* BaseNodeContainer, const UInterchangeSceneNode* SceneNode)
 	{
 		FString ParentUid = SceneNode->GetParentUid();
@@ -66,6 +66,21 @@ namespace UE::Interchange::Private
 
 		return UInterchangeBaseNode::InvalidNodeUid();
 	}
+
+	void UpdateReimportStrategyFlags(UInterchangeBaseNodeContainer& NodeContainer, UInterchangeFactoryBaseNode& FactoryNode, EReimportStrategyFlags ReimportPropertyStrategy)
+	{
+		FactoryNode.SetReimportStrategyFlags(ReimportPropertyStrategy);
+
+		TArray<FString> ActorDependencies;
+		FactoryNode.GetFactoryDependencies(ActorDependencies);
+		for (const FString& FactoryNodeID : ActorDependencies)
+		{
+			if (UInterchangeFactoryBaseNode* DependencyFactoryNode = NodeContainer.GetFactoryNode(FactoryNodeID))
+			{
+				UpdateReimportStrategyFlags(NodeContainer, *DependencyFactoryNode, ReimportPropertyStrategy);
+			}
+		}
+	}
 }
 
 void UInterchangeGenericLevelPipeline::ExecutePipeline(UInterchangeBaseNodeContainer* InBaseNodeContainer, const TArray<UInterchangeSourceData*>& InSourceDatas)
@@ -184,6 +199,8 @@ void UInterchangeGenericLevelPipeline::ExecutePipeline(UInterchangeBaseNodeConta
 
 void UInterchangeGenericLevelPipeline::ExecuteSceneNodePreImport(const FTransform& GlobalOffsetTransform, const UInterchangeSceneNode* SceneNode)
 {
+	using namespace UE::Interchange;
+
 	if (!BaseNodeContainer || !SceneNode)
 	{
 		return;
@@ -285,8 +302,8 @@ void UInterchangeGenericLevelPipeline::ExecuteSceneNodePreImport(const FTransfor
 		SetUpFactoryNode(ActorFactoryNode, SceneNode, TranslatedAssetNode);
 	}
 
-	//Make sure all actor factory nodes have the specified strategy
-	ActorFactoryNode->SetReimportStrategyFlags(ReimportStrategy);
+	//Make sure all actor factory nodes and dependencies have the specified strategy
+	Private::UpdateReimportStrategyFlags(*BaseNodeContainer, *ActorFactoryNode, ReimportStrategy);
 
 #if WITH_EDITORONLY_DATA
 	// Add dependency to newly created factory node
diff --git a/Engine/Source/Runtime/Interchange/Engine/Private/InterchangeImportCommon.cpp b/Engine/Source/Runtime/Interchange/Engine/Private/InterchangeImportCommon.cpp
--- a/Engine/Source/Runtime/Interchange/Engine/Private/InterchangeImportCommon.cpp
+++ b/Engine/Source/Runtime/Interchange/Engine/Private/InterchangeImportCommon.cpp
@@ -313,12 +313,15 @@ namespace UE::Interchange
 		}
 	}
 
-	UObject* FFactoryCommon::GetObjectToReimport(UObject* ReimportObject, const FString& PackageName, const FString& AssetName, const FString& SubPathString)
+	UObject* FFactoryCommon::GetObjectToReimport(UObject* ReimportObject, const UInterchangeFactoryBaseNode& FactoryNode, const FString& PackageName, const FString& AssetName, const FString& SubPathString)
 	{
 #if WITH_EDITORONLY_DATA
-		if (UInterchangeSceneImportAsset* SceneImportAsset = Cast<UInterchangeSceneImportAsset>(ReimportObject))
+		if (ReimportObject && !ReimportObject->GetClass()->IsChildOf(FactoryNode.GetObjectClass()))
 		{
-			return SceneImportAsset->GetSceneObject(PackageName, AssetName, SubPathString);
+			if (UInterchangeSceneImportAsset* SceneImportAsset = Cast<UInterchangeSceneImportAsset>(ReimportObject))
+			{
+				return SceneImportAsset->GetSceneObject(PackageName, AssetName, SubPathString);
+			}
 		}
 #endif
 
diff --git a/Engine/Source/Runtime/Interchange/Engine/Private/InterchangeSceneImportAsset.cpp b/Engine/Source/Runtime/Interchange/Engine/Private/InterchangeSceneImportAsset.cpp
--- a/Engine/Source/Runtime/Interchange/Engine/Private/InterchangeSceneImportAsset.cpp
+++ b/Engine/Source/Runtime/Interchange/Engine/Private/InterchangeSceneImportAsset.cpp
@@ -199,7 +199,33 @@ UObject* UInterchangeSceneImportAsset::GetSceneObject(const FString& PackageName
 
 	if (SceneObjects.Find(ObjectPath))
 	{
-		return ObjectPath.TryLoad();
+		if (UObject* SceneObject = ObjectPath.TryLoad())
+		{
+			if (IsValid(SceneObject))
+			{
+				if (SceneObject->IsA<AActor>())
+				{
+					return SceneObject;
+				}
+
+				// Most likely an asset, check whether SceneObject has actually already been imported
+				TArray<UObject*> SubObjects;
+				GetObjectsWithOuter(SceneObject, SubObjects);
+				for (UObject* SubObject : SubObjects)
+				{
+					if (SubObject && SubObject->IsA<UInterchangeAssetImportData>())
+					{
+						return SceneObject;
+					}
+				}
+
+				return nullptr;
+			}
+
+			// SceneObject is still in memory but invalid. Move it to TransientPackage
+			// Call UObject::Rename because for actors AActor::Rename will unnecessarily unregister and re-register components
+			SceneObject->UObject::Rename(nullptr, GetTransientPackage(), REN_DontCreateRedirectors | REN_ForceNoResetLoaders);
+		}
 	}
 #endif
 
diff --git a/Engine/Source/Runtime/Interchange/Engine/Private/Tasks/InterchangeTaskCreateSceneObjects.cpp b/Engine/Source/Runtime/Interchange/Engine/Private/Tasks/InterchangeTaskCreateSceneObjects.cpp
--- a/Engine/Source/Runtime/Interchange/Engine/Private/Tasks/InterchangeTaskCreateSceneObjects.cpp
+++ b/Engine/Source/Runtime/Interchange/Engine/Private/Tasks/InterchangeTaskCreateSceneObjects.cpp
@@ -64,6 +64,11 @@ void UE::Interchange::FTaskCreateSceneObjects::DoTask(ENamedThreads::Type Curren
 
 	for (UInterchangeFactoryBaseNode* FactoryNode : FactoryNodes)
 	{
+		if (!FactoryNode)
+		{
+			continue;
+		}
+
 		UInterchangeFactoryBase* Factory = NewObject<UInterchangeFactoryBase>(GetTransientPackage(), FactoryClass);
 		Factory->SetResultsContainer(AsyncHelper->AssetImportResult->GetResults());
 		AsyncHelper->AddCreatedFactory(FactoryNode->GetUniqueID(), Factory);
@@ -75,7 +80,7 @@ void UE::Interchange::FTaskCreateSceneObjects::DoTask(ENamedThreads::Type Curren
 		CreateSceneObjectsParams.ObjectName = NodeDisplayName;
 		CreateSceneObjectsParams.FactoryNode = FactoryNode;
 		CreateSceneObjectsParams.Level = CurrentLevel;
-		CreateSceneObjectsParams.ReimportObject = FFactoryCommon::GetObjectToReimport(ReimportObject, WorldPath, WorldName, NodePrefix + NodeDisplayName);
+		CreateSceneObjectsParams.ReimportObject = FFactoryCommon::GetObjectToReimport(ReimportObject, *FactoryNode, WorldPath, WorldName, NodePrefix + NodeDisplayName);
 		CreateSceneObjectsParams.ReimportFactoryNode = FFactoryCommon::GetFactoryNode(ReimportObject, WorldPath, WorldName, NodePrefix + NodeDisplayName);
 
 		if (AsyncHelper->BaseNodeContainers.IsValidIndex(SourceIndex))
diff --git a/Engine/Source/Runtime/Interchange/Engine/Private/Tasks/InterchangeTaskImportObject.cpp b/Engine/Source/Runtime/Interchange/Engine/Private/Tasks/InterchangeTaskImportObject.cpp
--- a/Engine/Source/Runtime/Interchange/Engine/Private/Tasks/InterchangeTaskImportObject.cpp
+++ b/Engine/Source/Runtime/Interchange/Engine/Private/Tasks/InterchangeTaskImportObject.cpp
@@ -197,7 +197,7 @@ namespace UE::Interchange::Private
 		FString AssetName;
 		Private::InternalGetPackageName(*AsyncHelper, SourceIndex, PackageBasePath, FactoryNode, PackageName, AssetName);
 		bool bSkipAsset = false;
-		UObject* ReimportObject = UE::Interchange::FFactoryCommon::GetObjectToReimport(AsyncHelper->TaskData.ReimportObject, PackageName, AssetName);
+		UObject* ReimportObject = UE::Interchange::FFactoryCommon::GetObjectToReimport(AsyncHelper->TaskData.ReimportObject, *FactoryNode, PackageName, AssetName);
 		if (ReimportObject)
 		{
 			UInterchangeBaseNodeContainer* NodeContainer = nullptr;
@@ -313,7 +313,7 @@ void UE::Interchange::FTaskImportObject_GameThread::DoTask(ENamedThreads::Type C
 	FString AssetName;
 	Private::InternalGetPackageName(*AsyncHelper, SourceIndex, PackageBasePath, FactoryNode, PackageName, AssetName);
 
-	UObject* ReimportObject = FFactoryCommon::GetObjectToReimport(AsyncHelper->TaskData.ReimportObject, PackageName, AssetName);
+	UObject* ReimportObject = FFactoryCommon::GetObjectToReimport(AsyncHelper->TaskData.ReimportObject, *FactoryNode, PackageName, AssetName);
 	if (!ensure(!IsGarbageCollecting()))
 	{
 		//Skip this asset
@@ -603,7 +603,10 @@ void UE::Interchange::FTaskImportObjectFinalize_GameThread::DoTask(ENamedThreads
 				AssetInfo.Factory = Factory;
 				AssetInfo.FactoryNode = FactoryNode;
 
-				const UObject* ReimportObject = FFactoryCommon::GetObjectToReimport(AsyncHelper->TaskData.ReimportObject, ImportAssetResult.ImportedObject->GetPathName(), ImportAssetResult.ImportedObject->GetName());
+
+				const FString PackagePathName = ImportAssetResult.ImportedObject->GetPathName();
+				const FString AssetName = ImportAssetResult.ImportedObject->GetName();
+				const UObject* ReimportObject = FFactoryCommon::GetObjectToReimport(AsyncHelper->TaskData.ReimportObject, *FactoryNode, PackagePathName, AssetName);
 				AssetInfo.bIsReimport = bool(ReimportObject != nullptr);
 			}
 
diff --git a/Engine/Source/Runtime/Interchange/Engine/Public/InterchangeImportCommon.h b/Engine/Source/Runtime/Interchange/Engine/Public/InterchangeImportCommon.h
--- a/Engine/Source/Runtime/Interchange/Engine/Public/InterchangeImportCommon.h
+++ b/Engine/Source/Runtime/Interchange/Engine/Public/InterchangeImportCommon.h
@@ -115,7 +115,7 @@ namespace UE::Interchange
 		 * @oaram AssetName: Asset name of the actual object to reimport
 		 * @oaram SubPathString: Optional subobject name
 		 */
-		static 	UObject* GetObjectToReimport(UObject* ReimportObject, const FString& PackageName, const FString& AssetName, const FString& SubPathString = FString());
+		static 	UObject* GetObjectToReimport(UObject* ReimportObject, const UInterchangeFactoryBaseNode& FactoryNode, const FString& PackageName, const FString& AssetName, const FString& SubPathString = FString());
 
 		/**
 		 * If the ReimportObject is a UInterchangeSceneImportAsset, returns the factory node
