Fixed crash when number of LODs is greater than MAX_STATIC_MESH_LODS

diff --git a/Engine/Plugins/Interchange/Runtime/Source/Import/Private/Mesh/InterchangeStaticMeshFactory.cpp b/Engine/Plugins/Interchange/Runtime/Source/Import/Private/Mesh/InterchangeStaticMeshFactory.cpp
--- a/Engine/Plugins/Interchange/Runtime/Source/Import/Private/Mesh/InterchangeStaticMeshFactory.cpp
+++ b/Engine/Plugins/Interchange/Runtime/Source/Import/Private/Mesh/InterchangeStaticMeshFactory.cpp
@@ -245,7 +245,12 @@ UInterchangeFactoryBase::FImportAssetResult UInterchangeStaticMeshFactory::Impor
 
 	ensure(!StaticMesh->AreRenderingResourcesInitialized());
 
-	const int32 LodCount = StaticMeshFactoryNode->GetLodDataCount();
+	const int32 LodCount = FMath::Min(StaticMeshFactoryNode->GetLodDataCount(), MAX_STATIC_MESH_LODS);
+	if (LodCount != StaticMeshFactoryNode->GetLodDataCount())
+	{
+		const int32 LodCountDiff = StaticMeshFactoryNode->GetLodDataCount() - MAX_STATIC_MESH_LODS;
+		UE_LOG(LogInterchangeImport, Warning, TEXT("Reached the maximum number of LODs for a Static Mesh(%d) - discarding %d LOD meshes."), MAX_STATIC_MESH_LODS, LodCountDiff);
+	}
 #if WITH_EDITOR
 	const int32 PrevLodCount = StaticMesh->GetNumSourceModels();
 	const int32 FinalLodCount = FMath::Max(PrevLodCount, LodCount);
@@ -310,7 +315,7 @@ UInterchangeFactoryBase::FImportAssetResult UInterchangeStaticMeshFactory::Impor
 	// Now import geometry for each LOD
 	TArray<FString> LodDataUniqueIds;
 	StaticMeshFactoryNode->GetLodDataUniqueIds(LodDataUniqueIds);
-	ensure(LodDataUniqueIds.Num() == LodCount);
+	ensure(LodDataUniqueIds.Num() >= LodCount);
 
 	TArray<FMeshDescription>& LodMeshDescriptions = ImportAssetObjectData.LodMeshDescriptions;
 	LodMeshDescriptions.SetNum(LodCount);
