PCX Invalid file handling.

diff --git a/Engine/Plugins/Interchange/Runtime/Source/Import/Private/Texture/InterchangePCXTranslator.cpp b/Engine/Plugins/Interchange/Runtime/Source/Import/Private/Texture/InterchangePCXTranslator.cpp
--- a/Engine/Plugins/Interchange/Runtime/Source/Import/Private/Texture/InterchangePCXTranslator.cpp
+++ b/Engine/Plugins/Interchange/Runtime/Source/Import/Private/Texture/InterchangePCXTranslator.cpp
@@ -118,6 +118,14 @@ TOptional<UE::Interchange::FImportImage> UInterchangePCXTranslator::GetTexturePa
 	const UE::Interchange::Private::FPCXFileHeader* PCX = (UE::Interchange::Private::FPCXFileHeader*)Buffer;
 	if (Length >= sizeof(UE::Interchange::Private::FPCXFileHeader) && PCX->Manufacturer == 10)
 	{
+		if (PCX->XMax < PCX->XMin ||
+			PCX->YMax < PCX->YMin)
+		{
+			FText ErrorMessage = NSLOCTEXT("InterchangePCXTranslator", "CorruptedFileDim", "Failed to import PCX, invalid dims.");
+			FTextureTranslatorUtilities::LogError(*this, MoveTemp(ErrorMessage));
+			return TOptional<UE::Interchange::FImportImage>();
+		}
+
 		int32 NewU = PCX->XMax + 1 - PCX->XMin;
 		int32 NewV = PCX->YMax + 1 - PCX->YMin;
 
@@ -169,9 +177,9 @@ TOptional<UE::Interchange::FImportImage> UInterchangePCXTranslator::GetTexturePa
 			);
 
 			uint8* Dest = static_cast<uint8*>(PayloadData.RawData.GetData());
-
+			uint64 DestSize = PayloadData.RawData.GetSize();
 			// Doing a memset to make sure the alpha channel is set to 0xff since we only have 3 color planes.
-			FMemory::Memset(Dest, 0xff, NewU * NewV * FTextureSource::GetBytesPerPixel(PayloadData.Format));
+			FMemory::Memset(Dest, 0xff, DestSize);
 
 			// Copy upside-down scanlines.
 			Buffer += 128;
@@ -203,6 +211,16 @@ TOptional<UE::Interchange::FImportImage> UInterchangePCXTranslator::GetTexturePa
 							Overflow = Overflow - RunLength;
 						}
 
+						// NewU is max uint16, i is max uint16, j is max uint16, runlength is max uint16.
+						uint64 FarthestOffset = ((uint64)i * NewU + (j + RunLength - 1))*4 + ColorPlane;
+
+						if (FarthestOffset >= DestSize)
+						{
+							FText ErrorMessage = NSLOCTEXT("InterchangePCXTranslator", "InvalidRunLength", "Failed to import PCX, RLE length is invalid during decode");
+							FTextureTranslatorUtilities::LogError(*this, MoveTemp(ErrorMessage));
+							return TOptional<UE::Interchange::FImportImage>();
+						}
+
 						//checkf(((i*NewU + RunLength) * 4 + ColorPlane) < (Texture->Source.CalcMipSize(0)),
 						//	TEXT("RLE going off the end of buffer"));
 						for (int32 k = j; k < j + RunLength; k++)
