Fixing regression where we can't override MaterialX material functions during import

diff --git a/Engine/Plugins/Interchange/Runtime/Source/Pipelines/Private/InterchangeMaterialXPipeline.cpp b/Engine/Plugins/Interchange/Runtime/Source/Pipelines/Private/InterchangeMaterialXPipeline.cpp
--- a/Engine/Plugins/Interchange/Runtime/Source/Pipelines/Private/InterchangeMaterialXPipeline.cpp
+++ b/Engine/Plugins/Interchange/Runtime/Source/Pipelines/Private/InterchangeMaterialXPipeline.cpp
@@ -12,6 +12,14 @@
 #include "Materials/MaterialFunction.h"
 #include "Misc/PackageName.h"
 
+TMap<FString, EInterchangeMaterialXShaders> UInterchangeMaterialXPipeline::PathToEnumMapping
+{
+	{TEXT("/Interchange/Functions/MX_StandardSurface.MX_StandardSurface"),EInterchangeMaterialXShaders::StandardSurface},
+	{TEXT("/Interchange/Functions/MX_TransmissionSurface.MX_TransmissionSurface"),EInterchangeMaterialXShaders::StandardSurfaceTransmission},
+	{TEXT("/Interchange/Functions/MX_SurfaceUnlit.MX_SurfaceUnlit"),EInterchangeMaterialXShaders::SurfaceUnlit},
+	{TEXT("/Interchange/Functions/MX_UsdPreviewSurface.MX_UsdPreviewSurface"),EInterchangeMaterialXShaders::UsdPreviewSurface},
+};
+
 UInterchangeMaterialXPipeline::UInterchangeMaterialXPipeline()
 	: MaterialXSettings(UMaterialXPipelineSettings::StaticClass()->GetDefaultObject< UMaterialXPipelineSettings>())
 {
diff --git a/Engine/Plugins/Interchange/Runtime/Source/Pipelines/Public/InterchangeMaterialXPipeline.h b/Engine/Plugins/Interchange/Runtime/Source/Pipelines/Public/InterchangeMaterialXPipeline.h
--- a/Engine/Plugins/Interchange/Runtime/Source/Pipelines/Public/InterchangeMaterialXPipeline.h
+++ b/Engine/Plugins/Interchange/Runtime/Source/Pipelines/Public/InterchangeMaterialXPipeline.h
@@ -97,5 +97,5 @@ protected:
 	}
 
 private:
-	TMap<FString, EInterchangeMaterialXShaders> PathToEnumMapping;
+	static TMap<FString, EInterchangeMaterialXShaders> PathToEnumMapping;
 };
\ No newline at end of file
