[Interchange] Plugin and Ini config fixes

diff --git a/Engine/Config/BaseEngine.ini b/Engine/Config/BaseEngine.ini
--- a/Engine/Config/BaseEngine.ini
+++ b/Engine/Config/BaseEngine.ini
@@ -1161,45 +1161,6 @@ DefaultWorldPartitionHLODLayer=/Engine/Maps/Templates/HLODs/HLODLayer_Instanced.
 +ClassRedirects=(OldName="XRDeviceVisualizationComponent", NewName="/Script/XRBase.XRDeviceVisualizationComponent")
 +ClassRedirects=(OldName="XRLoadingScreenFunctionLibrary", NewName="/Script/XRBase.XRLoadingScreenFunctionLibrary")
 
-; Rename MaterialX Material Expressions and Enums
-+ClassRedirects=(OldName="/Script/InterchangeImport.MaterialExpressionAppend3Vector", NewName="/Script/InterchangeImport.MaterialExpressionMaterialXAppend3Vector")
-+ClassRedirects=(OldName="/Script/InterchangeImport.MaterialExpressionAppend4Vector", NewName="/Script/InterchangeImport.MaterialExpressionMaterialXAppend4Vector")
-+ClassRedirects=(OldName="/Script/InterchangeImport.MaterialExpressionBurn", NewName="/Script/InterchangeImport.MaterialExpressionMaterialXBurn")
-+ClassRedirects=(OldName="/Script/InterchangeImport.MaterialExpressionDifference", NewName="/Script/InterchangeImport.MaterialExpressionMaterialXDifference")
-+ClassRedirects=(OldName="/Script/InterchangeImport.MaterialExpressionDisjointOver", NewName="/Script/InterchangeImport.MaterialExpressionMaterialXDisjointOver")
-+ClassRedirects=(OldName="/Script/InterchangeImport.MaterialExpressionDodge", NewName="/Script/InterchangeImport.MaterialExpressionMaterialXDodge")
-+ClassRedirects=(OldName="/Script/InterchangeImport.MaterialExpressionExponential", NewName="/Script/InterchangeImport.MaterialExpressionMaterialXExponential")
-+ClassRedirects=(OldName="/Script/InterchangeImport.MaterialExpressionFractal3D", NewName="/Script/InterchangeImport.MaterialExpressionMaterialXFractal3D")
-+ClassRedirects=(OldName="/Script/InterchangeImport.MaterialExpressionHsvToRgb", NewName="/Script/InterchangeImport.MaterialExpressionMaterialXHsvToRgb")
-+ClassRedirects=(OldName="/Script/InterchangeImport.MaterialExpressionIn", NewName="/Script/InterchangeImport.MaterialExpressionMaterialXIn")
-+ClassRedirects=(OldName="/Script/InterchangeImport.MaterialExpressionLength", NewName="/Script/InterchangeImport.MaterialExpressionMaterialXLength")
-+ClassRedirects=(OldName="/Script/InterchangeImport.MaterialExpressionLogarithm", NewName="/Script/InterchangeImport.MaterialExpressionMaterialXLogarithm")
-+ClassRedirects=(OldName="/Script/InterchangeImport.MaterialExpressionLuminance", NewName="/Script/InterchangeImport.MaterialExpressionMaterialXLuminance")
-+ClassRedirects=(OldName="/Script/InterchangeImport.MaterialExpressionMask", NewName="/Script/InterchangeImport.MaterialExpressionMaterialXMask")
-+ClassRedirects=(OldName="/Script/InterchangeImport.MaterialExpressionMatte", NewName="/Script/InterchangeImport.MaterialExpressionMaterialXMatte")
-+ClassRedirects=(OldName="/Script/InterchangeImport.MaterialExpressionMinus", NewName="/Script/InterchangeImport.MaterialExpressionMaterialXMinus")
-+ClassRedirects=(OldName="/Script/InterchangeImport.MaterialExpressionOut", NewName="/Script/InterchangeImport.MaterialExpressionMaterialXOut")
-+ClassRedirects=(OldName="/Script/InterchangeImport.MaterialExpressionOver", NewName="/Script/InterchangeImport.MaterialExpressionMaterialXOver")
-+ClassRedirects=(OldName="/Script/InterchangeImport.MaterialExpressionOverlay", NewName="/Script/InterchangeImport.MaterialExpressionMaterialXOverlay")
-+ClassRedirects=(OldName="/Script/InterchangeImport.MaterialExpressionPlace2D", NewName="/Script/InterchangeImport.MaterialExpressionMaterialXPlace2D")
-+ClassRedirects=(OldName="/Script/InterchangeImport.MaterialExpressionPlus", NewName="/Script/InterchangeImport.MaterialExpressionMaterialXPlus")
-+ClassRedirects=(OldName="/Script/InterchangeImport.MaterialExpressionPremult", NewName="/Script/InterchangeImport.MaterialExpressionMaterialXPremult")
-+ClassRedirects=(OldName="/Script/InterchangeImport.MaterialExpressionRamp4", NewName="/Script/InterchangeImport.MaterialExpressionMaterialXRamp4")
-+ClassRedirects=(OldName="/Script/InterchangeImport.MaterialExpressionRampLeftRight", NewName="/Script/InterchangeImport.MaterialExpressionMaterialXRampLeftRight")
-+ClassRedirects=(OldName="/Script/InterchangeImport.MaterialExpressionRampTopBottom", NewName="/Script/InterchangeImport.MaterialExpressionMaterialXRampTopBottom")
-+ClassRedirects=(OldName="/Script/InterchangeImport.MaterialExpressionRemap", NewName="/Script/InterchangeImport.MaterialExpressionMaterialXRemap")
-+ClassRedirects=(OldName="/Script/InterchangeImport.MaterialExpressionRgbToHsv", NewName="/Script/InterchangeImport.MaterialExpressionMaterialXRgbToHsv")
-+ClassRedirects=(OldName="/Script/InterchangeImport.MaterialExpressionRotate2D", NewName="/Script/InterchangeImport.MaterialExpressionMaterialXRotate2D")
-+ClassRedirects=(OldName="/Script/InterchangeImport.MaterialExpressionScreen", NewName="/Script/InterchangeImport.MaterialExpressionMaterialXScreen")
-+ClassRedirects=(OldName="/Script/InterchangeImport.MaterialExpressionSplitLeftRight", NewName="/Script/InterchangeImport.MaterialExpressionMaterialXSplitLeftRight")
-+ClassRedirects=(OldName="/Script/InterchangeImport.MaterialExpressionSplitTopBottom", NewName="/Script/InterchangeImport.MaterialExpressionMaterialXSplitTopBottom")
-+ClassRedirects=(OldName="/Script/InterchangeImport.MaterialExpressionTextureSampleParameterBlur", NewName="/Script/InterchangeImport.MaterialExpressionMaterialXTextureSampleParameterBlur")
-+ClassRedirects=(OldName="/Script/InterchangeImport.MaterialExpressionUnpremult", NewName="/Script/InterchangeImport.MaterialExpressionMaterialXUnpremult")
-+EnumRedirects=(OldName="ETextureSampleBlurFilter", NewName="EMaterialXTextureSampleBlurFilter")
-+EnumRedirects=(OldName="ELuminanceMode", NewName="EMaterialXLuminanceMode")
-
-
-
 
 
 [CoreUObject.Metadata]
@@ -3316,19 +3277,6 @@ EnvironmentCollisionThickness = 2.f
 ; Limit the number of extracted shapes per component as these collisions are very expensive
 EnvironmentCollisionMaxShapes = 32
 
-[/Script/InterchangeEngine.InterchangeProjectSettings]
-ContentImportSettings=(DefaultPipelineStackOverride=((Materials, "Materials"),(Textures, "Textures")),ShowPipelineStacksConfigurationDialogOverride=((Textures, False)),PipelineStacks=(("Assets", (Pipelines=("/Interchange/Pipelines/DefaultAssetsPipeline.DefaultAssetsPipeline"),PerTranslatorPipelines=((Translator="/Script/InterchangeImport.InterchangeMaterialXTranslator",Pipelines=("/Interchange/Pipelines/DefaultAssetsPipeline.DefaultAssetsPipeline","/Interchange/Pipelines/DefaultMaterialXPipeline.DefaultMaterialXPipeline")),(Translator="/Script/InterchangeImport.InterchangeGltfTranslator",Pipelines=("/Interchange/Pipelines/DefaultGLTFAssetsPipeline.DefaultGLTFAssetsPipeline","/Interchange/Pipelines/DefaultGLTFPipeline.DefaultGLTFPipeline"))))),("Materials", (Pipelines=("/Interchange/Pipelines/DefaultMaterialPipeline.DefaultMaterialPipeline"),PerTranslatorPipelines=((Translator="/Script/InterchangeImport.InterchangeMaterialXTranslator",Pipelines=("/Interchange/Pipelines/DefaultMaterialPipeline.DefaultMaterialPipeline","/Interchange/Pipelines/DefaultMaterialXPipeline.DefaultMaterialXPipeline"))))),("Textures", (Pipelines=("/Interchange/Pipelines/DefaultTexturePipeline.DefaultTexturePipeline")))),DefaultPipelineStack="Assets",PipelineConfigurationDialogClass="/Script/InterchangeEditorPipelines.InterchangePipelineConfigurationGeneric",bShowPipelineStacksConfigurationDialog=True)
-SceneImportSettings=(PipelineStacks=(("Scene", (Pipelines=("/Interchange/Pipelines/DefaultSceneAssetsPipeline.DefaultSceneAssetsPipeline","/Interchange/Pipelines/DefaultSceneLevelPipeline.DefaultSceneLevelPipeline"),PerTranslatorPipelines=((Translator="/Script/InterchangeImport.InterchangeMaterialXTranslator",Pipelines=("/Interchange/Pipelines/DefaultSceneAssetsPipeline.DefaultSceneAssetsPipeline","/Interchange/Pipelines/DefaultSceneLevelPipeline.DefaultSceneLevelPipeline","/Interchange/Pipelines/DefaultMaterialXPipeline.DefaultMaterialXPipeline")),(Translator="/Script/InterchangeImport.InterchangeGltfTranslator",Pipelines=("/Interchange/Pipelines/DefaultGLTFSceneAssetsPipeline.DefaultGLTFSceneAssetsPipeline","/Interchange/Pipelines/DefaultSceneLevelPipeline.DefaultSceneLevelPipeline","/Interchange/Pipelines/DefaultGLTFPipeline.DefaultGLTFPipeline")))))),DefaultPipelineStack="Scene",PipelineConfigurationDialogClass="/Script/InterchangeEditorPipelines.InterchangePipelineConfigurationGeneric",bShowPipelineStacksConfigurationDialog=True)
-FilePickerClass=/Script/InterchangeEditorUtilities.InterchangeFilePickerGeneric
-bStaticMeshUseSmoothEdgesIfSmoothingInformationIsMissing=True
-GenericPipelineClass=/Script/InterchangePipelines.InterchangeGenericAssetsPipeline
-
-[/Script/InterchangePipelines.MaterialXPipelineSettings]
-PredefinedSurfaceShaders=((StandardSurface, "/Interchange/Functions/MX_StandardSurface.MX_StandardSurface"),(StandardSurfaceTransmission, "/Interchange/Functions/MX_TransmissionSurface.MX_TransmissionSurface"),(SurfaceUnlit, "/Interchange/Functions/MX_SurfaceUnlit.MX_SurfaceUnlit"),(UsdPreviewSurface, "/Interchange/Functions/MX_UsdPreviewSurface.MX_UsdPreviewSurface"))
-
-[/Script/InterchangePipelines.GLTFPipelineSettings]
-MaterialParents=(("MI_Default_Opaque", "/Interchange/gltf/MaterialInstances/MI_Default_Opaque.MI_Default_Opaque"), ("MI_Default_Mask", "/Interchange/gltf/MaterialInstances/MI_Default_Mask.MI_Default_Mask"), ("MI_Default_Blend", "/Interchange/gltf/MaterialInstances/MI_Default_Blend.MI_Default_Blend"), ("MI_Unlit_Opaque", "/Interchange/gltf/MaterialInstances/MI_Unlit_Opaque.MI_Unlit_Opaque"), ("MI_Unlit_Mask", "/Interchange/gltf/MaterialInstances/MI_Unlit_Mask.MI_Unlit_Mask"), ("MI_Unlit_Blend", "/Interchange/gltf/MaterialInstances/MI_Unlit_Blend.MI_Unlit_Blend"), ("MI_ClearCoat_Opaque", "/Interchange/gltf/MaterialInstances/MI_ClearCoat_Opaque.MI_ClearCoat_Opaque"), ("MI_ClearCoat_Mask", "/Interchange/gltf/MaterialInstances/MI_ClearCoat_Mask.MI_ClearCoat_Mask"), ("MI_ClearCoat_Blend", "/Interchange/gltf/MaterialInstances/MI_ClearCoat_Blend.MI_ClearCoat_Blend"), ("MI_Sheen_Opaque", "/Interchange/gltf/MaterialInstances/MI_Sheen_Opaque.MI_Sheen_Opaque"), ("MI_Sheen_Mask", "/Interchange/gltf/MaterialInstances/MI_Sheen_Mask.MI_Sheen_Mask"), ("MI_Sheen_Blend", "/Interchange/gltf/MaterialInstances/MI_Sheen_Blend.MI_Sheen_Blend"), ("MI_Transmission", "/Interchange/gltf/MaterialInstances/MI_Transmission.MI_Transmission"), ("MI_SpecularGlossiness_Opaque", "/Interchange/gltf/MaterialInstances/MI_SpecularGlossiness_Opaque.MI_SpecularGlossiness_Opaque"), ("MI_SpecularGlossiness_Mask", "/Interchange/gltf/MaterialInstances/MI_SpecularGlossiness_Mask.MI_SpecularGlossiness_Mask"), ("MI_SpecularGlossiness_Blend", "/Interchange/gltf/MaterialInstances/MI_SpecularGlossiness_Blend.MI_SpecularGlossiness_Blend"), ("MI_Default_Opaque_DS", "/Interchange/gltf/MaterialInstances/MI_Default_Opaque_DS.MI_Default_Opaque_DS"), ("MI_Default_Mask_DS", "/Interchange/gltf/MaterialInstances/MI_Default_Mask_DS.MI_Default_Mask_DS"), ("MI_Default_Blend_DS", "/Interchange/gltf/MaterialInstances/MI_Default_Blend_DS.MI_Default_Blend_DS"), ("MI_Unlit_Opaque_DS", "/Interchange/gltf/MaterialInstances/MI_Unlit_Opaque_DS.MI_Unlit_Opaque_DS"), ("MI_Unlit_Mask_DS", "/Interchange/gltf/MaterialInstances/MI_Unlit_Mask_DS.MI_Unlit_Mask_DS"), ("MI_Unlit_Blend_DS", "/Interchange/gltf/MaterialInstances/MI_Unlit_Blend_DS.MI_Unlit_Blend_DS"), ("MI_ClearCoat_Opaque_DS", "/Interchange/gltf/MaterialInstances/MI_ClearCoat_Opaque_DS.MI_ClearCoat_Opaque_DS"), ("MI_ClearCoat_Mask_DS", "/Interchange/gltf/MaterialInstances/MI_ClearCoat_Mask_DS.MI_ClearCoat_Mask_DS"), ("MI_ClearCoat_Blend_DS", "/Interchange/gltf/MaterialInstances/MI_ClearCoat_Blend_DS.MI_ClearCoat_Blend_DS"), ("MI_Sheen_Opaque_DS", "/Interchange/gltf/MaterialInstances/MI_Sheen_Opaque_DS.MI_Sheen_Opaque_DS"), ("MI_Sheen_Mask_DS", "/Interchange/gltf/MaterialInstances/MI_Sheen_Mask_DS.MI_Sheen_Mask_DS"), ("MI_Sheen_Blend_DS", "/Interchange/gltf/MaterialInstances/MI_Sheen_Blend_DS.MI_Sheen_Blend_DS"), ("MI_Transmission_DS", "/Interchange/gltf/MaterialInstances/MI_Transmission_DS.MI_Transmission_DS"), ("MI_SpecularGlossiness_Opaque_DS", "/Interchange/gltf/MaterialInstances/MI_SpecularGlossiness_Opaque_DS.MI_SpecularGlossiness_Opaque_DS"), ("MI_SpecularGlossiness_Mask_DS", "/Interchange/gltf/MaterialInstances/MI_SpecularGlossiness_Mask_DS.MI_SpecularGlossiness_Mask_DS"), ("MI_SpecularGlossiness_Blend_DS", "/Interchange/gltf/MaterialInstances/MI_SpecularGlossiness_Blend_DS.MI_SpecularGlossiness_Blend_DS"))
-
 [BuildPatchServices]
 ; Enable/Disable mechanism that provides a disk overflow store for patch chunks.
 bEnableDiskOverflowStore=true
diff --git a/Engine/Plugins/Interchange/Runtime/Config/BaseInterchange.ini b/Engine/Plugins/Interchange/Runtime/Config/BaseInterchange.ini
--- a/Engine/Plugins/Interchange/Runtime/Config/BaseInterchange.ini
+++ b/Engine/Plugins/Interchange/Runtime/Config/BaseInterchange.ini
@@ -2,4 +2,64 @@
 +EnumRedirects=(OldName="/Script/InterchangePipelines.EInterchangeMaterialImportOption",ValueChanges=(("DoNotImport","ImportAsMaterials"))
 
 ; Rename Interchange CineCamera to PhysicalCamera FactoryNode
-+ClassRedirects=(OldName="/Script/InterchangeFactoryNodes.InterchangeCineCameraFactoryNode", NewName="/Script/InterchangeFactoryNodes.InterchangePhysicalCameraFactoryNode")
\ No newline at end of file
++ClassRedirects=(OldName="/Script/InterchangeFactoryNodes.InterchangeCineCameraFactoryNode", NewName="/Script/InterchangeFactoryNodes.InterchangePhysicalCameraFactoryNode")
+
+[/Script/InterchangeEngine.InterchangeProjectSettings]
+ContentImportSettings=(DefaultPipelineStackOverride=((Materials, "Materials"),(Textures, "Textures")),ShowPipelineStacksConfigurationDialogOverride=((Textures, False)),PipelineStacks=(("Assets", (Pipelines=("/Interchange/Pipelines/DefaultAssetsPipeline.DefaultAssetsPipeline"),PerTranslatorPipelines=((Translator="/Script/InterchangeImport.InterchangeMaterialXTranslator",Pipelines=("/Interchange/Pipelines/DefaultAssetsPipeline.DefaultAssetsPipeline","/Interchange/Pipelines/DefaultMaterialXPipeline.DefaultMaterialXPipeline")),(Translator="/Script/InterchangeImport.InterchangeGltfTranslator",Pipelines=("/Interchange/Pipelines/DefaultGLTFAssetsPipeline.DefaultGLTFAssetsPipeline","/Interchange/Pipelines/DefaultGLTFPipeline.DefaultGLTFPipeline"))))),("Materials", (Pipelines=("/Interchange/Pipelines/DefaultMaterialPipeline.DefaultMaterialPipeline"),PerTranslatorPipelines=((Translator="/Script/InterchangeImport.InterchangeMaterialXTranslator",Pipelines=("/Interchange/Pipelines/DefaultMaterialPipeline.DefaultMaterialPipeline","/Interchange/Pipelines/DefaultMaterialXPipeline.DefaultMaterialXPipeline"))))),("Textures", (Pipelines=("/Interchange/Pipelines/DefaultTexturePipeline.DefaultTexturePipeline")))),DefaultPipelineStack="Assets",PipelineConfigurationDialogClass="/Script/InterchangeEditorPipelines.InterchangePipelineConfigurationGeneric",bShowPipelineStacksConfigurationDialog=True)
+SceneImportSettings=(PipelineStacks=(("Scene", (Pipelines=("/Interchange/Pipelines/DefaultSceneAssetsPipeline.DefaultSceneAssetsPipeline","/Interchange/Pipelines/DefaultSceneLevelPipeline.DefaultSceneLevelPipeline"),PerTranslatorPipelines=((Translator="/Script/InterchangeImport.InterchangeMaterialXTranslator",Pipelines=("/Interchange/Pipelines/DefaultSceneAssetsPipeline.DefaultSceneAssetsPipeline","/Interchange/Pipelines/DefaultSceneLevelPipeline.DefaultSceneLevelPipeline","/Interchange/Pipelines/DefaultMaterialXPipeline.DefaultMaterialXPipeline")),(Translator="/Script/InterchangeImport.InterchangeGltfTranslator",Pipelines=("/Interchange/Pipelines/DefaultGLTFSceneAssetsPipeline.DefaultGLTFSceneAssetsPipeline","/Interchange/Pipelines/DefaultSceneLevelPipeline.DefaultSceneLevelPipeline","/Interchange/Pipelines/DefaultGLTFPipeline.DefaultGLTFPipeline")))))),DefaultPipelineStack="Scene",PipelineConfigurationDialogClass="/Script/InterchangeEditorPipelines.InterchangePipelineConfigurationGeneric",bShowPipelineStacksConfigurationDialog=True)
+FilePickerClass=/Script/InterchangeEditorUtilities.InterchangeFilePickerGeneric
+bStaticMeshUseSmoothEdgesIfSmoothingInformationIsMissing=True
+GenericPipelineClass=/Script/InterchangePipelines.InterchangeGenericAssetsPipeline
+
+[/Script/InterchangePipelines.GLTFPipelineSettings]
+MaterialParents=(("MI_Default_Opaque", "/Interchange/gltf/MaterialInstances/MI_Default_Opaque.MI_Default_Opaque"), ("MI_Default_Mask", "/Interchange/gltf/MaterialInstances/MI_Default_Mask.MI_Default_Mask"), ("MI_Default_Blend", "/Interchange/gltf/MaterialInstances/MI_Default_Blend.MI_Default_Blend"), ("MI_Unlit_Opaque", "/Interchange/gltf/MaterialInstances/MI_Unlit_Opaque.MI_Unlit_Opaque"), ("MI_Unlit_Mask", "/Interchange/gltf/MaterialInstances/MI_Unlit_Mask.MI_Unlit_Mask"), ("MI_Unlit_Blend", "/Interchange/gltf/MaterialInstances/MI_Unlit_Blend.MI_Unlit_Blend"), ("MI_ClearCoat_Opaque", "/Interchange/gltf/MaterialInstances/MI_ClearCoat_Opaque.MI_ClearCoat_Opaque"), ("MI_ClearCoat_Mask", "/Interchange/gltf/MaterialInstances/MI_ClearCoat_Mask.MI_ClearCoat_Mask"), ("MI_ClearCoat_Blend", "/Interchange/gltf/MaterialInstances/MI_ClearCoat_Blend.MI_ClearCoat_Blend"), ("MI_Sheen_Opaque", "/Interchange/gltf/MaterialInstances/MI_Sheen_Opaque.MI_Sheen_Opaque"), ("MI_Sheen_Mask", "/Interchange/gltf/MaterialInstances/MI_Sheen_Mask.MI_Sheen_Mask"), ("MI_Sheen_Blend", "/Interchange/gltf/MaterialInstances/MI_Sheen_Blend.MI_Sheen_Blend"), ("MI_Transmission", "/Interchange/gltf/MaterialInstances/MI_Transmission.MI_Transmission"), ("MI_SpecularGlossiness_Opaque", "/Interchange/gltf/MaterialInstances/MI_SpecularGlossiness_Opaque.MI_SpecularGlossiness_Opaque"), ("MI_SpecularGlossiness_Mask", "/Interchange/gltf/MaterialInstances/MI_SpecularGlossiness_Mask.MI_SpecularGlossiness_Mask"), ("MI_SpecularGlossiness_Blend", "/Interchange/gltf/MaterialInstances/MI_SpecularGlossiness_Blend.MI_SpecularGlossiness_Blend"), ("MI_Default_Opaque_DS", "/Interchange/gltf/MaterialInstances/MI_Default_Opaque_DS.MI_Default_Opaque_DS"), ("MI_Default_Mask_DS", "/Interchange/gltf/MaterialInstances/MI_Default_Mask_DS.MI_Default_Mask_DS"), ("MI_Default_Blend_DS", "/Interchange/gltf/MaterialInstances/MI_Default_Blend_DS.MI_Default_Blend_DS"), ("MI_Unlit_Opaque_DS", "/Interchange/gltf/MaterialInstances/MI_Unlit_Opaque_DS.MI_Unlit_Opaque_DS"), ("MI_Unlit_Mask_DS", "/Interchange/gltf/MaterialInstances/MI_Unlit_Mask_DS.MI_Unlit_Mask_DS"), ("MI_Unlit_Blend_DS", "/Interchange/gltf/MaterialInstances/MI_Unlit_Blend_DS.MI_Unlit_Blend_DS"), ("MI_ClearCoat_Opaque_DS", "/Interchange/gltf/MaterialInstances/MI_ClearCoat_Opaque_DS.MI_ClearCoat_Opaque_DS"), ("MI_ClearCoat_Mask_DS", "/Interchange/gltf/MaterialInstances/MI_ClearCoat_Mask_DS.MI_ClearCoat_Mask_DS"), ("MI_ClearCoat_Blend_DS", "/Interchange/gltf/MaterialInstances/MI_ClearCoat_Blend_DS.MI_ClearCoat_Blend_DS"), ("MI_Sheen_Opaque_DS", "/Interchange/gltf/MaterialInstances/MI_Sheen_Opaque_DS.MI_Sheen_Opaque_DS"), ("MI_Sheen_Mask_DS", "/Interchange/gltf/MaterialInstances/MI_Sheen_Mask_DS.MI_Sheen_Mask_DS"), ("MI_Sheen_Blend_DS", "/Interchange/gltf/MaterialInstances/MI_Sheen_Blend_DS.MI_Sheen_Blend_DS"), ("MI_Transmission_DS", "/Interchange/gltf/MaterialInstances/MI_Transmission_DS.MI_Transmission_DS"), ("MI_SpecularGlossiness_Opaque_DS", "/Interchange/gltf/MaterialInstances/MI_SpecularGlossiness_Opaque_DS.MI_SpecularGlossiness_Opaque_DS"), ("MI_SpecularGlossiness_Mask_DS", "/Interchange/gltf/MaterialInstances/MI_SpecularGlossiness_Mask_DS.MI_SpecularGlossiness_Mask_DS"), ("MI_SpecularGlossiness_Blend_DS", "/Interchange/gltf/MaterialInstances/MI_SpecularGlossiness_Blend_DS.MI_SpecularGlossiness_Blend_DS"))
+
+[/Script/InterchangePipelines.MaterialXPipelineSettings]
+PredefinedSurfaceShaders=((StandardSurface, "/Interchange/Functions/MX_StandardSurface.MX_StandardSurface"),(StandardSurfaceTransmission, "/Interchange/Functions/MX_TransmissionSurface.MX_TransmissionSurface"),(SurfaceUnlit, "/Interchange/Functions/MX_SurfaceUnlit.MX_SurfaceUnlit"),(UsdPreviewSurface, "/Interchange/Functions/MX_UsdPreviewSurface.MX_UsdPreviewSurface"))
+
+; Rename MaterialX Material Expressions and Enums
++ClassRedirects=(OldName="/Script/InterchangeImport.MaterialExpressionAppend3Vector", NewName="/Script/InterchangeImport.MaterialExpressionMaterialXAppend3Vector")
++ClassRedirects=(OldName="/Script/InterchangeImport.MaterialExpressionAppend4Vector", NewName="/Script/InterchangeImport.MaterialExpressionMaterialXAppend4Vector")
++ClassRedirects=(OldName="/Script/InterchangeImport.MaterialExpressionBurn", NewName="/Script/InterchangeImport.MaterialExpressionMaterialXBurn")
++ClassRedirects=(OldName="/Script/InterchangeImport.MaterialExpressionDifference", NewName="/Script/InterchangeImport.MaterialExpressionMaterialXDifference")
++ClassRedirects=(OldName="/Script/InterchangeImport.MaterialExpressionDisjointOver", NewName="/Script/InterchangeImport.MaterialExpressionMaterialXDisjointOver")
++ClassRedirects=(OldName="/Script/InterchangeImport.MaterialExpressionDodge", NewName="/Script/InterchangeImport.MaterialExpressionMaterialXDodge")
++ClassRedirects=(OldName="/Script/InterchangeImport.MaterialExpressionExponential", NewName="/Script/InterchangeImport.MaterialExpressionMaterialXExponential")
++ClassRedirects=(OldName="/Script/InterchangeImport.MaterialExpressionFractal3D", NewName="/Script/InterchangeImport.MaterialExpressionMaterialXFractal3D")
++ClassRedirects=(OldName="/Script/InterchangeImport.MaterialExpressionHsvToRgb", NewName="/Script/InterchangeImport.MaterialExpressionMaterialXHsvToRgb")
++ClassRedirects=(OldName="/Script/InterchangeImport.MaterialExpressionIn", NewName="/Script/InterchangeImport.MaterialExpressionMaterialXIn")
++ClassRedirects=(OldName="/Script/InterchangeImport.MaterialExpressionLength", NewName="/Script/InterchangeImport.MaterialExpressionMaterialXLength")
++ClassRedirects=(OldName="/Script/InterchangeImport.MaterialExpressionLogarithm", NewName="/Script/InterchangeImport.MaterialExpressionMaterialXLogarithm")
++ClassRedirects=(OldName="/Script/InterchangeImport.MaterialExpressionLuminance", NewName="/Script/InterchangeImport.MaterialExpressionMaterialXLuminance")
++ClassRedirects=(OldName="/Script/InterchangeImport.MaterialExpressionMask", NewName="/Script/InterchangeImport.MaterialExpressionMaterialXMask")
++ClassRedirects=(OldName="/Script/InterchangeImport.MaterialExpressionMatte", NewName="/Script/InterchangeImport.MaterialExpressionMaterialXMatte")
++ClassRedirects=(OldName="/Script/InterchangeImport.MaterialExpressionMinus", NewName="/Script/InterchangeImport.MaterialExpressionMaterialXMinus")
++ClassRedirects=(OldName="/Script/InterchangeImport.MaterialExpressionOut", NewName="/Script/InterchangeImport.MaterialExpressionMaterialXOut")
++ClassRedirects=(OldName="/Script/InterchangeImport.MaterialExpressionOver", NewName="/Script/InterchangeImport.MaterialExpressionMaterialXOver")
++ClassRedirects=(OldName="/Script/InterchangeImport.MaterialExpressionOverlay", NewName="/Script/InterchangeImport.MaterialExpressionMaterialXOverlay")
++ClassRedirects=(OldName="/Script/InterchangeImport.MaterialExpressionPlace2D", NewName="/Script/InterchangeImport.MaterialExpressionMaterialXPlace2D")
++ClassRedirects=(OldName="/Script/InterchangeImport.MaterialExpressionPlus", NewName="/Script/InterchangeImport.MaterialExpressionMaterialXPlus")
++ClassRedirects=(OldName="/Script/InterchangeImport.MaterialExpressionPremult", NewName="/Script/InterchangeImport.MaterialExpressionMaterialXPremult")
++ClassRedirects=(OldName="/Script/InterchangeImport.MaterialExpressionRamp4", NewName="/Script/InterchangeImport.MaterialExpressionMaterialXRamp4")
++ClassRedirects=(OldName="/Script/InterchangeImport.MaterialExpressionRampLeftRight", NewName="/Script/InterchangeImport.MaterialExpressionMaterialXRampLeftRight")
++ClassRedirects=(OldName="/Script/InterchangeImport.MaterialExpressionRampTopBottom", NewName="/Script/InterchangeImport.MaterialExpressionMaterialXRampTopBottom")
++ClassRedirects=(OldName="/Script/InterchangeImport.MaterialExpressionRemap", NewName="/Script/InterchangeImport.MaterialExpressionMaterialXRemap")
++ClassRedirects=(OldName="/Script/InterchangeImport.MaterialExpressionRgbToHsv", NewName="/Script/InterchangeImport.MaterialExpressionMaterialXRgbToHsv")
++ClassRedirects=(OldName="/Script/InterchangeImport.MaterialExpressionRotate2D", NewName="/Script/InterchangeImport.MaterialExpressionMaterialXRotate2D")
++ClassRedirects=(OldName="/Script/InterchangeImport.MaterialExpressionScreen", NewName="/Script/InterchangeImport.MaterialExpressionMaterialXScreen")
++ClassRedirects=(OldName="/Script/InterchangeImport.MaterialExpressionSplitLeftRight", NewName="/Script/InterchangeImport.MaterialExpressionMaterialXSplitLeftRight")
++ClassRedirects=(OldName="/Script/InterchangeImport.MaterialExpressionSplitTopBottom", NewName="/Script/InterchangeImport.MaterialExpressionMaterialXSplitTopBottom")
++ClassRedirects=(OldName="/Script/InterchangeImport.MaterialExpressionTextureSampleParameterBlur", NewName="/Script/InterchangeImport.MaterialExpressionMaterialXTextureSampleParameterBlur")
++ClassRedirects=(OldName="/Script/InterchangeImport.MaterialExpressionUnpremult", NewName="/Script/InterchangeImport.MaterialExpressionMaterialXUnpremult")
++EnumRedirects=(OldName="ETextureSampleBlurFilter", NewName="EMaterialXTextureSampleBlurFilter")
++EnumRedirects=(OldName="ELuminanceMode", NewName="EMaterialXLuminanceMode")
+
+[/Script/UnrealEd.ProjectPackagingSettings]
++DirectoriesToAlwaysCook=(Path="/Interchange/Functions")
++DirectoriesToAlwaysCook=(Path="/Interchange/gltf")
++DirectoriesToAlwaysCook=(Path="/Interchange/gltf/MaterialFunctions")
++DirectoriesToAlwaysCook=(Path="/Interchange/gltf/MaterialInstances")
++DirectoriesToAlwaysCook=(Path="/Interchange/gltf/Textures")
++DirectoriesToAlwaysCook=(Path="/Interchange/Materials")
++DirectoriesToAlwaysCook=(Path="/Interchange/Pipelines")
++DirectoriesToAlwaysCook=(Path="/Interchange/Utilities")
\ No newline at end of file
diff --git a/Engine/Plugins/Interchange/Runtime/Interchange.uplugin b/Engine/Plugins/Interchange/Runtime/Interchange.uplugin
--- a/Engine/Plugins/Interchange/Runtime/Interchange.uplugin
+++ b/Engine/Plugins/Interchange/Runtime/Interchange.uplugin
@@ -10,7 +10,7 @@
 	"DocsURL": "",
 	"MarketplaceURL": "",
 	"SupportURL": "",
-	"EnabledByDefault": false,
+	"EnabledByDefault": true,
 	"CanContainContent": true,
 	"IsBetaVersion": false,
 	"Installed": false,
diff --git a/Engine/Plugins/Interchange/Runtime/Source/Pipelines/Public/InterchangeMaterialXPipeline.h b/Engine/Plugins/Interchange/Runtime/Source/Pipelines/Public/InterchangeMaterialXPipeline.h
--- a/Engine/Plugins/Interchange/Runtime/Source/Pipelines/Public/InterchangeMaterialXPipeline.h
+++ b/Engine/Plugins/Interchange/Runtime/Source/Pipelines/Public/InterchangeMaterialXPipeline.h
@@ -44,7 +44,7 @@ enum class EInterchangeMaterialXShaders : uint8
 	MaxShaderCount
 };
 
-UCLASS(config = Engine, meta = (DisplayName = "Interchange MaterialX"))
+UCLASS(config = Interchange, meta = (DisplayName = "Interchange MaterialX"))
 class INTERCHANGEPIPELINES_API UMaterialXPipelineSettings : public UDeveloperSettings
 {
 	GENERATED_BODY()
diff --git a/Engine/Plugins/Interchange/Runtime/Source/Pipelines/Public/InterchangeglTFPipeline.h b/Engine/Plugins/Interchange/Runtime/Source/Pipelines/Public/InterchangeglTFPipeline.h
--- a/Engine/Plugins/Interchange/Runtime/Source/Pipelines/Public/InterchangeglTFPipeline.h
+++ b/Engine/Plugins/Interchange/Runtime/Source/Pipelines/Public/InterchangeglTFPipeline.h
@@ -16,7 +16,7 @@ class UInterchangeBaseNodeContainer;
 class UInterchangeMaterialInstanceFactoryNode;
 class UInterchangeShaderGraphNode;
 
-UCLASS(config = Engine, meta = (DisplayName = "Interchange glTF"))
+UCLASS(config = Interchange, meta = (DisplayName = "Interchange glTF"))
 class INTERCHANGEPIPELINES_API UGLTFPipelineSettings : public UDeveloperSettings
 {
 	GENERATED_BODY()
diff --git a/Engine/Source/Runtime/Interchange/Engine/Public/InterchangeProjectSettings.h b/Engine/Source/Runtime/Interchange/Engine/Public/InterchangeProjectSettings.h
--- a/Engine/Source/Runtime/Interchange/Engine/Public/InterchangeProjectSettings.h
+++ b/Engine/Source/Runtime/Interchange/Engine/Public/InterchangeProjectSettings.h
@@ -74,7 +74,7 @@ struct FInterchangeContentImportSettings : public FInterchangeImportSettings
 	TMap<EInterchangeTranslatorAssetType, bool> ShowPipelineStacksConfigurationDialogOverride;
 };
 
-UCLASS(config=Engine, meta=(DisplayName=Interchange), MinimalAPI)
+UCLASS(config=Interchange, meta=(DisplayName=Interchange), MinimalAPI)
 class UInterchangeProjectSettings : public UDeveloperSettings
 {
 	GENERATED_BODY()
