Bug fix in MaterialXLuminance.

After dragging another Luminance variable from the same connected node the Editor crashes with Assertion failed: Index >= 0 && Index < CurrentScopeChunks->Num()

diff --git a/Engine/Plugins/Interchange/Runtime/Source/Import/Private/MaterialX/MaterialExpressions/MaterialExpressionLuminance.cpp b/Engine/Plugins/Interchange/Runtime/Source/Import/Private/MaterialX/MaterialExpressions/MaterialExpressionLuminance.cpp
--- a/Engine/Plugins/Interchange/Runtime/Source/Import/Private/MaterialX/MaterialExpressions/MaterialExpressionLuminance.cpp
+++ b/Engine/Plugins/Interchange/Runtime/Source/Import/Private/MaterialX/MaterialExpressions/MaterialExpressionLuminance.cpp
@@ -36,9 +36,13 @@ int32 UMaterialExpressionMaterialXLuminance::Compile(FMaterialCompiler* Compiler
 	}
 
 	const int32 LuminanceFactorsIndex = Compiler->Constant3(LuminanceFactors.R, LuminanceFactors.G, LuminanceFactors.B);
-	auto Type = Compiler->GetParameterType(Input.Compile(Compiler));
+	const int32 InputIndex = Input.Compile(Compiler);
+	if (InputIndex == INDEX_NONE)
+	{
+		return INDEX_NONE;
+	}
 
-	return Compiler->Dot(Compiler->ComponentMask(Input.Compile(Compiler), true, true, true, false), LuminanceFactorsIndex);
+	return Compiler->Dot(Compiler->ComponentMask(InputIndex, true, true, true, false), LuminanceFactorsIndex);
 }
 
 void UMaterialExpressionMaterialXLuminance::GetCaption(TArray<FString>& OutCaptions) const
