[glTF - Interchange] Can't import gltf with bin file over 2 GB

diff --git a/Engine/Plugins/Interchange/Runtime/Source/Parsers/GLTFCore/Private/GLTF/GLTFAccessor.cpp b/Engine/Plugins/Interchange/Runtime/Source/Parsers/GLTFCore/Private/GLTF/GLTFAccessor.cpp
--- a/Engine/Plugins/Interchange/Runtime/Source/Parsers/GLTFCore/Private/GLTF/GLTFAccessor.cpp
+++ b/Engine/Plugins/Interchange/Runtime/Source/Parsers/GLTFCore/Private/GLTF/GLTFAccessor.cpp
@@ -258,7 +258,7 @@ namespace GLTF
 		}
 
 		template<typename ItemType, uint32 ItemElementCount>
-		void CopyWithoutConversion(const uint32 ByteStride, const FAccessor::EComponentType ComponentType, bool bNormalized, const uint32 Count, const FBufferView& BufferView, const uint32 ByteOffset, ItemType* Buffer)
+		void CopyWithoutConversion(const uint32 ByteStride, const FAccessor::EComponentType ComponentType, bool bNormalized, const uint32 Count, const FBufferView& BufferView, const uint64 ByteOffset, ItemType* Buffer)
 		{
 			// Stride equals item size => use simpler copy
 			if ((ByteStride == 0) || ByteStride == sizeof(ItemType))
@@ -297,7 +297,7 @@ namespace GLTF
 			CopyWithoutConversion<ItemType, ItemElementCount>(Accessor.ByteStride, Accessor.ComponentType, Accessor.bNormalized, Accessor.Count, Accessor.BufferView, Accessor.ByteOffset, Buffer);
 		}
 
-		bool FillFloatArray(bool bNormalized, const FBufferView& BufferView, const uint32 ByteOffset, FAccessor::EComponentType ComponentType, const uint32 ByteStride, const uint32 Count, float* Buffer)
+		bool FillFloatArray(bool bNormalized, const FBufferView& BufferView, const uint64 ByteOffset, FAccessor::EComponentType ComponentType, const uint32 ByteStride, const uint32 Count, float* Buffer)
 		{
 			if (!bNormalized)
 			{
@@ -333,7 +333,7 @@ namespace GLTF
 			return false;
 		}
 
-		bool FillUnsignedIntArray(const FBufferView& BufferView, const uint32 ByteOffset, FAccessor::EComponentType ComponentType, const uint32 Count, uint32* Buffer)
+		bool FillUnsignedIntArray(const FBufferView& BufferView, const uint64 ByteOffset, FAccessor::EComponentType ComponentType, const uint32 Count, uint32* Buffer)
 		{
 			const uint8* Src = BufferView.DataAt(ByteOffset);
 			switch (ComponentType)
@@ -352,7 +352,7 @@ namespace GLTF
 			}
 		}
 
-		TArray<uint32> AcquireUnsignedIntArray(const uint32 Count, const FBufferView& BufferView, const uint32 ByteOffset, FAccessor::EComponentType ComponentType)
+		TArray<uint32> AcquireUnsignedIntArray(const uint32 Count, const FBufferView& BufferView, const uint64 ByteOffset, FAccessor::EComponentType ComponentType)
 		{
 			TArray<uint32> ValueBuffer;
 			ValueBuffer.SetNumUninitialized(Count, false);
@@ -362,7 +362,7 @@ namespace GLTF
 			return ValueBuffer;
 		}
 
-		bool AcquireUnsignedInt(uint32 Index, const FBufferView& BufferView, const uint32 ByteOffset, const uint32 ByteStride, FAccessor::EComponentType ComponentType, uint32& Buffer)
+		bool AcquireUnsignedInt(uint32 Index, const FBufferView& BufferView, const uint64 ByteOffset, const uint32 ByteStride, FAccessor::EComponentType ComponentType, uint32& Buffer)
 		{
 			Buffer = 0;
 
@@ -383,7 +383,7 @@ namespace GLTF
 			}
 		}
 
-		bool AcquireUnsignedInt16x4(uint32 Index, const FBufferView& BufferView, const uint32 ByteOffset, const uint32 ByteStride, FAccessor::EComponentType ComponentType, uint16 Buffer[4])
+		bool AcquireUnsignedInt16x4(uint32 Index, const FBufferView& BufferView, const uint64 ByteOffset, const uint32 ByteStride, FAccessor::EComponentType ComponentType, uint16 Buffer[4])
 		{
 			Buffer[0] = Buffer[1] = Buffer[2] = Buffer[3] = 0;
 
@@ -407,7 +407,7 @@ namespace GLTF
 			}
 		}
 
-		bool AcquireFloat(uint32 Index, const FBufferView& BufferView, const uint32 ByteOffset, const uint32 ByteStride, FAccessor::EComponentType ComponentType, float& Buffer)
+		bool AcquireFloat(uint32 Index, const FBufferView& BufferView, const uint64 ByteOffset, const uint32 ByteStride, FAccessor::EComponentType ComponentType, float& Buffer)
 		{
 			Buffer = 0;
 
@@ -422,7 +422,7 @@ namespace GLTF
 			}
 		}
 
-		bool AcquireVec2(uint32 Index, const FBufferView& BufferView, const uint32 ByteOffset, const uint32 ByteStride, FAccessor::EComponentType ComponentType, bool bNormalized, FVector2D& Buffer)
+		bool AcquireVec2(uint32 Index, const FBufferView& BufferView, const uint64 ByteOffset, const uint32 ByteStride, FAccessor::EComponentType ComponentType, bool bNormalized, FVector2D& Buffer)
 		{
 			Buffer = FVector2D();
 
@@ -442,7 +442,7 @@ namespace GLTF
 			return false;
 		}
 
-		bool AcquireVec3(uint32 Index, const FBufferView& BufferView, const uint32 ByteOffset, const uint32 ByteStride, FAccessor::EComponentType ComponentType, bool bNormalized, FVector& Buffer)
+		bool AcquireVec3(uint32 Index, const FBufferView& BufferView, const uint64 ByteOffset, const uint32 ByteStride, FAccessor::EComponentType ComponentType, bool bNormalized, FVector& Buffer)
 		{
 			Buffer = FVector();
 
@@ -462,7 +462,7 @@ namespace GLTF
 			return false;
 		}
 
-		bool AcquireVec4(uint32 Index, const FBufferView& BufferView, const uint32 ByteOffset, const uint32 ByteStride, FAccessor::EComponentType ComponentType, bool bNormalized, FVector4& Buffer)
+		bool AcquireVec4(uint32 Index, const FBufferView& BufferView, const uint64 ByteOffset, const uint32 ByteStride, FAccessor::EComponentType ComponentType, bool bNormalized, FVector4& Buffer)
 		{
 			Buffer = FVector4();
 
@@ -482,7 +482,7 @@ namespace GLTF
 			return false;
 		}
 
-		bool AcquireMat4(uint32 Index, const FBufferView& BufferView, const uint32 ByteOffset, const uint32 ByteStride, FAccessor::EComponentType ComponentType, FMatrix& Buffer)
+		bool AcquireMat4(uint32 Index, const FBufferView& BufferView, const uint64 ByteOffset, const uint32 ByteStride, FAccessor::EComponentType ComponentType, FMatrix& Buffer)
 		{
 			Buffer = FMatrix();
 
@@ -497,7 +497,7 @@ namespace GLTF
 		}
 	}
 	
-	FAccessor::FSparse::FIndices::FIndices(uint32 InCount, const FBufferView& InBufferView, uint32 InByteOffset, EComponentType InComponentType)
+	FAccessor::FSparse::FIndices::FIndices(uint32 InCount, const FBufferView& InBufferView, uint64 InByteOffset, EComponentType InComponentType)
 		: Count(InCount)
 		, BufferView(InBufferView)
 		, ByteOffset(InByteOffset)
@@ -916,7 +916,7 @@ namespace GLTF
 	}
 
 	//
-	FValidAccessor::FValidAccessor(FBufferView& InBufferView, uint32 InOffset, uint32 InCount, EType InType, EComponentType InCompType,
+	FValidAccessor::FValidAccessor(FBufferView& InBufferView, uint64 InOffset, uint32 InCount, EType InType, EComponentType InCompType,
 	                               bool bInNormalized, const FSparse& InSparse)
 	    : FAccessor(InCount, InType, InCompType, bInNormalized, InSparse)
 	    , BufferView(InBufferView)
diff --git a/Engine/Plugins/Interchange/Runtime/Source/Parsers/GLTFCore/Private/GLTF/GLTFBinaryReader.cpp b/Engine/Plugins/Interchange/Runtime/Source/Parsers/GLTFCore/Private/GLTF/GLTFBinaryReader.cpp
--- a/Engine/Plugins/Interchange/Runtime/Source/Parsers/GLTFCore/Private/GLTF/GLTFBinaryReader.cpp
+++ b/Engine/Plugins/Interchange/Runtime/Source/Parsers/GLTFCore/Private/GLTF/GLTFBinaryReader.cpp
@@ -13,6 +13,10 @@ namespace GLTF
 		{
 			return (X + 3) & ~3;
 		}
+		inline uint64 Pad4(uint64 X)
+		{
+			return (X + 3) & ~3;
+		}
 
 		bool SignatureMatches(uint32 Signature, const char* ExpectedSignature)
 		{
@@ -47,7 +51,7 @@ namespace GLTF
 			return MagicOk && VersionOk && SizeOk;
 		}
 
-		bool ReadChunk(FArchive& FileReader, const char* ExpectedChunkType, bool& OutHasMoreData, TArray<uint8>& OutData)
+		bool ReadJSONChunk(FArchive& FileReader, const char* ExpectedChunkType, bool& OutHasMoreData, TArray<uint8>& OutData)
 		{
 			// Align to next 4-byte boundary before reading anything
 			uint32 Offset        = FileReader.Tell();
@@ -88,6 +92,48 @@ namespace GLTF
 				return false;
 			}
 		}
+
+		bool ReadBinChunk(FArchive& FileReader, const char* ExpectedChunkType, bool& OutHasMoreData, TArray64<uint8>& OutData)
+		{
+			// Align to next 4-byte boundary before reading anything
+			uint64 Offset = FileReader.Tell();
+			uint64 AlignedOffset = Pad4(Offset);
+			if (Offset != AlignedOffset)
+			{
+				FileReader.Seek(AlignedOffset);
+			}
+
+			// Each chunk has the form [Size][Type][...Data...]
+			uint32 ChunkType, ChunkDataSize;
+			FileReader << ChunkDataSize;
+			FileReader << ChunkType;
+
+			constexpr uint64 ChunkHeaderSize = 8;
+			const uint64     AvailableData = FileReader.TotalSize() - (AlignedOffset + ChunkHeaderSize);
+
+			// Is there room for another chunk after this one?
+			OutHasMoreData = AvailableData - Pad4(ChunkDataSize) >= ChunkHeaderSize;
+
+			// Is there room for this chunk's data? (should always be true)
+			if (ChunkDataSize > AvailableData)
+			{
+				return false;
+			}
+
+			if (SignatureMatches(ChunkType, ExpectedChunkType))
+			{
+				// Read this chunk's data
+				OutData.SetNumUninitialized(ChunkDataSize, true);
+				FileReader.Serialize(OutData.GetData(), ChunkDataSize);
+				return true;
+			}
+			else
+			{
+				// Skip past this chunk's data
+				FileReader.Seek(AlignedOffset + ChunkHeaderSize + ChunkDataSize);
+				return false;
+			}
+		}
 	}
 
 	FBinaryFileReader::FBinaryFileReader()
@@ -117,12 +163,12 @@ namespace GLTF
 		BinChunk->Empty();
 
 		bool HasMoreData;
-		if (ReadChunk(FileReader, "JSON", HasMoreData, JsonChunk))
+		if (ReadJSONChunk(FileReader, "JSON", HasMoreData, JsonChunk))
 		{
 			// Get BIN chunk if present
 			while (HasMoreData)
 			{
-				if (ReadChunk(FileReader, "BIN", HasMoreData, *BinChunk))
+				if (ReadBinChunk(FileReader, "BIN", HasMoreData, *BinChunk))
 				{
 					break;
 				}
diff --git a/Engine/Plugins/Interchange/Runtime/Source/Parsers/GLTFCore/Private/GLTF/GLTFBinaryReader.h b/Engine/Plugins/Interchange/Runtime/Source/Parsers/GLTFCore/Private/GLTF/GLTFBinaryReader.h
--- a/Engine/Plugins/Interchange/Runtime/Source/Parsers/GLTFCore/Private/GLTF/GLTFBinaryReader.h
+++ b/Engine/Plugins/Interchange/Runtime/Source/Parsers/GLTFCore/Private/GLTF/GLTFBinaryReader.h
@@ -16,17 +16,17 @@ namespace GLTF
 
 		bool ReadFile(FArchive& FileReader);
 
-		void SetBuffer(TArray<uint8>& InBuffer);
+		void SetBuffer(TArray64<uint8>& InBuffer);
 
 		const TArray<uint8>& GetJsonBuffer() const;
 
 	private:
 		TArray<uint8> JsonChunk;
 
-		TArray<uint8>* BinChunk;
+		TArray64<uint8>* BinChunk;
 	};
 
-	inline void FBinaryFileReader::SetBuffer(TArray<uint8>& InBuffer)
+	inline void FBinaryFileReader::SetBuffer(TArray64<uint8>& InBuffer)
 	{
 		BinChunk = &InBuffer;
 	}
diff --git a/Engine/Plugins/Interchange/Runtime/Source/Parsers/GLTFCore/Private/GLTF/GLTFReader.cpp b/Engine/Plugins/Interchange/Runtime/Source/Parsers/GLTFCore/Private/GLTF/GLTFReader.cpp
--- a/Engine/Plugins/Interchange/Runtime/Source/Parsers/GLTFCore/Private/GLTF/GLTFReader.cpp
+++ b/Engine/Plugins/Interchange/Runtime/Source/Parsers/GLTFCore/Private/GLTF/GLTFReader.cpp
@@ -125,7 +125,7 @@ namespace GLTF
 
 	void FFileReader::SetupBuffer(const FJsonObject& Object, const FString& Path)
 	{
-		const uint32 ByteLength = GetUnsignedInt(Object, TEXT("byteLength"), 0);
+		const uint64 ByteLength = GetUnsignedInt64(Object, TEXT("byteLength"), 0);
 		Asset->Buffers.Emplace(ByteLength);
 		FBuffer& Buffer = Asset->Buffers.Last();
 
@@ -210,8 +210,8 @@ namespace GLTF
 		const uint32 BufferIdx = GetUnsignedInt(Object, TEXT("buffer"), BufferCount);
 		if (BufferIdx < BufferCount)  // must be true
 		{
-			const uint32 ByteOffset = GetUnsignedInt(Object, TEXT("byteOffset"), 0);
-			const uint32 ByteLength = GetUnsignedInt(Object, TEXT("byteLength"), 0);
+			const uint64 ByteOffset = GetUnsignedInt64(Object, TEXT("byteOffset"), 0);
+			const uint64 ByteLength = GetUnsignedInt64(Object, TEXT("byteLength"), 0);
 			const uint32 ByteStride = GetUnsignedInt(Object, TEXT("byteStride"), 0);
 			Asset->BufferViews.Emplace(Asset->Buffers[BufferIdx], ByteOffset, ByteLength, ByteStride);
 			ExtensionsHandler->SetupBufferViewExtensions(Object, Asset->BufferViews.Last());
@@ -223,7 +223,7 @@ namespace GLTF
 		const uint32 BufferViewIdx = GetUnsignedInt(Object, TEXT("bufferView"), BufferViewCount);
 		if (BufferViewIdx < BufferViewCount)  // must be true
 		{
-			const uint32                    ByteOffset = GetUnsignedInt(Object, TEXT("byteOffset"), 0);
+			const uint64                    ByteOffset = GetUnsignedInt64(Object, TEXT("byteOffset"), 0);
 			const FAccessor::EComponentType CompType   = ComponentTypeFromNumber(GetUnsignedInt(Object, TEXT("componentType"), 0));
 			const uint32                    Count      = GetUnsignedInt(Object, TEXT("count"), 0);
 			const FAccessor::EType          Type       = AccessorTypeFromString(Object.GetStringField("type"));
@@ -238,12 +238,12 @@ namespace GLTF
 
 				const FJsonObject& IndicesObject = *SparseObject.GetObjectField(TEXT("indices"));
 				const uint32 IndicesBufferViewIdx = GetUnsignedInt(IndicesObject, TEXT("bufferView"), BufferViewCount);
-				const int32 IndicesByteOffset = GetUnsignedInt(IndicesObject, TEXT("byteOffset"), 0);
+				const uint64 IndicesByteOffset = GetUnsignedInt64(IndicesObject, TEXT("byteOffset"), 0);
 				const FAccessor::EComponentType IndicesCompType = ComponentTypeFromNumber(GetUnsignedInt(IndicesObject, TEXT("componentType"), 0));
 
 				const FJsonObject& ValuesObject = *SparseObject.GetObjectField(TEXT("values"));
 				const uint32 ValuesBufferViewIdx = GetUnsignedInt(ValuesObject, TEXT("bufferView"), BufferViewCount);
-				const int32 ValuesByteOffset = GetUnsignedInt(ValuesObject, TEXT("byteOffset"), 0);
+				const uint64 ValuesByteOffset = GetUnsignedInt64(ValuesObject, TEXT("byteOffset"), 0);
 
 				Asset->Accessors.Emplace(Asset->BufferViews[BufferViewIdx], ByteOffset, Count, Type, CompType, Normalized, 
 					FAccessor::FSparse(SparseCount, 
@@ -853,15 +853,15 @@ namespace GLTF
 		}
 	}
 
-	void FFileReader::AllocateExtraData(const FString& InResourcesPath, bool bInLoadImageData, TArray<uint8>& OutExtraData)
+	void FFileReader::AllocateExtraData(const FString& InResourcesPath, bool bInLoadImageData, TArray64<uint8>& OutExtraData)
 	{
-		uint32 ExtraBufferSize = 0;
+		uint64 ExtraBufferSize = 0;
 		if (BufferCount > 0)
 		{
 			for (const TSharedPtr<FJsonValue>& Value : JsonRoot->GetArrayField(TEXT("buffers")))
 			{
 				const FJsonObject& Object     = *Value->AsObject();
-				const uint32       ByteLength = GetUnsignedInt(Object, TEXT("byteLength"), 0);
+				const uint64       ByteLength = GetUnsignedInt64(Object, TEXT("byteLength"), 0);
 				if (!Object.HasTypedField<EJson::String>(TEXT("uri")))
 					continue;
 
diff --git a/Engine/Plugins/Interchange/Runtime/Source/Parsers/GLTFCore/Private/GLTF/JsonUtilities.h b/Engine/Plugins/Interchange/Runtime/Source/Parsers/GLTFCore/Private/GLTF/JsonUtilities.h
--- a/Engine/Plugins/Interchange/Runtime/Source/Parsers/GLTFCore/Private/GLTF/JsonUtilities.h
+++ b/Engine/Plugins/Interchange/Runtime/Source/Parsers/GLTFCore/Private/GLTF/JsonUtilities.h
@@ -38,6 +38,21 @@ namespace GLTF
 		return DefaultValue;
 	}
 
+	inline uint64 GetUnsignedInt64(const FJsonObject& Object, const TCHAR* Name, uint64 DefaultValue)
+	{
+		if (Object.HasTypedField<EJson::Number>(Name))
+		{
+			uint64 UnsignedInt64Value;
+			
+			if (Object.TryGetNumberField(Name, UnsignedInt64Value))
+			{
+				return UnsignedInt64Value;
+			}
+		}
+
+		return DefaultValue;
+	}
+
 	inline uint32 GetUnsignedInt(const FJsonObject& Object, const TCHAR* Name, uint32 DefaultValue)
 	{
 		if (Object.HasTypedField<EJson::Number>(Name))
diff --git a/Engine/Plugins/Interchange/Runtime/Source/Parsers/GLTFCore/Public/GLTFAccessor.h b/Engine/Plugins/Interchange/Runtime/Source/Parsers/GLTFCore/Public/GLTFAccessor.h
--- a/Engine/Plugins/Interchange/Runtime/Source/Parsers/GLTFCore/Public/GLTFAccessor.h
+++ b/Engine/Plugins/Interchange/Runtime/Source/Parsers/GLTFCore/Public/GLTFAccessor.h
@@ -10,10 +10,10 @@ namespace GLTF
 {
 	struct GLTFCORE_API FBuffer
 	{
-		const uint32 ByteLength;
+		const uint64 ByteLength;
 		const uint8* Data;
 
-		explicit FBuffer(uint32 InByteLength)
+		explicit FBuffer(uint64 InByteLength)
 		    : ByteLength(InByteLength)
 		    , Data(nullptr)
 		{
@@ -24,7 +24,7 @@ namespace GLTF
 			return Data != nullptr;
 		}
 
-		const uint8* DataAt(uint32 Offset) const
+		const uint8* DataAt(uint64 Offset) const
 		{
 			checkSlow(Data);
 			return Data + Offset;
@@ -34,12 +34,12 @@ namespace GLTF
 	struct GLTFCORE_API FBufferView
 	{
 		const FBuffer& Buffer;
-		const uint32 ByteOffset;
-		const uint32 ByteLength;
+		const uint64 ByteOffset;
+		const uint64 ByteLength;
 		// if zero then accessor elements are tightly packed, i.e., effective stride equals the size of the element
 		const uint32 ByteStride;  // range 4..252
 
-		explicit FBufferView(const FBuffer& InBuffer, uint32 InOffset, uint32 InLength, uint32 InStride)
+		explicit FBufferView(const FBuffer& InBuffer, uint64 InOffset, uint64 InLength, uint32 InStride)
 		    : Buffer(InBuffer)
 		    , ByteOffset(InOffset)
 		    , ByteLength(InLength)
@@ -53,7 +53,7 @@ namespace GLTF
 			return Buffer.IsValid();
 		}
 
-		const uint8* DataAt(uint32 Offset) const
+		const uint8* DataAt(uint64 Offset) const
 		{
 			return Buffer.DataAt(Offset + ByteOffset);
 		}
@@ -114,7 +114,7 @@ namespace GLTF
 				const int32          Count; //Helper for creating cache, equals to FSparse.Count
 
 				const FBufferView&   BufferView;
-				const uint32         ByteOffset;
+				const uint64         ByteOffset;
 				const EComponentType ComponentType;
 
 				FIndices()
@@ -125,14 +125,14 @@ namespace GLTF
 				{
 				}
 
-				FIndices(uint32 InCount, const FBufferView& InBufferView, uint32 InByteOffset, EComponentType InComponentType);
+				FIndices(uint32 InCount, const FBufferView& InBufferView, uint64 InByteOffset, EComponentType InComponentType);
 			} Indices;
 			
 			//Values:
 			struct FValues
 			{
 				const FBufferView&   BufferView;
-				const uint32         ByteOffset;
+				const uint64         ByteOffset;
 				
 				FValues()
 					: BufferView(FVoidBufferView::GetVoidBufferView())
@@ -140,7 +140,7 @@ namespace GLTF
 				{
 				}
 
-				FValues(const FBufferView& InBufferView, uint32 InByteOffset)
+				FValues(const FBufferView& InBufferView, uint64 InByteOffset)
 					: BufferView(InBufferView)
 					, ByteOffset(InByteOffset)
 				{
@@ -156,8 +156,8 @@ namespace GLTF
 			}
 
 			FSparse(uint32 InCount,
-				const FBufferView& InIndicesBufferView, uint32 InIndicesByteOffset, EComponentType InIndicesComponentType,
-				const FBufferView& InValuesBufferView, uint32 InValuesByteOffset)
+				const FBufferView& InIndicesBufferView, uint64 InIndicesByteOffset, EComponentType InIndicesComponentType,
+				const FBufferView& InValuesBufferView, uint64 InValuesByteOffset)
 				: bHasSparse(true)
 				, Count(InCount)
 				, Indices(InCount, InIndicesBufferView, InIndicesByteOffset, InIndicesComponentType)
@@ -225,7 +225,7 @@ namespace GLTF
 
 	struct GLTFCORE_API FValidAccessor final : FAccessor
 	{
-		FValidAccessor(FBufferView& InBufferView, uint32 InOffset, uint32 InCount, EType InType, EComponentType InCompType, bool bInNormalized, const FSparse& InSparse);
+		FValidAccessor(FBufferView& InBufferView, uint64 InOffset, uint32 InCount, EType InType, EComponentType InCompType, bool bInNormalized, const FSparse& InSparse);
 
 		bool IsValid() const override;
 
@@ -251,7 +251,7 @@ namespace GLTF
 		const uint8* DataAt(uint32 Index) const;
 
 		const FBufferView& BufferView;
-		const uint32       ByteOffset;
+		const uint64       ByteOffset;
 		const uint32       ElementSize;
 		const uint32	   ByteStride;
 
diff --git a/Engine/Plugins/Interchange/Runtime/Source/Parsers/GLTFCore/Public/GLTFAsset.h b/Engine/Plugins/Interchange/Runtime/Source/Parsers/GLTFCore/Public/GLTFAsset.h
--- a/Engine/Plugins/Interchange/Runtime/Source/Parsers/GLTFCore/Public/GLTFAsset.h
+++ b/Engine/Plugins/Interchange/Runtime/Source/Parsers/GLTFCore/Public/GLTFAsset.h
@@ -141,9 +141,9 @@ namespace GLTF
 	private:
 		// Binary glTF files can have embedded data after JSON.
 		// This will be empty when reading from a text glTF (common) or a binary glTF with no BIN chunk (rare).
-		TArray<uint8> BinData;
+		TArray64<uint8> BinData;
 		// Extra binary data used for images from disk, mime data and so on.
-		TArray<uint8> ExtraBinData;
+		TArray64<uint8> ExtraBinData;
 
 		friend class FFileReader;
 	};
diff --git a/Engine/Plugins/Interchange/Runtime/Source/Parsers/GLTFCore/Public/GLTFReader.h b/Engine/Plugins/Interchange/Runtime/Source/Parsers/GLTFCore/Public/GLTFReader.h
--- a/Engine/Plugins/Interchange/Runtime/Source/Parsers/GLTFCore/Public/GLTFReader.h
+++ b/Engine/Plugins/Interchange/Runtime/Source/Parsers/GLTFCore/Public/GLTFReader.h
@@ -38,7 +38,7 @@ namespace GLTF
 	private:
 		void LoadMetadata(GLTF::FAsset& Asset);
 		void ImportAsset(const FString& FilePath, bool bInLoadImageData, GLTF::FAsset& Asset);
-		void AllocateExtraData(const FString& InResourcesPath, bool bInLoadImageData, TArray<uint8>& OutExtraData);
+		void AllocateExtraData(const FString& InResourcesPath, bool bInLoadImageData, TArray64<uint8>& OutExtraData);
 
 		void SetupBuffer(const FJsonObject& Object, const FString& Path);
 		void SetupBufferView(const FJsonObject& Object) const;
