Rewrite the PSD importer to remove C++ undefined behavior, add buffer overread/write protection, graceful malformed file handling, and support for 16bit RLE compressed PSDs.

diff --git a/Engine/Plugins/Interchange/Runtime/Source/Import/Private/Texture/InterchangePSDTranslator.cpp b/Engine/Plugins/Interchange/Runtime/Source/Import/Private/Texture/InterchangePSDTranslator.cpp
--- a/Engine/Plugins/Interchange/Runtime/Source/Import/Private/Texture/InterchangePSDTranslator.cpp
+++ b/Engine/Plugins/Interchange/Runtime/Source/Import/Private/Texture/InterchangePSDTranslator.cpp
@@ -6,8 +6,10 @@
 #include "Engine/Texture2D.h"
 #include "IImageWrapper.h"
 #include "IImageWrapperModule.h"
+#include "ImageCoreUtils.h"
 #include "InterchangeImportLog.h"
 #include "InterchangeTextureNode.h"
+#include "Math/GuardedInt.h"
 #include "Misc/FileHelper.h"
 #include "Misc/Paths.h"
 #include "Misc/ScopedSlowTask.h"
@@ -80,264 +82,285 @@ namespace UE
 			};
 #pragma pack(pop)
 
+			static uint32 ReadNetwork32(const void* Ptr)
+			{
+				const uint8* Data = (const uint8*)Ptr;
+				return ((uint32)Data[0] << 24) +
+					((uint32)Data[1] << 16) +
+					((uint32)Data[2] << 8) +
+					((uint32)Data[3] << 0);
+			}
+			static uint16 ReadNetwork16(const void* Ptr)
+			{
+				const uint8* Data = (const uint8*)Ptr;
+				return ((uint16)Data[0] << 8) +
+					((uint16)Data[1] << 0);
+			}
 
-			static bool ReadData(uint8* pOut, const uint8*& pBuffer, FPSDFileHeader& Info)
+			static bool SkipSection(FMemoryView& InOutView)
 			{
-				const uint8* pPlane = nullptr;
-				const uint8* pRowTable = nullptr;
-				int32 iPlane;
-				int16 CompressionType;
-				int32 iPixel;
-				int32 iRow;
-				int32 CompressedBytes;
-				int32 iByte;
-				int32 Count;
-				uint8 Value;
+				// PSD has several 32 bit sized sections we ignore.
+				if (InOutView.GetSize() < 4)
+				{
+					return false;
+				}
+
+				uint32 SectionSize = ReadNetwork32(InOutView.GetData());
+				InOutView.RightChopInline(SectionSize + 4);
+				return InOutView.GetSize() != 0;
+			}
+
+			static bool ReadUInt16(FMemoryView& InOutView, uint16& OutValue)
+			{
+				if (InOutView.GetSize() < 2)
+				{
+					OutValue = 0;
+					return false;
+				}
+				OutValue = ReadNetwork16(InOutView.GetData());
+				InOutView.RightChopInline(2);
+				return true;
+			}
+
+			// Decodes the raw data for a row - this is the same independent of the scanline format.
+			static bool DecodeRLERow(const uint8* RowSource, uint16 RowSourceBytes, uint8* OutputScanlineData, uint64 OutputScanlineDatasize)
+			{
+				uint64 OutputByte = 0;
+				uint32 SourceByteIndex = 0;
+				while (SourceByteIndex < RowSourceBytes)
+				{
+					int8 code = (int8)RowSource[SourceByteIndex];
+					SourceByteIndex++;
+
+					// Is it a repeat?
+					if (code == -128) // nop code for alignment.
+					{
+					}
+					else if (code < 0)
+					{
+						int32 Count = -(int32)code + 1;
+						if (OutputByte + Count > OutputScanlineDatasize) // OutputScanlineDatasize originates as int32 + int8 can't overflow a uint64
+						{
+							return false;
+						}
+
+						if (SourceByteIndex >= RowSourceBytes)
+						{
+							return false;
+						}
+						uint8 Value = RowSource[SourceByteIndex];
+						SourceByteIndex++;
+
+						FMemory::Memset(OutputScanlineData + OutputByte, Value, Count);
+						OutputByte += Count;
+					}
+					// Must be a run of literals then
+					else
+					{
+						int32 Count = (int32)code + 1;
+						if (SourceByteIndex + Count > RowSourceBytes) // int32 + int8, can't overflow.
+						{
+							return false;
+						}
 
+						FMemory::Memcpy(OutputScanlineData + OutputByte, RowSource + SourceByteIndex, Count);
+						SourceByteIndex += Count;
+						OutputByte += Count;
+					}
+				}
+
+				// Confirm that we decoded the right number of bytes
+				return OutputByte == OutputScanlineDatasize;
+			}
+
+			static bool ReadData(FMutableMemoryView& Output, FMemoryView Buffer, FPSDFileHeader& Info)
+			{
 				// Double check to make sure this is a valid request
 				if (!Info.IsValid() || !Info.IsSupported())
 				{
 					return false;
 				}
 
-				const uint8* pCur = pBuffer + sizeof(FPSDFileHeader);
-				int32 NPixels = Info.Width * Info.Height;
-
-				int32 ClutSize = ((int32)pCur[0] << 24) +
-					((int32)pCur[1] << 16) +
-					((int32)pCur[2] << 8) +
-					((int32)pCur[3] << 0);
-				pCur += 4 + ClutSize;
-
-				// Skip Image Resource Section
-				int32 ImageResourceSize = ((int32)pCur[0] << 24) +
-					((int32)pCur[1] << 16) +
-					((int32)pCur[2] << 8) +
-					((int32)pCur[3] << 0);
-				pCur += 4 + ImageResourceSize;
-
-				// Skip Layer and Mask Section
-				int32 LayerAndMaskSize = ((int32)pCur[0] << 24) +
-					((int32)pCur[1] << 16) +
-					((int32)pCur[2] << 8) +
-					((int32)pCur[3] << 0);
-				pCur += 4 + LayerAndMaskSize;
-
-				// Determine number of bytes per pixel
-				int32 BytesPerPixel = 3;
-				const int32 BytesPerChannel = (Info.Depth / 8);
-				switch (Info.Mode)
+				if (Buffer.GetSize() <= sizeof(FPSDFileHeader) ||
+					FImageCoreUtils::IsImageImportPossible(Info.Width, Info.Height) == false)
 				{
-					case 1: // 'GrayScale'
-						BytesPerPixel = BytesPerChannel;
-						break;
-					case 2:
-						BytesPerPixel = 1;
-						return false;  // until we support indexed...
-						break;
-					case 3: // 'RGBColor'
-						if (Info.nChannels == 3)
-							BytesPerPixel = 3 * BytesPerChannel;
-						else
-							BytesPerPixel = 4 * BytesPerChannel;
-						break;
-					default:
-						return false;
-						break;
+					return false;
 				}
 
-				// Get Compression Type
-				CompressionType = ((int32)pCur[0] << 8) + ((int32)pCur[1] << 0);
-				pCur += 2;
+				FMemoryView Current = Buffer;
+				Current.RightChopInline(sizeof(FPSDFileHeader));
 
-				// Fail on 16 Bits/channel with RLE. This can occur when the file is not saved with 'Maximize Compatibility'. Compression doesn't appear to be standard.
-				if (CompressionType == 1 && Info.Depth == 16)
+				FGuardedInt64 GuardedPixelCount = FGuardedInt64(Info.Width) * Info.Height;
+				if (GuardedPixelCount.InvalidOrLessOrEqual(0))
 				{
 					return false;
 				}
 
-				// If no alpha channel, set alpha to opaque (255 or 65536).
-				if (Info.nChannels != 4)
+				if (Info.Depth != 8 && Info.Depth != 16)
 				{
-					if (Info.Depth == 8)
-					{
-						const uint32 Channels = 4;
-						const uint32 BufferSize = Info.Width * Info.Height * Channels * sizeof(uint8);
-						FMemory::Memset(pOut, 0xff, BufferSize);
-					}
-					else if (Info.Depth == 16)
-					{
-						const uint32 Channels = 4;
-						const uint32 BufferSize = Info.Width * Info.Height * Channels * sizeof(uint16);
-						FMemory::Memset(pOut, 0xff, BufferSize);
-					}
+					return false;
+				}
+				if (Info.nChannels != 1 && Info.nChannels != 3 && Info.nChannels != 4)
+				{
+					return false;
 				}
 
-				// Uncompressed?
-				if (CompressionType == 0)
+				FGuardedInt64 OutputBytesNeeded = GuardedPixelCount * 4 * (Info.Depth / 8);
+				if (OutputBytesNeeded.InvalidOrGreaterThan(Output.GetSize()))
 				{
-					if (Info.Depth == 8)
+					return false;
+				}
+
+				// Skip Clut, Image Resource Section, and Layer/Mask Section
+				if (SkipSection(Current) == false ||
+					SkipSection(Current) == false ||
+					SkipSection(Current) == false)
+				{
+					return false;
+				}
+
+				// Get Compression Type
+				uint16 CompressionType = 0;
+				if (ReadUInt16(Current, CompressionType) == false)
+				{
+					return false;
+				}
+				if (CompressionType != 0 && CompressionType != 1)
+				{
+					return false;
+				}
+
+				// Overflow checked above.
+				uint64 UncompressedScanLineSizePerChannel = Info.Width * (Info.Depth / 8);
+				uint64 OutputScanLineSize = UncompressedScanLineSizePerChannel * 4;
+				uint64 UncompressedPlaneSize = UncompressedScanLineSizePerChannel * Info.Height;
+
+				// For copying alpha when the source doesn't have alpha information.
+				uint8 OpaqueAlpha[2] = {255, 255};
+
+				const uint16* RLERowTable[4] = {};
+				const uint8* RLEPlaneSource[4] = {};
+				TArray<uint8> RLETempScanlines[4];
+
+				if (CompressionType == 1)
+				{
+					uint64 RowTableBytesPerChannel = Info.Height * sizeof(uint16);
+					if (Current.GetSize() < RowTableBytesPerChannel * Info.nChannels)
 					{
-						FColor* Dest = (FColor*)pOut;
-						for (int32 Pixel = 0; Pixel < NPixels; Pixel++)
-						{
-							if (Info.nChannels == 1)
-							{
-								Dest[Pixel].R = pCur[Pixel];
-								Dest[Pixel].G = pCur[Pixel];
-								Dest[Pixel].B = pCur[Pixel];
-							}
-							else
-							{
-								// Each channel live in a separate plane
-								Dest[Pixel].R = pCur[Pixel];
-								Dest[Pixel].G = pCur[NPixels + Pixel];
-								Dest[Pixel].B = pCur[NPixels * 2 + Pixel];
-								if (Info.nChannels == 4)
-								{
-									Dest[Pixel].A = pCur[NPixels * 3 + Pixel];
-								}
-							}
-						}
+						return false;
 					}
-					else if (Info.Depth == 16)
-					{
-						uint32 SrcOffset = 0;
 
-						if (Info.nChannels == 1)
-						{
-							uint16* Dest = (uint16*)pOut;
-							uint32 ChannelOffset = 0;
+					FMemoryView RowTableSource = Current;
+					RowTableSource.LeftInline(RowTableBytesPerChannel * Info.nChannels);
+					Current.RightChopInline(RowTableBytesPerChannel * Info.nChannels);
 
-							for (int32 Pixel = 0; Pixel < NPixels; Pixel++)
-							{
-								Dest[ChannelOffset + 0] = ((pCur[SrcOffset] << 8) + (pCur[SrcOffset + 1] << 0));
-								Dest[ChannelOffset + 1] = ((pCur[SrcOffset] << 8) + (pCur[SrcOffset + 1] << 0));
-								Dest[ChannelOffset + 2] = ((pCur[SrcOffset] << 8) + (pCur[SrcOffset + 1] << 0));
+					// We want the row tables for each plane, which means we have to decode them.
+					uint64 CurrentOffset = 0;
+					for (uint64 Plane = 0; Plane < (uint64)Info.nChannels; Plane++)
+					{
+						RLETempScanlines[Plane].AddUninitialized(UncompressedScanLineSizePerChannel);
+						RLERowTable[Plane] = (const uint16*)RowTableSource.GetData();
+						RowTableSource.RightChopInline(RowTableBytesPerChannel);
 
-								//Increment offsets
-								ChannelOffset += 4;
-								SrcOffset += BytesPerChannel;
-							}
-						}
-						else
+						// Save off where this plane's source is.
+						RLEPlaneSource[Plane] = (const uint8*)Current.GetData() + CurrentOffset;
+						for (uint64 Row = 0; Row < Info.Height; Row++)
 						{
-							// Loop through the planes	
-							for (iPlane = 0; iPlane < Info.nChannels; iPlane++)
-							{
-								uint16* Dest = (uint16*)pOut;
-								uint32 ChannelOffset = iPlane;
-
-								for (int32 Pixel = 0; Pixel < NPixels; Pixel++)
-								{
-									Dest[ChannelOffset] = ((pCur[SrcOffset] << 8) + (pCur[SrcOffset + 1] << 0));
+							// can't overflow: we're adding 16-bit uints and Info.Height is bounded by a 32-bit value, so sum fits in 48 bits
+							CurrentOffset += ReadNetwork16(&RLERowTable[Plane][Row]);
+						}
 
-									//Increment offsets
-									ChannelOffset += 4;
-									SrcOffset += BytesPerChannel;
-								}
-							}
+						// Now that we know the size, verify we have it.
+						if (Current.GetSize() < CurrentOffset)
+						{
+							return false;
 						}
 					}
 				}
-				// RLE?
-				else if (CompressionType == 1)
+				else
 				{
-					// Setup RowTable
-					pRowTable = pCur;
-					pCur += Info.nChannels * Info.Height * 2;
+					if (Current.GetSize() < UncompressedPlaneSize * Info.nChannels) // overflow checked on function entry.
+					{
+						return false;
+					}
+				}
 
-					FColor* Dest = (FColor*)pOut;
+				uint8* OutputScanline = (uint8*)Output.GetData();
+				for (uint64 Row = 0; Row < (uint64)Info.Height; Row++, OutputScanline += OutputScanLineSize)
+				{
+					const uint8* SourceScanLine[4] = {};
+					uint64 AlphaMask = ~0ULL;
 
-					// Loop through the planes
-					for (iPlane = 0; iPlane < Info.nChannels; iPlane++)
+					// Init the source scanlines from the file data. For RLE we decode into a temp buffer,
+					// otherwise we read directly. File size has already been validated.
+					if (CompressionType == 0)
 					{
-						int32 iWritePlane = iPlane;
-						if (iWritePlane > BytesPerPixel - 1) iWritePlane = BytesPerPixel - 1;
+						SourceScanLine[0] = (const uint8*)Current.GetData() + Row * UncompressedScanLineSizePerChannel;
 
-						// Loop through the rows
-						for (iRow = 0; iRow < Info.Height; iRow++)
+						for (uint16 Channel = 1; Channel < Info.nChannels; Channel++)
 						{
-							// Load a row
-							CompressedBytes = (pRowTable[(iPlane * Info.Height + iRow) * 2] << 8) +
-								(pRowTable[(iPlane * Info.Height + iRow) * 2 + 1] << 0);
-
-							// Setup Plane
-							pPlane = pCur;
-							pCur += CompressedBytes;
-
-							// Decompress Row
-							iPixel = 0;
-							iByte = 0;
-							while ((iPixel < Info.Width) && (iByte < CompressedBytes))
+							SourceScanLine[Channel] = SourceScanLine[Channel-1] + UncompressedPlaneSize;
+						}
+					}
+					else
+					{
+						for (uint16 Channel = 0; Channel < Info.nChannels; Channel++)
+						{
+							if (DecodeRLERow(RLEPlaneSource[Channel], ReadNetwork16(&RLERowTable[Channel][Row]), RLETempScanlines[Channel].GetData(), RLETempScanlines[Channel].Num()) == false)
 							{
-								int8 code = (int8)pPlane[iByte++];
-
-								// Is it a repeat?
-								if (code < 0)
-								{
-									Count = -(int32)code + 1;
-									Value = pPlane[iByte++];
-									while (Count-- > 0)
-									{
-										int32 idx = (iPixel)+(iRow * Info.Width);
-										if (Info.nChannels == 1)
-										{
-											Dest[idx].R = Value;
-											Dest[idx].G = Value;
-											Dest[idx].B = Value;
-										}
-										else
-										{
-											switch (iWritePlane)
-											{
-												case 0: Dest[idx].R = Value; break;
-												case 1: Dest[idx].G = Value; break;
-												case 2: Dest[idx].B = Value; break;
-												case 3: Dest[idx].A = Value; break;
-											}
-										}
-										iPixel++;
-									}
-								}
-								// Must be a literal then
-								else
-								{
-									Count = (int32)code + 1;
-									while (Count-- > 0)
-									{
-										Value = pPlane[iByte++];
-										int32 idx = (iPixel)+(iRow * Info.Width);
-
-										if (Info.nChannels == 1)
-										{
-											Dest[idx].R = Value;
-											Dest[idx].G = Value;
-											Dest[idx].B = Value;
-										}
-										else
-										{
-											switch (iWritePlane)
-											{
-												case 0: Dest[idx].R = Value; break;
-												case 1: Dest[idx].G = Value; break;
-												case 2: Dest[idx].B = Value; break;
-												case 3: Dest[idx].A = Value; break;
-											}
-										}
-										iPixel++;
-									}
-								}
+								return false;
 							}
+							RLEPlaneSource[Channel] += ReadNetwork16(&RLERowTable[Channel][Row]);
+							SourceScanLine[Channel] = RLETempScanlines[Channel].GetData();
+						}
+					}
+
+					// If we don't have all 4 channels, set up the scanlines to valid data.
+					if (Info.nChannels == 1)
+					{
+						SourceScanLine[1] = SourceScanLine[0];
+						SourceScanLine[2] = SourceScanLine[0];
+						SourceScanLine[3] = OpaqueAlpha;
+						AlphaMask = 0;
+					}
+					else if (Info.nChannels == 3)
+					{
+						SourceScanLine[3] = OpaqueAlpha;
+						AlphaMask = 0;
+					}
+					else if (Info.nChannels == 4)
+					{
+						AlphaMask = ~0ULL;
+					}
 
-							// Confirm that we decoded the right number of bytes
-							check(iByte == CompressedBytes);
-							check(iPixel == Info.Width);
+					// Do the plane interleaving.
+					if (Info.Depth == 8)
+					{
+						FColor* ScanLine8 = (FColor*)OutputScanline;
+						for (uint64 X = 0; X < Info.Width; X++)
+						{
+							ScanLine8[X].R = SourceScanLine[0][X];
+							ScanLine8[X].G = SourceScanLine[1][X];
+							ScanLine8[X].B = SourceScanLine[2][X];
+							ScanLine8[X].A = SourceScanLine[3][X & AlphaMask];
 						}
 					}
-				}
-				else
-					return false;
+					else if (Info.Depth == 16)
+					{
+						const uint16* SourceScanLineR16 = (const uint16*)SourceScanLine[0];
+						const uint16* SourceScanLineG16 = (const uint16*)SourceScanLine[1];
+						const uint16* SourceScanLineB16 = (const uint16*)SourceScanLine[2];
+						const uint16* SourceScanLineA16 = (const uint16*)SourceScanLine[3];
+						uint16* ScanLine16 = (uint16*)OutputScanline;
+						for (uint64 X = 0; X < Info.Width; X++)
+						{
+							ScanLine16[4*X + 0] = ReadNetwork16(&SourceScanLineR16[X]);
+							ScanLine16[4*X + 1] = ReadNetwork16(&SourceScanLineG16[X]);
+							ScanLine16[4*X + 2] = ReadNetwork16(&SourceScanLineB16[X]);
+							ScanLine16[4*X + 3] = ReadNetwork16(&SourceScanLineA16[X & AlphaMask]);
+						}
+					}
+				} // end each row.
 
 				// Success!
 				return true;
@@ -345,32 +368,13 @@ namespace UE
 
 			static void GetPSDHeader(const uint8* Buffer, FPSDFileHeader& Info)
 			{
-				Info.Signature = ((int32)Buffer[0] << 24) +
-					((int32)Buffer[1] << 16) +
-					((int32)Buffer[2] << 8) +
-					((int32)Buffer[3] << 0);
-
-				Info.Version = ((int32)Buffer[4] << 8) +
-					((int32)Buffer[5] << 0);
-
-				Info.nChannels = ((int32)Buffer[12] << 8) +
-					((int32)Buffer[13] << 0);
-
-				Info.Height = ((int32)Buffer[14] << 24) +
-					((int32)Buffer[15] << 16) +
-					((int32)Buffer[16] << 8) +
-					((int32)Buffer[17] << 0);
-
-				Info.Width = ((int32)Buffer[18] << 24) +
-					((int32)Buffer[19] << 16) +
-					((int32)Buffer[20] << 8) +
-					((int32)Buffer[21] << 0);
-
-				Info.Depth = ((int32)Buffer[22] << 8) +
-					((int32)Buffer[23] << 0);
-
-				Info.Mode = ((int32)Buffer[24] << 8) +
-					((int32)Buffer[25] << 0);
+				Info.Signature = (int32)ReadNetwork32(Buffer + 0);
+				Info.Version = (int16)ReadNetwork16(Buffer + 4);
+				Info.nChannels = (int16)ReadNetwork16(Buffer + 12);
+				Info.Height = (int32)ReadNetwork32(Buffer + 14);
+				Info.Width = (int32)ReadNetwork32(Buffer + 18);
+				Info.Depth = (int16)ReadNetwork16(Buffer + 22);
+				Info.Mode = (int16)ReadNetwork16(Buffer + 24);
 			}
 
 		}//ns Private
@@ -405,21 +409,16 @@ TOptional<UE::Interchange::FImportImage> UInterchangePSDTranslator::GetTexturePa
 		return {};
 	}
 
-	const uint8* Buffer = SourceDataBuffer.GetData();
-	const uint8* BufferEnd = Buffer + SourceDataBuffer.Num();
-
-	const int64 Length = BufferEnd - Buffer;
-
 	//
 	// PSD File
 	//
 	UE::Interchange::Private::FPSDFileHeader PSDHeader;
-	if (Length < sizeof(UE::Interchange::Private::FPSDFileHeader))
+	if (SourceDataBuffer.Num() < sizeof(UE::Interchange::Private::FPSDFileHeader))
 	{
 		return TOptional<UE::Interchange::FImportImage>();
 	}
 
-	UE::Interchange::Private::GetPSDHeader(Buffer, PSDHeader);
+	UE::Interchange::Private::GetPSDHeader(SourceDataBuffer.GetData(), PSDHeader);
 
 	if (!PSDHeader.IsValid())
 	{
@@ -458,7 +457,8 @@ TOptional<UE::Interchange::FImportImage> UInterchangePSDTranslator::GetTexturePa
 		TextureFormat
 	);
 
-	if (!UE::Interchange::Private::ReadData(static_cast<uint8*>(PayloadData.RawData.GetData()), Buffer, PSDHeader))
+	FMutableMemoryView Output(PayloadData.RawData.GetData(), PayloadData.RawData.GetSize());
+	if (!UE::Interchange::Private::ReadData(Output, MakeMemoryView(SourceDataBuffer), PSDHeader))
 	{
 		FTextureTranslatorUtilities::LogError(*this, NSLOCTEXT("InterchangePSDTranslator", "FailedRead", "Failed to read this PSD."));
 		return TOptional<UE::Interchange::FImportImage>();
