[glTF - Interchange] Shouldn't combine recursive skeletal animations

diff --git a/Engine/Plugins/Interchange/Runtime/Source/Import/Private/Gltf/InterchangeGltfMesh.cpp b/Engine/Plugins/Interchange/Runtime/Source/Import/Private/Gltf/InterchangeGltfMesh.cpp
--- a/Engine/Plugins/Interchange/Runtime/Source/Import/Private/Gltf/InterchangeGltfMesh.cpp
+++ b/Engine/Plugins/Interchange/Runtime/Source/Import/Private/Gltf/InterchangeGltfMesh.cpp
@@ -73,8 +73,19 @@ namespace UE::Interchange::Gltf::Private
 
 	bool GetSkeletalMeshDescriptionForPayLoadKey(const GLTF::FAsset& GltfAsset, const FString& PayLoadKey, FMeshDescription& MeshDescription, TArray<FString>* OutJointUniqueNames)
 	{
+		TArray<FString> BakingAndPayloadKey;
+		PayLoadKey.ParseIntoArray(BakingAndPayloadKey, TEXT(";"));
+
+		bool bBakeSkinJointTransform = false;
+		if (BakingAndPayloadKey.Num() == 2)
+		{
+			LexFromString(bBakeSkinJointTransform, *BakingAndPayloadKey[0]);
+		}
+
+		FString PayloadKeyToParse = BakingAndPayloadKey.Num() == 2 ? BakingAndPayloadKey[1] : PayLoadKey;
+
 		TArray<FString> PayLoadKeys;
-		PayLoadKey.ParseIntoArray(PayLoadKeys, TEXT(":"));
+		PayloadKeyToParse.ParseIntoArray(PayLoadKeys, TEXT(":"));
 		TMap<int32, TArray<int32>> MeshIndexToSkinIndicesMap;
 
 		int32 MeshIndex = 0;
@@ -131,7 +142,7 @@ namespace UE::Interchange::Gltf::Private
 			FStaticMeshAttributes StaticMeshAttributes(SkinnedMeshDescription);
 
 			//for instanced meshes we need to bake the transforms of the joints:
-			if (SkinIndices.Num() > 1 && GltfAsset.Skins[SkinIndex].Joints.Num() > 0)
+			if ((bBakeSkinJointTransform || SkinIndices.Num() > 1) && GltfAsset.Skins[SkinIndex].Joints.Num() > 0)
 			{
 				FTransform TransformLocalToWorld3d(FTransform::Identity);
 
diff --git a/Engine/Plugins/Interchange/Runtime/Source/Import/Private/Gltf/InterchangeGltfTranslator.cpp b/Engine/Plugins/Interchange/Runtime/Source/Import/Private/Gltf/InterchangeGltfTranslator.cpp
--- a/Engine/Plugins/Interchange/Runtime/Source/Import/Private/Gltf/InterchangeGltfTranslator.cpp
+++ b/Engine/Plugins/Interchange/Runtime/Source/Import/Private/Gltf/InterchangeGltfTranslator.cpp
@@ -2010,14 +2010,14 @@ void UInterchangeGLTFTranslator::HandleGltfSkeletons(UInterchangeBaseNodeContain
 			continue;
 		}
 		
-		int32 RootJointIndex = GltfAsset.Nodes[GltfAsset.Skins[SkinnedMeshNode.Skindex].Joints[0]].RootJointIndex;
-		if (!GltfAsset.Nodes.IsValidIndex(RootJointIndex))
+		int32 RootSkinJointIndex = UE::Interchange::Gltf::Private::GetRootNodeIndex(GltfAsset, GltfAsset.Skins[SkinnedMeshNode.Skindex].Joints);
+		if (!GltfAsset.Nodes.IsValidIndex(RootSkinJointIndex))
 		{
 			continue;
 		}
 
 		//basedon that root joint group the SkinnedMeshNodes:
-		TArray<int32>& GroupedSkinnedMeshNodes = RootJointGroupedSkinnedMeshNodes.FindOrAdd(RootJointIndex);
+		TArray<int32>& GroupedSkinnedMeshNodes = RootJointGroupedSkinnedMeshNodes.FindOrAdd(RootSkinJointIndex);
 		GroupedSkinnedMeshNodes.Add(SkinnedMeshNodeIndex);
 	}
 
@@ -2032,41 +2032,44 @@ void UInterchangeGLTFTranslator::HandleGltfSkeletons(UInterchangeBaseNodeContain
 		for (const TTuple<int32, TArray<int32>>& RootJointToSkinnedMeshNodes : RootJointGroupedSkinnedMeshNodes)
 		{
 			//Duplicate MeshNode for each group:
-			int32 RootJointIndex = RootJointToSkinnedMeshNodes.Key;
+			int32 RootSkinJointNodeIndex = RootJointToSkinnedMeshNodes.Key;
+			int32 RootJointNodeIndex = GltfAsset.Nodes[RootSkinJointNodeIndex].RootJointIndex;
 
 			//Skeletal Mesh's naming policy: (Mesh.Name)_(RootJointNode.Name) naming policy:
-			FString SkeletalName = GltfAsset.Meshes[MeshIndex].Name + TEXT("_") + GltfAsset.Nodes[RootJointIndex].Name;
-			FString SkeletalId = GltfAsset.Meshes[MeshIndex].UniqueId + TEXT("_") + GltfAsset.Nodes[RootJointIndex].UniqueId;
+			FString SkeletalName = GltfAsset.Meshes[MeshIndex].Name + TEXT("_") + GltfAsset.Nodes[RootSkinJointNodeIndex].Name;
+			FString SkeletalId = GltfAsset.Meshes[MeshIndex].UniqueId + TEXT("_") + GltfAsset.Nodes[RootSkinJointNodeIndex].UniqueId;
 
 			UInterchangeMeshNode* SkeletalMeshNode = HandleGltfMesh(NodeContainer, GltfAsset.Meshes[MeshIndex], MeshIndex, UnusedMeshIndices, SkeletalName, SkeletalId);
-
 			SkeletalMeshNode->SetSkinnedMesh(true);
 
+			//set the root joint node as the skeletonDependency:
+			const GLTF::FNode& RootJointNode = GltfAsset.Nodes[RootJointNodeIndex];
+			const FString* SkeletonNodeUid = NodeUidMap.Find(&RootJointNode);
+			if (ensure(SkeletonNodeUid))
+			{
+				SkeletalMeshNode->SetSkeletonDependencyUid(*SkeletonNodeUid);
+			}
+
+			bool bBakeSkinJointTransform = RootSkinJointNodeIndex != RootJointNodeIndex || RootJointNode.ParentIndex == -1;
+
 			//generate payload key:
 			//of template:
-			//"LexToString(SkinnedMeshNode.MeshIndex | (SkinnedMeshNode.Skindex << 16))":"LexToString(SkinnedMeshNode.MeshIndex | (SkinnedMeshNode.Skindex << 16))".....
-			FString Payload = "";
+			//"LexToString(bBakeSkinJointTransform)";"LexToString(SkinnedMeshNode.MeshIndex | (SkinnedMeshNode.Skindex << 16))":"LexToString(SkinnedMeshNode.MeshIndex | (SkinnedMeshNode.Skindex << 16))".....
+			FString Payload = TEXT("");
+			bool bAddSeparator = false;
 			for (int32 SkinnedMeshIndex : RootJointToSkinnedMeshNodes.Value)
 			{
 				const GLTF::FNode& SkinnedMeshNode = GltfAsset.Nodes[SkinnedMeshIndex];
 				if (Payload.Len() > 0)
 				{
-					Payload += ":";
+					Payload += TEXT(":");
 				}
+				
 				Payload += LexToString(SkinnedMeshNode.MeshIndex | (SkinnedMeshNode.Skindex << 16));
 			}
+			Payload = LexToString(bBakeSkinJointTransform) + TEXT(";") + Payload;
 			SkeletalMeshNode->SetPayLoadKey(Payload, EInterchangeMeshPayLoadType::SKELETAL);
 
-			//set the root joint node as the skeletonDependency:
-			int32 RootJointNodeIndex = RootJointToSkinnedMeshNodes.Key;
-			const GLTF::FNode& RootJointNode = GltfAsset.Nodes[RootJointNodeIndex];
-			const FString* SkeletonNodeUid = NodeUidMap.Find(&RootJointNode);
-			if (ensure(SkeletonNodeUid))
-			{
-				SkeletalMeshNode->SetSkeletonDependencyUid(*SkeletonNodeUid);
-			}
-				
-
 			//set the mesh actor node's custom asset instance uid to the new duplicated mesh
 			//if there are more than one skins, then choose the topmost (root node of the collection, top most in a hierarchical tree term) occurance of SkinnedMeshIndex
 			int32 MeshActorNodeIndex = UE::Interchange::Gltf::Private::GetRootNodeIndex(GltfAsset, RootJointToSkinnedMeshNodes.Value);
