[Interchange] [glTF] Removing unused variable

diff --git a/Engine/Plugins/Interchange/Runtime/Source/Import/Private/Gltf/InterchangeGltfTranslator.cpp b/Engine/Plugins/Interchange/Runtime/Source/Import/Private/Gltf/InterchangeGltfTranslator.cpp
--- a/Engine/Plugins/Interchange/Runtime/Source/Import/Private/Gltf/InterchangeGltfTranslator.cpp
+++ b/Engine/Plugins/Interchange/Runtime/Source/Import/Private/Gltf/InterchangeGltfTranslator.cpp
@@ -2056,7 +2056,6 @@ void UInterchangeGLTFTranslator::HandleGltfSkeletons(UInterchangeBaseNodeContain
 			//of template:
 			//"LexToString(bBakeSkinJointTransform)";"LexToString(SkinnedMeshNode.MeshIndex | (SkinnedMeshNode.Skindex << 16))":"LexToString(SkinnedMeshNode.MeshIndex | (SkinnedMeshNode.Skindex << 16))".....
 			FString Payload = TEXT("");
-			bool bAddSeparator = false;
 			for (int32 SkinnedMeshIndex : RootJointToSkinnedMeshNodes.Value)
 			{
 				const GLTF::FNode& SkinnedMeshNode = GltfAsset.Nodes[SkinnedMeshIndex];
