[interchange][gltf] Imports hangs and crash on that model

diff --git a/Engine/Plugins/Interchange/Runtime/Source/Parsers/GLTFCore/Private/GLTF/GLTFReader.cpp b/Engine/Plugins/Interchange/Runtime/Source/Parsers/GLTFCore/Private/GLTF/GLTFReader.cpp
--- a/Engine/Plugins/Interchange/Runtime/Source/Parsers/GLTFCore/Private/GLTF/GLTFReader.cpp
+++ b/Engine/Plugins/Interchange/Runtime/Source/Parsers/GLTFCore/Private/GLTF/GLTFReader.cpp
@@ -19,6 +19,8 @@
 
 #include "GenericPlatform/GenericPlatformHttp.h"
 
+#include "Math/UnrealMathVectorConstants.h"
+
 namespace GLTF
 {
 	namespace
@@ -109,6 +111,49 @@ namespace GLTF
 				return Void;
 			}
 		}
+
+		template<typename T>
+		void SetTransformFromMatrix(FTransform& Transform, const UE::Math::TMatrix<T>& InMatrix)
+		{
+			using namespace UE::Math;
+			TMatrix<T> M = InMatrix;
+
+			// Get the 3D scale from the matrix
+			TVector<T> InScale = M.ExtractScaling(0); //we want Precise Scaling with 0 tolerance.
+			TPersistentVectorRegisterType<T> Scale3D = VectorLoadFloat3_W0(&InScale);
+
+			// If there is negative scaling going on, we handle that here
+			if (InMatrix.Determinant() < 0.f)
+			{
+				// Assume it is along X and modify transform accordingly. 
+				// It doesn't actually matter which axis we choose, the 'appearance' will be the same
+				//(UE::Math::TTransform<T>::TransformVectorRegister)
+				Scale3D = VectorMultiply(Scale3D, UE::Math::TTransform<T>::TransformVectorRegister(GlobalVectorConstants::FloatMinus1_111));
+				M.SetAxis(0, -M.GetScaledAxis(EAxis::X));
+			}
+
+			TQuat<T> InRotation = TQuat<T>(M);
+			TPersistentVectorRegisterType<T> Rotation = VectorLoadAligned(&InRotation);
+			TVector<T> InTranslation = InMatrix.GetOrigin();
+			TPersistentVectorRegisterType<T> Translation = VectorLoadFloat3_W0(&InTranslation);
+
+			// Normalize rotation
+			Rotation = VectorNormalizeQuaternion(Rotation);
+
+			TVector<T> OutScale3D;
+			VectorStoreFloat3(Scale3D, &OutScale3D);
+
+			TQuat<T> OutRotation;
+			VectorStoreAligned(Rotation, &OutRotation);
+			OutRotation = GLTF::ConvertQuat(OutRotation);
+
+			TVector<T> OutTranslation;
+			VectorStoreFloat3(Translation, &OutTranslation);
+
+			Transform.SetScale3D(OutScale3D);
+			Transform.SetRotation(OutRotation);
+			Transform.SetTranslation(OutTranslation);
+		}
 	}
 
 	FFileReader::FFileReader()
@@ -413,8 +458,11 @@ namespace GLTF
 
 		if (Object.HasField(TEXT("matrix")))
 		{
-			Node.Transform.SetFromMatrix(GetMat4(Object, TEXT("matrix")));
-			Node.Transform.SetRotation(GLTF::ConvertQuat(Node.Transform.GetRotation()));
+			FMatrix Matrix = GetMat4(Object, TEXT("matrix"));
+
+			//we cannot use Transform.SetFromMatrix, because we want precise scaling
+			// (GLTF::SetTransformFromMatrix also Converts the Rotation Quaternion)
+			GLTF::SetTransformFromMatrix<double>(Node.Transform, Matrix);
 		}
 		else
 		{
