Improve conversion of Lambert/Phong materials to UE materials
  - Added a proprietary material function, MF_PhongToMetalRoughness, to make the conversion. The computation is based on various implementations in open source code, Blender, yocto-gl,...
  - Refactored the creation of the UInterchangeShaderGraphNode in FBX
  - Modified the OBJ translator to take in account the logic in the new material function
  - Modified the generic material pipeline to use the new material function
  - Move Anisotropy to common parameters since Anisotropic Phong is a thing
  - Renamed PBR namespace PBRMR, for Metallic|Roughness PBR model
  - Created PBRSG namespace for Specular|Glossiness PBR model with DiffuseColor, SpecularColor and GLossiness parameters.
  - Renamed IsPBRModel to IsMetalRoughModel
  - Added IsSpecGlossModel to check against PBRSG parameters
  - Added suffix '_MR' to existing materials and material functions used when user requests to create material instances
  - Added new materials and material functions with '_SG' suffix for specular|glossiness model
  - Added factory dependency between material instances and newly added Map inputs
  - Removed ability to generate material instances from the generic material pipeline for the glTF translator

diff --git a/Engine/Plugins/Experimental/Enterprise/DatasmithInterchange/Source/DatasmithInterchange/Private/InterchangeDatasmithMaterialNode.cpp b/Engine/Plugins/Experimental/Enterprise/DatasmithInterchange/Source/DatasmithInterchange/Private/InterchangeDatasmithMaterialNode.cpp
--- a/Engine/Plugins/Experimental/Enterprise/DatasmithInterchange/Source/DatasmithInterchange/Private/InterchangeDatasmithMaterialNode.cpp
+++ b/Engine/Plugins/Experimental/Enterprise/DatasmithInterchange/Source/DatasmithInterchange/Private/InterchangeDatasmithMaterialNode.cpp
@@ -665,23 +665,23 @@ namespace UE::DatasmithInterchange::MaterialUtils
 		}
 		else
 		{
-			ConnectOutputToInput(MaterialElement.GetBaseColor(), PBR::Parameters::BaseColor);
+			ConnectOutputToInput(MaterialElement.GetBaseColor(), PBRMR::Parameters::BaseColor);
 
-			ConnectOutputToInput(MaterialElement.GetMetallic(), PBR::Parameters::Metallic);
+			ConnectOutputToInput(MaterialElement.GetMetallic(), PBRMR::Parameters::Metallic);
 
-			ConnectOutputToInput(MaterialElement.GetSpecular(), PBR::Parameters::Specular);
+			ConnectOutputToInput(MaterialElement.GetSpecular(), PBRMR::Parameters::Specular);
 
-			ConnectOutputToInput(MaterialElement.GetRoughness(), PBR::Parameters::Roughness);
+			ConnectOutputToInput(MaterialElement.GetRoughness(), PBRMR::Parameters::Roughness);
 
-			ConnectOutputToInput(MaterialElement.GetEmissiveColor(), PBR::Parameters::EmissiveColor);
+			ConnectOutputToInput(MaterialElement.GetEmissiveColor(), PBRMR::Parameters::EmissiveColor);
 
-			ConnectOutputToInput(MaterialElement.GetAmbientOcclusion(), PBR::Parameters::Occlusion);
+			ConnectOutputToInput(MaterialElement.GetAmbientOcclusion(), PBRMR::Parameters::Occlusion);
 
-			ConnectOutputToInput(MaterialElement.GetNormal(), PBR::Parameters::Normal);
+			ConnectOutputToInput(MaterialElement.GetNormal(), PBRMR::Parameters::Normal);
 
-			ConnectOutputToInput(MaterialElement.GetOpacity(), PBR::Parameters::Opacity);
+			ConnectOutputToInput(MaterialElement.GetOpacity(), PBRMR::Parameters::Opacity);
 
-			ConnectOutputToInput(MaterialElement.GetRefraction(), PBR::Parameters::IndexOfRefraction);
+			ConnectOutputToInput(MaterialElement.GetRefraction(), PBRMR::Parameters::IndexOfRefraction);
 
 			ConnectOutputToInput(MaterialElement.GetClearCoat(), ClearCoat::Parameters::ClearCoat);
 
diff --git a/Engine/Plugins/Interchange/Runtime/Config/BaseInterchange.ini b/Engine/Plugins/Interchange/Runtime/Config/BaseInterchange.ini
--- a/Engine/Plugins/Interchange/Runtime/Config/BaseInterchange.ini
+++ b/Engine/Plugins/Interchange/Runtime/Config/BaseInterchange.ini
@@ -41,9 +41,16 @@
 +EnumRedirects=(OldName="ETextureSampleBlurFilter", NewName="EMaterialXTextureSampleBlurFilter")
 +EnumRedirects=(OldName="ELuminanceMode", NewName="EMaterialXLuminanceMode")
 
+; Add proper suffix to Interchange predefined materials
++PackageRedirects=(OldName="/Interchange/Materials/PBRSurfaceMaterial.PBRSurfaceMaterial", NewName="/Interchange/Materials/PBRSurfaceMaterial_MR.PBRSurfaceMaterial_MR")
++PackageRedirects=(OldName="/Interchange/Materials/ClearCoatMaterial.ClearCoatMaterial", NewName="/Interchange/Materials/ClearCoatMaterial_MR.ClearCoatMaterial_MR")
++PackageRedirects=(OldName="/Interchange/Materials/SheenMaterial.SheenMaterial", NewName="/Interchange/Materials/SheenMaterial_MR.SheenMaterial_MR")
++PackageRedirects=(OldName="/Interchange/Materials/SubsurfaceMaterial.SubsurfaceMaterial", NewName="/Interchange/Materials/SubsurfaceMaterial_MR.SubsurfaceMaterial_MR")
++PackageRedirects=(OldName="/Interchange/Materials/ThinTranslucentMaterial.ThinTranslucentMaterial", NewName="/Interchange/Materials/ThinTranslucentMaterial_MR.ThinTranslucentMaterial_MR")
+
 [/Script/InterchangeEngine.InterchangeProjectSettings]
 ContentImportSettings=(DefaultPipelineStackOverride=((Materials, "Materials"),(Textures, "Textures")),ShowPipelineStacksConfigurationDialogOverride=((Textures, False)),PipelineStacks=(("Assets", (Pipelines=("/Interchange/Pipelines/DefaultAssetsPipeline.DefaultAssetsPipeline"),PerTranslatorPipelines=((Translator="/Script/InterchangeImport.InterchangeMaterialXTranslator",Pipelines=("/Interchange/Pipelines/DefaultAssetsPipeline.DefaultAssetsPipeline","/Interchange/Pipelines/DefaultMaterialXPipeline.DefaultMaterialXPipeline")),(Translator="/Script/InterchangeImport.InterchangeGltfTranslator",Pipelines=("/Interchange/Pipelines/DefaultGLTFAssetsPipeline.DefaultGLTFAssetsPipeline","/Interchange/Pipelines/DefaultGLTFPipeline.DefaultGLTFPipeline"))))),("Materials", (Pipelines=("/Interchange/Pipelines/DefaultMaterialPipeline.DefaultMaterialPipeline"),PerTranslatorPipelines=((Translator="/Script/InterchangeImport.InterchangeMaterialXTranslator",Pipelines=("/Interchange/Pipelines/DefaultMaterialPipeline.DefaultMaterialPipeline","/Interchange/Pipelines/DefaultMaterialXPipeline.DefaultMaterialXPipeline"))))),("Textures", (Pipelines=("/Interchange/Pipelines/DefaultTexturePipeline.DefaultTexturePipeline")))),DefaultPipelineStack="Assets",PipelineConfigurationDialogClass="/Script/InterchangeEditorPipelines.InterchangePipelineConfigurationGeneric",bShowPipelineStacksConfigurationDialog=True)
-SceneImportSettings=(PipelineStacks=(("Scene", (Pipelines=("/Interchange/Pipelines/DefaultSceneAssetsPipeline.DefaultSceneAssetsPipeline","/Interchange/Pipelines/DefaultSceneLevelPipeline.DefaultSceneLevelPipeline"),PerTranslatorPipelines=((Translator="/Script/InterchangeImport.InterchangeMaterialXTranslator",Pipelines=("/Interchange/Pipelines/DefaultSceneAssetsPipeline.DefaultSceneAssetsPipeline","/Interchange/Pipelines/DefaultSceneLevelPipeline.DefaultSceneLevelPipeline","/Interchange/Pipelines/DefaultMaterialXPipeline.DefaultMaterialXPipeline")),(Translator="/Script/InterchangeImport.InterchangeGltfTranslator",Pipelines=("/Interchange/Pipelines/DefaultGLTFSceneAssetsPipeline.DefaultGLTFSceneAssetsPipeline","/Interchange/Pipelines/DefaultSceneLevelPipeline.DefaultSceneLevelPipeline","/Interchange/Pipelines/DefaultGLTFPipeline.DefaultGLTFPipeline")))))),DefaultPipelineStack="Scene",PipelineConfigurationDialogClass="/Script/InterchangeEditorPipelines.InterchangePipelineConfigurationGeneric",bShowPipelineStacksConfigurationDialog=True)
+SceneImportSettings=(PipelineStacks=(("Scene", (Pipelines=("/Interchange/Pipelines/DefaultAssetsPipeline.DefaultAssetsPipeline","/Interchange/Pipelines/DefaultSceneLevelPipeline.DefaultSceneLevelPipeline"),PerTranslatorPipelines=((Translator="/Script/InterchangeImport.InterchangeMaterialXTranslator",Pipelines=("/Interchange/Pipelines/DefaultAssetsPipeline.DefaultAssetsPipeline","/Interchange/Pipelines/DefaultSceneLevelPipeline.DefaultSceneLevelPipeline","/Interchange/Pipelines/DefaultMaterialXPipeline.DefaultMaterialXPipeline")),(Translator="/Script/InterchangeImport.InterchangeGltfTranslator",Pipelines=("/Interchange/Pipelines/DefaultGLTFSceneAssetsPipeline.DefaultGLTFSceneAssetsPipeline","/Interchange/Pipelines/DefaultSceneLevelPipeline.DefaultSceneLevelPipeline","/Interchange/Pipelines/DefaultGLTFPipeline.DefaultGLTFPipeline")))))),DefaultPipelineStack="Scene",PipelineConfigurationDialogClass="/Script/InterchangeEditorPipelines.InterchangePipelineConfigurationGeneric",bShowPipelineStacksConfigurationDialog=True)
 FilePickerClass=/Script/InterchangeEditorUtilities.InterchangeFilePickerGeneric
 bStaticMeshUseSmoothEdgesIfSmoothingInformationIsMissing=True
 GenericPipelineClass=/Script/InterchangePipelines.InterchangeGenericAssetsPipeline
diff --git a/Engine/Plugins/Interchange/Runtime/Source/FactoryNodes/Private/InterchangeMaterialFactoryNode.cpp b/Engine/Plugins/Interchange/Runtime/Source/FactoryNodes/Private/InterchangeMaterialFactoryNode.cpp
--- a/Engine/Plugins/Interchange/Runtime/Source/FactoryNodes/Private/InterchangeMaterialFactoryNode.cpp
+++ b/Engine/Plugins/Interchange/Runtime/Source/FactoryNodes/Private/InterchangeMaterialFactoryNode.cpp
@@ -37,77 +37,77 @@ UClass* UInterchangeMaterialFactoryNode::GetObjectClass() const
 
 bool UInterchangeMaterialFactoryNode::GetBaseColorConnection(FString& ExpressionNodeUid, FString& OutputName) const
 {
-	return UInterchangeShaderPortsAPI::GetInputConnection(this, UE::Interchange::Materials::PBR::Parameters::BaseColor.ToString(), ExpressionNodeUid, OutputName);
+	return UInterchangeShaderPortsAPI::GetInputConnection(this, UE::Interchange::Materials::PBRMR::Parameters::BaseColor.ToString(), ExpressionNodeUid, OutputName);
 }
 
 bool UInterchangeMaterialFactoryNode::ConnectToBaseColor(const FString& AttributeValue)
 {
-	return UInterchangeShaderPortsAPI::ConnectDefaultOuputToInput(this, UE::Interchange::Materials::PBR::Parameters::BaseColor.ToString(), AttributeValue);
+	return UInterchangeShaderPortsAPI::ConnectDefaultOuputToInput(this, UE::Interchange::Materials::PBRMR::Parameters::BaseColor.ToString(), AttributeValue);
 }
 	
 bool UInterchangeMaterialFactoryNode::ConnectOutputToBaseColor(const FString& ExpressionNodeUid, const FString& OutputName)
 {
-	return UInterchangeShaderPortsAPI::ConnectOuputToInputByName(this, UE::Interchange::Materials::PBR::Parameters::BaseColor.ToString(), ExpressionNodeUid, OutputName);
+	return UInterchangeShaderPortsAPI::ConnectOuputToInputByName(this, UE::Interchange::Materials::PBRMR::Parameters::BaseColor.ToString(), ExpressionNodeUid, OutputName);
 }
 
 bool UInterchangeMaterialFactoryNode::GetMetallicConnection(FString& ExpressionNodeUid, FString& OutputName) const
 {
-	return UInterchangeShaderPortsAPI::GetInputConnection(this, UE::Interchange::Materials::PBR::Parameters::Metallic.ToString(), ExpressionNodeUid, OutputName);
+	return UInterchangeShaderPortsAPI::GetInputConnection(this, UE::Interchange::Materials::PBRMR::Parameters::Metallic.ToString(), ExpressionNodeUid, OutputName);
 }
 
 bool UInterchangeMaterialFactoryNode::ConnectToMetallic(const FString& AttributeValue)
 {
-	return UInterchangeShaderPortsAPI::ConnectDefaultOuputToInput(this, UE::Interchange::Materials::PBR::Parameters::Metallic.ToString(), AttributeValue);
+	return UInterchangeShaderPortsAPI::ConnectDefaultOuputToInput(this, UE::Interchange::Materials::PBRMR::Parameters::Metallic.ToString(), AttributeValue);
 }
 
 bool UInterchangeMaterialFactoryNode::ConnectOutputToMetallic(const FString& ExpressionNodeUid, const FString& OutputName)
 {
-	return UInterchangeShaderPortsAPI::ConnectOuputToInputByName(this, UE::Interchange::Materials::PBR::Parameters::Metallic.ToString(), ExpressionNodeUid, OutputName);
+	return UInterchangeShaderPortsAPI::ConnectOuputToInputByName(this, UE::Interchange::Materials::PBRMR::Parameters::Metallic.ToString(), ExpressionNodeUid, OutputName);
 }
 
 bool UInterchangeMaterialFactoryNode::GetSpecularConnection(FString& ExpressionNodeUid, FString& OutputName) const
 {
-	return UInterchangeShaderPortsAPI::GetInputConnection(this, UE::Interchange::Materials::PBR::Parameters::Specular.ToString(), ExpressionNodeUid, OutputName);
+	return UInterchangeShaderPortsAPI::GetInputConnection(this, UE::Interchange::Materials::PBRMR::Parameters::Specular.ToString(), ExpressionNodeUid, OutputName);
 }
 
 bool UInterchangeMaterialFactoryNode::ConnectToSpecular(const FString& ExpressionNodeUid)
 {
-	return UInterchangeShaderPortsAPI::ConnectDefaultOuputToInput(this, UE::Interchange::Materials::PBR::Parameters::Specular.ToString(), ExpressionNodeUid);
+	return UInterchangeShaderPortsAPI::ConnectDefaultOuputToInput(this, UE::Interchange::Materials::PBRMR::Parameters::Specular.ToString(), ExpressionNodeUid);
 }
 
 bool UInterchangeMaterialFactoryNode::ConnectOutputToSpecular(const FString& ExpressionNodeUid, const FString& OutputName)
 {
-	return UInterchangeShaderPortsAPI::ConnectOuputToInputByName(this, UE::Interchange::Materials::PBR::Parameters::Specular.ToString(), ExpressionNodeUid, OutputName);
+	return UInterchangeShaderPortsAPI::ConnectOuputToInputByName(this, UE::Interchange::Materials::PBRMR::Parameters::Specular.ToString(), ExpressionNodeUid, OutputName);
 }
 
 bool UInterchangeMaterialFactoryNode::GetRoughnessConnection(FString& ExpressionNodeUid, FString& OutputName) const
 {
-	return UInterchangeShaderPortsAPI::GetInputConnection(this, UE::Interchange::Materials::PBR::Parameters::Roughness.ToString(), ExpressionNodeUid, OutputName);
+	return UInterchangeShaderPortsAPI::GetInputConnection(this, UE::Interchange::Materials::PBRMR::Parameters::Roughness.ToString(), ExpressionNodeUid, OutputName);
 }
 
 bool UInterchangeMaterialFactoryNode::ConnectToRoughness(const FString& ExpressionNodeUid)
 {
-	return UInterchangeShaderPortsAPI::ConnectDefaultOuputToInput(this, UE::Interchange::Materials::PBR::Parameters::Roughness.ToString(), ExpressionNodeUid);
+	return UInterchangeShaderPortsAPI::ConnectDefaultOuputToInput(this, UE::Interchange::Materials::PBRMR::Parameters::Roughness.ToString(), ExpressionNodeUid);
 }
 
 bool UInterchangeMaterialFactoryNode::ConnectOutputToRoughness(const FString& ExpressionNodeUid, const FString& OutputName)
 {
-	return UInterchangeShaderPortsAPI::ConnectOuputToInputByName(this, UE::Interchange::Materials::PBR::Parameters::Roughness.ToString(), ExpressionNodeUid, OutputName);
+	return UInterchangeShaderPortsAPI::ConnectOuputToInputByName(this, UE::Interchange::Materials::PBRMR::Parameters::Roughness.ToString(), ExpressionNodeUid, OutputName);
 }
 
 bool UInterchangeMaterialFactoryNode::GetAnisotropyConnection(FString& ExpressionNodeUid, FString& OutputName) const
 {
-	return UInterchangeShaderPortsAPI::GetInputConnection(this, UE::Interchange::Materials::PBR::Parameters::Anisotropy.ToString(), ExpressionNodeUid, OutputName);
+	return UInterchangeShaderPortsAPI::GetInputConnection(this, UE::Interchange::Materials::PBRMR::Parameters::Anisotropy.ToString(), ExpressionNodeUid, OutputName);
 }
 
 bool UInterchangeMaterialFactoryNode::ConnectToAnisotropy(const FString& ExpressionNodeUid)
 {
-	return UInterchangeShaderPortsAPI::ConnectDefaultOuputToInput(this, UE::Interchange::Materials::PBR::Parameters::Anisotropy.ToString(), ExpressionNodeUid);
+	return UInterchangeShaderPortsAPI::ConnectDefaultOuputToInput(this, UE::Interchange::Materials::PBRMR::Parameters::Anisotropy.ToString(), ExpressionNodeUid);
 }
 
 bool UInterchangeMaterialFactoryNode::ConnectOutputToAnisotropy(const FString& ExpressionNodeUid, const FString& OutputName)
 {
-	return UInterchangeShaderPortsAPI::ConnectOuputToInputByName(this, UE::Interchange::Materials::PBR::Parameters::Anisotropy.ToString(), ExpressionNodeUid, OutputName);
+	return UInterchangeShaderPortsAPI::ConnectOuputToInputByName(this, UE::Interchange::Materials::PBRMR::Parameters::Anisotropy.ToString(), ExpressionNodeUid, OutputName);
 }
 
 bool UInterchangeMaterialFactoryNode::GetEmissiveColorConnection(FString& ExpressionNodeUid, FString& OutputName) const
@@ -122,7 +122,7 @@ bool UInterchangeMaterialFactoryNode::ConnectToEmissiveColor(const FString& Expr
 
 bool UInterchangeMaterialFactoryNode::ConnectOutputToEmissiveColor(const FString& ExpressionNodeUid, const FString& OutputName)
 {
-	return UInterchangeShaderPortsAPI::ConnectOuputToInputByName(this, UE::Interchange::Materials::PBR::Parameters::EmissiveColor.ToString(), ExpressionNodeUid, OutputName);
+	return UInterchangeShaderPortsAPI::ConnectOuputToInputByName(this, UE::Interchange::Materials::PBRMR::Parameters::EmissiveColor.ToString(), ExpressionNodeUid, OutputName);
 }
 
 bool UInterchangeMaterialFactoryNode::GetNormalConnection(FString& ExpressionNodeUid, FString& OutputName) const
@@ -137,7 +137,7 @@ bool UInterchangeMaterialFactoryNode::ConnectToNormal(const FString& ExpressionN
 
 bool UInterchangeMaterialFactoryNode::ConnectOutputToNormal(const FString& ExpressionNodeUid, const FString& OutputName)
 {
-	return UInterchangeShaderPortsAPI::ConnectOuputToInputByName(this, UE::Interchange::Materials::PBR::Parameters::Normal.ToString(), ExpressionNodeUid, OutputName);
+	return UInterchangeShaderPortsAPI::ConnectOuputToInputByName(this, UE::Interchange::Materials::PBRMR::Parameters::Normal.ToString(), ExpressionNodeUid, OutputName);
 }
 
 bool UInterchangeMaterialFactoryNode::GetTangentConnection(FString& ExpressionNodeUid, FString& OutputName) const
@@ -152,7 +152,7 @@ bool UInterchangeMaterialFactoryNode::ConnectToTangent(const FString& Expression
 
 bool UInterchangeMaterialFactoryNode::ConnectOutputToTangent(const FString& ExpressionNodeUid, const FString& OutputName)
 {
-	return UInterchangeShaderPortsAPI::ConnectOuputToInputByName(this, UE::Interchange::Materials::PBR::Parameters::Tangent.ToString(), ExpressionNodeUid, OutputName);
+	return UInterchangeShaderPortsAPI::ConnectOuputToInputByName(this, UE::Interchange::Materials::PBRMR::Parameters::Tangent.ToString(), ExpressionNodeUid, OutputName);
 }
 
 bool UInterchangeMaterialFactoryNode::GetSubsurfaceConnection(FString& ExpressionNodeUid, FString& OutputName) const
@@ -182,7 +182,7 @@ bool UInterchangeMaterialFactoryNode::ConnectToOpacity(const FString& AttributeV
 
 bool UInterchangeMaterialFactoryNode::ConnectOutputToOpacity(const FString& ExpressionNodeUid, const FString& OutputName)
 {
-	return UInterchangeShaderPortsAPI::ConnectOuputToInputByName(this, UE::Interchange::Materials::PBR::Parameters::Opacity.ToString(), ExpressionNodeUid, OutputName);
+	return UInterchangeShaderPortsAPI::ConnectOuputToInputByName(this, UE::Interchange::Materials::PBRMR::Parameters::Opacity.ToString(), ExpressionNodeUid, OutputName);
 }
 
 bool UInterchangeMaterialFactoryNode::GetOcclusionConnection(FString& ExpressionNodeUid, FString& OutputName) const
@@ -197,7 +197,7 @@ bool UInterchangeMaterialFactoryNode::ConnectToOcclusion(const FString& Attribut
 
 bool UInterchangeMaterialFactoryNode::ConnectOutputToOcclusion(const FString& ExpressionNodeUid, const FString& OutputName)
 {
-	return UInterchangeShaderPortsAPI::ConnectOuputToInputByName(this, UE::Interchange::Materials::PBR::Parameters::Occlusion.ToString(), ExpressionNodeUid, OutputName);
+	return UInterchangeShaderPortsAPI::ConnectOuputToInputByName(this, UE::Interchange::Materials::PBRMR::Parameters::Occlusion.ToString(), ExpressionNodeUid, OutputName);
 }
 
 bool UInterchangeMaterialFactoryNode::GetRefractionConnection(FString& ExpressionNodeUid, FString& OutputName) const
diff --git a/Engine/Plugins/Interchange/Runtime/Source/Import/Private/Gltf/InterchangeGltfTranslator.cpp b/Engine/Plugins/Interchange/Runtime/Source/Import/Private/Gltf/InterchangeGltfTranslator.cpp
--- a/Engine/Plugins/Interchange/Runtime/Source/Import/Private/Gltf/InterchangeGltfTranslator.cpp
+++ b/Engine/Plugins/Interchange/Runtime/Source/Import/Private/Gltf/InterchangeGltfTranslator.cpp
@@ -527,7 +527,7 @@ void UInterchangeGLTFTranslator::HandleGltfMaterial( UInterchangeBaseNodeContain
 			TVariant< FLinearColor, float > OpacityFactor;
 			OpacityFactor.Set< float >(GltfMaterial.BaseColorFactor.W);
 
-			HandleGltfMaterialParameter(NodeContainer, GltfMaterial.BaseColor, ShaderGraphNode, PBR::Parameters::Opacity.ToString(),
+			HandleGltfMaterialParameter(NodeContainer, GltfMaterial.BaseColor, ShaderGraphNode, PBRMR::Parameters::Opacity.ToString(),
 				OpacityFactor, Standard::Nodes::TextureSample::Outputs::A.ToString());
 		}
 
@@ -546,7 +546,7 @@ void UInterchangeGLTFTranslator::HandleGltfMaterial( UInterchangeBaseNodeContain
 			TVariant< FLinearColor, float > BaseColorFactor;
 			BaseColorFactor.Set< FLinearColor >( FLinearColor( GltfMaterial.BaseColorFactor ) );
 
-			HandleGltfMaterialParameter( NodeContainer, GltfMaterial.BaseColor, ShaderGraphNode, PBR::Parameters::BaseColor.ToString(),
+			HandleGltfMaterialParameter( NodeContainer, GltfMaterial.BaseColor, ShaderGraphNode, PBRMR::Parameters::BaseColor.ToString(),
 				BaseColorFactor, Standard::Nodes::TextureSample::Outputs::RGB.ToString(), false, false, bUseVertexColor );
 		}
 
@@ -555,7 +555,7 @@ void UInterchangeGLTFTranslator::HandleGltfMaterial( UInterchangeBaseNodeContain
 			TVariant< FLinearColor, float > MetallicFactor;
 			MetallicFactor.Set< float >( GltfMaterial.MetallicRoughness.MetallicFactor );
 
-			HandleGltfMaterialParameter( NodeContainer, GltfMaterial.MetallicRoughness.Map, ShaderGraphNode, PBR::Parameters::Metallic.ToString(),
+			HandleGltfMaterialParameter( NodeContainer, GltfMaterial.MetallicRoughness.Map, ShaderGraphNode, PBRMR::Parameters::Metallic.ToString(),
 				MetallicFactor, Standard::Nodes::TextureSample::Outputs::B.ToString() );
 		}
 
@@ -564,7 +564,7 @@ void UInterchangeGLTFTranslator::HandleGltfMaterial( UInterchangeBaseNodeContain
 			TVariant< FLinearColor, float > RoughnessFactor;
 			RoughnessFactor.Set< float >( GltfMaterial.MetallicRoughness.RoughnessFactor );
 
-			HandleGltfMaterialParameter( NodeContainer, GltfMaterial.MetallicRoughness.Map, ShaderGraphNode, PBR::Parameters::Roughness.ToString(),
+			HandleGltfMaterialParameter( NodeContainer, GltfMaterial.MetallicRoughness.Map, ShaderGraphNode, PBRMR::Parameters::Roughness.ToString(),
 				RoughnessFactor, Standard::Nodes::TextureSample::Outputs::G.ToString() );
 		}
 
@@ -574,7 +574,7 @@ void UInterchangeGLTFTranslator::HandleGltfMaterial( UInterchangeBaseNodeContain
 			TVariant< FLinearColor, float > SpecularFactor;
 			SpecularFactor.Set< float >( GltfMaterial.Specular.SpecularFactor );
 
-			HandleGltfMaterialParameter( NodeContainer, GltfMaterial.Specular.SpecularMap, ShaderGraphNode, PBR::Parameters::Specular.ToString(),
+			HandleGltfMaterialParameter( NodeContainer, GltfMaterial.Specular.SpecularMap, ShaderGraphNode, PBRMR::Parameters::Specular.ToString(),
 				SpecularFactor, Standard::Nodes::TextureSample::Outputs::RGB.ToString() );
 		}
 	}
@@ -585,7 +585,7 @@ void UInterchangeGLTFTranslator::HandleGltfMaterial( UInterchangeBaseNodeContain
 			TVariant< FLinearColor, float > DiffuseColorFactor;
 			DiffuseColorFactor.Set< FLinearColor >( FLinearColor( GltfMaterial.BaseColorFactor ) );
 
-			HandleGltfMaterialParameter( NodeContainer, GltfMaterial.BaseColor, ShaderGraphNode, Phong::Parameters::DiffuseColor.ToString(),
+			HandleGltfMaterialParameter( NodeContainer, GltfMaterial.BaseColor, ShaderGraphNode, PBRSG::Parameters::DiffuseColor.ToString(),
 				DiffuseColorFactor, Standard::Nodes::TextureSample::Outputs::RGB.ToString(), false, false, bUseVertexColor );
 		}
 
@@ -594,7 +594,7 @@ void UInterchangeGLTFTranslator::HandleGltfMaterial( UInterchangeBaseNodeContain
 			TVariant< FLinearColor, float > SpecularColorFactor;
 			SpecularColorFactor.Set< FLinearColor >( FLinearColor( GltfMaterial.SpecularGlossiness.SpecularFactor ) );
 
-			HandleGltfMaterialParameter( NodeContainer, GltfMaterial.SpecularGlossiness.Map, ShaderGraphNode, Phong::Parameters::SpecularColor.ToString(),
+			HandleGltfMaterialParameter( NodeContainer, GltfMaterial.SpecularGlossiness.Map, ShaderGraphNode, PBRSG::Parameters::SpecularColor.ToString(),
 				SpecularColorFactor, Standard::Nodes::TextureSample::Outputs::RGB.ToString() );
 		}
 
@@ -603,9 +603,8 @@ void UInterchangeGLTFTranslator::HandleGltfMaterial( UInterchangeBaseNodeContain
 			TVariant< FLinearColor, float > GlossinessFactor;
 			GlossinessFactor.Set< float >( GltfMaterial.SpecularGlossiness.GlossinessFactor );
 
-			const bool bInverse = true;
-			HandleGltfMaterialParameter( NodeContainer, GltfMaterial.SpecularGlossiness.Map, ShaderGraphNode, PBR::Parameters::Roughness.ToString(),
-				GlossinessFactor, Standard::Nodes::TextureSample::Outputs::A.ToString(), bInverse );
+			HandleGltfMaterialParameter( NodeContainer, GltfMaterial.SpecularGlossiness.Map, ShaderGraphNode, PBRSG::Parameters::Glossiness.ToString(),
+				GlossinessFactor, Standard::Nodes::TextureSample::Outputs::A.ToString() );
 		}
 	}
 
@@ -640,7 +639,7 @@ void UInterchangeGLTFTranslator::HandleGltfMaterial( UInterchangeBaseNodeContain
 			TVariant< FLinearColor, float > OcclusionFactor;
 			OcclusionFactor.Set< float >( GltfMaterial.OcclusionStrength );
 
-			HandleGltfMaterialParameter( NodeContainer, GltfMaterial.Occlusion, ShaderGraphNode, PBR::Parameters::Occlusion.ToString(),
+			HandleGltfMaterialParameter( NodeContainer, GltfMaterial.Occlusion, ShaderGraphNode, PBRMR::Parameters::Occlusion.ToString(),
 				OcclusionFactor, Standard::Nodes::TextureSample::Outputs::R.ToString() );
 		}
 
@@ -650,7 +649,7 @@ void UInterchangeGLTFTranslator::HandleGltfMaterial( UInterchangeBaseNodeContain
 			TVariant< FLinearColor, float > OpacityFactor;
 			OpacityFactor.Set< float >( GltfMaterial.BaseColorFactor.W );
 
-			HandleGltfMaterialParameter( NodeContainer, GltfMaterial.BaseColor, ShaderGraphNode, PBR::Parameters::Opacity.ToString(),
+			HandleGltfMaterialParameter( NodeContainer, GltfMaterial.BaseColor, ShaderGraphNode, PBRMR::Parameters::Opacity.ToString(),
 				OpacityFactor, Standard::Nodes::TextureSample::Outputs::A.ToString() );
 		}
 
@@ -663,7 +662,7 @@ void UInterchangeGLTFTranslator::HandleGltfMaterial( UInterchangeBaseNodeContain
 		// IOR
 		if ( GltfMaterial.bHasIOR )
 		{
-			ShaderGraphNode.AddFloatAttribute( UInterchangeShaderPortsAPI::MakeInputValueKey( PBR::Parameters::IndexOfRefraction.ToString() ), GltfMaterial.IOR );
+			ShaderGraphNode.AddFloatAttribute( UInterchangeShaderPortsAPI::MakeInputValueKey( PBRMR::Parameters::IndexOfRefraction.ToString() ), GltfMaterial.IOR );
 		}
 	}
 
@@ -804,7 +803,7 @@ void UInterchangeGLTFTranslator::HandleGltfTransmission( UInterchangeBaseNodeCon
 			FString MetallicNodeUid;
 			FString MetallicNodeOutput;
 
-			if ( UInterchangeShaderPortsAPI::GetInputConnection( &ShaderGraphNode, PBR::Parameters::Metallic.ToString(), MetallicNodeUid, MetallicNodeOutput ) )
+			if ( UInterchangeShaderPortsAPI::GetInputConnection( &ShaderGraphNode, PBRMR::Parameters::Metallic.ToString(), MetallicNodeUid, MetallicNodeOutput ) )
 			{
 				const FString MetallicLerpNodeName =  TEXT("OpacityMetallicLerp");
 				UInterchangeShaderNode* LerpMetallicNode = UInterchangeShaderNode::Create( &NodeContainer, MetallicLerpNodeName, ShaderGraphNode.GetUniqueID() );
@@ -822,7 +821,7 @@ void UInterchangeGLTFTranslator::HandleGltfTransmission( UInterchangeBaseNodeCon
 
 		if ( GltfMaterial.AlphaMode != GLTF::FMaterial::EAlphaMode::Opaque )
 		{
-			if ( UInterchangeShaderPortsAPI::GetInputConnection( &ShaderGraphNode, PBR::Parameters::Opacity.ToString(), OpacityNodeUid, OpacityNodeOutput ) )
+			if ( UInterchangeShaderPortsAPI::GetInputConnection( &ShaderGraphNode, PBRMR::Parameters::Opacity.ToString(), OpacityNodeUid, OpacityNodeOutput ) )
 			{
 				const FString OpacityLerpNodeName =  TEXT("OpacityLerp");
 				UInterchangeShaderNode* OpacityLerpNode = UInterchangeShaderNode::Create( &NodeContainer, OpacityLerpNodeName, ShaderGraphNode.GetUniqueID() );
@@ -851,7 +850,7 @@ void UInterchangeGLTFTranslator::HandleGltfTransmission( UInterchangeBaseNodeCon
 		FString BaseColorNodeUid;
 		FString BaseColorNodeOutput;
 
-		if ( UInterchangeShaderPortsAPI::GetInputConnection( CurrentNode, PBR::Parameters::BaseColor.ToString(), BaseColorNodeUid, BaseColorNodeOutput ) )
+		if ( UInterchangeShaderPortsAPI::GetInputConnection( CurrentNode, PBRMR::Parameters::BaseColor.ToString(), BaseColorNodeUid, BaseColorNodeOutput ) )
 		{
 			CurrentNode = NodeContainer.GetNode( BaseColorNodeUid );
 			CurrentOuput = BaseColorNodeOutput;
@@ -859,7 +858,7 @@ void UInterchangeGLTFTranslator::HandleGltfTransmission( UInterchangeBaseNodeCon
 		else
 		{
 			FLinearColor BaseColor;
-			if ( ShaderGraphNode.GetLinearColorAttribute( UInterchangeShaderPortsAPI::MakeInputValueKey( PBR::Parameters::BaseColor.ToString() ), BaseColor ) )
+			if ( ShaderGraphNode.GetLinearColorAttribute( UInterchangeShaderPortsAPI::MakeInputValueKey( PBRMR::Parameters::BaseColor.ToString() ), BaseColor ) )
 			{
 				CurrentNode = nullptr;
 				CurrentColor = BaseColor;
diff --git a/Engine/Plugins/Interchange/Runtime/Source/Import/Private/Material/InterchangeMaterialFactory.cpp b/Engine/Plugins/Interchange/Runtime/Source/Import/Private/Material/InterchangeMaterialFactory.cpp
--- a/Engine/Plugins/Interchange/Runtime/Source/Import/Private/Material/InterchangeMaterialFactory.cpp
+++ b/Engine/Plugins/Interchange/Runtime/Source/Import/Private/Material/InterchangeMaterialFactory.cpp
@@ -1575,16 +1575,16 @@ void UInterchangeMaterialFunctionFactory::SetupMaterial(UMaterialFunction* Mater
 	// Only material functions with BSDF output are supported for now
 	if (!ConnectInput(Common::Parameters::BxDF.ToString(), Attrib->GetInput(0)))
 	{
-		ConnectInput(PBR::Parameters::BaseColor.ToString(), &Attrib->BaseColor);
-		ConnectInput(PBR::Parameters::Metallic.ToString(), &Attrib->Metallic);
-		ConnectInput(PBR::Parameters::Specular.ToString(), &Attrib->Specular);
-		ConnectInput(PBR::Parameters::Roughness.ToString(), &Attrib->Roughness);
-		ConnectInput(PBR::Parameters::EmissiveColor.ToString(), &Attrib->EmissiveColor);
-		ConnectInput(PBR::Parameters::Normal.ToString(), &Attrib->Normal);
-		ConnectInput(PBR::Parameters::Anisotropy.ToString(), &Attrib->Anisotropy);
-		ConnectInput(PBR::Parameters::Tangent.ToString(), &Attrib->Tangent);
-		ConnectInput(PBR::Parameters::IndexOfRefraction.ToString(), &Attrib->Refraction);
-		ConnectInput(PBR::Parameters::Occlusion.ToString(), &Attrib->AmbientOcclusion);
+		ConnectInput(PBRMR::Parameters::BaseColor.ToString(), &Attrib->BaseColor);
+		ConnectInput(PBRMR::Parameters::Metallic.ToString(), &Attrib->Metallic);
+		ConnectInput(PBRMR::Parameters::Specular.ToString(), &Attrib->Specular);
+		ConnectInput(PBRMR::Parameters::Roughness.ToString(), &Attrib->Roughness);
+		ConnectInput(PBRMR::Parameters::EmissiveColor.ToString(), &Attrib->EmissiveColor);
+		ConnectInput(PBRMR::Parameters::Normal.ToString(), &Attrib->Normal);
+		ConnectInput(PBRMR::Parameters::Anisotropy.ToString(), &Attrib->Anisotropy);
+		ConnectInput(PBRMR::Parameters::Tangent.ToString(), &Attrib->Tangent);
+		ConnectInput(PBRMR::Parameters::IndexOfRefraction.ToString(), &Attrib->Refraction);
+		ConnectInput(PBRMR::Parameters::Occlusion.ToString(), &Attrib->AmbientOcclusion);
 		ConnectInput(ClearCoat::Parameters::ClearCoat.ToString(), &Attrib->ClearCoat);
 		ConnectInput(ClearCoat::Parameters::ClearCoatRoughness.ToString(), &Attrib->ClearCoatRoughness);
 		ConnectInput(Subsurface::Parameters::SubsurfaceColor.ToString(), &Attrib->SubsurfaceColor);
diff --git a/Engine/Plugins/Interchange/Runtime/Source/Import/Private/MaterialX/MaterialXUtils/MaterialXStandardSurfaceShader.cpp b/Engine/Plugins/Interchange/Runtime/Source/Import/Private/MaterialX/MaterialXUtils/MaterialXStandardSurfaceShader.cpp
--- a/Engine/Plugins/Interchange/Runtime/Source/Import/Private/MaterialX/MaterialXUtils/MaterialXStandardSurfaceShader.cpp
+++ b/Engine/Plugins/Interchange/Runtime/Source/Import/Private/MaterialX/MaterialXUtils/MaterialXStandardSurfaceShader.cpp
@@ -138,20 +138,20 @@ void FMaterialXStandardSurfaceShader::Translate(mx::NodePtr StandardSurfaceNode)
 	ConnectNodeOutputToInput(mx::StandardSurface::Input::TransmissionExtraRoughness, StandardSurfaceShaderNode, StandardSurface::Parameters::TransmissionExtraRoughness.ToString(), mx::StandardSurface::DefaultValue::Float::TransmissionExtraRoughness);
 
 	// Outputs
-	UInterchangeShaderPortsAPI::ConnectOuputToInputByName(ShaderGraphNode, PBR::Parameters::BaseColor.ToString(), StandardSurfaceShaderNode->GetUniqueID(), TEXT("Base Color"));
-	UInterchangeShaderPortsAPI::ConnectOuputToInputByName(ShaderGraphNode, PBR::Parameters::Metallic.ToString(), StandardSurfaceShaderNode->GetUniqueID(), PBR::Parameters::Metallic.ToString());
-	UInterchangeShaderPortsAPI::ConnectOuputToInputByName(ShaderGraphNode, PBR::Parameters::Specular.ToString(), StandardSurfaceShaderNode->GetUniqueID(), PBR::Parameters::Specular.ToString());
-	UInterchangeShaderPortsAPI::ConnectOuputToInputByName(ShaderGraphNode, PBR::Parameters::Roughness.ToString(), StandardSurfaceShaderNode->GetUniqueID(), PBR::Parameters::Roughness.ToString());
-	UInterchangeShaderPortsAPI::ConnectOuputToInputByName(ShaderGraphNode, PBR::Parameters::EmissiveColor.ToString(), StandardSurfaceShaderNode->GetUniqueID(), PBR::Parameters::EmissiveColor.ToString());
-	UInterchangeShaderPortsAPI::ConnectOuputToInputByName(ShaderGraphNode, PBR::Parameters::Anisotropy.ToString(), StandardSurfaceShaderNode->GetUniqueID(), PBR::Parameters::Anisotropy.ToString());
-	UInterchangeShaderPortsAPI::ConnectOuputToInputByName(ShaderGraphNode, PBR::Parameters::Normal.ToString(), StandardSurfaceShaderNode->GetUniqueID(), PBR::Parameters::Normal.ToString());
-	UInterchangeShaderPortsAPI::ConnectOuputToInputByName(ShaderGraphNode, PBR::Parameters::Tangent.ToString(), StandardSurfaceShaderNode->GetUniqueID(), PBR::Parameters::Tangent.ToString());
+	UInterchangeShaderPortsAPI::ConnectOuputToInputByName(ShaderGraphNode, PBRMR::Parameters::BaseColor.ToString(), StandardSurfaceShaderNode->GetUniqueID(), TEXT("Base Color"));
+	UInterchangeShaderPortsAPI::ConnectOuputToInputByName(ShaderGraphNode, PBRMR::Parameters::Metallic.ToString(), StandardSurfaceShaderNode->GetUniqueID(), PBRMR::Parameters::Metallic.ToString());
+	UInterchangeShaderPortsAPI::ConnectOuputToInputByName(ShaderGraphNode, PBRMR::Parameters::Specular.ToString(), StandardSurfaceShaderNode->GetUniqueID(), PBRMR::Parameters::Specular.ToString());
+	UInterchangeShaderPortsAPI::ConnectOuputToInputByName(ShaderGraphNode, PBRMR::Parameters::Roughness.ToString(), StandardSurfaceShaderNode->GetUniqueID(), PBRMR::Parameters::Roughness.ToString());
+	UInterchangeShaderPortsAPI::ConnectOuputToInputByName(ShaderGraphNode, PBRMR::Parameters::EmissiveColor.ToString(), StandardSurfaceShaderNode->GetUniqueID(), PBRMR::Parameters::EmissiveColor.ToString());
+	UInterchangeShaderPortsAPI::ConnectOuputToInputByName(ShaderGraphNode, PBRMR::Parameters::Anisotropy.ToString(), StandardSurfaceShaderNode->GetUniqueID(), PBRMR::Parameters::Anisotropy.ToString());
+	UInterchangeShaderPortsAPI::ConnectOuputToInputByName(ShaderGraphNode, PBRMR::Parameters::Normal.ToString(), StandardSurfaceShaderNode->GetUniqueID(), PBRMR::Parameters::Normal.ToString());
+	UInterchangeShaderPortsAPI::ConnectOuputToInputByName(ShaderGraphNode, PBRMR::Parameters::Tangent.ToString(), StandardSurfaceShaderNode->GetUniqueID(), PBRMR::Parameters::Tangent.ToString());
 
 	// We can't have all shading models at once, so we have to make a choice here
 	if(UInterchangeShaderPortsAPI::HasInput(StandardSurfaceShaderNode, StandardSurface::Parameters::Transmission))
 	{
 		StandardSurfaceShaderNode->SetCustomMaterialFunction(TEXT("/Interchange/Functions/MX_TransmissionSurface.MX_TransmissionSurface"));
-		UInterchangeShaderPortsAPI::ConnectOuputToInputByName(ShaderGraphNode, PBR::Parameters::Opacity.ToString(), StandardSurfaceShaderNode->GetUniqueID(), PBR::Parameters::Opacity.ToString());
+		UInterchangeShaderPortsAPI::ConnectOuputToInputByName(ShaderGraphNode, PBRMR::Parameters::Opacity.ToString(), StandardSurfaceShaderNode->GetUniqueID(), PBRMR::Parameters::Opacity.ToString());
 		UInterchangeShaderPortsAPI::ConnectOuputToInputByName(ShaderGraphNode, ThinTranslucent::Parameters::TransmissionColor.ToString(), StandardSurfaceShaderNode->GetUniqueID(), ThinTranslucent::Parameters::TransmissionColor.ToString());
 		UInterchangeShaderPortsAPI::ConnectOuputToInputByName(ShaderGraphNode, Common::Parameters::Refraction.ToString(), StandardSurfaceShaderNode->GetUniqueID(), Common::Parameters::Refraction.ToString());
 	}
@@ -168,7 +168,7 @@ void FMaterialXStandardSurfaceShader::Translate(mx::NodePtr StandardSurfaceNode)
 	}
 	else if(UInterchangeShaderPortsAPI::HasInput(StandardSurfaceShaderNode, StandardSurface::Parameters::Subsurface))
 	{
-		UInterchangeShaderPortsAPI::ConnectOuputToInputByName(ShaderGraphNode, PBR::Parameters::Opacity.ToString(), StandardSurfaceShaderNode->GetUniqueID(), PBR::Parameters::Opacity.ToString());
+		UInterchangeShaderPortsAPI::ConnectOuputToInputByName(ShaderGraphNode, PBRMR::Parameters::Opacity.ToString(), StandardSurfaceShaderNode->GetUniqueID(), PBRMR::Parameters::Opacity.ToString());
 		UInterchangeShaderPortsAPI::ConnectOuputToInputByName(ShaderGraphNode, Subsurface::Parameters::SubsurfaceColor.ToString(), StandardSurfaceShaderNode->GetUniqueID(), Subsurface::Parameters::SubsurfaceColor.ToString());
 	}
 }
diff --git a/Engine/Plugins/Interchange/Runtime/Source/Import/Private/MaterialX/MaterialXUtils/MaterialXSurfaceUnlitShader.cpp b/Engine/Plugins/Interchange/Runtime/Source/Import/Private/MaterialX/MaterialXUtils/MaterialXSurfaceUnlitShader.cpp
--- a/Engine/Plugins/Interchange/Runtime/Source/Import/Private/MaterialX/MaterialXUtils/MaterialXSurfaceUnlitShader.cpp
+++ b/Engine/Plugins/Interchange/Runtime/Source/Import/Private/MaterialX/MaterialXUtils/MaterialXSurfaceUnlitShader.cpp
@@ -56,16 +56,16 @@ void FMaterialXSurfaceUnlitShader::Translate(mx::NodePtr SurfaceUnlitNode)
 	ConnectNodeOutputToInput(mx::SurfaceUnlit::Input::TransmissionColor, SurfaceUnlitShaderNode, SurfaceUnlit::Parameters::TransmissionColor.ToString(), mx::SurfaceUnlit::DefaultValue::Color3::TransmissionColor);
 
 	// Outputs
-	UInterchangeShaderPortsAPI::ConnectOuputToInputByName(ShaderGraphNode, PBR::Parameters::EmissiveColor.ToString(), SurfaceUnlitShaderNode->GetUniqueID(), PBR::Parameters::EmissiveColor.ToString());
+	UInterchangeShaderPortsAPI::ConnectOuputToInputByName(ShaderGraphNode, PBRMR::Parameters::EmissiveColor.ToString(), SurfaceUnlitShaderNode->GetUniqueID(), PBRMR::Parameters::EmissiveColor.ToString());
 
 	//We can't have both Opacity and Opacity Mask we need to make a choice	
 	if(UInterchangeShaderPortsAPI::HasInput(SurfaceUnlitShaderNode, SurfaceUnlit::Parameters::Transmission))
 	{
-		UInterchangeShaderPortsAPI::ConnectOuputToInputByName(ShaderGraphNode, PBR::Parameters::Opacity.ToString(), SurfaceUnlitShaderNode->GetUniqueID(), PBR::Parameters::Opacity.ToString());
+		UInterchangeShaderPortsAPI::ConnectOuputToInputByName(ShaderGraphNode, PBRMR::Parameters::Opacity.ToString(), SurfaceUnlitShaderNode->GetUniqueID(), PBRMR::Parameters::Opacity.ToString());
 	}
 	else if(UInterchangeShaderPortsAPI::HasInput(SurfaceUnlitShaderNode, SurfaceUnlit::Parameters::Opacity))
 	{
-		UInterchangeShaderPortsAPI::ConnectOuputToInputByName(ShaderGraphNode, PBR::Parameters::Opacity.ToString(), SurfaceUnlitShaderNode->GetUniqueID(), TEXT("OpacityMask"));
+		UInterchangeShaderPortsAPI::ConnectOuputToInputByName(ShaderGraphNode, PBRMR::Parameters::Opacity.ToString(), SurfaceUnlitShaderNode->GetUniqueID(), TEXT("OpacityMask"));
 		ShaderGraphNode->SetCustomOpacityMaskClipValue(1.f); // just to connect to the opacity mask
 	}
 }
diff --git a/Engine/Plugins/Interchange/Runtime/Source/Import/Private/MaterialX/MaterialXUtils/MaterialXUsdPreviewSurfaceShader.cpp b/Engine/Plugins/Interchange/Runtime/Source/Import/Private/MaterialX/MaterialXUtils/MaterialXUsdPreviewSurfaceShader.cpp
--- a/Engine/Plugins/Interchange/Runtime/Source/Import/Private/MaterialX/MaterialXUtils/MaterialXUsdPreviewSurfaceShader.cpp
+++ b/Engine/Plugins/Interchange/Runtime/Source/Import/Private/MaterialX/MaterialXUtils/MaterialXUsdPreviewSurfaceShader.cpp
@@ -82,20 +82,20 @@ void FMaterialXUsdPreviewSurfaceShader::Translate(MaterialX::NodePtr UsdPreviewS
 	ConnectNodeOutputToInput(mx::UsdPreviewSurface::Input::Occlusion, UsdPreviewSurfaceShaderNode, UsdPreviewSurface::Parameters::Occlusion.ToString(), mx::UsdPreviewSurface::DefaultValue::Float::Occlusion);
 
 	// Outputs
-	UInterchangeShaderPortsAPI::ConnectOuputToInputByName(ShaderGraphNode, PBR::Parameters::BaseColor.ToString(), UsdPreviewSurfaceShaderNode->GetUniqueID(), PBR::Parameters::BaseColor.ToString());
-	UInterchangeShaderPortsAPI::ConnectOuputToInputByName(ShaderGraphNode, PBR::Parameters::Metallic.ToString(), UsdPreviewSurfaceShaderNode->GetUniqueID(), PBR::Parameters::Metallic.ToString());
-	UInterchangeShaderPortsAPI::ConnectOuputToInputByName(ShaderGraphNode, PBR::Parameters::Specular.ToString(), UsdPreviewSurfaceShaderNode->GetUniqueID(), PBR::Parameters::Specular.ToString());
-	UInterchangeShaderPortsAPI::ConnectOuputToInputByName(ShaderGraphNode, PBR::Parameters::Roughness.ToString(), UsdPreviewSurfaceShaderNode->GetUniqueID(), PBR::Parameters::Roughness.ToString());
-	UInterchangeShaderPortsAPI::ConnectOuputToInputByName(ShaderGraphNode, PBR::Parameters::EmissiveColor.ToString(), UsdPreviewSurfaceShaderNode->GetUniqueID(), PBR::Parameters::EmissiveColor.ToString());
-	UInterchangeShaderPortsAPI::ConnectOuputToInputByName(ShaderGraphNode, PBR::Parameters::Normal.ToString(), UsdPreviewSurfaceShaderNode->GetUniqueID(), PBR::Parameters::Normal.ToString());
+	UInterchangeShaderPortsAPI::ConnectOuputToInputByName(ShaderGraphNode, PBRMR::Parameters::BaseColor.ToString(), UsdPreviewSurfaceShaderNode->GetUniqueID(), PBRMR::Parameters::BaseColor.ToString());
+	UInterchangeShaderPortsAPI::ConnectOuputToInputByName(ShaderGraphNode, PBRMR::Parameters::Metallic.ToString(), UsdPreviewSurfaceShaderNode->GetUniqueID(), PBRMR::Parameters::Metallic.ToString());
+	UInterchangeShaderPortsAPI::ConnectOuputToInputByName(ShaderGraphNode, PBRMR::Parameters::Specular.ToString(), UsdPreviewSurfaceShaderNode->GetUniqueID(), PBRMR::Parameters::Specular.ToString());
+	UInterchangeShaderPortsAPI::ConnectOuputToInputByName(ShaderGraphNode, PBRMR::Parameters::Roughness.ToString(), UsdPreviewSurfaceShaderNode->GetUniqueID(), PBRMR::Parameters::Roughness.ToString());
+	UInterchangeShaderPortsAPI::ConnectOuputToInputByName(ShaderGraphNode, PBRMR::Parameters::EmissiveColor.ToString(), UsdPreviewSurfaceShaderNode->GetUniqueID(), PBRMR::Parameters::EmissiveColor.ToString());
+	UInterchangeShaderPortsAPI::ConnectOuputToInputByName(ShaderGraphNode, PBRMR::Parameters::Normal.ToString(), UsdPreviewSurfaceShaderNode->GetUniqueID(), PBRMR::Parameters::Normal.ToString());
 	
 	if(UInterchangeShaderPortsAPI::HasInput(UsdPreviewSurfaceShaderNode, UsdPreviewSurface::Parameters::Opacity))
 	{
-		UInterchangeShaderPortsAPI::ConnectOuputToInputByName(ShaderGraphNode, PBR::Parameters::Opacity.ToString(), UsdPreviewSurfaceShaderNode->GetUniqueID(), PBR::Parameters::Opacity.ToString());
+		UInterchangeShaderPortsAPI::ConnectOuputToInputByName(ShaderGraphNode, PBRMR::Parameters::Opacity.ToString(), UsdPreviewSurfaceShaderNode->GetUniqueID(), PBRMR::Parameters::Opacity.ToString());
 	}
 
-	UInterchangeShaderPortsAPI::ConnectOuputToInputByName(ShaderGraphNode, PBR::Parameters::Occlusion.ToString(), UsdPreviewSurfaceShaderNode->GetUniqueID(), PBR::Parameters::Occlusion.ToString());
-	UInterchangeShaderPortsAPI::ConnectOuputToInputByName(ShaderGraphNode, PBR::Parameters::Refraction.ToString(), UsdPreviewSurfaceShaderNode->GetUniqueID(), PBR::Parameters::Refraction.ToString());
+	UInterchangeShaderPortsAPI::ConnectOuputToInputByName(ShaderGraphNode, PBRMR::Parameters::Occlusion.ToString(), UsdPreviewSurfaceShaderNode->GetUniqueID(), PBRMR::Parameters::Occlusion.ToString());
+	UInterchangeShaderPortsAPI::ConnectOuputToInputByName(ShaderGraphNode, PBRMR::Parameters::Refraction.ToString(), UsdPreviewSurfaceShaderNode->GetUniqueID(), PBRMR::Parameters::Refraction.ToString());
 	UInterchangeShaderPortsAPI::ConnectOuputToInputByName(ShaderGraphNode, ClearCoat::Parameters::ClearCoat.ToString(), UsdPreviewSurfaceShaderNode->GetUniqueID(), ClearCoat::Parameters::ClearCoat.ToString());
 	UInterchangeShaderPortsAPI::ConnectOuputToInputByName(ShaderGraphNode, ClearCoat::Parameters::ClearCoatRoughness.ToString(), UsdPreviewSurfaceShaderNode->GetUniqueID(), ClearCoat::Parameters::ClearCoatRoughness.ToString());
 }
diff --git a/Engine/Plugins/Interchange/Runtime/Source/Import/Private/Mesh/InterchangeOBJTranslator.cpp b/Engine/Plugins/Interchange/Runtime/Source/Import/Private/Mesh/InterchangeOBJTranslator.cpp
--- a/Engine/Plugins/Interchange/Runtime/Source/Import/Private/Mesh/InterchangeOBJTranslator.cpp
+++ b/Engine/Plugins/Interchange/Runtime/Source/Import/Private/Mesh/InterchangeOBJTranslator.cpp
@@ -1160,14 +1160,14 @@ namespace ObjTranslatorUtils
 			return false;
 		}
 
-		AddTexturedColoredInput(BaseNodeContainer, ShaderGraphNode, Materials::PBR::Parameters::BaseColor.ToString(), MaterialData.DiffuseColor, MaterialData.DiffuseTexture);
-		AddTexturedWeightedInput(BaseNodeContainer, ShaderGraphNode, Materials::PBR::Parameters::Metallic.ToString(), MaterialData.Metallic, MaterialData.MetallicTexture);
+		AddTexturedColoredInput(BaseNodeContainer, ShaderGraphNode, Materials::PBRMR::Parameters::BaseColor.ToString(), MaterialData.DiffuseColor, MaterialData.DiffuseTexture);
+		AddTexturedWeightedInput(BaseNodeContainer, ShaderGraphNode, Materials::PBRMR::Parameters::Metallic.ToString(), MaterialData.Metallic, MaterialData.MetallicTexture);
 
 		const float SpecularValue = (MaterialData.SpecularColor.X + MaterialData.SpecularColor.Y + MaterialData.SpecularColor.Z) / 3;
-		AddTexturedWeightedInput(BaseNodeContainer, ShaderGraphNode, Materials::PBR::Parameters::Specular.ToString(), SpecularValue, MaterialData.SpecularTexture);
+		AddTexturedWeightedInput(BaseNodeContainer, ShaderGraphNode, Materials::PBRMR::Parameters::Specular.ToString(), SpecularValue, MaterialData.SpecularTexture);
 
 		// Roughness (Glossiness, Shininess, Specular Power/Exponent in non-PBR terms)
-		AddTexturedWeightedInput(BaseNodeContainer, ShaderGraphNode, Materials::PBR::Parameters::Roughness.ToString(), MaterialData.Roughness, MaterialData.RoughnessTexture);
+		AddTexturedWeightedInput(BaseNodeContainer, ShaderGraphNode, Materials::PBRMR::Parameters::Roughness.ToString(), MaterialData.Roughness, MaterialData.RoughnessTexture);
 
 		if (IsScalarInitialized(MaterialData.ClearCoatThickness) && !FMath::IsNearlyZero(MaterialData.ClearCoatThickness))
 		{
@@ -1181,7 +1181,7 @@ namespace ObjTranslatorUtils
 
 		if (IsScalarInitialized(MaterialData.Anisotropy)  && !FMath::IsNearlyZero(MaterialData.Anisotropy))
 		{
-			ShaderGraphNode->AddFloatAttribute(UInterchangeShaderPortsAPI::MakeInputValueKey(Materials::PBR::Parameters::Anisotropy.ToString()), MaterialData.Anisotropy);
+			ShaderGraphNode->AddFloatAttribute(UInterchangeShaderPortsAPI::MakeInputValueKey(Materials::PBRMR::Parameters::Anisotropy.ToString()), MaterialData.Anisotropy);
 		}
 
 		return true;
@@ -1429,13 +1429,20 @@ bool UInterchangeOBJTranslator::Translate(UInterchangeBaseNodeContainer& BaseNod
 						}
 						else if (IsScalarInitialized(MaterialData.SpecularExponent))
 						{
-							ShaderGraphNode->AddFloatAttribute(UInterchangeShaderPortsAPI::MakeInputValueKey(Materials::Phong::Parameters::Shininess.ToString()), MaterialData.SpecularExponent * SpecularExponentToShininessFactor);
+							ShaderGraphNode->AddFloatAttribute(UInterchangeShaderPortsAPI::MakeInputValueKey(Materials::Phong::Parameters::Shininess.ToString()), MaterialData.SpecularExponent);
 						}
 					}
 
+					const FString AmbientColorPropertyName = Materials::Phong::Parameters::AmbientColor.ToString();
 					if (bReflectionMetal)
 					{
-						ShaderGraphNode->AddFloatAttribute(UInterchangeShaderPortsAPI::MakeInputValueKey(Materials::PBR::Parameters::Metallic.ToString()), 1.0);
+						// Set ambient to white as it drives the metallic parameter of UE Pbr material.
+						ShaderGraphNode->AddLinearColorAttribute(UInterchangeShaderPortsAPI::MakeInputValueKey(AmbientColorPropertyName), FLinearColor::White);
+					}
+					else
+					{
+						const FVector3f AmbientColor = MaterialData.AmbientColor.GetMin() < 0.f ? FVector3f(0.8f, 0.8f, 0.8f) : MaterialData.AmbientColor;
+						AddTexturedColoredInput(BaseNodeContainer, ShaderGraphNode, AmbientColorPropertyName, AmbientColor, MaterialData.AmbientTexture);
 					}
 				}
 			}
diff --git a/Engine/Plugins/Interchange/Runtime/Source/Nodes/Public/InterchangeMaterialDefinitions.h b/Engine/Plugins/Interchange/Runtime/Source/Nodes/Public/InterchangeMaterialDefinitions.h
--- a/Engine/Plugins/Interchange/Runtime/Source/Nodes/Public/InterchangeMaterialDefinitions.h
+++ b/Engine/Plugins/Interchange/Runtime/Source/Nodes/Public/InterchangeMaterialDefinitions.h
@@ -294,6 +294,7 @@ namespace Interchange
 				const FName IndexOfRefraction = TEXT("IOR"); // Type: float
 				const FName BxDF = TEXT("BxDF"); // input/output of BSDF or BRDF or BXDF or BTDF data
 				const FName Refraction = TEXT("Refraction"); // input/output of BSDF or BRDF or BXDF or BTDF data
+				const FName Anisotropy = TEXT("Anisotropy"); // Type: float
 			}
 		}
 
@@ -315,10 +316,25 @@ namespace Interchange
 
 				const FName SpecularColor = TEXT("SpecularColor"); // Type: linear color
 				const FName Shininess = TEXT("Shininess"); // Type: float, this is the specular exponent, expected range: 2-100
+				const FName AmbientColor = TEXT("AmbientColor"); // Type: linear color
+			}
+		}
+
+		/** PBR Specular/Glossiness model */
+		namespace PBRSG
+		{
+			namespace Parameters
+			{
+				using namespace Common::Parameters;
+
+				const FName DiffuseColor = TEXT("DiffuseColor"); // Type: vector3
+				const FName SpecularColor = TEXT("SpecularColor"); // Type: vector3
+				const FName Glossiness = TEXT("Glossiness"); // Type: float
 			}
 		}
 
-		namespace PBR
+		/** PBR Metallic/Roughness model */
+		namespace PBRMR
 		{
 			namespace Parameters
 			{
@@ -328,7 +344,6 @@ namespace Interchange
 				const FName Metallic = TEXT("Metallic"); // Type: float
 				const FName Specular = TEXT("Specular"); // Type: float
 				const FName Roughness = TEXT("Roughness"); // Type: float
-				const FName Anisotropy = TEXT("Anisotropy"); // Type: float
 			}
 		}
 
@@ -336,8 +351,6 @@ namespace Interchange
 		{
 			namespace Parameters
 			{
-				using namespace PBR::Parameters;
-
 				const FName ClearCoat = TEXT("ClearCoat"); // Type: float
 				const FName ClearCoatRoughness = TEXT("ClearCoatRoughness"); // Type: float
 				const FName ClearCoatNormal = TEXT("ClearCoatNormal"); // Type: vector3
@@ -348,8 +361,6 @@ namespace Interchange
 		{
 			namespace Parameters
 			{
-				using namespace PBR::Parameters;
-
 				const FName TransmissionColor = TEXT("TransmissionColor"); // Type: vector3
 			}
 		}
@@ -358,8 +369,6 @@ namespace Interchange
 		{
 			namespace Parameters
 			{
-				using namespace PBR::Parameters;
-
 				const FName SheenColor = TEXT("SheenColor"); // Type: vector3
 				const FName SheenRoughness = TEXT("SheenRoughness"); // Type: float
 			}
@@ -369,8 +378,6 @@ namespace Interchange
 		{
 			namespace Parameters
 			{
-				using namespace PBR::Parameters;
-
 				const FName SubsurfaceColor = TEXT("SubsurfaceColor"); // Type: linear color
 			}
 		}
diff --git a/Engine/Plugins/Interchange/Runtime/Source/Parsers/Fbx/Private/FbxMaterial.cpp b/Engine/Plugins/Interchange/Runtime/Source/Parsers/Fbx/Private/FbxMaterial.cpp
--- a/Engine/Plugins/Interchange/Runtime/Source/Parsers/Fbx/Private/FbxMaterial.cpp
+++ b/Engine/Plugins/Interchange/Runtime/Source/Parsers/Fbx/Private/FbxMaterial.cpp
@@ -48,7 +48,56 @@ namespace UE
 				return TextureNode;
 			}
 
-			void FFbxMaterial::ConvertPropertyToShaderNode(UInterchangeBaseNodeContainer& NodeContainer, UInterchangeShaderGraphNode* ShaderGraphNode, FbxProperty& Property, float Factor, FName InputName,
+			UInterchangeShaderNode* FFbxMaterial::CreateTextureSampler(FbxFileTexture* FbxTexture, UInterchangeBaseNodeContainer& NodeContainer, const FString& ShaderUniqueID)
+			{
+				using namespace Materials::Standard::Nodes;
+
+				if (!FbxTexture)
+				{
+					return nullptr;
+				}
+
+				const FString TextureFilename = FbxTexture ? UTF8_TO_TCHAR(FbxTexture->GetFileName()) : TEXT("");
+				const FString TextureName = FPaths::GetBaseFilename(TextureFilename);
+
+				UInterchangeShaderNode* TextureSampleShader = UInterchangeShaderNode::Create(&NodeContainer, TextureName, ShaderUniqueID);
+				TextureSampleShader->SetCustomShaderType(TextureSample::Name.ToString());
+
+				// Return incomplete texture sampler if texture file does not exist
+				if (TextureFilename.IsEmpty() || !FPaths::FileExists(TextureFilename))
+				{
+					const UInterchangeBaseNode* ShaderGraphNode = NodeContainer.GetNode(ShaderUniqueID);
+					UInterchangeResultTextureWarning_TextureFileDoNotExist* Message = Parser.AddMessage<UInterchangeResultTextureWarning_TextureFileDoNotExist>();
+					Message->TextureName = TextureFilename;
+					Message->MaterialName = ShaderGraphNode ? ShaderGraphNode->GetDisplayLabel() : TEXT("Unknown");
+
+					return TextureSampleShader;
+				}
+
+				const UInterchangeTexture2DNode* TextureNode = Cast<const UInterchangeTexture2DNode>(NodeContainer.GetNode(UInterchangeTextureNode::MakeNodeUid(TextureName)));
+
+				if (!TextureNode)
+				{
+					TextureNode = CreateTexture2DNode(NodeContainer, TextureFilename);
+				}
+
+				TextureSampleShader->AddStringAttribute(UInterchangeShaderPortsAPI::MakeInputValueKey(TextureSample::Inputs::Texture.ToString()), TextureNode->GetUniqueID());
+
+				if (!FMath::IsNearlyEqual(FbxTexture->GetScaleU(), 1.0) || !FMath::IsNearlyEqual(FbxTexture->GetScaleV(), 1.0))
+				{
+					UInterchangeShaderNode* TextureCoordinateShader = UInterchangeShaderNode::Create(&NodeContainer, TextureName + TEXT("_Coordinate"), ShaderUniqueID);
+					TextureCoordinateShader->SetCustomShaderType(TextureCoordinate::Name.ToString());
+
+					TextureCoordinateShader->AddFloatAttribute(UInterchangeShaderPortsAPI::MakeInputValueKey(TextureCoordinate::Inputs::UTiling.ToString()), (float)FbxTexture->GetScaleU());
+					TextureCoordinateShader->AddFloatAttribute(UInterchangeShaderPortsAPI::MakeInputValueKey(TextureCoordinate::Inputs::VTiling.ToString()), (float)FbxTexture->GetScaleV());
+
+					UInterchangeShaderPortsAPI::ConnectDefaultOuputToInput(TextureSampleShader, TextureSample::Inputs::Coordinates.ToString(), TextureCoordinateShader->GetUniqueID());
+				}
+
+				return TextureSampleShader;
+			}
+
+			bool FFbxMaterial::ConvertPropertyToShaderNode(UInterchangeBaseNodeContainer& NodeContainer, UInterchangeShaderGraphNode* ShaderGraphNode, FbxProperty& Property, float Factor, FName InputName,
 				TVariant<FLinearColor, float> DefaultValue, bool bInverse)
 			{
 				using namespace Materials::Standard::Nodes;
@@ -99,41 +148,17 @@ namespace UE
 				if (TextureCount > 0)
 				{
 					FbxFileTexture* FbxTexture = Property.GetSrcObject<FbxFileTexture>(0);
-					FString TextureFilename = UTF8_TO_TCHAR(FbxTexture->GetFileName());
-					//Only import texture that exist on disk
-					if(!TextureFilename.IsEmpty() && FPaths::FileExists(TextureFilename))
+					if (UInterchangeShaderNode* TextureSampleShader = CreateTextureSampler(FbxTexture, NodeContainer, ShaderGraphNode->GetUniqueID()))
 					{
-						FString TextureName = FPaths::GetBaseFilename(TextureFilename);
-						UInterchangeShaderNode* TextureSampleShader = UInterchangeShaderNode::Create(&NodeContainer, TextureName, ShaderGraphNode->GetUniqueID());
-						TextureSampleShader->SetCustomShaderType(TextureSample::Name.ToString());
-
-						const UInterchangeTexture2DNode* TextureNode = Cast<const UInterchangeTexture2DNode>(NodeContainer.GetNode(UInterchangeTextureNode::MakeNodeUid(TextureName)));
-
-						if (!TextureNode)
-						{
-							TextureNode = CreateTexture2DNode(NodeContainer, TextureFilename);
-						}
-
-						TextureSampleShader->AddStringAttribute(UInterchangeShaderPortsAPI::MakeInputValueKey(TextureSample::Inputs::Texture.ToString()), TextureNode->GetUniqueID());
-
-						if (!FMath::IsNearlyEqual(FbxTexture->GetScaleU(), 1.0) || !FMath::IsNearlyEqual(FbxTexture->GetScaleV(), 1.0))
-						{
-							UInterchangeShaderNode* TextureCoordinateShader = UInterchangeShaderNode::Create(&NodeContainer, TextureName + TEXT("_Coordinate"), ShaderGraphNode->GetUniqueID());
-							TextureCoordinateShader->SetCustomShaderType(TextureCoordinate::Name.ToString());
-
-							TextureCoordinateShader->AddFloatAttribute(UInterchangeShaderPortsAPI::MakeInputValueKey(TextureCoordinate::Inputs::UTiling.ToString()), (float)FbxTexture->GetScaleU());
-							TextureCoordinateShader->AddFloatAttribute(UInterchangeShaderPortsAPI::MakeInputValueKey(TextureCoordinate::Inputs::VTiling.ToString()), (float)FbxTexture->GetScaleV());
-
-							UInterchangeShaderPortsAPI::ConnectDefaultOuputToInput(TextureSampleShader, TextureSample::Inputs::Coordinates.ToString(), TextureCoordinateShader->GetUniqueID());
-						}
-
 						UInterchangeShaderPortsAPI::ConnectDefaultOuputToInput(NodeToConnectTo, InputToConnectTo, TextureSampleShader->GetUniqueID());
 					}
 					else
 					{
 						UInterchangeResultTextureWarning_TextureFileDoNotExist* Message = Parser.AddMessage<UInterchangeResultTextureWarning_TextureFileDoNotExist>();
-						Message->TextureName = TextureFilename;
+						Message->TextureName = FbxTexture ? UTF8_TO_TCHAR(FbxTexture->GetFileName()) : TEXT("Undefined");
 						Message->MaterialName = ShaderGraphNode->GetDisplayLabel();
+
+						return false;
 					}
 				}
 				else if (DataType.GetType() == eFbxDouble || DataType.GetType() == eFbxFloat || DataType.GetType() == eFbxInt)
@@ -154,12 +179,58 @@ namespace UE
 						NodeToConnectTo->AddFloatAttribute(UInterchangeShaderPortsAPI::MakeInputValueKey(InputToConnectTo), PropertyValue.R);
 					}
 				}
+
+				return true;
 			}
 
+			void FFbxMaterial::ConvertShininessToShaderNode(FbxSurfaceMaterial& SurfaceMaterial, UInterchangeBaseNodeContainer& NodeContainer, UInterchangeShaderGraphNode* ShaderGraphNode)
+			{
+				using namespace UE::Interchange::Materials;
+				using namespace UE::Interchange::Materials::Standard::Nodes;
+
+				FBXSDK_NAMESPACE::FbxProperty MaterialProperty = SurfaceMaterial.FindProperty(FbxSurfaceMaterial::sShininess);
+
+				if (!MaterialProperty.IsValid())
+				{
+					return;
+				}
+
+				const FString InputName = Phong::Parameters::Shininess.ToString();
+
+				if (MaterialProperty.GetSrcObjectCount<FBXSDK_NAMESPACE::FbxTexture>() > 0)
+				{
+					FbxFileTexture* FbxTexture = MaterialProperty.GetSrcObject<FbxFileTexture>(0);
+					if (UInterchangeShaderNode* TextureSampleShader = CreateTextureSampler(FbxTexture, NodeContainer, ShaderGraphNode->GetUniqueID()))
+					{
+						FString MultiplyNodeName = Phong::Parameters::Shininess.ToString() + TEXT("_Multiply");
+						UInterchangeShaderNode* MultiplyNode = UInterchangeShaderNode::Create(&NodeContainer, MultiplyNodeName, ShaderGraphNode->GetUniqueID());
+						MultiplyNode->SetCustomShaderType(Multiply::Name.ToString());
+						
+						UInterchangeShaderPortsAPI::ConnectDefaultOuputToInput(ShaderGraphNode, InputName, MultiplyNode->GetUniqueID());
+
+						// Scale texture output from [0-1] to [0-1000]
+						MultiplyNode->AddFloatAttribute(UInterchangeShaderPortsAPI::MakeInputValueKey(Multiply::Inputs::B.ToString()), 1000.0f);
+
+						UInterchangeShaderPortsAPI::ConnectDefaultOuputToInput(MultiplyNode, Multiply::Inputs::A.ToString(), TextureSampleShader->GetUniqueID());
+					}
+				}
+				else if (!FBXSDK_NAMESPACE::FbxProperty::HasDefaultValue(MaterialProperty))
+				{
+					const FString InputValueKey = UInterchangeShaderPortsAPI::MakeInputValueKey(InputName);
+					ShaderGraphNode->AddFloatAttribute(InputValueKey, MaterialProperty.Get<float>());
+				}
+			}
+
+
 			const UInterchangeShaderGraphNode* FFbxMaterial::AddShaderGraphNode(FbxSurfaceMaterial* SurfaceMaterial, UInterchangeBaseNodeContainer& NodeContainer)
 			{
 				using namespace UE::Interchange::Materials;
 
+				if (!SurfaceMaterial)
+				{
+					return nullptr;
+				}
+
 				//Create a material node
 				FString MaterialName = Parser.GetFbxHelper()->GetFbxObjectName(SurfaceMaterial);
 				FString NodeUid = TEXT("\\Material\\") + MaterialName;
@@ -182,78 +253,123 @@ namespace UE
 					return nullptr;
 				}
 
-				// for each material property, add an input value or a shader node
-				FbxProperty CurrentProperty = SurfaceMaterial->GetFirstProperty();
-				while (CurrentProperty.IsValid())
+				// If not a Lambert nor a Phong material, match legacy FBX import
+				// Use random color because there may be multiple materials, then they can be different
+				if (!SurfaceMaterial->Is<FbxSurfaceLambert>())
 				{
-					if (CurrentProperty.Modified())
+					FLinearColor BaseColor;
+					BaseColor.R = 0.5f + 0.5f * FMath::FRand();
+					BaseColor.G = 0.5f + 0.5f * FMath::FRand();
+					BaseColor.B = 0.5f + 0.5f * FMath::FRand();
+
+					const FString InputValueKey = UInterchangeShaderPortsAPI::MakeInputValueKey(Phong::Parameters::DiffuseColor.ToString());
+					ShaderGraphNode->AddLinearColorAttribute(InputValueKey, BaseColor);
+
+					return ShaderGraphNode;
+				}
+
+				TFunction<bool(FBXSDK_NAMESPACE::FbxProperty&)> ShouldConvertProperty = [&](FBXSDK_NAMESPACE::FbxProperty& MaterialProperty) -> bool
+				{
+					bool bShouldConvertProperty = false;
+
+					if (MaterialProperty.IsValid())
+					{
+						// FbxProperty::HasDefaultValue(..) can return true while the property has textures attached to it.
+						bShouldConvertProperty = MaterialProperty.GetSrcObjectCount<FBXSDK_NAMESPACE::FbxTexture>() > 0
+							|| !FBXSDK_NAMESPACE::FbxProperty::HasDefaultValue(MaterialProperty);
+					}
+
+					return bShouldConvertProperty;
+				};
+
+				TFunction<bool(FName, const char*, float , TVariant<FLinearColor, float>&, bool)>  ConnectInput;
+				ConnectInput = [&](FName InputName, const char* FbxPropertyName, float Factor, TVariant<FLinearColor, float>& DefaultValue, bool bInverse) -> bool
+				{
+					FBXSDK_NAMESPACE::FbxProperty MaterialProperty = SurfaceMaterial->FindProperty(FbxPropertyName);
+					if (ShouldConvertProperty(MaterialProperty))
 					{
-						const FName PropertyName(CurrentProperty.GetNameAsCStr());
+						return ConvertPropertyToShaderNode(NodeContainer, ShaderGraphNode, MaterialProperty, Factor, InputName, DefaultValue, bInverse);
+					}
+
+					return false;
+				};
 
-						// Convert FBX property names to standard Interchange Phong parameter names when possible
+				// Diffuse
+				{
+					const float Factor = (float)((FbxSurfaceLambert*)SurfaceMaterial)->DiffuseFactor.Get();
+					TVariant<FLinearColor, float> DefaultValue;
+					DefaultValue.Set<FLinearColor>(FLinearColor::Black);
+					if (!ConnectInput(Phong::Parameters::DiffuseColor, FbxSurfaceMaterial::sDiffuse, Factor, DefaultValue, false))
+					{
+						// Nothing to retrieve from the Diffuse property, get Diffuse directly from the class
+						const FbxDouble3 DiffuseColor = ((FbxSurfaceLambert*)SurfaceMaterial)->Diffuse.Get();
+						const FString InputValueKey = UInterchangeShaderPortsAPI::MakeInputValueKey(Phong::Parameters::DiffuseColor.ToString());
+						ShaderGraphNode->AddLinearColorAttribute(InputValueKey, FFbxConvert::ConvertColor(DiffuseColor));
+					}
+				}
+
+				// Ambient
+				{
+					const float Factor = (float)((FbxSurfacePhong*)SurfaceMaterial)->AmbientFactor.Get();
+					TVariant<FLinearColor, float> DefaultValue;
+					DefaultValue.Set<FLinearColor>(FLinearColor::Black);
+					ConnectInput(Phong::Parameters::AmbientColor, FbxSurfaceMaterial::sAmbient, Factor, DefaultValue, false);
+				}
+
+				// Emissive
+				{
+					const float Factor = (float)((FbxSurfaceLambert*)SurfaceMaterial)->EmissiveFactor.Get();
+					TVariant<FLinearColor, float> DefaultValue;
+					DefaultValue.Set<FLinearColor>(FLinearColor::Black);
+					ConnectInput(Phong::Parameters::EmissiveColor, FbxSurfaceMaterial::sEmissive, Factor, DefaultValue, false);
+				}
+
+				// Normal
+				{
+					// FbxSurfaceMaterial can have either a normal map or a bump map, check for both
+					float Factor = 1.f;
+					const char* PropertyName = FbxSurfaceMaterial::sNormalMap;
+					FBXSDK_NAMESPACE::FbxProperty MaterialProperty = SurfaceMaterial->FindProperty(PropertyName);
+					if (!MaterialProperty.IsValid() || MaterialProperty.GetSrcObjectCount<FbxTexture>() < 1)
+					{
+						PropertyName = FbxSurfaceMaterial::sBump;
+						Factor = (float)((FbxSurfaceLambert*)SurfaceMaterial)->BumpFactor.Get();
+						MaterialProperty = SurfaceMaterial->FindProperty(PropertyName);
+						if (!MaterialProperty.IsValid() || MaterialProperty.GetSrcObjectCount<FbxTexture>() < 1)
 						{
-							if (PropertyName == FbxSurfaceMaterial::sDiffuse && SurfaceMaterial->Is<FbxSurfaceLambert>())
-							{
-								const float Factor = (float)((FbxSurfaceLambert*)SurfaceMaterial)->DiffuseFactor.Get();
-								TVariant<FLinearColor, float> DefaultValue;
-								DefaultValue.Set<FLinearColor>(FLinearColor::Black);
-								ConvertPropertyToShaderNode(NodeContainer, ShaderGraphNode, CurrentProperty, Factor, PBR::Parameters::BaseColor, DefaultValue);
-							}
-							else if (PropertyName == FbxSurfaceMaterial::sSpecular && SurfaceMaterial->Is<FbxSurfacePhong>())
-							{
-								const float Factor = (float)((FbxSurfacePhong*)SurfaceMaterial)->SpecularFactor.Get();
-								TVariant<FLinearColor, float> DefaultValue;
-								DefaultValue.Set<FLinearColor>(FLinearColor::Black);
-								ConvertPropertyToShaderNode(NodeContainer, ShaderGraphNode, CurrentProperty, Factor, Phong::Parameters::SpecularColor, DefaultValue);
-							}
-							else if (PropertyName == FbxSurfaceMaterial::sShininess)
-							{
-								const float Factor = 1.f;
-								TVariant<FLinearColor, float> DefaultValue;
-								DefaultValue.Set<float>(0.2f);
-								ConvertPropertyToShaderNode(NodeContainer, ShaderGraphNode, CurrentProperty, Factor, Phong::Parameters::Shininess, DefaultValue);
-							}
-							else if (PropertyName == FbxSurfaceMaterial::sEmissive && SurfaceMaterial->Is<FbxSurfaceLambert>())
-							{
-								const float Factor = (float)((FbxSurfaceLambert*)SurfaceMaterial)->EmissiveFactor.Get();
-								TVariant<FLinearColor, float> DefaultValue;
-								DefaultValue.Set<FLinearColor>(FLinearColor::Black);
-								ConvertPropertyToShaderNode(NodeContainer, ShaderGraphNode, CurrentProperty, Factor, Phong::Parameters::EmissiveColor, DefaultValue);
-							}
-							else if (PropertyName == FbxSurfaceMaterial::sNormalMap)
-							{
-								const float Factor = 1.f;
-								TVariant<FLinearColor, float> DefaultValue;
-								DefaultValue.Set<FLinearColor>(FLinearColor(FVector::UpVector));
-								ConvertPropertyToShaderNode(NodeContainer, ShaderGraphNode, CurrentProperty, Factor, Phong::Parameters::Normal, DefaultValue);
-							}
-							else if (PropertyName == FbxSurfaceMaterial::sTransparentColor && SurfaceMaterial->Is<FbxSurfaceLambert>())
-							{
-								const float TransparencyFactor = (float)((FbxSurfaceLambert*)SurfaceMaterial)->TransparencyFactor.Get();
-								TVariant<FLinearColor, float> DefaultValue;
-								DefaultValue.Set<float>(0.f); // Opaque
-								const bool bInverse = true;
-								ConvertPropertyToShaderNode(NodeContainer, ShaderGraphNode, CurrentProperty, TransparencyFactor, Phong::Parameters::Opacity, DefaultValue, bInverse);
-							}
-							else if (PropertyName == FbxSurfaceMaterial::sDiffuseFactor || PropertyName == FbxSurfaceMaterial::sSpecularFactor ||
-										PropertyName == FbxSurfaceMaterial::sEmissiveFactor || PropertyName == FbxSurfaceMaterial::sTransparencyFactor)
-							{
-								// Skip factors as they were processed with their corresponding inputs
-							}
-							else
-							{
-								if (!UInterchangeShaderPortsAPI::HasInput(ShaderGraphNode, PropertyName))
-								{
-									const float DummyFactor = 1.f;
-									TVariant<FLinearColor, float> DummyDefaultValue;
-									DummyDefaultValue.Set<float>(1.f); // Will be ignored because we have a factor of 1
-									ConvertPropertyToShaderNode(NodeContainer, ShaderGraphNode, CurrentProperty, DummyFactor, PropertyName, DummyDefaultValue);
-								}
-							}
+							PropertyName = nullptr;
 						}
 					}
 
-					CurrentProperty = SurfaceMaterial->GetNextProperty(CurrentProperty);
+					if (PropertyName)
+					{
+						TVariant<FLinearColor, float> DefaultValue;
+						DefaultValue.Set<FLinearColor>(FLinearColor(FVector::UpVector));
+						ConvertPropertyToShaderNode(NodeContainer, ShaderGraphNode, MaterialProperty, Factor, Phong::Parameters::Normal, DefaultValue);
+					}
+				}
+
+				// Opacity
+				{
+					const float TransparencyFactor = (float)((FbxSurfaceLambert*)SurfaceMaterial)->TransparencyFactor.Get();
+					TVariant<FLinearColor, float> DefaultValue;
+					DefaultValue.Set<float>(0.f); // Opaque
+					ConnectInput(Phong::Parameters::Opacity, FbxSurfaceMaterial::sTransparentColor, TransparencyFactor, DefaultValue, true);
+				}
+
+				// Get specular color and shininess if FBX material is a Phong one
+				if (SurfaceMaterial->Is<FbxSurfacePhong>())
+				{
+					// Specular
+					{
+						const float Factor = (float)((FbxSurfacePhong*)SurfaceMaterial)->SpecularFactor.Get();
+						TVariant<FLinearColor, float> DefaultValue;
+						DefaultValue.Set<FLinearColor>(FLinearColor::Black);
+						ConnectInput(Phong::Parameters::SpecularColor, FbxSurfaceMaterial::sSpecular, Factor, DefaultValue, false);
+					}
+
+					// Shininess
+					ConvertShininessToShaderNode(*SurfaceMaterial, NodeContainer, ShaderGraphNode);
 				}
 
 				return ShaderGraphNode;
diff --git a/Engine/Plugins/Interchange/Runtime/Source/Parsers/Fbx/Private/FbxMaterial.h b/Engine/Plugins/Interchange/Runtime/Source/Parsers/Fbx/Private/FbxMaterial.h
--- a/Engine/Plugins/Interchange/Runtime/Source/Parsers/Fbx/Private/FbxMaterial.h
+++ b/Engine/Plugins/Interchange/Runtime/Source/Parsers/Fbx/Private/FbxMaterial.h
@@ -54,11 +54,13 @@ namespace UE
 			protected:
 				UInterchangeShaderGraphNode* CreateShaderGraphNode(UInterchangeBaseNodeContainer& NodeContainer, const FString& NodeUID, const FString& NodeName);
 				UInterchangeTexture2DNode* CreateTexture2DNode(UInterchangeBaseNodeContainer& NodeContainer, const FString& TextureFilePath);
-				void ConvertPropertyToShaderNode(UInterchangeBaseNodeContainer& NodeContainer, UInterchangeShaderGraphNode* ShaderGraphNode,
+				bool ConvertPropertyToShaderNode(UInterchangeBaseNodeContainer& NodeContainer, UInterchangeShaderGraphNode* ShaderGraphNode,
 					FbxProperty& Property, float Factor, FName PropertyName, TVariant<FLinearColor, float> DefaultValue, bool bInverse = false);
 
 			private:
 				const UInterchangeShaderGraphNode* AddShaderGraphNode(FbxSurfaceMaterial* SurfaceMaterial, UInterchangeBaseNodeContainer& NodeContainer);
+				void ConvertShininessToShaderNode(FbxSurfaceMaterial& SurfaceMaterial, UInterchangeBaseNodeContainer& NodeContainer, UInterchangeShaderGraphNode* ShaderGraphNode);
+				UInterchangeShaderNode* CreateTextureSampler(FbxFileTexture* FbxTexture, UInterchangeBaseNodeContainer& NodeContainer, const FString& ShaderUniqueID);
 
 				FFbxParser& Parser;
 			};
diff --git a/Engine/Plugins/Interchange/Runtime/Source/Pipelines/Private/InterchangeGenericMaterialPipeline.cpp b/Engine/Plugins/Interchange/Runtime/Source/Pipelines/Private/InterchangeGenericMaterialPipeline.cpp
--- a/Engine/Plugins/Interchange/Runtime/Source/Pipelines/Private/InterchangeGenericMaterialPipeline.cpp
+++ b/Engine/Plugins/Interchange/Runtime/Source/Pipelines/Private/InterchangeGenericMaterialPipeline.cpp
@@ -118,6 +118,7 @@ namespace UE::Interchange::InterchangeGenericMaterialPipeline::Private
 			TEXT("MaterialFunction'/Engine/Functions/Engine_MaterialFunctions01/Texturing/FlattenNormal.FlattenNormal'"),
 			TEXT("MaterialFunction'/Engine/Functions/Engine_MaterialFunctions02/Utility/MakeFloat3.MakeFloat3'"),
 			TEXT("MaterialFunction'/Engine/Functions/Engine_MaterialFunctions02/Texturing/CustomRotator.CustomRotator'"),
+			TEXT("MaterialFunction'/Interchange/Functions/MF_PhongToMetalRoughness.MF_PhongToMetalRoughness'"),
 		};
 
 		static const bool bRequiredPackagesLoaded = ArePackagesLoaded(RequiredPackages);
@@ -530,7 +531,7 @@ UInterchangeBaseMaterialFactoryNode* UInterchangeGenericMaterialPipeline::Create
 	return MaterialFactoryNode;
 }
 
-bool UInterchangeGenericMaterialPipeline::IsClearCoatModel(const UInterchangeShaderGraphNode* ShaderGraphNode) const
+bool UInterchangeGenericMaterialPipeline::HasClearCoat(const UInterchangeShaderGraphNode* ShaderGraphNode) const
 {
 	using namespace UE::Interchange::Materials::ClearCoat;
 
@@ -539,7 +540,7 @@ bool UInterchangeGenericMaterialPipeline::IsClearCoatModel(const UInterchangeSha
 	return bHasClearCoatInput;
 }
 
-bool UInterchangeGenericMaterialPipeline::IsSheenModel(const UInterchangeShaderGraphNode* ShaderGraphNode) const
+bool UInterchangeGenericMaterialPipeline::HasSheen(const UInterchangeShaderGraphNode* ShaderGraphNode) const
 {
 	using namespace UE::Interchange::Materials::Sheen;
 
@@ -548,7 +549,7 @@ bool UInterchangeGenericMaterialPipeline::IsSheenModel(const UInterchangeShaderG
 	return bHasSheenColorInput;
 }
 
-bool UInterchangeGenericMaterialPipeline::IsSubsurfaceModel(const UInterchangeShaderGraphNode* ShaderGraphNode) const
+bool UInterchangeGenericMaterialPipeline::HasSubsurface(const UInterchangeShaderGraphNode* ShaderGraphNode) const
 {
 	using namespace UE::Interchange::Materials::Subsurface;
 
@@ -557,7 +558,7 @@ bool UInterchangeGenericMaterialPipeline::IsSubsurfaceModel(const UInterchangeSh
 	return bHasSubsurfaceColorInput;
 }
 
-bool UInterchangeGenericMaterialPipeline::IsThinTranslucentModel(const UInterchangeShaderGraphNode* ShaderGraphNode) const
+bool UInterchangeGenericMaterialPipeline::HasThinTranslucency(const UInterchangeShaderGraphNode* ShaderGraphNode) const
 {
 	using namespace UE::Interchange::Materials::ThinTranslucent;
 
@@ -566,15 +567,25 @@ bool UInterchangeGenericMaterialPipeline::IsThinTranslucentModel(const UIntercha
 	return bHasTransmissionColorInput;
 }
 
-bool UInterchangeGenericMaterialPipeline::IsPBRModel(const UInterchangeShaderGraphNode* ShaderGraphNode) const
+bool UInterchangeGenericMaterialPipeline::IsMetalRoughModel(const UInterchangeShaderGraphNode* ShaderGraphNode) const
 {
-	using namespace UE::Interchange::Materials::PBR;
+	using namespace UE::Interchange::Materials::PBRMR;
 
 	const bool bHasBaseColorInput = UInterchangeShaderPortsAPI::HasInput(ShaderGraphNode, Parameters::BaseColor);
 
 	return bHasBaseColorInput;
 }
 
+bool UInterchangeGenericMaterialPipeline::IsSpecGlossModel(const UInterchangeShaderGraphNode* ShaderGraphNode) const
+{
+	using namespace UE::Interchange::Materials::PBRSG;
+
+	const bool bHasSpecularInput = UInterchangeShaderPortsAPI::HasInput(ShaderGraphNode, Parameters::SpecularColor);
+	const bool bHasGlossinessInput = UInterchangeShaderPortsAPI::HasInput(ShaderGraphNode, Parameters::Glossiness);
+
+	return bHasSpecularInput && bHasGlossinessInput;
+}
+
 bool UInterchangeGenericMaterialPipeline::IsPhongModel(const UInterchangeShaderGraphNode* ShaderGraphNode) const
 {
 	using namespace UE::Interchange::Materials::Phong;
@@ -622,30 +633,31 @@ bool UInterchangeGenericMaterialPipeline::IsSurfaceUnlitModel(const UInterchange
 	return false;
 }
 
-bool UInterchangeGenericMaterialPipeline::HandlePhongModel(const UInterchangeShaderGraphNode* ShaderGraphNode, UInterchangeMaterialFactoryNode* MaterialFactoryNode)
+
+bool UInterchangeGenericMaterialPipeline::HandleSpecGlossModel(const UInterchangeShaderGraphNode* ShaderGraphNode, UInterchangeMaterialFactoryNode* MaterialFactoryNode)
 {
-	using namespace UE::Interchange::Materials::Phong;
+	using namespace UE::Interchange::Materials;
 	using namespace UE::Interchange::InterchangeGenericMaterialPipeline::Private;
 
-	if (IsPhongModel(ShaderGraphNode))
+	if (IsSpecGlossModel(ShaderGraphNode))
 	{
 		// ConvertFromDiffSpec function call
 		const FString MaterialFunctionPath = TEXT("MaterialFunction'/Engine/Functions/Engine_MaterialFunctions01/Shading/ConvertFromDiffSpec.ConvertFromDiffSpec'");
 		UInterchangeMaterialExpressionFactoryNode* FunctionCallExpression = CreateExpressionWithMaterialFunction(BaseNodeContainer, MaterialFactoryNode, TEXT("DiffSpecFunc"), MaterialFunctionPath);
 
 		const FString FunctionCallExpressionUid = FunctionCallExpression->GetUniqueID();
-		MaterialFactoryNode->ConnectOutputToBaseColor(FunctionCallExpressionUid, UE::Interchange::Materials::PBR::Parameters::BaseColor.ToString());
-		MaterialFactoryNode->ConnectOutputToMetallic(FunctionCallExpressionUid, UE::Interchange::Materials::PBR::Parameters::Metallic.ToString());
-		MaterialFactoryNode->ConnectOutputToSpecular(FunctionCallExpressionUid, UE::Interchange::Materials::PBR::Parameters::Specular.ToString());
-		
-		// Diffuse
+		MaterialFactoryNode->ConnectOutputToBaseColor(FunctionCallExpressionUid, PBRMR::Parameters::BaseColor.ToString());
+		MaterialFactoryNode->ConnectOutputToMetallic(FunctionCallExpressionUid, PBRMR::Parameters::Metallic.ToString());
+		MaterialFactoryNode->ConnectOutputToSpecular(FunctionCallExpressionUid, PBRMR::Parameters::Specular.ToString());
+
+		// DiffuseColor
 		{
 			TTuple<UInterchangeMaterialExpressionFactoryNode*, FString> DiffuseExpressionFactoryNode =
-				CreateMaterialExpressionForInput(MaterialFactoryNode, ShaderGraphNode, Parameters::DiffuseColor.ToString(), FunctionCallExpression->GetUniqueID());
+				CreateMaterialExpressionForInput(MaterialFactoryNode, ShaderGraphNode, PBRSG::Parameters::DiffuseColor.ToString(), FunctionCallExpressionUid);
 
 			if (DiffuseExpressionFactoryNode.Get<0>())
 			{
-				UInterchangeShaderPortsAPI::ConnectOuputToInputByName(FunctionCallExpression, Parameters::DiffuseColor.ToString(),
+				UInterchangeShaderPortsAPI::ConnectOuputToInputByName(FunctionCallExpression, PBRSG::Parameters::DiffuseColor.ToString(),
 					DiffuseExpressionFactoryNode.Get<0>()->GetUniqueID(), DiffuseExpressionFactoryNode.Get<1>());
 			}
 		}
@@ -653,49 +665,76 @@ bool UInterchangeGenericMaterialPipeline::HandlePhongModel(const UInterchangeSha
 		// Specular Color
 		{
 			TTuple<UInterchangeMaterialExpressionFactoryNode*, FString> SpecularExpressionFactoryNode =
-				CreateMaterialExpressionForInput(MaterialFactoryNode, ShaderGraphNode, Parameters::SpecularColor.ToString(), FunctionCallExpression->GetUniqueID());
+				CreateMaterialExpressionForInput(MaterialFactoryNode, ShaderGraphNode, PBRSG::Parameters::SpecularColor.ToString(), FunctionCallExpressionUid);
 
 			if (SpecularExpressionFactoryNode.Get<0>())
 			{
-				UInterchangeShaderPortsAPI::ConnectOuputToInputByName(FunctionCallExpression, Parameters::SpecularColor.ToString(),
+				UInterchangeShaderPortsAPI::ConnectOuputToInputByName(FunctionCallExpression, PBRSG::Parameters::SpecularColor.ToString(),
 					SpecularExpressionFactoryNode.Get<0>()->GetUniqueID(), SpecularExpressionFactoryNode.Get<1>());
 			}
 		}
-		
-		// Shininess
+
+		// Glossiness
 		{
-			const bool bHasShininessInput = UInterchangeShaderPortsAPI::HasInput(ShaderGraphNode, Parameters::Shininess);
-			if (bHasShininessInput)
+			TTuple<UInterchangeMaterialExpressionFactoryNode*, FString> GlossinessExpressionFactoryNode =
+				CreateMaterialExpressionForInput(MaterialFactoryNode, ShaderGraphNode, PBRSG::Parameters::Glossiness.ToString(), MaterialFactoryNode->GetUniqueID());
+
+			if (UInterchangeMaterialExpressionFactoryNode* GlossinessFactoryNode = GlossinessExpressionFactoryNode.Get<0>())
 			{
-				TGuardValue<EMaterialInputType> InputTypeBeingProcessedGuard(MaterialCreationContext.InputTypeBeingProcessed, EMaterialInputType::Scalar);
+				UInterchangeMaterialExpressionFactoryNode* OneMinusNode =
+					CreateExpressionNode(TEXT("InverseGlossiness"), MaterialFactoryNode->GetUniqueID(), UMaterialExpressionOneMinus::StaticClass());
 
-				TTuple<UInterchangeMaterialExpressionFactoryNode*, FString> ShininessExpressionFactoryNode =
-					CreateMaterialExpressionForInput(MaterialFactoryNode, ShaderGraphNode, Parameters::Shininess.ToString(), MaterialFactoryNode->GetUniqueID());
+				const FString OneMinusNodeInput = GET_MEMBER_NAME_CHECKED(UMaterialExpressionOneMinus, Input).ToString();
+				const FString OutputName = GlossinessExpressionFactoryNode.Get<1>();
+				UInterchangeShaderPortsAPI::ConnectOuputToInputByName(OneMinusNode, OneMinusNodeInput, GlossinessFactoryNode->GetUniqueID(), OutputName);
 
-				if (ShininessExpressionFactoryNode.Get<0>())
-				{
-					UInterchangeMaterialExpressionFactoryNode* DivideShininessNode =
-						CreateExpressionNode(TEXT("DivideShininess"), ShininessExpressionFactoryNode.Get<0>()->GetUniqueID(), UMaterialExpressionDivide::StaticClass());
+				MaterialFactoryNode->ConnectOutputToRoughness(OneMinusNode->GetUniqueID(), PBRMR::Parameters::Roughness.ToString());
+			}
+		}
 
-					const float ShininessScale = 100.f; // Divide shininess by 100 to bring it into a 0-1 range for roughness.
-					const FName ShininessScaleParameterName = GET_MEMBER_NAME_CHECKED(UMaterialExpressionDivide, ConstB);
-					DivideShininessNode->AddFloatAttribute(ShininessScaleParameterName.ToString(), ShininessScale);
-					DivideShininessNode->AddApplyAndFillDelegates<float>(ShininessScaleParameterName.ToString(), UMaterialExpressionDivide::StaticClass(), GET_MEMBER_NAME_CHECKED(UMaterialExpressionDivide, ConstB));
+		return true;
+	}
 
-					// Connect Shininess to Divide
-					UInterchangeShaderPortsAPI::ConnectOuputToInputByName(DivideShininessNode, GET_MEMBER_NAME_CHECKED(UMaterialExpressionDivide, A).ToString(),
-						ShininessExpressionFactoryNode.Get<0>()->GetUniqueID(), ShininessExpressionFactoryNode.Get<1>());
+	return false;
+}
 
-					UInterchangeMaterialExpressionFactoryNode* InverseShininessNode =
-						CreateExpressionNode(TEXT("InverseShininess"), ShininessExpressionFactoryNode.Get<0>()->GetUniqueID(), UMaterialExpressionOneMinus::StaticClass());
+bool UInterchangeGenericMaterialPipeline::HandlePhongModel(const UInterchangeShaderGraphNode* ShaderGraphNode, UInterchangeMaterialFactoryNode* MaterialFactoryNode)
+{
+	using namespace UE::Interchange::Materials::Phong;
+	using namespace UE::Interchange::InterchangeGenericMaterialPipeline::Private;
 
-					// Connect Divide to Inverse
-					UInterchangeShaderPortsAPI::ConnectDefaultOuputToInput(InverseShininessNode, GET_MEMBER_NAME_CHECKED(UMaterialExpressionOneMinus, Input).ToString(),
-						DivideShininessNode->GetUniqueID());
+	if (IsPhongModel(ShaderGraphNode))
+	{
+		// ConvertFromDiffSpec function call
+		const FString MaterialFunctionPath = TEXT("MaterialFunction'/Interchange/Functions/MF_PhongToMetalRoughness.MF_PhongToMetalRoughness'");
+		UInterchangeMaterialExpressionFactoryNode* FunctionCallExpression = CreateExpressionWithMaterialFunction(BaseNodeContainer, MaterialFactoryNode, TEXT("DiffSpecFunc"), MaterialFunctionPath);
 
-					MaterialFactoryNode->ConnectToRoughness(InverseShininessNode->GetUniqueID());
+		const FString FunctionCallExpressionUid = FunctionCallExpression->GetUniqueID();
+		MaterialFactoryNode->ConnectOutputToBaseColor(FunctionCallExpressionUid, UE::Interchange::Materials::PBRMR::Parameters::BaseColor.ToString());
+		MaterialFactoryNode->ConnectOutputToMetallic(FunctionCallExpressionUid, UE::Interchange::Materials::PBRMR::Parameters::Metallic.ToString());
+		MaterialFactoryNode->ConnectOutputToSpecular(FunctionCallExpressionUid, UE::Interchange::Materials::PBRMR::Parameters::Specular.ToString());
+		MaterialFactoryNode->ConnectOutputToRoughness(FunctionCallExpressionUid, UE::Interchange::Materials::PBRMR::Parameters::Roughness.ToString());
+
+		{
+			const FString UniqueID = FunctionCallExpression->GetUniqueID();
+
+			TFunction<void(const FString&)> ConnectInput;
+			ConnectInput = [&](const FString& InputName) -> void
+			{
+				TTuple<UInterchangeMaterialExpressionFactoryNode*, FString> ExpressionFactoryNode =
+					this->CreateMaterialExpressionForInput(MaterialFactoryNode, ShaderGraphNode, InputName, UniqueID);
+
+				if (ExpressionFactoryNode.Get<0>())
+				{
+					UInterchangeShaderPortsAPI::ConnectOuputToInputByName(FunctionCallExpression, InputName,
+						ExpressionFactoryNode.Get<0>()->GetUniqueID(), ExpressionFactoryNode.Get<1>());
 				}
-			}
+			};
+
+			ConnectInput(Parameters::AmbientColor.ToString());
+			ConnectInput(Parameters::DiffuseColor.ToString());
+			ConnectInput(Parameters::Shininess.ToString());
+			ConnectInput(Parameters::SpecularColor.ToString());
 		}
 
 		return true;
@@ -724,9 +763,9 @@ bool UInterchangeGenericMaterialPipeline::HandleLambertModel(const UInterchangeS
 	return false;
 }
 
-bool UInterchangeGenericMaterialPipeline::HandlePBRModel(const UInterchangeShaderGraphNode* ShaderGraphNode, UInterchangeMaterialFactoryNode* MaterialFactoryNode)
+bool UInterchangeGenericMaterialPipeline::HandleMetalRoughnessModel(const UInterchangeShaderGraphNode* ShaderGraphNode, UInterchangeMaterialFactoryNode* MaterialFactoryNode)
 {
-	using namespace UE::Interchange::Materials::PBR;
+	using namespace UE::Interchange::Materials::PBRMR;
 
 	bool bShadingModelHandled = false;
 
@@ -808,26 +847,6 @@ bool UInterchangeGenericMaterialPipeline::HandlePBRModel(const UInterchangeShade
 		}
 	}
 
-	// Anisotropy
-	{
-		const bool bHasInput = UInterchangeShaderPortsAPI::HasInput(ShaderGraphNode, Parameters::Anisotropy);
-
-		if(bHasInput)
-		{
-			TGuardValue<EMaterialInputType> InputTypeBeingProcessedGuard(MaterialCreationContext.InputTypeBeingProcessed, EMaterialInputType::Scalar);
-
-			TTuple<UInterchangeMaterialExpressionFactoryNode*, FString> ExpressionFactoryNode =
-				CreateMaterialExpressionForInput(MaterialFactoryNode, ShaderGraphNode, Parameters::Anisotropy.ToString(), MaterialFactoryNode->GetUniqueID());
-
-			if(ExpressionFactoryNode.Get<0>())
-			{
-				MaterialFactoryNode->ConnectOutputToAnisotropy(ExpressionFactoryNode.Get<0>()->GetUniqueID(), ExpressionFactoryNode.Get<1>());
-			}
-
-			bShadingModelHandled = true;
-		}
-	}
-
 	return bShadingModelHandled;
 }
 
@@ -1027,13 +1046,31 @@ void UInterchangeGenericMaterialPipeline::HandleCommonParameters(const UIntercha
 	using namespace UE::Interchange::Materials::Common;
 
 	// Two sidedness (ignored for thin translucency as it looks wrong)
-	if (!IsThinTranslucentModel(ShaderGraphNode))
+	if (!HasThinTranslucency(ShaderGraphNode))
 	{
 		bool bTwoSided = false;
 		ShaderGraphNode->GetCustomTwoSided(bTwoSided);
 		MaterialFactoryNode->SetCustomTwoSided(bTwoSided);
 	}
 
+	// Anisotropy
+	{
+		const bool bHasInput = UInterchangeShaderPortsAPI::HasInput(ShaderGraphNode, Parameters::Anisotropy);
+
+		if (bHasInput)
+		{
+			TGuardValue<EMaterialInputType> InputTypeBeingProcessedGuard(MaterialCreationContext.InputTypeBeingProcessed, EMaterialInputType::Scalar);
+
+			TTuple<UInterchangeMaterialExpressionFactoryNode*, FString> ExpressionFactoryNode =
+				CreateMaterialExpressionForInput(MaterialFactoryNode, ShaderGraphNode, Parameters::Anisotropy.ToString(), MaterialFactoryNode->GetUniqueID());
+
+			if (ExpressionFactoryNode.Get<0>())
+			{
+				MaterialFactoryNode->ConnectOutputToAnisotropy(ExpressionFactoryNode.Get<0>()->GetUniqueID(), ExpressionFactoryNode.Get<1>());
+			}
+		}
+	}
+
 	// Emissive
 	{
 		const bool bHasEmissiveInput = UInterchangeShaderPortsAPI::HasInput(ShaderGraphNode, Parameters::EmissiveColor);
@@ -2306,13 +2343,17 @@ UInterchangeMaterialFactoryNode* UInterchangeGenericMaterialPipeline::CreateMate
 		return MaterialFactoryNode;
 	}
 
-	if (!HandlePhongModel(ShaderGraphNode, MaterialFactoryNode))
+	if (!HandleMetalRoughnessModel(ShaderGraphNode, MaterialFactoryNode))
 	{
-		HandleLambertModel(ShaderGraphNode, MaterialFactoryNode);
+		if (!HandleSpecGlossModel(ShaderGraphNode, MaterialFactoryNode))
+		{
+			if (!HandlePhongModel(ShaderGraphNode, MaterialFactoryNode))
+			{
+				HandleLambertModel(ShaderGraphNode, MaterialFactoryNode);
+			}
+		}
 	}
 
-	HandlePBRModel(ShaderGraphNode, MaterialFactoryNode); // Always process the PBR parameters. If they were already assigned from Phong or Lambert, they will be ignored.
-	
 	// Can't have different shading models
 	// Favor translucency over coats (clear coat, sheen, etc.) since it tends to have a bigger impact visually
 	if (!HandleThinTranslucent(ShaderGraphNode, MaterialFactoryNode))
@@ -2336,29 +2377,46 @@ UInterchangeMaterialInstanceFactoryNode* UInterchangeGenericMaterialPipeline::Cr
 	UInterchangeMaterialInstanceFactoryNode* MaterialInstanceFactoryNode =
 		Cast<UInterchangeMaterialInstanceFactoryNode>( CreateBaseMaterialFactoryNode(ShaderGraphNode, UInterchangeMaterialInstanceFactoryNode::StaticClass()) );
 
+	TFunction<void(const FString&)> ChooseParent = [this, MaterialInstanceFactoryNode, ShaderGraphNode](const FString& Model) -> void
+	{
+		FString ParentRootName;
+
+		if (HasThinTranslucency(ShaderGraphNode))
+		{
+			ParentRootName = TEXT("ThinTranslucentMaterial_");
+		}
+		else if (HasClearCoat(ShaderGraphNode))
+		{
+			ParentRootName = TEXT("ClearCoatMaterial_");
+		}
+		else if (HasSheen(ShaderGraphNode))
+		{
+			ParentRootName = TEXT("SheenMaterial_");
+		}
+		else if (HasSubsurface(ShaderGraphNode))
+		{
+			ParentRootName = TEXT("SubsurfaceMaterial_");
+		}
+		else
+		{
+			ParentRootName = TEXT("PBRSurfaceMaterial_");
+		}
+
+		const FString ParentAssetPath = TEXT("Material'/Interchange/Materials/") + ParentRootName + Model + TEXT(".") + ParentRootName + Model + TEXT("'");
+		MaterialInstanceFactoryNode->SetCustomParent(ParentAssetPath);
+	};
+
 	if (UMaterialInterface* ParentMaterialObj = Cast<UMaterialInterface>(ParentMaterial.TryLoad()))
 	{
 		MaterialInstanceFactoryNode->SetCustomParent(ParentMaterialObj->GetPathName());
 	}
-	else if (IsThinTranslucentModel(ShaderGraphNode))
-	{
-		MaterialInstanceFactoryNode->SetCustomParent(TEXT("Material'/Interchange/Materials/ThinTranslucentMaterial.ThinTranslucentMaterial'"));
-	}
-	else if (IsClearCoatModel(ShaderGraphNode))
+	else if (IsSpecGlossModel(ShaderGraphNode))
 	{
-		MaterialInstanceFactoryNode->SetCustomParent(TEXT("Material'/Interchange/Materials/ClearCoatMaterial.ClearCoatMaterial'"));
+		ChooseParent(TEXT("SG"));
 	}
-	else if (IsSheenModel(ShaderGraphNode))
+	else if (IsMetalRoughModel(ShaderGraphNode))
 	{
-		MaterialInstanceFactoryNode->SetCustomParent(TEXT("Material'/Interchange/Materials/SheenMaterial.SheenMaterial'"));
-	}
-	else if(IsSubsurfaceModel(ShaderGraphNode))
-	{
-		MaterialInstanceFactoryNode->SetCustomParent(TEXT("Material'/Interchange/Materials/SubsurfaceMaterial.SubsurfaceMaterial'"));
-	}
-	else if (IsPBRModel(ShaderGraphNode))
-	{
-		MaterialInstanceFactoryNode->SetCustomParent(TEXT("Material'/Interchange/Materials/PBRSurfaceMaterial.PBRSurfaceMaterial'"));
+		ChooseParent(TEXT("MR"));
 	}
 	else if (IsPhongModel(ShaderGraphNode))
 	{
@@ -2437,6 +2495,8 @@ UInterchangeMaterialInstanceFactoryNode* UInterchangeGenericMaterialPipeline::Cr
 
 			const FString MapWeightName(MapName + TEXT("Weight"));
 			MaterialInstanceFactoryNode->AddFloatAttribute(UInterchangeShaderPortsAPI::MakeInputValueKey(MapWeightName), 1.f);
+
+			MaterialInstanceFactoryNode->AddFactoryDependencyUid(InputValue.Get<FString>());
 		}
 	}
 
diff --git a/Engine/Plugins/Interchange/Runtime/Source/Pipelines/Private/InterchangeMaterialXPipeline.cpp b/Engine/Plugins/Interchange/Runtime/Source/Pipelines/Private/InterchangeMaterialXPipeline.cpp
--- a/Engine/Plugins/Interchange/Runtime/Source/Pipelines/Private/InterchangeMaterialXPipeline.cpp
+++ b/Engine/Plugins/Interchange/Runtime/Source/Pipelines/Private/InterchangeMaterialXPipeline.cpp
@@ -163,14 +163,14 @@ TSet<FName> UMaterialXPipelineSettings::StandardSurfaceInputs
 TSet<FName> UMaterialXPipelineSettings::StandardSurfaceOutputs
 {
 	TEXT("Base Color"), // MX_StandardSurface has BaseColor with a whitespace, this should be fixed in further release
-	UE::Interchange::Materials::PBR::Parameters::Metallic,
-	UE::Interchange::Materials::PBR::Parameters::Specular,
-	UE::Interchange::Materials::PBR::Parameters::Roughness,
-	UE::Interchange::Materials::PBR::Parameters::Anisotropy,
-	UE::Interchange::Materials::PBR::Parameters::EmissiveColor,
-	UE::Interchange::Materials::PBR::Parameters::Opacity,
-	UE::Interchange::Materials::PBR::Parameters::Normal,
-	UE::Interchange::Materials::PBR::Parameters::Tangent,
+	UE::Interchange::Materials::PBRMR::Parameters::Metallic,
+	UE::Interchange::Materials::PBRMR::Parameters::Specular,
+	UE::Interchange::Materials::PBRMR::Parameters::Roughness,
+	UE::Interchange::Materials::PBRMR::Parameters::Anisotropy,
+	UE::Interchange::Materials::PBRMR::Parameters::EmissiveColor,
+	UE::Interchange::Materials::PBRMR::Parameters::Opacity,
+	UE::Interchange::Materials::PBRMR::Parameters::Normal,
+	UE::Interchange::Materials::PBRMR::Parameters::Tangent,
 	UE::Interchange::Materials::Sheen::Parameters::SheenRoughness,
 	UE::Interchange::Materials::Sheen::Parameters::SheenColor,
 	UE::Interchange::Materials::Subsurface::Parameters::SubsurfaceColor,
@@ -216,16 +216,16 @@ TSet<FName> UMaterialXPipelineSettings::TransmissionSurfaceInputs
 
 TSet<FName> UMaterialXPipelineSettings::TransmissionSurfaceOutputs
 {
-	UE::Interchange::Materials::PBR::Parameters::BaseColor,
-	UE::Interchange::Materials::PBR::Parameters::Metallic,
-	UE::Interchange::Materials::PBR::Parameters::Specular,
-	UE::Interchange::Materials::PBR::Parameters::Roughness,
-	UE::Interchange::Materials::PBR::Parameters::Anisotropy,
-	UE::Interchange::Materials::PBR::Parameters::EmissiveColor,
-	UE::Interchange::Materials::PBR::Parameters::Opacity,
-	UE::Interchange::Materials::PBR::Parameters::Normal,
-	UE::Interchange::Materials::PBR::Parameters::Tangent,
-	UE::Interchange::Materials::PBR::Parameters::Refraction,
+	UE::Interchange::Materials::PBRMR::Parameters::BaseColor,
+	UE::Interchange::Materials::PBRMR::Parameters::Metallic,
+	UE::Interchange::Materials::PBRMR::Parameters::Specular,
+	UE::Interchange::Materials::PBRMR::Parameters::Roughness,
+	UE::Interchange::Materials::PBRMR::Parameters::Anisotropy,
+	UE::Interchange::Materials::PBRMR::Parameters::EmissiveColor,
+	UE::Interchange::Materials::PBRMR::Parameters::Opacity,
+	UE::Interchange::Materials::PBRMR::Parameters::Normal,
+	UE::Interchange::Materials::PBRMR::Parameters::Tangent,
+	UE::Interchange::Materials::PBRMR::Parameters::Refraction,
 	UE::Interchange::Materials::ThinTranslucent::Parameters::TransmissionColor
 };
 
@@ -264,13 +264,13 @@ TSet<FName> UMaterialXPipelineSettings::UsdPreviewSurfaceInputs
 
 TSet<FName> UMaterialXPipelineSettings::UsdPreviewSurfaceOutputs
 {
-	UE::Interchange::Materials::PBR::Parameters::BaseColor,
-	UE::Interchange::Materials::PBR::Parameters::Metallic,
-	UE::Interchange::Materials::PBR::Parameters::Specular,
-	UE::Interchange::Materials::PBR::Parameters::Roughness,
-	UE::Interchange::Materials::PBR::Parameters::EmissiveColor,
-	UE::Interchange::Materials::PBR::Parameters::Opacity,
-	UE::Interchange::Materials::PBR::Parameters::Normal,
+	UE::Interchange::Materials::PBRMR::Parameters::BaseColor,
+	UE::Interchange::Materials::PBRMR::Parameters::Metallic,
+	UE::Interchange::Materials::PBRMR::Parameters::Specular,
+	UE::Interchange::Materials::PBRMR::Parameters::Roughness,
+	UE::Interchange::Materials::PBRMR::Parameters::EmissiveColor,
+	UE::Interchange::Materials::PBRMR::Parameters::Opacity,
+	UE::Interchange::Materials::PBRMR::Parameters::Normal,
 	UE::Interchange::Materials::Common::Parameters::Refraction,
 	UE::Interchange::Materials::Common::Parameters::Occlusion,
 	UE::Interchange::Materials::ClearCoat::Parameters::ClearCoat,
diff --git a/Engine/Plugins/Interchange/Runtime/Source/Pipelines/Public/InterchangeGenericMaterialPipeline.h b/Engine/Plugins/Interchange/Runtime/Source/Pipelines/Public/InterchangeGenericMaterialPipeline.h
--- a/Engine/Plugins/Interchange/Runtime/Source/Pipelines/Public/InterchangeGenericMaterialPipeline.h
+++ b/Engine/Plugins/Interchange/Runtime/Source/Pipelines/Public/InterchangeGenericMaterialPipeline.h
@@ -87,21 +87,46 @@ private:
 	UInterchangeMaterialInstanceFactoryNode* CreateMaterialInstanceFactoryNode(const UInterchangeShaderGraphNode* ShaderGraphNode);
 
 	/** True if the shader graph has a clear coat input. */
-	bool IsClearCoatModel(const UInterchangeShaderGraphNode* ShaderGraphNode) const;
+	bool HasClearCoat(const UInterchangeShaderGraphNode* ShaderGraphNode) const;
+	UE_DEPRECATED(5.3, "Deprecated. Use HasClearCoat.")
+	bool IsClearCoatModel(const UInterchangeShaderGraphNode* ShaderGraphNode) const
+	{
+		return HasClearCoat(ShaderGraphNode);
+	}
 
 	/** True if the shader graph has a sheen color input. */
-	bool IsSheenModel(const UInterchangeShaderGraphNode* ShaderGraphNode) const;
+	bool HasSheen(const UInterchangeShaderGraphNode* ShaderGraphNode) const;
+	UE_DEPRECATED(5.3, "Deprecated. Use HasSheen.")
+	bool IsSheenModel(const UInterchangeShaderGraphNode* ShaderGraphNode) const
+	{
+		return HasSheen(ShaderGraphNode);
+	}
 
 	/** True if the shader graph has a subsurface color input. */
-	bool IsSubsurfaceModel(const UInterchangeShaderGraphNode* ShaderGraphNode) const;
+	bool HasSubsurface(const UInterchangeShaderGraphNode* ShaderGraphNode) const;
+	UE_DEPRECATED(5.3, "Deprecated. Use HasSubsurface.")
+	bool IsSubsurfaceModel(const UInterchangeShaderGraphNode* ShaderGraphNode) const
+	{
+		return HasSubsurface(ShaderGraphNode);
+	}
 
 	/** True if the shader graph has a transmission color input. */
-	bool IsThinTranslucentModel(const UInterchangeShaderGraphNode* ShaderGraphNode) const;
+	bool HasThinTranslucency(const UInterchangeShaderGraphNode* ShaderGraphNode) const;
+	UE_DEPRECATED(5.3, "Deprecated. Use HasThinTranslucency.")
+	bool IsThinTranslucentModel(const UInterchangeShaderGraphNode* ShaderGraphNode) const
+	{
+		return HasThinTranslucency(ShaderGraphNode);
+	}
 
-	/** True if the shader graph has a base color input. */
-	bool IsPBRModel(const UInterchangeShaderGraphNode* ShaderGraphNode) const;
+	/** True if the shader graph has a base color input (Metallic/Roughness model. */
+	bool IsMetalRoughModel(const UInterchangeShaderGraphNode* ShaderGraphNode) const;
+	UE_DEPRECATED(5.3, "Deprecated. Use IsMetalRoughModel and IsSpecGlossModel to identify the correct PBR model.")
+	bool IsPBRModel(const UInterchangeShaderGraphNode* ShaderGraphNode) const
+	{
+		return IsMetalRoughModel(ShaderGraphNode);
+	}
 
-	/** True if the shader graph has diffuse color and specular color inputs. */
+	/** True if the shader graph has diffuse color and specular inputs. */
 	bool IsPhongModel(const UInterchangeShaderGraphNode* ShaderGraphNode) const;
 
 	/** True if the shader graph has a diffuse color input. */
@@ -116,9 +141,12 @@ private:
 	/** True if the shader graph has an unlit color input. */
 	bool IsUnlitModel(const UInterchangeShaderGraphNode* ShaderGraphNode) const;
 
+	/** True if the shader graph has specular color and glossiness scalar inputs. */
+	bool IsSpecGlossModel(const UInterchangeShaderGraphNode* ShaderGraphNode) const;
+
 	bool HandlePhongModel(const UInterchangeShaderGraphNode* ShaderGraphNode, UInterchangeMaterialFactoryNode* MaterialFactoryNode);
 	bool HandleLambertModel(const UInterchangeShaderGraphNode* ShaderGraphNode, UInterchangeMaterialFactoryNode* MaterialFactoryNode);
-	bool HandlePBRModel(const UInterchangeShaderGraphNode* ShaderGraphNode, UInterchangeMaterialFactoryNode* MaterialFactoryNode);
+	bool HandleMetalRoughnessModel(const UInterchangeShaderGraphNode* ShaderGraphNode, UInterchangeMaterialFactoryNode* MaterialFactoryNode);
 	bool HandleClearCoat(const UInterchangeShaderGraphNode* ShaderGraphNode, UInterchangeMaterialFactoryNode* MaterialFactoryNode);
 	bool HandleSubsurface(const UInterchangeShaderGraphNode* ShaderGraphNode, UInterchangeMaterialFactoryNode* MaterialFactoryNode);
 	bool HandleSheen(const UInterchangeShaderGraphNode* ShaderGraphNode, UInterchangeMaterialFactoryNode* MaterialFactoryNode);
@@ -126,6 +154,7 @@ private:
 	void HandleCommonParameters(const UInterchangeShaderGraphNode* ShaderGraphNode, UInterchangeMaterialFactoryNode* MaterialFactoryNode);
 	bool HandleBxDFInput(const UInterchangeShaderGraphNode* ShaderGraphNode, UInterchangeMaterialFactoryNode* MaterialFactoryNode);
 	bool HandleUnlitModel(const UInterchangeShaderGraphNode* ShaderGraphNode, UInterchangeMaterialFactoryNode* MaterialFactoryNode);
+	bool HandleSpecGlossModel(const UInterchangeShaderGraphNode* ShaderGraphNode, UInterchangeMaterialFactoryNode* MaterialFactoryNode);
 
 	void HandleFlattenNormalNode(const UInterchangeShaderNode* ShaderNode, UInterchangeBaseMaterialFactoryNode* MaterialFactoryNode, UInterchangeMaterialExpressionFactoryNode* FlattenNormalFactoryNode);
 	void HandleNormalFromHeightMapNode(const UInterchangeShaderNode* ShaderNode, UInterchangeBaseMaterialFactoryNode* MaterialFactoryNode, UInterchangeMaterialExpressionFactoryNode* NormalFromHeightMapFactoryNode);
