Fixed handling of bump map in FBX material

diff --git a/Engine/Plugins/Interchange/Runtime/Source/Parsers/Fbx/Private/FbxMaterial.cpp b/Engine/Plugins/Interchange/Runtime/Source/Parsers/Fbx/Private/FbxMaterial.cpp
--- a/Engine/Plugins/Interchange/Runtime/Source/Parsers/Fbx/Private/FbxMaterial.cpp
+++ b/Engine/Plugins/Interchange/Runtime/Source/Parsers/Fbx/Private/FbxMaterial.cpp
@@ -35,17 +35,24 @@ namespace UE
 				return ShaderGraphNode;
 			}
 
-			UInterchangeTexture2DNode* FFbxMaterial::CreateTexture2DNode(UInterchangeBaseNodeContainer& NodeContainer, const FString& TextureFilePath)
+			const UInterchangeTexture2DNode* FFbxMaterial::CreateTexture2DNode(UInterchangeBaseNodeContainer& NodeContainer, const FString& TextureFilePath)
 			{
 				const FString TextureName = FPaths::GetBaseFilename(TextureFilePath);
-				UInterchangeTexture2DNode* TextureNode = UInterchangeTexture2DNode::Create(&NodeContainer, TextureName);
+				const FString TextureNodeID = UInterchangeTextureNode::MakeNodeUid(TextureName);
+
+				if (const UInterchangeTexture2DNode* TextureNode = Cast<const UInterchangeTexture2DNode>(NodeContainer.GetNode(TextureNodeID)))
+				{
+					return TextureNode;
+				}
+
+				UInterchangeTexture2DNode* NewTextureNode = UInterchangeTexture2DNode::Create(&NodeContainer, TextureName);
 
 				//All texture translator expect a file as the payload key
 				FString NormalizeFilePath = TextureFilePath;
 				FPaths::NormalizeFilename(NormalizeFilePath);
-				TextureNode->SetPayLoadKey(NormalizeFilePath);
+				NewTextureNode->SetPayLoadKey(NormalizeFilePath);
 
-				return TextureNode;
+				return NewTextureNode;
 			}
 
 			UInterchangeShaderNode* FFbxMaterial::CreateTextureSampler(FbxFileTexture* FbxTexture, UInterchangeBaseNodeContainer& NodeContainer, const FString& ShaderUniqueID)
@@ -328,25 +335,45 @@ namespace UE
 				{
 					// FbxSurfaceMaterial can have either a normal map or a bump map, check for both
 					float Factor = 1.f;
-					const char* PropertyName = FbxSurfaceMaterial::sNormalMap;
-					FBXSDK_NAMESPACE::FbxProperty MaterialProperty = SurfaceMaterial->FindProperty(PropertyName);
-					if (!MaterialProperty.IsValid() || MaterialProperty.GetSrcObjectCount<FbxTexture>() < 1)
-					{
-						PropertyName = FbxSurfaceMaterial::sBump;
-						Factor = (float)((FbxSurfaceLambert*)SurfaceMaterial)->BumpFactor.Get();
-						MaterialProperty = SurfaceMaterial->FindProperty(PropertyName);
-						if (!MaterialProperty.IsValid() || MaterialProperty.GetSrcObjectCount<FbxTexture>() < 1)
-						{
-							PropertyName = nullptr;
-						}
-					}
-
-					if (PropertyName)
+					FBXSDK_NAMESPACE::FbxProperty MaterialProperty = SurfaceMaterial->FindProperty(FbxSurfaceMaterial::sNormalMap);
+					if (MaterialProperty.IsValid() && MaterialProperty.GetSrcObjectCount<FbxTexture>() > 0)
 					{
 						TVariant<FLinearColor, float> DefaultValue;
 						DefaultValue.Set<FLinearColor>(FLinearColor(FVector::UpVector));
 						ConvertPropertyToShaderNode(NodeContainer, ShaderGraphNode, MaterialProperty, Factor, Phong::Parameters::Normal, DefaultValue);
 					}
+					else
+					{
+						Factor = (float)((FbxSurfaceLambert*)SurfaceMaterial)->BumpFactor.Get();
+						MaterialProperty = SurfaceMaterial->FindProperty(FbxSurfaceMaterial::sBump);
+						if (MaterialProperty.IsValid() && MaterialProperty.GetSrcObjectCount<FbxTexture>() > 0)
+						{
+							using namespace Materials::Standard::Nodes;
+
+							if (FbxFileTexture* FbxTexture = MaterialProperty.GetSrcObject<FbxFileTexture>(0))
+							{
+								const FString TexturePath = FbxTexture->GetFileName();
+
+								if (const UInterchangeTexture2DNode* TextureNode = CreateTexture2DNode(NodeContainer, TexturePath))
+								{
+									const FString TextureName = FPaths::GetBaseFilename(TexturePath);
+
+									// NormalFromHeightmap needs TextureObject(not just a sample as it takes multiple samples from it)
+									UInterchangeShaderNode* TextureObjectNode = UInterchangeShaderNode::Create(&NodeContainer, TextureName, ShaderGraphNode->GetUniqueID());
+									TextureObjectNode->SetCustomShaderType(TextureObject::Name.ToString());
+									TextureObjectNode->AddStringAttribute(UInterchangeShaderPortsAPI::MakeInputValueKey(TextureObject::Inputs::Texture.ToString()), TextureNode->GetUniqueID());
+
+									UInterchangeShaderNode* HeightMapNode = UInterchangeShaderNode::Create(&NodeContainer, NormalFromHeightMap::Name.ToString(), ShaderGraphNode->GetUniqueID());
+									HeightMapNode->SetCustomShaderType(NormalFromHeightMap::Name.ToString());
+
+									UInterchangeShaderPortsAPI::ConnectDefaultOuputToInput(HeightMapNode, NormalFromHeightMap::Inputs::HeightMap.ToString(), TextureObjectNode->GetUniqueID());
+									HeightMapNode->AddFloatAttribute(UInterchangeShaderPortsAPI::MakeInputValueKey(NormalFromHeightMap::Inputs::Intensity.ToString()), Factor);
+
+									UInterchangeShaderPortsAPI::ConnectDefaultOuputToInput(ShaderGraphNode, Materials::Common::Parameters::Normal.ToString(), HeightMapNode->GetUniqueID());
+								}
+							}
+						}
+					}
 				}
 
 				// Opacity
diff --git a/Engine/Plugins/Interchange/Runtime/Source/Parsers/Fbx/Private/FbxMaterial.h b/Engine/Plugins/Interchange/Runtime/Source/Parsers/Fbx/Private/FbxMaterial.h
--- a/Engine/Plugins/Interchange/Runtime/Source/Parsers/Fbx/Private/FbxMaterial.h
+++ b/Engine/Plugins/Interchange/Runtime/Source/Parsers/Fbx/Private/FbxMaterial.h
@@ -53,7 +53,7 @@ namespace UE
 
 			protected:
 				UInterchangeShaderGraphNode* CreateShaderGraphNode(UInterchangeBaseNodeContainer& NodeContainer, const FString& NodeUID, const FString& NodeName);
-				UInterchangeTexture2DNode* CreateTexture2DNode(UInterchangeBaseNodeContainer& NodeContainer, const FString& TextureFilePath);
+				const UInterchangeTexture2DNode* CreateTexture2DNode(UInterchangeBaseNodeContainer& NodeContainer, const FString& TextureFilePath);
 				bool ConvertPropertyToShaderNode(UInterchangeBaseNodeContainer& NodeContainer, UInterchangeShaderGraphNode* ShaderGraphNode,
 					FbxProperty& Property, float Factor, FName PropertyName, TVariant<FLinearColor, float> DefaultValue, bool bInverse = false);
 
