Using Engine material expression name instead of MaterialX prefix.

diff --git a/Engine/Plugins/Interchange/Runtime/Source/Import/Private/MaterialX/MaterialXUtils/MaterialXManager.cpp b/Engine/Plugins/Interchange/Runtime/Source/Import/Private/MaterialX/MaterialXUtils/MaterialXManager.cpp
--- a/Engine/Plugins/Interchange/Runtime/Source/Import/Private/MaterialX/MaterialXUtils/MaterialXManager.cpp
+++ b/Engine/Plugins/Interchange/Runtime/Source/Import/Private/MaterialX/MaterialXUtils/MaterialXManager.cpp
@@ -99,11 +99,11 @@ FMaterialXManager::FMaterialXManager()
 		{MaterialX::Category::CrossProduct, TEXT("Crossproduct")},
 		{MaterialX::Category::Divide,       TEXT("Divide")},
 		{MaterialX::Category::DotProduct,   TEXT("Dotproduct")},
-		{MaterialX::Category::Exp,          TEXT("MaterialXExponential")},
+		{MaterialX::Category::Exp,          TEXT("Exponential")},
 		{MaterialX::Category::Floor,        TEXT("Floor")},
 		{MaterialX::Category::Invert,       TEXT("Subtract")},
-		{MaterialX::Category::Ln,           TEXT("MaterialXLogarithm")},
-		{MaterialX::Category::Magnitude,    TEXT("MaterialXLength")},
+		{MaterialX::Category::Ln,           TEXT("Logarithm")},
+		{MaterialX::Category::Magnitude,    TEXT("Length")},
 		{MaterialX::Category::Max,          TEXT("Max")},
 		{MaterialX::Category::Min,          TEXT("Min")},
 		{MaterialX::Category::Modulo,       TEXT("Fmod")},
@@ -154,10 +154,10 @@ FMaterialXManager::FMaterialXManager()
 		// Geometric nodes 
 		{MaterialX::Category::GeomColor,	TEXT("VertexColor")},
 		// Adjustment nodes,
-		{MaterialX::Category::HsvToRgb,		TEXT("MaterialXHsvToRgb")},
+		{MaterialX::Category::HsvToRgb,		TEXT("HsvToRgb")},
 		{MaterialX::Category::Luminance,    TEXT("MaterialXLuminance")},
 		{MaterialX::Category::Remap,        TEXT("MaterialXRemap")},
-		{MaterialX::Category::RgbToHsv,		TEXT("MaterialXRgbToHsv")},
+		{MaterialX::Category::RgbToHsv,		TEXT("RgbToHsv")},
 		{MaterialX::Category::Saturate,		TEXT("Desaturation")},
 		{MaterialX::Category::Smoothstep,   TEXT("SmoothStep")}
 	}
