Tagging MaterialX material expressions private in order to hide them in the Material Editor

diff --git a/Engine/Plugins/Interchange/Runtime/Source/Import/Private/MaterialX/MaterialExpressions/MaterialExpressionAppend3Vector.h b/Engine/Plugins/Interchange/Runtime/Source/Import/Private/MaterialX/MaterialExpressions/MaterialExpressionAppend3Vector.h
--- a/Engine/Plugins/Interchange/Runtime/Source/Import/Private/MaterialX/MaterialExpressions/MaterialExpressionAppend3Vector.h
+++ b/Engine/Plugins/Interchange/Runtime/Source/Import/Private/MaterialX/MaterialExpressions/MaterialExpressionAppend3Vector.h
@@ -9,7 +9,7 @@
 /**
  * A material expression that allows combining 3 channels together to create a vector with more channel than the original
  */
-UCLASS(MinimalAPI, collapsecategories, hidecategories = Object)
+UCLASS(MinimalAPI, collapsecategories, hidecategories = Object, meta = ( Private))
 class UMaterialExpressionMaterialXAppend3Vector : public UMaterialExpression
 {
 	GENERATED_UCLASS_BODY()
diff --git a/Engine/Plugins/Interchange/Runtime/Source/Import/Private/MaterialX/MaterialExpressions/MaterialExpressionAppend4Vector.h b/Engine/Plugins/Interchange/Runtime/Source/Import/Private/MaterialX/MaterialExpressions/MaterialExpressionAppend4Vector.h
--- a/Engine/Plugins/Interchange/Runtime/Source/Import/Private/MaterialX/MaterialExpressions/MaterialExpressionAppend4Vector.h
+++ b/Engine/Plugins/Interchange/Runtime/Source/Import/Private/MaterialX/MaterialExpressions/MaterialExpressionAppend4Vector.h
@@ -11,7 +11,7 @@
 /**
  * A material expression that allows combining 4 channels together to create a vector with more channel than the original
  */
-UCLASS(MinimalAPI, collapsecategories, hidecategories = Object)
+UCLASS(MinimalAPI, collapsecategories, hidecategories = Object, meta = (Private))
 class UMaterialExpressionMaterialXAppend4Vector : public UMaterialExpression
 {
 	GENERATED_UCLASS_BODY()
diff --git a/Engine/Plugins/Interchange/Runtime/Source/Import/Private/MaterialX/MaterialExpressions/MaterialExpressionBurn.h b/Engine/Plugins/Interchange/Runtime/Source/Import/Private/MaterialX/MaterialExpressions/MaterialExpressionBurn.h
--- a/Engine/Plugins/Interchange/Runtime/Source/Import/Private/MaterialX/MaterialExpressions/MaterialExpressionBurn.h
+++ b/Engine/Plugins/Interchange/Runtime/Source/Import/Private/MaterialX/MaterialExpressions/MaterialExpressionBurn.h
@@ -13,7 +13,7 @@
  * Operation: 1-(1-B)/A
  * Result: Lerp(B, 1-(1-B)/A, Alpha)
  */
-UCLASS(MinimalAPI, collapsecategories, hidecategories = Object)
+UCLASS(MinimalAPI, collapsecategories, hidecategories = Object, meta = (Private))
 class UMaterialExpressionMaterialXBurn : public UMaterialExpression
 {
 	GENERATED_UCLASS_BODY()
diff --git a/Engine/Plugins/Interchange/Runtime/Source/Import/Private/MaterialX/MaterialExpressions/MaterialExpressionDifference.h b/Engine/Plugins/Interchange/Runtime/Source/Import/Private/MaterialX/MaterialExpressions/MaterialExpressionDifference.h
--- a/Engine/Plugins/Interchange/Runtime/Source/Import/Private/MaterialX/MaterialExpressions/MaterialExpressionDifference.h
+++ b/Engine/Plugins/Interchange/Runtime/Source/Import/Private/MaterialX/MaterialExpressions/MaterialExpressionDifference.h
@@ -13,7 +13,7 @@
  * Operation: abs(A - B)
  * Result: Lerp(B, abs(A - B), Alpha)
  */
-UCLASS(MinimalAPI, collapsecategories, hidecategories = Object)
+UCLASS(MinimalAPI, collapsecategories, hidecategories = Object, meta = (Private))
 class UMaterialExpressionMaterialXDifference : public UMaterialExpression
 {
 	GENERATED_UCLASS_BODY()
diff --git a/Engine/Plugins/Interchange/Runtime/Source/Import/Private/MaterialX/MaterialExpressions/MaterialExpressionDisjointOver.h b/Engine/Plugins/Interchange/Runtime/Source/Import/Private/MaterialX/MaterialExpressions/MaterialExpressionDisjointOver.h
--- a/Engine/Plugins/Interchange/Runtime/Source/Import/Private/MaterialX/MaterialExpressions/MaterialExpressionDisjointOver.h
+++ b/Engine/Plugins/Interchange/Runtime/Source/Import/Private/MaterialX/MaterialExpressions/MaterialExpressionDisjointOver.h
@@ -18,7 +18,7 @@
  *            A + B(1-a)/b if a + b > 1
  * Result: Lerp(B, Op, Alpha)
  */
-UCLASS(MinimalAPI, collapsecategories, hidecategories = Object)
+UCLASS(MinimalAPI, collapsecategories, hidecategories = Object, meta = (Private))
 class UMaterialExpressionMaterialXDisjointOver : public UMaterialExpression
 {
 	GENERATED_UCLASS_BODY()
diff --git a/Engine/Plugins/Interchange/Runtime/Source/Import/Private/MaterialX/MaterialExpressions/MaterialExpressionDodge.h b/Engine/Plugins/Interchange/Runtime/Source/Import/Private/MaterialX/MaterialExpressions/MaterialExpressionDodge.h
--- a/Engine/Plugins/Interchange/Runtime/Source/Import/Private/MaterialX/MaterialExpressions/MaterialExpressionDodge.h
+++ b/Engine/Plugins/Interchange/Runtime/Source/Import/Private/MaterialX/MaterialExpressions/MaterialExpressionDodge.h
@@ -13,7 +13,7 @@
  * Operation: B/(1-A)
  * Result: Lerp(B, B/(1-A), Alpha)
  */
-UCLASS(MinimalAPI, collapsecategories, hidecategories = Object)
+UCLASS(MinimalAPI, collapsecategories, hidecategories = Object, meta = (Private))
 class UMaterialExpressionMaterialXDodge: public UMaterialExpression
 {
 	GENERATED_UCLASS_BODY()
diff --git a/Engine/Plugins/Interchange/Runtime/Source/Import/Private/MaterialX/MaterialExpressions/MaterialExpressionFractal3D.h b/Engine/Plugins/Interchange/Runtime/Source/Import/Private/MaterialX/MaterialExpressions/MaterialExpressionFractal3D.h
--- a/Engine/Plugins/Interchange/Runtime/Source/Import/Private/MaterialX/MaterialExpressions/MaterialExpressionFractal3D.h
+++ b/Engine/Plugins/Interchange/Runtime/Source/Import/Private/MaterialX/MaterialExpressions/MaterialExpressionFractal3D.h
@@ -10,7 +10,7 @@
  * octaves of 3D Perlin noise, increasing the frequency and decreasing the amplitude at each octave.
  */
 
-UCLASS(collapsecategories, hidecategories = Object, MinimalAPI)
+UCLASS(collapsecategories, hidecategories = Object, MinimalAPI, meta = (Private))
 class UMaterialExpressionMaterialXFractal3D : public UMaterialExpression
 {
 	GENERATED_UCLASS_BODY()
diff --git a/Engine/Plugins/Interchange/Runtime/Source/Import/Private/MaterialX/MaterialExpressions/MaterialExpressionIn.h b/Engine/Plugins/Interchange/Runtime/Source/Import/Private/MaterialX/MaterialExpressions/MaterialExpressionIn.h
--- a/Engine/Plugins/Interchange/Runtime/Source/Import/Private/MaterialX/MaterialExpressions/MaterialExpressionIn.h
+++ b/Engine/Plugins/Interchange/Runtime/Source/Import/Private/MaterialX/MaterialExpressions/MaterialExpressionIn.h
@@ -17,7 +17,7 @@
  * Operation: Ab
  * Result: Lerp(B, Ab, Alpha)
  */
-UCLASS(MinimalAPI, collapsecategories, hidecategories = Object)
+UCLASS(MinimalAPI, collapsecategories, hidecategories = Object, meta = (Private))
 class UMaterialExpressionMaterialXIn: public UMaterialExpression
 {
 	GENERATED_UCLASS_BODY()
diff --git a/Engine/Plugins/Interchange/Runtime/Source/Import/Private/MaterialX/MaterialExpressions/MaterialExpressionLuminance.h b/Engine/Plugins/Interchange/Runtime/Source/Import/Private/MaterialX/MaterialExpressions/MaterialExpressionLuminance.h
--- a/Engine/Plugins/Interchange/Runtime/Source/Import/Private/MaterialX/MaterialExpressions/MaterialExpressionLuminance.h
+++ b/Engine/Plugins/Interchange/Runtime/Source/Import/Private/MaterialX/MaterialExpressions/MaterialExpressionLuminance.h
@@ -21,7 +21,7 @@ enum class EMaterialXLuminanceMode : uint8
 	Custom
 };
 
-UCLASS(collapsecategories, hidecategories = Object, MinimalAPI)
+UCLASS(collapsecategories, hidecategories = Object, MinimalAPI, meta = (Private))
 class UMaterialExpressionMaterialXLuminance : public UMaterialExpression
 {
 	GENERATED_UCLASS_BODY()
diff --git a/Engine/Plugins/Interchange/Runtime/Source/Import/Private/MaterialX/MaterialExpressions/MaterialExpressionMask.h b/Engine/Plugins/Interchange/Runtime/Source/Import/Private/MaterialX/MaterialExpressions/MaterialExpressionMask.h
--- a/Engine/Plugins/Interchange/Runtime/Source/Import/Private/MaterialX/MaterialExpressions/MaterialExpressionMask.h
+++ b/Engine/Plugins/Interchange/Runtime/Source/Import/Private/MaterialX/MaterialExpressions/MaterialExpressionMask.h
@@ -17,7 +17,7 @@
  * Operation: B*a
  * Result: Lerp(B, B*a, Alpha)
  */
-UCLASS(MinimalAPI, collapsecategories, hidecategories = Object)
+UCLASS(MinimalAPI, collapsecategories, hidecategories = Object, meta = (Private))
 class UMaterialExpressionMaterialXMask : public UMaterialExpression
 {
 	GENERATED_UCLASS_BODY()
diff --git a/Engine/Plugins/Interchange/Runtime/Source/Import/Private/MaterialX/MaterialExpressions/MaterialExpressionMatte.h b/Engine/Plugins/Interchange/Runtime/Source/Import/Private/MaterialX/MaterialExpressions/MaterialExpressionMatte.h
--- a/Engine/Plugins/Interchange/Runtime/Source/Import/Private/MaterialX/MaterialExpressions/MaterialExpressionMatte.h
+++ b/Engine/Plugins/Interchange/Runtime/Source/Import/Private/MaterialX/MaterialExpressions/MaterialExpressionMatte.h
@@ -17,7 +17,7 @@
  * Operation: A*a + B*(1-a)
  * Result: Lerp(B, A*a + B*(1-a), Alpha)
  */
-UCLASS(MinimalAPI, collapsecategories, hidecategories = Object)
+UCLASS(MinimalAPI, collapsecategories, hidecategories = Object, meta = (Private))
 class UMaterialExpressionMaterialXMatte : public UMaterialExpression
 {
 	GENERATED_UCLASS_BODY()
diff --git a/Engine/Plugins/Interchange/Runtime/Source/Import/Private/MaterialX/MaterialExpressions/MaterialExpressionMinus.h b/Engine/Plugins/Interchange/Runtime/Source/Import/Private/MaterialX/MaterialExpressions/MaterialExpressionMinus.h
--- a/Engine/Plugins/Interchange/Runtime/Source/Import/Private/MaterialX/MaterialExpressions/MaterialExpressionMinus.h
+++ b/Engine/Plugins/Interchange/Runtime/Source/Import/Private/MaterialX/MaterialExpressions/MaterialExpressionMinus.h
@@ -13,7 +13,7 @@
  * Operation: A - B
  * Result: Lerp(B, A - B, Alpha)
  */
-UCLASS(MinimalAPI, collapsecategories, hidecategories = Object)
+UCLASS(MinimalAPI, collapsecategories, hidecategories = Object, meta = (Private))
 class UMaterialExpressionMaterialXMinus : public UMaterialExpression
 {
 	GENERATED_UCLASS_BODY()
diff --git a/Engine/Plugins/Interchange/Runtime/Source/Import/Private/MaterialX/MaterialExpressions/MaterialExpressionOut.h b/Engine/Plugins/Interchange/Runtime/Source/Import/Private/MaterialX/MaterialExpressions/MaterialExpressionOut.h
--- a/Engine/Plugins/Interchange/Runtime/Source/Import/Private/MaterialX/MaterialExpressions/MaterialExpressionOut.h
+++ b/Engine/Plugins/Interchange/Runtime/Source/Import/Private/MaterialX/MaterialExpressions/MaterialExpressionOut.h
@@ -17,7 +17,7 @@
  * Operation: A(1-b)
  * Result: Lerp(B, A(1-b), Alpha)
  */
-UCLASS(MinimalAPI, collapsecategories, hidecategories = Object)
+UCLASS(MinimalAPI, collapsecategories, hidecategories = Object, meta = (Private))
 class UMaterialExpressionMaterialXOut : public UMaterialExpression
 {
 	GENERATED_UCLASS_BODY()
diff --git a/Engine/Plugins/Interchange/Runtime/Source/Import/Private/MaterialX/MaterialExpressions/MaterialExpressionOver.h b/Engine/Plugins/Interchange/Runtime/Source/Import/Private/MaterialX/MaterialExpressions/MaterialExpressionOver.h
--- a/Engine/Plugins/Interchange/Runtime/Source/Import/Private/MaterialX/MaterialExpressions/MaterialExpressionOver.h
+++ b/Engine/Plugins/Interchange/Runtime/Source/Import/Private/MaterialX/MaterialExpressions/MaterialExpressionOver.h
@@ -17,7 +17,7 @@
  * Operation: A + B(1-a)
  * Result: Lerp(B, A + B(1-a), Alpha)
  */
-UCLASS(MinimalAPI, collapsecategories, hidecategories = Object)
+UCLASS(MinimalAPI, collapsecategories, hidecategories = Object, meta = (Private))
 class UMaterialExpressionMaterialXOver : public UMaterialExpression
 {
 	GENERATED_UCLASS_BODY()
diff --git a/Engine/Plugins/Interchange/Runtime/Source/Import/Private/MaterialX/MaterialExpressions/MaterialExpressionOverlay.h b/Engine/Plugins/Interchange/Runtime/Source/Import/Private/MaterialX/MaterialExpressions/MaterialExpressionOverlay.h
--- a/Engine/Plugins/Interchange/Runtime/Source/Import/Private/MaterialX/MaterialExpressions/MaterialExpressionOverlay.h
+++ b/Engine/Plugins/Interchange/Runtime/Source/Import/Private/MaterialX/MaterialExpressions/MaterialExpressionOverlay.h
@@ -14,7 +14,7 @@
  *            1-(1-A)(1-B) if A >= 0.5
  * Result: Lerp(B, Op, Alpha)
  */
-UCLASS(MinimalAPI, collapsecategories, hidecategories = Object)
+UCLASS(MinimalAPI, collapsecategories, hidecategories = Object, meta = (Private))
 class UMaterialExpressionMaterialXOverlay : public UMaterialExpression
 {
 	GENERATED_UCLASS_BODY()
diff --git a/Engine/Plugins/Interchange/Runtime/Source/Import/Private/MaterialX/MaterialExpressions/MaterialExpressionPlace2D.h b/Engine/Plugins/Interchange/Runtime/Source/Import/Private/MaterialX/MaterialExpressions/MaterialExpressionPlace2D.h
--- a/Engine/Plugins/Interchange/Runtime/Source/Import/Private/MaterialX/MaterialExpressions/MaterialExpressionPlace2D.h
+++ b/Engine/Plugins/Interchange/Runtime/Source/Import/Private/MaterialX/MaterialExpressions/MaterialExpressionPlace2D.h
@@ -13,7 +13,7 @@
  * implementation matching the behavior of certain DCC packages, and "1" or "TRS" performs
  * -pivot, translate, rotate, scale, +pivot which does not introduce texture shear.
  * Default is 0 "SRT" for backward compatibility.*/
-UCLASS(MinimalAPI, collapsecategories, hidecategories = Object)
+UCLASS(MinimalAPI, collapsecategories, hidecategories = Object, meta = (Private))
 class UMaterialExpressionMaterialXPlace2D : public UMaterialExpression
 {
 	GENERATED_UCLASS_BODY()
diff --git a/Engine/Plugins/Interchange/Runtime/Source/Import/Private/MaterialX/MaterialExpressions/MaterialExpressionPlus.h b/Engine/Plugins/Interchange/Runtime/Source/Import/Private/MaterialX/MaterialExpressions/MaterialExpressionPlus.h
--- a/Engine/Plugins/Interchange/Runtime/Source/Import/Private/MaterialX/MaterialExpressions/MaterialExpressionPlus.h
+++ b/Engine/Plugins/Interchange/Runtime/Source/Import/Private/MaterialX/MaterialExpressions/MaterialExpressionPlus.h
@@ -13,7 +13,7 @@
  * Operation: A + B
  * Result: Lerp(B, A + B, Alpha)
  */
-UCLASS(MinimalAPI, collapsecategories, hidecategories = Object)
+UCLASS(MinimalAPI, collapsecategories, hidecategories = Object, meta = (Private))
 class UMaterialExpressionMaterialXPlus : public UMaterialExpression
 {
 	GENERATED_UCLASS_BODY()
diff --git a/Engine/Plugins/Interchange/Runtime/Source/Import/Private/MaterialX/MaterialExpressions/MaterialExpressionPremult.h b/Engine/Plugins/Interchange/Runtime/Source/Import/Private/MaterialX/MaterialExpressions/MaterialExpressionPremult.h
--- a/Engine/Plugins/Interchange/Runtime/Source/Import/Private/MaterialX/MaterialExpressions/MaterialExpressionPremult.h
+++ b/Engine/Plugins/Interchange/Runtime/Source/Import/Private/MaterialX/MaterialExpressions/MaterialExpressionPremult.h
@@ -10,7 +10,7 @@
  * Multiply the RGB channels of the input by the Alpha channel of the input.
  * Input must be of type float4
  */
-UCLASS(collapsecategories, hidecategories = Object, MinimalAPI)
+UCLASS(collapsecategories, hidecategories = Object, MinimalAPI, meta = (Private))
 class UMaterialExpressionMaterialXPremult : public UMaterialExpression
 {
 	GENERATED_UCLASS_BODY()
diff --git a/Engine/Plugins/Interchange/Runtime/Source/Import/Private/MaterialX/MaterialExpressions/MaterialExpressionRamp4.h b/Engine/Plugins/Interchange/Runtime/Source/Import/Private/MaterialX/MaterialExpressions/MaterialExpressionRamp4.h
--- a/Engine/Plugins/Interchange/Runtime/Source/Import/Private/MaterialX/MaterialExpressions/MaterialExpressionRamp4.h
+++ b/Engine/Plugins/Interchange/Runtime/Source/Import/Private/MaterialX/MaterialExpressions/MaterialExpressionRamp4.h
@@ -9,7 +9,7 @@
 /**
  * A material expression that computes a 4-corner bilinear value ramp..
  */
-UCLASS(collapsecategories, hidecategories = Object, MinimalAPI)
+UCLASS(collapsecategories, hidecategories = Object, MinimalAPI, meta = (Private))
 class UMaterialExpressionMaterialXRamp4: public UMaterialExpression
 {
 	GENERATED_UCLASS_BODY()
diff --git a/Engine/Plugins/Interchange/Runtime/Source/Import/Private/MaterialX/MaterialExpressions/MaterialExpressionRampLeftRight.h b/Engine/Plugins/Interchange/Runtime/Source/Import/Private/MaterialX/MaterialExpressions/MaterialExpressionRampLeftRight.h
--- a/Engine/Plugins/Interchange/Runtime/Source/Import/Private/MaterialX/MaterialExpressions/MaterialExpressionRampLeftRight.h
+++ b/Engine/Plugins/Interchange/Runtime/Source/Import/Private/MaterialX/MaterialExpressions/MaterialExpressionRampLeftRight.h
@@ -9,7 +9,7 @@
 /**
  * A material expression that computes a left-to-right bilinear value ramp.
  */
-UCLASS(collapsecategories, hidecategories = Object, MinimalAPI)
+UCLASS(collapsecategories, hidecategories = Object, MinimalAPI, meta = (Private))
 class UMaterialExpressionMaterialXRampLeftRight : public UMaterialExpression
 {
 	GENERATED_UCLASS_BODY()
diff --git a/Engine/Plugins/Interchange/Runtime/Source/Import/Private/MaterialX/MaterialExpressions/MaterialExpressionRampTopBottom.h b/Engine/Plugins/Interchange/Runtime/Source/Import/Private/MaterialX/MaterialExpressions/MaterialExpressionRampTopBottom.h
--- a/Engine/Plugins/Interchange/Runtime/Source/Import/Private/MaterialX/MaterialExpressions/MaterialExpressionRampTopBottom.h
+++ b/Engine/Plugins/Interchange/Runtime/Source/Import/Private/MaterialX/MaterialExpressions/MaterialExpressionRampTopBottom.h
@@ -9,7 +9,7 @@
 /**
  * A material expression that computes a top-to-bottom bilinear value ramp.
  */
-UCLASS(collapsecategories, hidecategories = Object, MinimalAPI)
+UCLASS(collapsecategories, hidecategories = Object, MinimalAPI, meta = (Private))
 class UMaterialExpressionMaterialXRampTopBottom : public UMaterialExpression
 {
 	GENERATED_UCLASS_BODY()
diff --git a/Engine/Plugins/Interchange/Runtime/Source/Import/Private/MaterialX/MaterialExpressions/MaterialExpressionRemap.h b/Engine/Plugins/Interchange/Runtime/Source/Import/Private/MaterialX/MaterialExpressions/MaterialExpressionRemap.h
--- a/Engine/Plugins/Interchange/Runtime/Source/Import/Private/MaterialX/MaterialExpressions/MaterialExpressionRemap.h
+++ b/Engine/Plugins/Interchange/Runtime/Source/Import/Private/MaterialX/MaterialExpressions/MaterialExpressionRemap.h
@@ -9,7 +9,7 @@
 /**
  * A material expression that Remap a value from one range to another.
  */
-UCLASS(collapsecategories, hidecategories = Object, MinimalAPI)
+UCLASS(collapsecategories, hidecategories = Object, MinimalAPI, meta = (Private))
 class UMaterialExpressionMaterialXRemap : public UMaterialExpression
 {
 	GENERATED_UCLASS_BODY()
diff --git a/Engine/Plugins/Interchange/Runtime/Source/Import/Private/MaterialX/MaterialExpressions/MaterialExpressionRotate2D.h b/Engine/Plugins/Interchange/Runtime/Source/Import/Private/MaterialX/MaterialExpressions/MaterialExpressionRotate2D.h
--- a/Engine/Plugins/Interchange/Runtime/Source/Import/Private/MaterialX/MaterialExpressions/MaterialExpressionRotate2D.h
+++ b/Engine/Plugins/Interchange/Runtime/Source/Import/Private/MaterialX/MaterialExpressions/MaterialExpressionRotate2D.h
@@ -8,7 +8,7 @@
 /**
  * A material expression that rotates a vector2 value about the origin in 2D.
  */
-UCLASS(collapsecategories, hidecategories = Object, MinimalAPI)
+UCLASS(collapsecategories, hidecategories = Object, MinimalAPI, meta = (Private))
 class UMaterialExpressionMaterialXRotate2D : public UMaterialExpression
 {
 	GENERATED_UCLASS_BODY()
diff --git a/Engine/Plugins/Interchange/Runtime/Source/Import/Private/MaterialX/MaterialExpressions/MaterialExpressionScreen.h b/Engine/Plugins/Interchange/Runtime/Source/Import/Private/MaterialX/MaterialExpressions/MaterialExpressionScreen.h
--- a/Engine/Plugins/Interchange/Runtime/Source/Import/Private/MaterialX/MaterialExpressions/MaterialExpressionScreen.h
+++ b/Engine/Plugins/Interchange/Runtime/Source/Import/Private/MaterialX/MaterialExpressions/MaterialExpressionScreen.h
@@ -13,7 +13,7 @@
  * Operation: 1-(1-A)(1-B)
  * Result: Lerp(B, 1-(1-F)(1-B), Alpha)
  */
-UCLASS(MinimalAPI, collapsecategories, hidecategories = Object)
+UCLASS(MinimalAPI, collapsecategories, hidecategories = Object, meta = (Private))
 class UMaterialExpressionMaterialXScreen : public UMaterialExpression
 {
 	GENERATED_UCLASS_BODY()
diff --git a/Engine/Plugins/Interchange/Runtime/Source/Import/Private/MaterialX/MaterialExpressions/MaterialExpressionSplitLeftRight.h b/Engine/Plugins/Interchange/Runtime/Source/Import/Private/MaterialX/MaterialExpressions/MaterialExpressionSplitLeftRight.h
--- a/Engine/Plugins/Interchange/Runtime/Source/Import/Private/MaterialX/MaterialExpressions/MaterialExpressionSplitLeftRight.h
+++ b/Engine/Plugins/Interchange/Runtime/Source/Import/Private/MaterialX/MaterialExpressions/MaterialExpressionSplitLeftRight.h
@@ -9,7 +9,7 @@
 /**
  * A material expression that computes a left-right split matte, split at a specified u value.
  */
-UCLASS(collapsecategories, hidecategories = Object, MinimalAPI)
+UCLASS(collapsecategories, hidecategories = Object, MinimalAPI, meta = (Private))
 class UMaterialExpressionMaterialXSplitLeftRight : public UMaterialExpression
 {
 	GENERATED_UCLASS_BODY()
diff --git a/Engine/Plugins/Interchange/Runtime/Source/Import/Private/MaterialX/MaterialExpressions/MaterialExpressionSplitTopBottom.h b/Engine/Plugins/Interchange/Runtime/Source/Import/Private/MaterialX/MaterialExpressions/MaterialExpressionSplitTopBottom.h
--- a/Engine/Plugins/Interchange/Runtime/Source/Import/Private/MaterialX/MaterialExpressions/MaterialExpressionSplitTopBottom.h
+++ b/Engine/Plugins/Interchange/Runtime/Source/Import/Private/MaterialX/MaterialExpressions/MaterialExpressionSplitTopBottom.h
@@ -9,7 +9,7 @@
 /**
  * A material expression that computes a top-bottom split matte, split at a specified v value.
  */
-UCLASS(collapsecategories, hidecategories = Object, MinimalAPI)
+UCLASS(collapsecategories, hidecategories = Object, MinimalAPI, meta = (Private))
 class UMaterialExpressionMaterialXSplitTopBottom : public UMaterialExpression
 {
 	GENERATED_UCLASS_BODY()
diff --git a/Engine/Plugins/Interchange/Runtime/Source/Import/Private/MaterialX/MaterialExpressions/MaterialExpressionUnpremult.h b/Engine/Plugins/Interchange/Runtime/Source/Import/Private/MaterialX/MaterialExpressions/MaterialExpressionUnpremult.h
--- a/Engine/Plugins/Interchange/Runtime/Source/Import/Private/MaterialX/MaterialExpressions/MaterialExpressionUnpremult.h
+++ b/Engine/Plugins/Interchange/Runtime/Source/Import/Private/MaterialX/MaterialExpressions/MaterialExpressionUnpremult.h
@@ -11,7 +11,7 @@
  * Alpha value is zero, the original color4 input value is passed through unchanged.
  * Input must be of type float4
  */
-UCLASS(collapsecategories, hidecategories = Object, MinimalAPI)
+UCLASS(collapsecategories, hidecategories = Object, MinimalAPI, meta = (Private))
 class UMaterialExpressionMaterialXUnpremult : public UMaterialExpression
 {
 	GENERATED_UCLASS_BODY()
diff --git a/Engine/Plugins/Interchange/Runtime/Source/Import/Public/MaterialX/MaterialExpressions/MaterialExpressionSwizzle.h b/Engine/Plugins/Interchange/Runtime/Source/Import/Public/MaterialX/MaterialExpressions/MaterialExpressionSwizzle.h
--- a/Engine/Plugins/Interchange/Runtime/Source/Import/Public/MaterialX/MaterialExpressions/MaterialExpressionSwizzle.h
+++ b/Engine/Plugins/Interchange/Runtime/Source/Import/Public/MaterialX/MaterialExpressions/MaterialExpressionSwizzle.h
@@ -6,7 +6,7 @@
 
 #include "MaterialExpressionSwizzle.generated.h"
 
-UCLASS(collapsecategories, hidecategories = Object, MinimalAPI)
+UCLASS(collapsecategories, hidecategories = Object, MinimalAPI, meta = (Private))
 class UMaterialExpressionMaterialXSwizzle : public UMaterialExpression
 {
 	GENERATED_UCLASS_BODY()
diff --git a/Engine/Plugins/Interchange/Runtime/Source/Import/Public/MaterialX/MaterialExpressions/MaterialExpressionTextureSampleParameterBlur.h b/Engine/Plugins/Interchange/Runtime/Source/Import/Public/MaterialX/MaterialExpressions/MaterialExpressionTextureSampleParameterBlur.h
--- a/Engine/Plugins/Interchange/Runtime/Source/Import/Public/MaterialX/MaterialExpressions/MaterialExpressionTextureSampleParameterBlur.h
+++ b/Engine/Plugins/Interchange/Runtime/Source/Import/Public/MaterialX/MaterialExpressions/MaterialExpressionTextureSampleParameterBlur.h
@@ -20,7 +20,7 @@ enum class EMAterialXTextureSampleBlurKernel
 	Kernel7 UMETA(DisplayName = "7x7"),
 };
 
-UCLASS(MinimalAPI, collapsecategories, hidecategories = Object)
+UCLASS(MinimalAPI, collapsecategories, hidecategories = Object, meta = (Private))
 class UMaterialExpressionMaterialXTextureSampleParameterBlur: public UMaterialExpressionTextureSampleParameter2D
 {
 	GENERATED_UCLASS_BODY()
