[Interchange] glTF Transmission material handling issues.
Exposed MaterialAttribute possibilities for Interchange.
Updated Shading Model priorities for glTF.
Updated Material Generation to use new MaterialFunction for glTF Transmission.
GLTFCore removed AlphaMode restricion on Transmission materials.

diff --git a/Engine/Plugins/Interchange/Runtime/Source/FactoryNodes/Private/InterchangeMaterialFactoryNode.cpp b/Engine/Plugins/Interchange/Runtime/Source/FactoryNodes/Private/InterchangeMaterialFactoryNode.cpp
--- a/Engine/Plugins/Interchange/Runtime/Source/FactoryNodes/Private/InterchangeMaterialFactoryNode.cpp
+++ b/Engine/Plugins/Interchange/Runtime/Source/FactoryNodes/Private/InterchangeMaterialFactoryNode.cpp
@@ -365,6 +365,16 @@ bool UInterchangeMaterialFactoryNode::SetCustomRefractionMethod(const TEnumAsByt
 	IMPLEMENT_NODE_ATTRIBUTE_SETTER(UInterchangeMaterialFactoryNode, RefractionMethod, TEnumAsByte<ERefractionMode>, UMaterial);
 }
 
+bool UInterchangeMaterialFactoryNode::GetCustomScreenSpaceReflections(bool& AttributeValue) const
+{
+	IMPLEMENT_NODE_ATTRIBUTE_GETTER(ScreenSpaceReflections, bool);
+}
+
+bool UInterchangeMaterialFactoryNode::SetCustomScreenSpaceReflections(const bool& AttributeValue)
+{
+	IMPLEMENT_NODE_ATTRIBUTE_SETTER_NODELEGATE(ScreenSpaceReflections, bool);
+}
+
 FString UInterchangeMaterialExpressionFactoryNode::GetTypeName() const
 {
 	const FString TypeName = TEXT("MaterialExpressionFactoryNode");
diff --git a/Engine/Plugins/Interchange/Runtime/Source/FactoryNodes/Public/InterchangeMaterialFactoryNode.h b/Engine/Plugins/Interchange/Runtime/Source/FactoryNodes/Public/InterchangeMaterialFactoryNode.h
--- a/Engine/Plugins/Interchange/Runtime/Source/FactoryNodes/Public/InterchangeMaterialFactoryNode.h
+++ b/Engine/Plugins/Interchange/Runtime/Source/FactoryNodes/Public/InterchangeMaterialFactoryNode.h
@@ -237,6 +237,12 @@ public:
 	UFUNCTION(BlueprintCallable, Category = "Interchange | Node | Material")
 	bool SetCustomRefractionMethod(const TEnumAsByte<ERefractionMode>& AttributeValue, bool bAddApplyDelegate = true);
 
+	UFUNCTION(BlueprintCallable, Category = "Interchange | Node | Material")
+	bool GetCustomScreenSpaceReflections(bool& AttributeValue) const;
+
+	UFUNCTION(BlueprintCallable, Category = "Interchange | Node | Material")
+	bool SetCustomScreenSpaceReflections(const bool& AttributeValue);
+
 private:
 	const UE::Interchange::FAttributeKey Macro_CustomShadingModelKey = UE::Interchange::FAttributeKey(TEXT("ShadingModel"));
 	const UE::Interchange::FAttributeKey Macro_CustomTranslucencyLightingModeKey = UE::Interchange::FAttributeKey(TEXT("TranslucencyLightingMode"));
@@ -244,6 +250,7 @@ private:
 	const UE::Interchange::FAttributeKey Macro_CustomTwoSidedKey = UE::Interchange::FAttributeKey(TEXT("TwoSided"));
 	const UE::Interchange::FAttributeKey Macro_CustomOpacityMaskClipValueKey = UE::Interchange::FAttributeKey(TEXT("OpacityMaskClipValue"));
 	const UE::Interchange::FAttributeKey Macro_CustomRefractionMethodKey = UE::Interchange::FAttributeKey(TEXT("RefractionMethod"));
+	const UE::Interchange::FAttributeKey Macro_CustomScreenSpaceReflectionsKey = UE::Interchange::FAttributeKey(TEXT("ScreenSpaceReflections"));
 };
 
 UCLASS(BlueprintType)
diff --git a/Engine/Plugins/Interchange/Runtime/Source/Import/Private/Gltf/InterchangeGLTFMaterial.cpp b/Engine/Plugins/Interchange/Runtime/Source/Import/Private/Gltf/InterchangeGLTFMaterial.cpp
new file mode 100644
--- /dev/null
+++ b/Engine/Plugins/Interchange/Runtime/Source/Import/Private/Gltf/InterchangeGLTFMaterial.cpp
@@ -0,0 +1,641 @@
+// Copyright Epic Games, Inc. All Rights Reserved.
+
+#include "Gltf/InterchangeGLTFMaterial.h"
+#include "GLTFMaterial.h"
+#include "GLTFTexture.h"
+
+#include "Nodes/InterchangeBaseNodeContainer.h"
+#include "InterchangeTextureNode.h"
+#include "InterchangeShaderGraphNode.h"
+#include "InterchangeMaterialInstanceNode.h"
+#include "InterchangeShaderGraphNode.h"
+#include "InterchangeMaterialDefinitions.h"
+#include "InterchangeImportLog.h"
+
+#include "Materials/MaterialInstance.h"
+#include "Materials/MaterialInstanceConstant.h"
+#include "Materials/MaterialInstanceDynamic.h"
+
+#include "Async/Async.h"
+#include "Misc/Paths.h"
+#include "Misc/PackageName.h"
+
+namespace UE::Interchange::GLTFMaterials
+{
+
+	FString ToString(EShadingModel ShadingModel)
+	{
+		switch (ShadingModel)
+		{
+		case DEFAULT:
+			return TEXT("Default");
+		case UNLIT:
+			return TEXT("Unlit");
+		case CLEARCOAT:
+			return TEXT("ClearCoat");
+		case SHEEN:
+			return TEXT("Sheen");
+		case TRANSMISSION:
+			return TEXT("Transmission");
+		case SPECULARGLOSSINESS:
+			return TEXT("SpecularGlossiness");
+		default:
+			return TEXT("Unexpected");
+		}
+	}
+
+	FString ToString(EAlphaMode AlphaMode)
+	{
+		switch (AlphaMode)
+		{
+		case EAlphaMode::Opaque:
+			return TEXT("Opaque");
+		case EAlphaMode::Mask:
+			return TEXT("Mask");
+		case EAlphaMode::Blend:
+			return TEXT("Blend");
+		default:
+			return TEXT("Unexpected");
+		}
+	}
+
+	FString GetIdentifier(EShadingModel ShadingModel, EAlphaMode AlphaMode, bool bTwoSided)
+	{
+		FString Name = TEXT("MI_");
+
+		Name += ToString(ShadingModel);
+
+		if (ShadingModel != EShadingModel::TRANSMISSION)
+		{
+			Name += TEXT("_") + ToString(AlphaMode);
+		}
+
+		Name += bTwoSided ? TEXT("_DS") : TEXT("");
+
+		return Name;
+	}
+
+	FString GetMaterialFunctionPath(EShadingModel ShadingModel)
+	{
+		if (ShadingModelToMaterialFunctions.Contains(ShadingModel))
+		{
+			return ShadingModelToMaterialFunctions[ShadingModel].Key;
+		}
+
+		ensure(false);
+		return ShadingModelToMaterialFunctions.begin()->Value.Key;
+	}
+
+	TArray<FString> GetOutputs(EShadingModel ShadingModel)
+	{
+		if (ShadingModelToMaterialFunctions.Contains(ShadingModel))
+		{
+			return ShadingModelToMaterialFunctions[ShadingModel].Value;
+		}
+
+		ensure(false);
+		return ShadingModelToMaterialFunctions.begin()->Value.Value;
+	}
+
+	//Precedence order based on the logic in the InterchangeGenericMaterialPipeline: ClearCoat > Sheen > unlit
+	EShadingModel GetShadingModel(GLTF::FMaterial::EShadingModel GltfShadingModel, bool bHasClearCoat, bool bHasSheen, bool bUnlit, bool bHasTransmission)
+	{
+		if (bUnlit)
+		{
+			return EShadingModel::UNLIT;
+		}
+
+		if (GltfShadingModel == GLTF::FMaterial::EShadingModel::SpecularGlossiness)
+		{
+			return EShadingModel::SPECULARGLOSSINESS;
+		}
+
+		if (bHasTransmission)
+		{
+			return EShadingModel::TRANSMISSION;
+		}
+
+		if (bHasClearCoat)
+		{
+			return EShadingModel::CLEARCOAT;
+		}
+
+		if (bHasSheen)
+		{
+			return EShadingModel::SHEEN;
+		}
+
+		return EShadingModel::DEFAULT;
+	}
+
+	FString GetMaterialInstanceInputName(const FString& Name)
+	{
+		return *InterchangeGltfMaterialAttributeIdentifier + Name;
+	}
+
+	FString GetMaterialInputName(const FString& Name)
+	{
+		return UInterchangeShaderPortsAPI::MakeInputValueKey(Name);
+	}
+
+	struct FGLTFMaterialProcessor
+	{
+		UInterchangeBaseNodeContainer& NodeContainer;
+		const GLTF::FMaterial& GltfMaterial;
+		const TArray<GLTF::FTexture>& Textures;
+		
+		UInterchangeShaderGraphNode* ShaderGraphNode = nullptr; //Set by the Constructor
+		UInterchangeShaderNode* MaterialNode = nullptr;			//Set by ConfigureMaterialNode
+		UInterchangeShaderNode* MaterialInstanceNode = nullptr; //Set by ConfigureMaterialInstanceNode
+
+		EShadingModel ShadingModel = EShadingModel::DEFAULT;
+		EAlphaMode AlphaMode = EAlphaMode::Opaque;
+
+		FGLTFMaterialProcessor(UInterchangeBaseNodeContainer& InNodeContainer, const GLTF::FMaterial& InGltfMaterial, const TArray<GLTF::FTexture>& InTextures, UInterchangeShaderGraphNode* InShaderGraphNode)
+			: NodeContainer(InNodeContainer)
+			, GltfMaterial(InGltfMaterial)
+			, Textures(InTextures)
+			, ShaderGraphNode(InShaderGraphNode)
+		{
+			//Generate base properties:
+			ShadingModel = GetShadingModel(GltfMaterial.ShadingModel, GltfMaterial.bHasClearCoat, GltfMaterial.bHasSheen, GltfMaterial.bIsUnlitShadingModel, GltfMaterial.bHasTransmission);
+			AlphaMode = EAlphaMode(GltfMaterial.AlphaMode);
+
+			//Configure Nodes:
+			ConfigureMaterialNode();
+			ConfigureMaterialInstanceNode();
+		}
+
+		////////////////////////////
+		//Start of Config Functions:
+		////////////////////////////
+		void ConfigureMaterialNode()
+		{
+			FString MaterialFunctionPath = GetMaterialFunctionPath(ShadingModel);
+			FString MaterailFunctionFileName = FPaths::GetBaseFilename(MaterialFunctionPath);
+
+			UInterchangeFunctionCallShaderNode* MF_Body_Node = NewObject<UInterchangeFunctionCallShaderNode>(&NodeContainer);
+			MF_Body_Node->InitializeNode(UInterchangeShaderNode::MakeNodeUid(MaterailFunctionFileName, ShaderGraphNode->GetUniqueID()), MaterailFunctionFileName, EInterchangeNodeContainerType::TranslatedAsset);
+			MF_Body_Node->SetCustomMaterialFunction(MaterialFunctionPath);
+			NodeContainer.AddNode(MF_Body_Node);
+
+			MaterialNode = MF_Body_Node;
+
+			ShaderGraphNode->SetCustomTwoSided(GltfMaterial.bIsDoubleSided);
+			if (ShadingModel == EShadingModel::TRANSMISSION)
+			{
+				ShaderGraphNode->SetCustomScreenSpaceReflections(true);
+			}
+
+			TArray<FString> Outputs = GetOutputs(ShadingModel);
+
+			switch (AlphaMode)
+			{
+			case UE::Interchange::GLTFMaterials::Opaque:
+				Outputs.Remove(TEXT("Opacity"));
+				Outputs.Remove(TEXT("OpacityMask"));
+				break;
+			case UE::Interchange::GLTFMaterials::Mask:
+				if (ShadingModel != EShadingModel::TRANSMISSION)
+				{
+					Outputs.Remove(TEXT("Opacity"));
+				}
+				break;
+			case UE::Interchange::GLTFMaterials::Blend:
+				Outputs.Remove(TEXT("OpacityMask"));
+				break;
+			default:
+				break;
+			}
+
+			for (const FString& Input : Outputs)
+			{
+				UInterchangeShaderPortsAPI::ConnectOuputToInputByName(ShaderGraphNode, Input, MF_Body_Node->GetUniqueID(), Input);
+			}
+
+			//Set OpacityMaskClipValue if needed:
+			if (!GltfMaterial.bIsUnlitShadingModel &&
+				GltfMaterial.AlphaMode == GLTF::FMaterial::EAlphaMode::Mask &&
+				ShadingModel != EShadingModel::TRANSMISSION) //Transmission does not support AlphaCutoff
+			{
+				ShaderGraphNode->SetCustomOpacityMaskClipValue(GltfMaterial.AlphaCutoff);
+			}
+		}
+
+		void ConfigureMaterialInstanceNode()
+		{
+			MaterialInstanceNode = ShaderGraphNode;
+
+			FString ParentIdentifier = GetIdentifier(ShadingModel, AlphaMode, GltfMaterial.bIsDoubleSided && (ShadingModel != EShadingModel::TRANSMISSION)); //Transmission does not support two sided materials, see UInterchangeGenericMaterialPipeline::HandleCommonParameters
+			MaterialInstanceNode->AddStringAttribute(*(InterchangeGltfMaterialAttributeIdentifier + TEXT("ParentIdentifier")), ParentIdentifier);
+
+			//Set AlphaCutOff if needed:
+			if (!GltfMaterial.bIsUnlitShadingModel &&
+				GltfMaterial.AlphaMode == GLTF::FMaterial::EAlphaMode::Mask &&
+				ShadingModel != EShadingModel::TRANSMISSION) //Transmission does not support AlphaCutoff
+			{
+				//AlphaCutoff
+				SetScalar(Inputs::AlphaCutoff, GltfMaterial.AlphaCutoff, 0.f, EProcessType::MATERIALINSTANCE);
+			}
+		}
+		////////////////////////////
+		//End of Config Functions:
+		////////////////////////////
+
+
+		///////////////////////////////////
+		//Start of Helper Setter Functions:
+		///////////////////////////////////
+		enum EProcessType : uint8
+		{
+			MATERIAL = 1,
+			MATERIALINSTANCE = 2
+		};
+
+		void SetScalar(const FString& Name, float Value, float DefaultValue, int ProcessType = (EProcessType::MATERIAL | EProcessType::MATERIALINSTANCE))
+		{
+			if (!FMath::IsNearlyEqual(Value, DefaultValue))
+			{
+				if (ProcessType & EProcessType::MATERIAL)
+				{
+					MaterialNode->AddFloatAttribute(GetMaterialInputName(Name), Value);
+				}
+
+				if (ProcessType & EProcessType::MATERIALINSTANCE)
+				{
+					MaterialInstanceNode->AddFloatAttribute(GetMaterialInstanceInputName(Name), Value);
+				}
+			}
+		}
+
+		void SetVec3(const FString& Name, const FVector3f& Value, const FVector3f& DefaultValue, int ProcessType = (EProcessType::MATERIAL | EProcessType::MATERIALINSTANCE))
+		{
+			if (!Value.Equals(DefaultValue, UE_SMALL_NUMBER))
+			{
+
+				if (ProcessType & EProcessType::MATERIAL)
+				{
+					MaterialNode->AddLinearColorAttribute(GetMaterialInputName(Name), Value);
+				}
+
+				if (ProcessType & EProcessType::MATERIALINSTANCE)
+				{
+					MaterialInstanceNode->AddLinearColorAttribute(GetMaterialInstanceInputName(Name), Value);
+				}
+			}
+		}
+
+		void SetVec4(const FString& Name, const FVector4f& Value, const FVector4f& DefaultValue, int ProcessType = (EProcessType::MATERIAL | EProcessType::MATERIALINSTANCE))
+		{
+			if (!Value.Equals(DefaultValue, UE_SMALL_NUMBER))
+			{
+				if (ProcessType & EProcessType::MATERIAL)
+				{
+					MaterialNode->AddLinearColorAttribute(GetMaterialInputName(Name), Value);
+				}
+
+				if (ProcessType & EProcessType::MATERIALINSTANCE)
+				{
+					MaterialInstanceNode->AddLinearColorAttribute(GetMaterialInstanceInputName(Name), Value);
+				}
+			}
+		}
+
+		void SetColor(const FString& Name, const FVector4f& Value, const FVector4f& DefaultValue)
+		{
+			{//Material specific settings:
+				SetVec3(Name + Inputs::Color_RGB, FVector3f(Value), FVector3f(DefaultValue), EProcessType::MATERIAL);
+				SetScalar(Name + Inputs::Color_A, Value.W, DefaultValue.W, EProcessType::MATERIAL);
+			}
+
+			{//MaterialInstance specific settings:
+				SetVec4(Name, Value, DefaultValue, EProcessType::MATERIALINSTANCE);
+			}
+		}
+
+		void SetMap(const GLTF::FTextureMap& TextureMap, const FString& Name)
+		{
+			if (Textures.IsValidIndex(TextureMap.TextureIndex))
+			{
+				{//Material specific settings:
+					//create TextureObject:
+					UInterchangeShaderNode* ColorNode = UInterchangeShaderNode::Create(&NodeContainer, Name, MaterialNode->GetUniqueID());
+					ColorNode->SetCustomShaderType(UE::Interchange::Materials::Standard::Nodes::TextureObject::Name.ToString());
+
+					ColorNode->AddStringAttribute(UInterchangeShaderPortsAPI::MakeInputValueKey(UE::Interchange::Materials::Standard::Nodes::TextureObject::Inputs::Texture.ToString()), UInterchangeTextureNode::MakeNodeUid(Textures[TextureMap.TextureIndex].UniqueId));
+
+					UInterchangeShaderPortsAPI::ConnectDefaultOuputToInput(MaterialNode, Name, ColorNode->GetUniqueID());
+				}
+
+				{//MaterialInstance specific settings:
+					MaterialInstanceNode->AddStringAttribute(GetMaterialInstanceInputName(Name), UInterchangeTextureNode::MakeNodeUid(Textures[TextureMap.TextureIndex].UniqueId));
+				}
+
+				//TexCoord decision in MaterialInstances work based on a Switch Node
+				//Currently Only supports UV0 and UV1
+				// [0...1) -> UV0
+				// [1...2) -> UV1
+				// [2...3) -> UV2
+				// [3...4) -> UV3
+				// else    -> UV0 (defaults to 0)
+				SetScalar(Name + TEXT("_TexCoord"), TextureMap.TexCoord, 0.f);
+			}
+
+			if (TextureMap.bHasTextureTransform)
+			{
+				{//Material specific settings:
+					SetScalar(Name + TEXT("_Offset_X"), TextureMap.TextureTransform.Offset[0], 0.f, EProcessType::MATERIAL);
+					SetScalar(Name + TEXT("_Offset_Y"), TextureMap.TextureTransform.Offset[1], 0.f, EProcessType::MATERIAL);
+
+					SetScalar(Name + TEXT("_Scale_X"), TextureMap.TextureTransform.Scale[0], 1.f, EProcessType::MATERIAL);
+					SetScalar(Name + TEXT("_Scale_Y"), TextureMap.TextureTransform.Scale[1], 1.f, EProcessType::MATERIAL);
+				}
+
+				{//MaterialInstance specific settings:
+					FVector4f OffsetScale(TextureMap.TextureTransform.Offset[0], TextureMap.TextureTransform.Offset[1], TextureMap.TextureTransform.Scale[0], TextureMap.TextureTransform.Scale[1]);
+					SetVec4(Name + TEXT("_OffsetScale"), OffsetScale, FVector4f(0, 0, 1, 1), EProcessType::MATERIALINSTANCE);
+				}
+
+				if (!FMath::IsNearlyZero(TextureMap.TextureTransform.Rotation))
+				{
+					float AngleRadians = TextureMap.TextureTransform.Rotation;
+
+					if (AngleRadians < 0.0f)
+					{
+						AngleRadians = TWO_PI - AngleRadians;
+					}
+
+					AngleRadians = 1.0f - (AngleRadians / TWO_PI);
+
+					SetScalar(Name + TEXT("_Rotation"), AngleRadians, 0.f);
+				}
+			}
+		}
+		////////////////////////////////
+		//End Of Helper Setter Functions
+		////////////////////////////////
+		
+
+		////////////////////////////////////////////////////////////////////////////////////
+		/// Parse Material Values from GltfMaterial to MaterialNode and MaterialInstanceNode
+		////////////////////////////////////////////////////////////////////////////////////
+		void ParseMaterialValues()
+		{
+			if (GltfMaterial.bIsUnlitShadingModel)
+			{
+				//BaseColorTexture
+				//BaseColorTexture_OffsetScale
+				//BaseColorTExture_Rotation
+				//BaseColorTExture_TexCoord
+				SetMap(GltfMaterial.BaseColor, Inputs::BaseColorTexture);
+
+				//BaseColorFactor
+				SetColor(Inputs::BaseColorFactor, GltfMaterial.BaseColorFactor, FVector4f(1, 1, 1, 1));
+				return;
+			}
+
+			{
+				//NormalTexture
+				//NormalTexture_OffsetScale
+				//NormalTexture_Rotation
+				//NormalTexture_TexCoord
+				SetMap(GltfMaterial.Normal, Inputs::NormalTexture);
+
+				//NormalScale
+				SetScalar(Inputs::NormalScale, GltfMaterial.NormalScale, 1.f);
+			}
+
+			if (!GltfMaterial.bHasTransmission)
+			{
+				//EmissiveTexture
+				//EmissiveTexture_OffsetScale
+				//EmissiveTexture_Rotation
+				//EmissiveTexture_TexCoord
+				SetMap(GltfMaterial.Emissive, Inputs::EmissiveTexture);
+
+				//EmissiveFactor
+				SetVec3(Inputs::EmissiveFactor, GltfMaterial.EmissiveFactor, FVector3f(0, 0, 0));
+
+				//EmissiveStrength
+				SetScalar(Inputs::EmissiveStrength, GltfMaterial.bHasEmissiveStrength ? GltfMaterial.EmissiveStrength : 1.f, 1.f);
+			}
+
+			{
+				//OcclusionTexture
+				//OcclusionTexture_OffsetScale
+				//OcclusionTexture_Rotation
+				//OcclusionTexture_TexCoord
+				SetMap(GltfMaterial.Occlusion, Inputs::OcclusionTexture);
+
+				//OcclusionStrength
+				SetScalar(Inputs::OcclusionStrength, GltfMaterial.OcclusionStrength, 1.f);
+			}
+
+			if (GltfMaterial.ShadingModel == GLTF::FMaterial::EShadingModel::SpecularGlossiness)
+			{
+				//////////////////////
+				//specular glossiness:
+				//////////////////////
+
+				{
+					//DiffuseTexture
+					//DiffuseTexture_OffsetScale
+					//DiffuseTexture_Rotation
+					//DiffuseTexture_TexCoord
+					SetMap(GltfMaterial.BaseColor, Inputs::DiffuseTexture);
+
+					//DiffuseFactor
+					SetColor(Inputs::DiffuseFactor, GltfMaterial.BaseColorFactor, FVector4f(1, 1, 1, 1));
+				}
+
+				{
+					//SpecularGlossinessTexture
+					//SpecularGlossinessTexture_OffsetScale
+					//SpecularGlossinessTexture_Rotation
+					//SpecularGlossinessTexture_TexCoord
+					SetMap(GltfMaterial.SpecularGlossiness.Map, Inputs::SpecularGlossinessTexture);
+
+					//SpecFactor
+					FVector3f SpecularFactor(GltfMaterial.SpecularGlossiness.SpecularFactor[0], GltfMaterial.SpecularGlossiness.SpecularFactor[1], GltfMaterial.SpecularGlossiness.SpecularFactor[2]);
+					SetVec3(Inputs::SpecFactor, SpecularFactor, FVector3f(1, 1, 1));
+
+					//GlossinessFactor
+					SetScalar(Inputs::GlossinessFactor, GltfMaterial.SpecularGlossiness.GlossinessFactor, 1.f);
+				}
+
+				return;
+			}
+			else if (GltfMaterial.ShadingModel == GLTF::FMaterial::EShadingModel::MetallicRoughness)
+			{
+				{
+					//BaseColorTexture
+					//BaseColorTexture_OffsetScale
+					//BaseColorTexture_Rotation
+					//BaseColorTexture_TexCoord
+					SetMap(GltfMaterial.BaseColor, Inputs::BaseColorTexture);
+
+					//BaseColorFactor
+					SetColor(Inputs::BaseColorFactor, GltfMaterial.BaseColorFactor, FVector4f(1, 1, 1, 1));
+				}
+
+				{
+					//MetallicRoughnessTexture
+					//MetallicRoughnessTexture_OffsetScale
+					//MetallicRoughnessTexture_Rotation
+					//MetallicRoughnessTexture_TexCoord
+					SetMap(GltfMaterial.MetallicRoughness.Map, Inputs::MetallicRoughnessTexture);
+
+					//MetallicFactor
+					SetScalar(Inputs::MetallicFactor, GltfMaterial.MetallicRoughness.MetallicFactor, 1.f);
+
+					//RoughnessFactor
+					SetScalar(Inputs::RoughnessFactor, GltfMaterial.MetallicRoughness.RoughnessFactor, 1.f);
+				}
+
+				if (GltfMaterial.bHasSpecular)
+				{
+					//SpecularTexture
+					//SpecularTexture_OffsetScale
+					//SpecularTexture_Rotation
+					//SpecularTexture_TexCoord
+					SetMap(GltfMaterial.Specular.SpecularMap, Inputs::SpecularTexture);
+
+					//SpecularFactor
+					SetScalar(Inputs::SpecularFactor, GltfMaterial.Specular.SpecularFactor, 1.f);
+				}
+			}
+			else
+			{
+				ensure(false);
+			}
+
+			if (GltfMaterial.bHasIOR)
+			{
+				//ior
+				SetScalar(Inputs::IOR, GltfMaterial.IOR, 1.5f);
+			}
+
+			if (GltfMaterial.bHasTransmission)
+			{
+				{
+					//TransmissionTexture
+					//TransmissionTexture_OffsetScale
+					//TransmissionTexture_Rotation
+					//TransmissionTexture_TexCoord
+					SetMap(GltfMaterial.Transmission.TransmissionMap, Inputs::TransmissionTexture);
+
+					//TransmissionFactor
+					SetScalar(Inputs::TransmissionFactor, GltfMaterial.Transmission.TransmissionFactor, 0.f);
+				}
+			}
+			else if (GltfMaterial.bHasClearCoat)
+			{
+				{
+					//ClearCoatTexture
+					//ClearCoatTexture_OffsetScale
+					//ClearCoatTexture_Rotation
+					//ClearCoatTexture_TexCoord
+					SetMap(GltfMaterial.ClearCoat.ClearCoatMap, Inputs::ClearCoatTexture);
+
+					//ClearCoatFactor
+					SetScalar(Inputs::ClearCoatFactor, GltfMaterial.ClearCoat.ClearCoatFactor, 0.f);
+				}
+
+				{
+					//ClearCoatRoughnessTexture
+					//ClearCoatRoughnessTexture_OffsetScale
+					//ClearCoatRoughnessTexture_Rotation
+					//ClearCoatRoughnessTexture_TexCoord
+					SetMap(GltfMaterial.ClearCoat.RoughnessMap, Inputs::ClearCoatRoughnessTexture);
+
+					//ClearCoatRoughnessFactor
+					SetScalar(Inputs::ClearCoatRoughnessFactor, GltfMaterial.ClearCoat.Roughness, 0.f);
+				}
+
+				{
+					//ClearCoatNormalTexture
+					//ClearCoatNormalTexture_OffsetScale
+					//ClearCoatNormalTexture_Rotation
+					//ClearCoatNormalTexture_TexCoord
+					SetMap(GltfMaterial.ClearCoat.NormalMap, Inputs::ClearCoatNormalTexture);
+
+					//ClearCoatNormalFactor
+					SetScalar(Inputs::ClearCoatNormalScale, GltfMaterial.ClearCoat.NormalMapUVScale, 1.f);
+				}
+			}
+			else if (GltfMaterial.bHasSheen)
+			{
+				{
+					//SheenColorTexture
+					//SheenColorTexture_OffsetScale
+					//SheenColorTexture_Rotation
+					//SheenColorTexture_TexCoord
+					SetMap(GltfMaterial.Sheen.SheenColorMap, Inputs::SheenColorTexture);
+
+					//SheenColorFactor
+					FVector3f SheenColorFactor(GltfMaterial.Sheen.SheenColorFactor[0], GltfMaterial.Sheen.SheenColorFactor[1], GltfMaterial.Sheen.SheenColorFactor[2]);
+					SetVec3(Inputs::SheenColorFactor, SheenColorFactor, FVector3f(0, 0, 0));
+
+					//SheenRoughnessTexture
+					//SheenRoughnessTexture_OffsetScale
+					//SheenRoughnessTexture_Rotation
+					//SheenRoughnessTexture_TexCoord
+					SetMap(GltfMaterial.Sheen.SheenRoughnessMap, Inputs::SheenRoughnessTexture);
+
+					//SheenRoughnessFactor
+					SetScalar(Inputs::SheenRoughnessFactor, GltfMaterial.Sheen.SheenRoughnessFactor, 0.f);
+				}
+			}
+		}
+	};
+
+	void HandleGltfMaterial(UInterchangeBaseNodeContainer& NodeContainer, const GLTF::FMaterial& GltfMaterial, const TArray<GLTF::FTexture>& Textures, UInterchangeShaderGraphNode* ShaderGraphNode)
+	{
+		using namespace UE::Interchange::GLTFMaterials;
+
+		//Constructor of FGLTFMaterialProcessor will configure and generate necessary properties for Parsing:
+		FGLTFMaterialProcessor GLTFMaterailProcessor(NodeContainer, GltfMaterial, Textures, ShaderGraphNode);
+
+		GLTFMaterailProcessor.ParseMaterialValues();
+	}
+
+	bool AreRequiredPackagesLoaded()
+	{
+		auto ArePackagesLoaded = [](const TArray<FString>& PackagePaths) -> bool
+		{
+			bool bAllLoaded = true;
+
+			for (const FString& PackagePath : PackagePaths)
+			{
+				const FString ObjectPath(FPackageName::ExportTextPathToObjectPath(PackagePath));
+
+				if (FPackageName::DoesPackageExist(ObjectPath))
+				{
+					if (FSoftObjectPath(ObjectPath).TryLoad())
+					{
+						continue;
+					}
+					else
+					{
+						UE_LOG(LogInterchangeImport, Warning, TEXT("Couldn't load %s"), *PackagePath);
+					}
+				}
+				else
+				{
+					UE_LOG(LogInterchangeImport, Warning, TEXT("Couldn't find %s"), *PackagePath);
+				}
+
+				bAllLoaded = false;
+			}
+
+			return bAllLoaded;
+		};
+
+		TArray<FString> RequiredPackages = GetRequiredMaterialFunctionPaths();
+
+		static const bool bRequiredPackagesLoaded = ArePackagesLoaded(RequiredPackages);
+
+		return bRequiredPackagesLoaded;
+	}
+}
+
diff --git a/Engine/Plugins/Interchange/Runtime/Source/Import/Private/Gltf/InterchangeGLTFMaterialInstances.cpp b/Engine/Plugins/Interchange/Runtime/Source/Import/Private/Gltf/InterchangeGLTFMaterialInstances.cpp
deleted file mode 100644
diff --git a/Engine/Plugins/Interchange/Runtime/Source/Import/Private/Gltf/InterchangeGltfTranslator.cpp b/Engine/Plugins/Interchange/Runtime/Source/Import/Private/Gltf/InterchangeGltfTranslator.cpp
--- a/Engine/Plugins/Interchange/Runtime/Source/Import/Private/Gltf/InterchangeGltfTranslator.cpp
+++ b/Engine/Plugins/Interchange/Runtime/Source/Import/Private/Gltf/InterchangeGltfTranslator.cpp
@@ -38,7 +38,7 @@
 #include "SkeletalMeshOperations.h"
 
 #include "Gltf/InterchangeGltfPrivate.h"
-#include "Gltf/InterchangeGLTFMaterialInstances.h"
+#include "Gltf/InterchangeGLTFMaterial.h"
 
 #include "EngineAnalytics.h"
 #include "Engine/RendererSettings.h"
@@ -352,601 +352,6 @@ void UInterchangeGLTFTranslator::HandleGltfNode( UInterchangeBaseNodeContainer&
 	}
 }
 
-void UInterchangeGLTFTranslator::HandleGltfMaterialParameter( UInterchangeBaseNodeContainer& NodeContainer, const GLTF::FTextureMap& TextureMap, UInterchangeShaderNode& ShaderNode,
-		const FString& MapName, const TVariant< FLinearColor, float >& MapFactor, const FString& OutputChannel, const bool bInverse, const bool bIsNormal, const bool bUseVertexColor) const
-{
-	using namespace UE::Interchange::Materials;
-	using namespace UE::Interchange::Gltf::Private;
-
-	UInterchangeShaderNode* NodeToConnectTo = &ShaderNode;
-	FString InputToConnectTo = MapName;
-
-	if (bUseVertexColor)
-	{
-		//From GLTF Specification:
-		// "if a primitive specifies a vertex color using the attribute semantic property COLOR_0, then this value acts as an additional linear multiplier to base color."
-		const FString MultiplierNodeName = MapName + TEXT("VertexColorMultiply");
-		UInterchangeShaderNode* MultiplierNode = UInterchangeShaderNode::Create(&NodeContainer, MultiplierNodeName, ShaderNode.GetUniqueID());
-		MultiplierNode->SetCustomShaderType(Standard::Nodes::Multiply::Name.ToString());
-
-		UInterchangeShaderPortsAPI::ConnectDefaultOuputToInput(NodeToConnectTo, InputToConnectTo, MultiplierNode->GetUniqueID());
-		NodeToConnectTo = MultiplierNode;
-		InputToConnectTo = Standard::Nodes::Multiply::Inputs::B.ToString();
-
-
-		const FString VertexColorNodeName = MapName + TEXT("VertexColor");
-		UInterchangeShaderNode* VertexColorNode = UInterchangeShaderNode::Create(&NodeContainer, VertexColorNodeName, ShaderNode.GetUniqueID());
-		VertexColorNode->SetCustomShaderType(Standard::Nodes::VertexColor::Name.ToString());
-
-		UInterchangeShaderPortsAPI::ConnectDefaultOuputToInput(MultiplierNode, Standard::Nodes::Multiply::Inputs::A.ToString(), VertexColorNode->GetUniqueID());
-	}
-
-	if (bInverse)
-	{
-		const FString OneMinusNodeName = MapName + TEXT("OneMinus");
-		UInterchangeShaderNode* OneMinusNode = UInterchangeShaderNode::Create( &NodeContainer, OneMinusNodeName, ShaderNode.GetUniqueID() );
-
-		OneMinusNode->SetCustomShaderType(Standard::Nodes::OneMinus::Name.ToString());
-
-		UInterchangeShaderPortsAPI::ConnectDefaultOuputToInput(NodeToConnectTo, InputToConnectTo, OneMinusNode->GetUniqueID());
-
-		NodeToConnectTo = OneMinusNode;
-		InputToConnectTo = Standard::Nodes::OneMinus::Inputs::Input.ToString();
-	}
-
-	bool bTextureHasImportance = true;
-
-	if ( MapFactor.IsType< float >() )
-	{
-		bTextureHasImportance = !FMath::IsNearlyZero( MapFactor.Get< float >() );
-	}
-	else if ( MapFactor.IsType< FLinearColor >() )
-	{
-		bTextureHasImportance = !MapFactor.Get< FLinearColor >().IsAlmostBlack();
-	}
-
-	if ( bTextureHasImportance && GltfAsset.Textures.IsValidIndex( TextureMap.TextureIndex ) )
-	{
-		const FString ColorNodeName = MapName;
-		UInterchangeShaderNode* ColorNode = UInterchangeShaderNode::Create( &NodeContainer, ColorNodeName, ShaderNode.GetUniqueID() );
-		ColorNode->SetCustomShaderType( Standard::Nodes::TextureSample::Name.ToString() );
-
-		const FString TextureUid = UInterchangeTextureNode::MakeNodeUid(GltfAsset.Textures[TextureMap.TextureIndex].UniqueId);
-
-		ColorNode->AddStringAttribute( UInterchangeShaderPortsAPI::MakeInputValueKey( Standard::Nodes::TextureSample::Inputs::Texture.ToString() ), TextureUid );
-
-		if (TextureMap.TexCoord > 0 || TextureMap.bHasTextureTransform)
-		{
-			UInterchangeShaderNode* TexCoordNode = UInterchangeShaderNode::Create(&NodeContainer, MapName + TEXT("\\TexCoord"), ShaderNode.GetUniqueID());
-			TexCoordNode->SetCustomShaderType(Standard::Nodes::TextureCoordinate::Name.ToString());
-
-			TexCoordNode->AddInt32Attribute(UInterchangeShaderPortsAPI::MakeInputValueKey(Standard::Nodes::TextureCoordinate::Inputs::Index.ToString()), TextureMap.TexCoord);
-
-			UInterchangeShaderPortsAPI::ConnectDefaultOuputToInput(ColorNode, Standard::Nodes::TextureSample::Inputs::Coordinates.ToString(), TexCoordNode->GetUniqueID());
-
-			if (TextureMap.bHasTextureTransform)
-			{
-				HandleGltfTextureTransform(NodeContainer, TextureMap.TextureTransform, TextureMap.TexCoord, *TexCoordNode);
-			}
-		}
-
-		bool bNeedsFactorNode = false;
-
-		if ( MapFactor.IsType< float >() )
-		{
-			bNeedsFactorNode = !FMath::IsNearlyEqual( MapFactor.Get< float >(), 1.f );
-		}
-		else if ( MapFactor.IsType< FLinearColor >() )
-		{
-			bNeedsFactorNode = !MapFactor.Get< FLinearColor >().Equals( FLinearColor::White );
-		}
-
-		if ( bNeedsFactorNode )
-		{
-			UInterchangeShaderNode* FactorNode = UInterchangeShaderNode::Create( &NodeContainer, ColorNodeName + TEXT("_Factor"), ShaderNode.GetUniqueID() );
-
-			if ( bIsNormal )
-			{
-				FactorNode->SetCustomShaderType( Standard::Nodes::FlattenNormal::Name.ToString() );
-
-				UInterchangeShaderNode* FactorOneMinusNode = UInterchangeShaderNode::Create( &NodeContainer, ColorNodeName + TEXT("_Factor_OneMinus"), ShaderNode.GetUniqueID() );
-				FactorOneMinusNode->SetCustomShaderType(Standard::Nodes::OneMinus::Name.ToString());
-
-				if ( MapFactor.IsType< float >() )
-				{
-					FactorOneMinusNode->AddFloatAttribute( UInterchangeShaderPortsAPI::MakeInputValueKey( Standard::Nodes::OneMinus::Inputs::Input.ToString() ), MapFactor.Get< float >() );
-				}
-
-				UInterchangeShaderPortsAPI::ConnectOuputToInputByName( FactorNode, Standard::Nodes::FlattenNormal::Inputs::Normal.ToString(), ColorNode->GetUniqueID(), OutputChannel );
-				UInterchangeShaderPortsAPI::ConnectDefaultOuputToInput( FactorNode, Standard::Nodes::FlattenNormal::Inputs::Flatness.ToString(), FactorOneMinusNode->GetUniqueID() );
-			}
-			else
-			{
-				FactorNode->SetCustomShaderType( Standard::Nodes::Multiply::Name.ToString() );
-
-				if ( MapFactor.IsType< float >() )
-				{
-					FactorNode->AddFloatAttribute( UInterchangeShaderPortsAPI::MakeInputValueKey( Standard::Nodes::Multiply::Inputs::B.ToString() ), MapFactor.Get< float >() );
-				}
-				else if ( MapFactor.IsType< FLinearColor >() )
-				{
-					FactorNode->AddLinearColorAttribute( UInterchangeShaderPortsAPI::MakeInputValueKey( Standard::Nodes::Multiply::Inputs::B.ToString() ), MapFactor.Get< FLinearColor >() );
-				}
-
-				UInterchangeShaderPortsAPI::ConnectOuputToInputByName( FactorNode, Standard::Nodes::Multiply::Inputs::A.ToString(), ColorNode->GetUniqueID(), OutputChannel );
-			}
-
-			UInterchangeShaderPortsAPI::ConnectDefaultOuputToInput( NodeToConnectTo, InputToConnectTo, FactorNode->GetUniqueID() );
-		}
-		else
-		{
-			UInterchangeShaderPortsAPI::ConnectOuputToInputByName( NodeToConnectTo, InputToConnectTo, ColorNode->GetUniqueID(), OutputChannel );
-		}
-	}
-	else
-	{
-		if ( bIsNormal && !bTextureHasImportance )
-		{
-			//default normal value is 0,0,1 (blue)
-			NodeToConnectTo->AddLinearColorAttribute( UInterchangeShaderPortsAPI::MakeInputValueKey(InputToConnectTo), FLinearColor::Blue );
-		}
-		else
-		{
-			if ( MapFactor.IsType< FLinearColor >() )
-			{
-				NodeToConnectTo->AddLinearColorAttribute( UInterchangeShaderPortsAPI::MakeInputValueKey(InputToConnectTo), MapFactor.Get< FLinearColor >() );
-			}
-			else if ( MapFactor.IsType< float >() )
-			{
-				NodeToConnectTo->AddFloatAttribute( UInterchangeShaderPortsAPI::MakeInputValueKey(InputToConnectTo), MapFactor.Get< float >() );
-			}
-		}
-	}
-}
-
-void UInterchangeGLTFTranslator::HandleGltfMaterial( UInterchangeBaseNodeContainer& NodeContainer, const GLTF::FMaterial& GltfMaterial, UInterchangeShaderGraphNode& ShaderGraphNode, bool bUseVertexColor) const
-{
-	using namespace UE::Interchange::Materials;
-
-	ShaderGraphNode.SetCustomTwoSided( GltfMaterial.bIsDoubleSided );
-
-	if (GltfMaterial.bIsUnlitShadingModel)
-	{
-		// Base Color
-		{
-			TVariant< FLinearColor, float > BaseColorFactor;
-			BaseColorFactor.Set< FLinearColor >(FLinearColor(GltfMaterial.BaseColorFactor));
-
-			HandleGltfMaterialParameter(NodeContainer, GltfMaterial.BaseColor, ShaderGraphNode, Unlit::Parameters::UnlitColor.ToString(),
-				BaseColorFactor, Standard::Nodes::TextureSample::Outputs::RGB.ToString(), false, false, bUseVertexColor);
-		}
-
-		// Opacity (use the base color alpha channel)
-		if (GltfMaterial.AlphaMode != GLTF::FMaterial::EAlphaMode::Opaque)
-		{
-			TVariant< FLinearColor, float > OpacityFactor;
-			OpacityFactor.Set< float >(GltfMaterial.BaseColorFactor.W);
-
-			HandleGltfMaterialParameter(NodeContainer, GltfMaterial.BaseColor, ShaderGraphNode, PBRMR::Parameters::Opacity.ToString(),
-				OpacityFactor, Standard::Nodes::TextureSample::Outputs::A.ToString());
-		}
-
-		return;
-	}
-	
-	//if there is a clearcoatnormal map then we want to swap it with the normals map
-	//as the Interchange pipeline will connect the clearcoatnormal map to ClearCoatBottomNormalMap
-	//however as per specification the gltf.clearcoatnormal.map should be the top clearcoat and the gltf.normal.map should be the bottom one.
-	bool bSwapNormalAndClearCoatNormal = GltfMaterial.bHasClearCoat;
-
-	if ( GltfMaterial.ShadingModel == GLTF::FMaterial::EShadingModel::MetallicRoughness )
-	{
-		// Base Color
-		{
-			TVariant< FLinearColor, float > BaseColorFactor;
-			BaseColorFactor.Set< FLinearColor >( FLinearColor( GltfMaterial.BaseColorFactor ) );
-
-			HandleGltfMaterialParameter( NodeContainer, GltfMaterial.BaseColor, ShaderGraphNode, PBRMR::Parameters::BaseColor.ToString(),
-				BaseColorFactor, Standard::Nodes::TextureSample::Outputs::RGB.ToString(), false, false, bUseVertexColor );
-		}
-
-		// Metallic
-		{
-			TVariant< FLinearColor, float > MetallicFactor;
-			MetallicFactor.Set< float >( GltfMaterial.MetallicRoughness.MetallicFactor );
-
-			HandleGltfMaterialParameter( NodeContainer, GltfMaterial.MetallicRoughness.Map, ShaderGraphNode, PBRMR::Parameters::Metallic.ToString(),
-				MetallicFactor, Standard::Nodes::TextureSample::Outputs::B.ToString() );
-		}
-
-		// Roughness
-		{
-			TVariant< FLinearColor, float > RoughnessFactor;
-			RoughnessFactor.Set< float >( GltfMaterial.MetallicRoughness.RoughnessFactor );
-
-			HandleGltfMaterialParameter( NodeContainer, GltfMaterial.MetallicRoughness.Map, ShaderGraphNode, PBRMR::Parameters::Roughness.ToString(),
-				RoughnessFactor, Standard::Nodes::TextureSample::Outputs::G.ToString() );
-		}
-
-		// Specular
-		if (GltfMaterial.bHasSpecular)
-		{
-			TVariant< FLinearColor, float > SpecularFactor;
-			SpecularFactor.Set< float >( GltfMaterial.Specular.SpecularFactor );
-
-			HandleGltfMaterialParameter( NodeContainer, GltfMaterial.Specular.SpecularMap, ShaderGraphNode, PBRMR::Parameters::Specular.ToString(),
-				SpecularFactor, Standard::Nodes::TextureSample::Outputs::RGB.ToString() );
-		}
-	}
-	else if ( GltfMaterial.ShadingModel == GLTF::FMaterial::EShadingModel::SpecularGlossiness )
-	{
-		// Diffuse Color
-		{
-			TVariant< FLinearColor, float > DiffuseColorFactor;
-			DiffuseColorFactor.Set< FLinearColor >( FLinearColor( GltfMaterial.BaseColorFactor ) );
-
-			HandleGltfMaterialParameter( NodeContainer, GltfMaterial.BaseColor, ShaderGraphNode, PBRSG::Parameters::DiffuseColor.ToString(),
-				DiffuseColorFactor, Standard::Nodes::TextureSample::Outputs::RGB.ToString(), false, false, bUseVertexColor );
-		}
-
-		// Specular Color
-		{
-			TVariant< FLinearColor, float > SpecularColorFactor;
-			SpecularColorFactor.Set< FLinearColor >( FLinearColor( GltfMaterial.SpecularGlossiness.SpecularFactor ) );
-
-			HandleGltfMaterialParameter( NodeContainer, GltfMaterial.SpecularGlossiness.Map, ShaderGraphNode, PBRSG::Parameters::SpecularColor.ToString(),
-				SpecularColorFactor, Standard::Nodes::TextureSample::Outputs::RGB.ToString() );
-		}
-
-		// Glossiness
-		{
-			TVariant< FLinearColor, float > GlossinessFactor;
-			GlossinessFactor.Set< float >( GltfMaterial.SpecularGlossiness.GlossinessFactor );
-
-			HandleGltfMaterialParameter( NodeContainer, GltfMaterial.SpecularGlossiness.Map, ShaderGraphNode, PBRSG::Parameters::Glossiness.ToString(),
-				GlossinessFactor, Standard::Nodes::TextureSample::Outputs::A.ToString() );
-		}
-	}
-
-	// Additional maps
-	{
-		// Normal
-		if ( GltfMaterial.Normal.TextureIndex != INDEX_NONE )
-		{
-			TVariant< FLinearColor, float > NormalFactor;
-			NormalFactor.Set< float >( GltfMaterial.NormalScale );
-
-			const bool bInverse = false;
-			const bool bIsNormal = true;
-
-			HandleGltfMaterialParameter( NodeContainer, GltfMaterial.Normal, ShaderGraphNode, bSwapNormalAndClearCoatNormal ? ClearCoat::Parameters::ClearCoatNormal.ToString() : Common::Parameters::Normal.ToString(),
-				NormalFactor, Standard::Nodes::TextureSample::Outputs::RGB.ToString(), bInverse, bIsNormal );
-		}
-
-		// Emissive
-		if ( GltfMaterial.Emissive.TextureIndex != INDEX_NONE || !GltfMaterial.EmissiveFactor.IsNearlyZero() || GltfMaterial.bHasEmissiveStrength)
-		{
-			TVariant< FLinearColor, float > EmissiveFactor;
-			EmissiveFactor.Set< FLinearColor >(GltfMaterial.bHasEmissiveStrength ? FLinearColor(GltfMaterial.EmissiveFactor) * GltfMaterial.EmissiveStrength : FLinearColor(GltfMaterial.EmissiveFactor));
-
-			HandleGltfMaterialParameter( NodeContainer, GltfMaterial.Emissive, ShaderGraphNode, Common::Parameters::EmissiveColor.ToString(),
-				EmissiveFactor, Standard::Nodes::TextureSample::Outputs::RGB.ToString() );
-		}
-
-		// Occlusion
-		if ( GltfMaterial.Occlusion .TextureIndex != INDEX_NONE )
-		{
-			TVariant< FLinearColor, float > OcclusionFactor;
-			OcclusionFactor.Set< float >( GltfMaterial.OcclusionStrength );
-
-			HandleGltfMaterialParameter( NodeContainer, GltfMaterial.Occlusion, ShaderGraphNode, PBRMR::Parameters::Occlusion.ToString(),
-				OcclusionFactor, Standard::Nodes::TextureSample::Outputs::R.ToString() );
-		}
-
-		// Opacity (use the base color alpha channel)
-		if ( GltfMaterial.AlphaMode != GLTF::FMaterial::EAlphaMode::Opaque )
-		{
-			TVariant< FLinearColor, float > OpacityFactor;
-			OpacityFactor.Set< float >( GltfMaterial.BaseColorFactor.W );
-
-			HandleGltfMaterialParameter( NodeContainer, GltfMaterial.BaseColor, ShaderGraphNode, PBRMR::Parameters::Opacity.ToString(),
-				OpacityFactor, Standard::Nodes::TextureSample::Outputs::A.ToString() );
-		}
-
-		// Alpha cutoff
-		if ( GltfMaterial.AlphaMode == GLTF::FMaterial::EAlphaMode::Mask )
-		{
-			ShaderGraphNode.SetCustomOpacityMaskClipValue( GltfMaterial.AlphaCutoff );
-		}
-
-		// IOR
-		if ( GltfMaterial.bHasIOR )
-		{
-			ShaderGraphNode.AddFloatAttribute( UInterchangeShaderPortsAPI::MakeInputValueKey( PBRMR::Parameters::IndexOfRefraction.ToString() ), GltfMaterial.IOR );
-		}
-	}
-
-	if ( GltfMaterial.bHasClearCoat )
-	{
-		HandleGltfClearCoat( NodeContainer, GltfMaterial, ShaderGraphNode, bSwapNormalAndClearCoatNormal );
-	}
-
-	if ( GltfMaterial.bHasSheen )
-	{
-		HandleGltfSheen( NodeContainer, GltfMaterial, ShaderGraphNode );
-	}
-
-	if ( GltfMaterial.bHasTransmission )
-	{
-		HandleGltfTransmission( NodeContainer, GltfMaterial, ShaderGraphNode );
-	}
-}
-
-void UInterchangeGLTFTranslator::HandleGltfClearCoat( UInterchangeBaseNodeContainer& NodeContainer, const GLTF::FMaterial& GltfMaterial, UInterchangeShaderGraphNode& ShaderGraphNode, const bool bSwapNormalAndClearCoatNormal) const
-{
-	using namespace UE::Interchange::Materials;
-
-	if ( !GltfMaterial.bHasClearCoat || FMath::IsNearlyZero( GltfMaterial.ClearCoat.ClearCoatFactor ) )
-	{
-		return;
-	}
-
-	// ClearCoat::Parameters::ClearCoat
-	{
-		TVariant< FLinearColor, float > ClearCoatFactor;
-		ClearCoatFactor.Set< float >( GltfMaterial.ClearCoat.ClearCoatFactor );
-
-		HandleGltfMaterialParameter( NodeContainer, GltfMaterial.ClearCoat.ClearCoatMap, ShaderGraphNode, ClearCoat::Parameters::ClearCoat.ToString(),
-			ClearCoatFactor, Standard::Nodes::TextureSample::Outputs::R.ToString() );
-	}
-
-	//  ClearCoat::Parameters::ClearCoatRoughness
-	{
-		TVariant< FLinearColor, float > ClearCoatRoughnessFactor;
-		ClearCoatRoughnessFactor.Set< float >( GltfMaterial.ClearCoat.Roughness );
-
-		HandleGltfMaterialParameter( NodeContainer, GltfMaterial.ClearCoat.RoughnessMap, ShaderGraphNode, ClearCoat::Parameters::ClearCoatRoughness.ToString(),
-			ClearCoatRoughnessFactor, Standard::Nodes::TextureSample::Outputs::G.ToString() );
-	}
-
-	// ClearCoat::Parameters::ClearCoatNormal
-	{
-		TVariant< FLinearColor, float > ClearCoatNormalFactor;
-		ClearCoatNormalFactor.Set< float >( GltfMaterial.ClearCoat.NormalMapUVScale );
-
-		const bool bInverse = false;
-		const bool bIsNormal = true;
-
-		HandleGltfMaterialParameter( NodeContainer, GltfMaterial.ClearCoat.NormalMap, ShaderGraphNode, bSwapNormalAndClearCoatNormal ? Common::Parameters::Normal.ToString() : ClearCoat::Parameters::ClearCoatNormal.ToString(),
-			ClearCoatNormalFactor, Standard::Nodes::TextureSample::Outputs::RGB.ToString(), bInverse, bIsNormal );
-
-		//check if ClearCoat second normal is enabled in the render settings:
-		if (GltfAsset.Textures.IsValidIndex(GltfMaterial.Normal.TextureIndex) && !bRenderSettingsClearCoatEnableSecondNormal)
-		{
-			UE_LOG(LogInterchangeImport, Warning, TEXT("GLTF Material[%s] uses ClearCoat and has Normal map, however ClearCoat Second Normal is disabled in the Render Settings."), *GltfMaterial.Name);
-		}
-	}
-}
-
-void UInterchangeGLTFTranslator::HandleGltfSheen( UInterchangeBaseNodeContainer& NodeContainer, const GLTF::FMaterial& GltfMaterial, UInterchangeShaderGraphNode& ShaderGraphNode ) const
-{
-	using namespace UE::Interchange::Materials;
-
-	if ( !GltfMaterial.bHasSheen )
-	{
-		return;
-	}
-
-	// Sheen::Parameters::SheenColor
-	{
-		TVariant< FLinearColor, float > SheenColorFactor;
-		SheenColorFactor.Set< FLinearColor >( FLinearColor( GltfMaterial.Sheen.SheenColorFactor ) );
-
-		HandleGltfMaterialParameter( NodeContainer, GltfMaterial.Sheen.SheenColorMap, ShaderGraphNode, Sheen::Parameters::SheenColor.ToString(),
-			SheenColorFactor, Standard::Nodes::TextureSample::Outputs::RGB.ToString() );
-	}
-
-	// Sheen::Parameters::SheenRoughness
-	{
-		TVariant< FLinearColor, float > SheenRoughnessFactor;
-		SheenRoughnessFactor.Set< float >( GltfMaterial.Sheen.SheenRoughnessFactor );
-
-		const bool bInverse = true;
-		HandleGltfMaterialParameter( NodeContainer, GltfMaterial.Sheen.SheenRoughnessMap, ShaderGraphNode, Sheen::Parameters::SheenRoughness.ToString(),
-			SheenRoughnessFactor, Standard::Nodes::TextureSample::Outputs::A.ToString(), bInverse);
-	}
-}
-
-/**
- * GLTF transmission is handled a little differently than UE's.
- * GLTF doesn't allow having different reflected and transmitted colors, UE does (base color vs transmittance color).
- * GLTF controls the amount of reflected light vs transmitted light using the transmission factor, UE does that through opacity.
- * GLTF opacity means that the medium is present of not, so it's normal for transmission materials to be considered opaque,
- * meaning that the medium is full present, and the transmission factor determines how much lights is transmitted.
- * When a transmission material isn't fully opaque, we reduce the transmission color by the opacity to mimic GLTF's BTDF.
- * Ideally, this would be better represented by blending a default lit alpha blended material with a thin translucent material based on GLTF's opacity.
- */
-void UInterchangeGLTFTranslator::HandleGltfTransmission( UInterchangeBaseNodeContainer& NodeContainer, const GLTF::FMaterial& GltfMaterial, UInterchangeShaderGraphNode& ShaderGraphNode ) const
-{
-	using namespace UE::Interchange::Materials;
-
-	if ( !GltfMaterial.bHasTransmission || FMath::IsNearlyZero( GltfMaterial.Transmission.TransmissionFactor ) )
-	{
-		return;
-	}
- 
-	FString OpacityNodeUid;
-	FString OpacityNodeOutput;
-
-	// Common::Parameters::Opacity
-	{
-		/**
-		 * Per the spec, the red channel of the transmission texture drives how much light is transmitted vs diffused.
-		 * So we're setting the inverse of the red channel as the opacity.
-		 */
-		const FString OneMinusNodeName =  TEXT("OpacityOneMinus");
-		UInterchangeShaderNode* OneMinusNode = UInterchangeShaderNode::Create( &NodeContainer, OneMinusNodeName, ShaderGraphNode.GetUniqueID() );
-
-		OneMinusNode->SetCustomShaderType(Standard::Nodes::OneMinus::Name.ToString());
-
-		UInterchangeShaderNode* CurrentNode = OneMinusNode;
-		FString CurrentInput = Standard::Nodes::OneMinus::Inputs::Input.ToString();
-
-		TVariant< FLinearColor, float > TransmissionFactor;
-		TransmissionFactor.Set< float >( GltfMaterial.Transmission.TransmissionFactor );
-
-		HandleGltfMaterialParameter( NodeContainer, GltfMaterial.Transmission.TransmissionMap, *CurrentNode, CurrentInput,
-			TransmissionFactor, Standard::Nodes::TextureSample::Outputs::R.ToString() );
-
-		// The GLTF transmission model specifies that metallic surfaces don't transmit light, so adjust Common::Parameters::Opacity so that metallic surfaces are opaque.
-		{
-			FString MetallicNodeUid;
-			FString MetallicNodeOutput;
-
-			if ( UInterchangeShaderPortsAPI::GetInputConnection( &ShaderGraphNode, PBRMR::Parameters::Metallic.ToString(), MetallicNodeUid, MetallicNodeOutput ) )
-			{
-				const FString MetallicLerpNodeName =  TEXT("OpacityMetallicLerp");
-				UInterchangeShaderNode* LerpMetallicNode = UInterchangeShaderNode::Create( &NodeContainer, MetallicLerpNodeName, ShaderGraphNode.GetUniqueID() );
-				LerpMetallicNode->SetCustomShaderType( Standard::Nodes::Lerp::Name.ToString() );
-
-
-				LerpMetallicNode->AddFloatAttribute( UInterchangeShaderPortsAPI::MakeInputValueKey( Standard::Nodes::Lerp::Inputs::B.ToString() ), 1.f );
-				UInterchangeShaderPortsAPI::ConnectDefaultOuputToInput( LerpMetallicNode, Standard::Nodes::Lerp::Inputs::A.ToString(), CurrentNode->GetUniqueID() );
-				UInterchangeShaderPortsAPI::ConnectOuputToInputByName( LerpMetallicNode, Standard::Nodes::Lerp::Inputs::Factor.ToString(), MetallicNodeUid, MetallicNodeOutput );
-
-				CurrentNode = LerpMetallicNode;
-				CurrentInput = TEXT("");
-			}
-		}
-
-		if ( GltfMaterial.AlphaMode != GLTF::FMaterial::EAlphaMode::Opaque )
-		{
-			if ( UInterchangeShaderPortsAPI::GetInputConnection( &ShaderGraphNode, PBRMR::Parameters::Opacity.ToString(), OpacityNodeUid, OpacityNodeOutput ) )
-			{
-				const FString OpacityLerpNodeName =  TEXT("OpacityLerp");
-				UInterchangeShaderNode* OpacityLerpNode = UInterchangeShaderNode::Create( &NodeContainer, OpacityLerpNodeName, ShaderGraphNode.GetUniqueID() );
-				OpacityLerpNode->SetCustomShaderType( Standard::Nodes::Lerp::Name.ToString() );
-
-				OpacityLerpNode->AddFloatAttribute( UInterchangeShaderPortsAPI::MakeInputValueKey( Standard::Nodes::Lerp::Inputs::A.ToString() ), 0.f );
-				UInterchangeShaderPortsAPI::ConnectDefaultOuputToInput( OpacityLerpNode, Standard::Nodes::Lerp::Inputs::B.ToString(), CurrentNode->GetUniqueID() );
-				UInterchangeShaderPortsAPI::ConnectOuputToInputByName( OpacityLerpNode, Standard::Nodes::Lerp::Inputs::Factor.ToString(), OpacityNodeUid, OpacityNodeOutput );
-
-				CurrentNode = OpacityLerpNode;
-				CurrentInput = TEXT("");
-			}
-		}
-
-		UInterchangeShaderPortsAPI::ConnectDefaultOuputToInput( &ShaderGraphNode, Common::Parameters::Opacity.ToString(), CurrentNode->GetUniqueID() );
-	}
-
-	// ThinTranslucent::Parameters::Transmissioncolor
-	{
-		// There's no separation of reflected and transmitted color in this model. So the same color is used for the base color and the transmitted color.
-		// Since this extension is only supported with the metallic-roughness model, we can reuse its base color
-		const UInterchangeBaseNode* CurrentNode = &ShaderGraphNode;
-		FString CurrentOuput = TEXT("");
-		FLinearColor CurrentColor = FLinearColor::White;
-
-		FString BaseColorNodeUid;
-		FString BaseColorNodeOutput;
-
-		if ( UInterchangeShaderPortsAPI::GetInputConnection( CurrentNode, PBRMR::Parameters::BaseColor.ToString(), BaseColorNodeUid, BaseColorNodeOutput ) )
-		{
-			CurrentNode = NodeContainer.GetNode( BaseColorNodeUid );
-			CurrentOuput = BaseColorNodeOutput;
-		}
-		else
-		{
-			FLinearColor BaseColor;
-			if ( ShaderGraphNode.GetLinearColorAttribute( UInterchangeShaderPortsAPI::MakeInputValueKey( PBRMR::Parameters::BaseColor.ToString() ), BaseColor ) )
-			{
-				CurrentNode = nullptr;
-				CurrentColor = BaseColor;
-			}
-		}
-
-		if ( GltfMaterial.AlphaMode != GLTF::FMaterial::EAlphaMode::Opaque )
-		{
-			if ( !OpacityNodeUid.IsEmpty() )
-			{
-				const FString OpacityLerpNodeName =  TEXT("OpacityTransmissionLerp");
-				UInterchangeShaderNode* OpacityLerpNode = UInterchangeShaderNode::Create( &NodeContainer, OpacityLerpNodeName, ShaderGraphNode.GetUniqueID() );
-				OpacityLerpNode->SetCustomShaderType( Standard::Nodes::Lerp::Name.ToString() );
-
-				OpacityLerpNode->AddLinearColorAttribute( UInterchangeShaderPortsAPI::MakeInputValueKey( Standard::Nodes::Lerp::Inputs::A.ToString() ), FLinearColor::White );
-				UInterchangeShaderPortsAPI::ConnectOuputToInputByName( OpacityLerpNode, Standard::Nodes::Lerp::Inputs::Factor.ToString(), OpacityNodeUid, OpacityNodeOutput );
-
-				if ( CurrentNode )
-				{
-					UInterchangeShaderPortsAPI::ConnectOuputToInputByName( OpacityLerpNode, Standard::Nodes::Lerp::Inputs::B.ToString(), CurrentNode->GetUniqueID(), CurrentOuput );
-				}
-				else
-				{
-					OpacityLerpNode->AddLinearColorAttribute( UInterchangeShaderPortsAPI::MakeInputValueKey( Standard::Nodes::Lerp::Inputs::B.ToString() ), CurrentColor );
-				}
-
-				CurrentNode = OpacityLerpNode;
-				CurrentOuput = TEXT("");
-			}
-		}
-
-		if ( CurrentNode )
-		{
-			UInterchangeShaderPortsAPI::ConnectOuputToInputByName( &ShaderGraphNode, ThinTranslucent::Parameters::TransmissionColor.ToString(), CurrentNode->GetUniqueID(), CurrentOuput );
-		}
-		else
-		{
-			ShaderGraphNode.AddLinearColorAttribute( UInterchangeShaderPortsAPI::MakeInputValueKey( ThinTranslucent::Parameters::TransmissionColor.ToString() ), CurrentColor );
-		}
-	}
-}
-
-void UInterchangeGLTFTranslator::HandleGltfTextureTransform( UInterchangeBaseNodeContainer& NodeContainer, const GLTF::FTextureTransform& TextureTransform, const int32 TexCoordIndex, UInterchangeShaderNode& ShaderNode ) const
-{
-	using namespace UE::Interchange::Materials;
-
-	// Scale
-	if ( !FMath::IsNearlyEqual( TextureTransform.Scale[0], 1.f ) || 
-		 !FMath::IsNearlyEqual( TextureTransform.Scale[1], 1.f ) )
-	{
-		FVector2f TextureScale;
-		TextureScale.X = TextureTransform.Scale[0];
-		TextureScale.Y = TextureTransform.Scale[1];
-
-		ShaderNode.SetAttribute< FVector2f >( UInterchangeShaderPortsAPI::MakeInputValueKey( Standard::Nodes::TextureCoordinate::Inputs::Scale.ToString() ), TextureScale );
-	}
-
-	// Offset
-	if ( !FMath::IsNearlyZero( TextureTransform.Offset[0] ) || 
-		 !FMath::IsNearlyZero( TextureTransform.Offset[1] ) )
-	{
-		FVector2f TextureOffset;
-		TextureOffset.X = TextureTransform.Offset[0];
-		TextureOffset.Y = TextureTransform.Offset[1];
-
-		ShaderNode.SetAttribute< FVector2f >( UInterchangeShaderPortsAPI::MakeInputValueKey( Standard::Nodes::TextureCoordinate::Inputs::Offset.ToString() ), TextureOffset );
-	}
-
-	// Rotate
-	if ( !FMath::IsNearlyZero( TextureTransform.Rotation ) )
-	{
-		float AngleRadians = TextureTransform.Rotation;
-
-		if ( AngleRadians < 0.0f )
-		{
-			AngleRadians = TWO_PI - AngleRadians;
-		}
-
-		AngleRadians = 1.0f - ( AngleRadians / TWO_PI );
-
-		ShaderNode.AddFloatAttribute( UInterchangeShaderPortsAPI::MakeInputValueKey( Standard::Nodes::TextureCoordinate::Inputs::Rotate.ToString() ), AngleRadians );
-
-		FVector2f RotationCenter = FVector2f::Zero();
-		ShaderNode.SetAttribute< FVector2f >( UInterchangeShaderPortsAPI::MakeInputValueKey( Standard::Nodes::TextureCoordinate::Inputs::RotationCenter.ToString() ), RotationCenter );
-	}
-
-}
-
 EInterchangeTranslatorType UInterchangeGLTFTranslator::GetTranslatorType() const
 {
 	return EInterchangeTranslatorType::Scenes;
@@ -1041,6 +446,24 @@ bool UInterchangeGLTFTranslator::Translate( UInterchangeBaseNodeContainer& NodeC
 
 	ScaleNodeTranslations(const_cast<UInterchangeGLTFTranslator*>(this)->GltfAsset.Nodes, GltfUnitConversionMultiplier);
 
+	//Check Normal Textures:
+	TSet<int32> NormalTextureIndices;
+	{
+		auto AddTextureIndex = [&](int32 TextureIndex)
+		{
+			if (GltfAsset.Textures.IsValidIndex(TextureIndex))
+			{
+				NormalTextureIndices.Add(TextureIndex);
+			}
+		};
+
+		for (const GLTF::FMaterial& GltfMaterial : GltfAsset.Materials)
+		{
+			AddTextureIndex(GltfMaterial.Normal.TextureIndex);
+			AddTextureIndex(GltfMaterial.ClearCoat.NormalMap.TextureIndex);
+		}
+	}
+
 	// Textures
 	{
 		int32 TextureIndex = 0;
@@ -1072,7 +495,19 @@ bool UInterchangeGLTFTranslator::Translate( UInterchangeBaseNodeContainer& NodeC
 			TextureNode->SetDisplayLabel(GltfTexture.Name);
 
 			TextureNode->SetCustomFilter(ConvertFilter(GltfTexture.Sampler.MinFilter));
-			TextureNode->SetPayLoadKey( LexToString( TextureIndex++ ) );
+
+			bool TextureUsedAsNormal = NormalTextureIndices.Contains(TextureIndex);
+
+			if (TextureUsedAsNormal)
+			{
+				//According to GLTF documentation the normal maps are right handed (following OpenGL convention),
+				//however UE expects left handed normal maps, this can be resolved by flipping the green channel of the normal textures:
+				//(based on [REDACTED_PROFILE_URL]/glTF-Sample-Models/tree/master/2.0/NormalTangentTest#problem-flipped-y-axis-or-flipped-green-channel)
+				TextureNode->SetCustombFlipGreenChannel(true);
+			}
+
+			FString Payloadkey = LexToString(TextureIndex++) + TEXT(":") + LexToString(TextureUsedAsNormal);
+			TextureNode->SetPayLoadKey(Payloadkey);
 
 			TextureNode->SetCustomWrapU( UE::Interchange::Gltf::Private::ConvertWrap( GltfTexture.Sampler.WrapS ) );
 			TextureNode->SetCustomWrapV( UE::Interchange::Gltf::Private::ConvertWrap( GltfTexture.Sampler.WrapT ) );
@@ -1119,22 +554,13 @@ bool UInterchangeGLTFTranslator::Translate( UInterchangeBaseNodeContainer& NodeC
 			SetTextureSRGB(NodeContainer, GltfMaterial.ClearCoat.RoughnessMap, false);
 			SetTextureSRGB(NodeContainer, GltfMaterial.Transmission.TransmissionMap, false);
 
-			//According to GLTF documentation the normal maps are right handed (following OpenGL convention),
-			//however UE expects left handed normal maps, this can be resolved by flipping the green channel of the normal textures:
-			//(based on [REDACTED_PROFILE_URL]/glTF-Sample-Models/tree/master/2.0/NormalTangentTest#problem-flipped-y-axis-or-flipped-green-channel)
-			SetTextureFlipGreenChannel(NodeContainer, GltfMaterial.Normal);
-			SetTextureFlipGreenChannel(NodeContainer, GltfMaterial.ClearCoat.NormalMap);
-
 			const FString ShaderGraphNodeUid = UInterchangeShaderGraphNode::MakeNodeUid(GltfMaterial.UniqueId);
 			bool bUseVertexColor = MaterialsUsedOnMeshesWithVertexColor.Contains(ShaderGraphNodeUid);
 
 			UInterchangeShaderGraphNode* ShaderGraphNode = UInterchangeShaderGraphNode::Create(&NodeContainer, GltfMaterial.UniqueId);
 			ShaderGraphNode->SetDisplayLabel(GltfMaterial.Name);
 
-			HandleGltfMaterial(NodeContainer, GltfMaterial, *ShaderGraphNode, bUseVertexColor);
-
-			//Store Gltf Material Attribute values in ShaderGraphNode but do not create MaterialInstances
-			UE::Interchange::GLTFMaterialInstances::AddGltfMaterialValuesToShaderGraphNode(GltfMaterial, GltfAsset.Textures, ShaderGraphNode);
+			UE::Interchange::GLTFMaterials::HandleGltfMaterial(NodeContainer, GltfMaterial, GltfAsset.Textures, ShaderGraphNode);
 			
 			++MaterialIndex;
 		}
@@ -1327,8 +753,16 @@ bool UInterchangeGLTFTranslator::Translate( UInterchangeBaseNodeContainer& NodeC
 
 TOptional< UE::Interchange::FImportImage > UInterchangeGLTFTranslator::GetTexturePayloadData(const FString& PayloadKey, TOptional<FString>& AlternateTexturePath) const
 {
+	TArray<FString> PayloadKeys;
+	PayloadKey.ParseIntoArray(PayloadKeys, TEXT(":"));
+
+	if (PayloadKeys.Num() == 0)
+	{
+		return TOptional< UE::Interchange::FImportImage >();
+	}
+
 	int32 TextureIndex = 0;
-	LexFromString( TextureIndex, *PayloadKey);
+	LexFromString( TextureIndex, *PayloadKeys[0]);
 
 	if ( !GltfAsset.Textures.IsValidIndex( TextureIndex ) )
 	{
@@ -1337,12 +771,14 @@ TOptional< UE::Interchange::FImportImage > UInterchangeGLTFTranslator::GetTextur
 
 	const GLTF::FTexture& GltfTexture = GltfAsset.Textures[ TextureIndex ];
 
+	TOptional<UE::Interchange::FImportImage> TexturePayloadData = TOptional<UE::Interchange::FImportImage>();
+
 	if (GltfTexture.Source.FilePath.IsEmpty())
 	{
 		// Embedded texture -- try using ImageWrapper to decode it
 		TArray64<uint8> ImageData(GltfTexture.Source.Data, GltfTexture.Source.DataByteLength);
 		UInterchangeImageWrapperTranslator* ImageWrapperTranslator = NewObject<UInterchangeImageWrapperTranslator>(GetTransientPackage(), NAME_None);
-		return ImageWrapperTranslator->GetTexturePayloadDataFromBuffer(ImageData);
+		TexturePayloadData = ImageWrapperTranslator->GetTexturePayloadDataFromBuffer(ImageData);
 	}
 	else
 	{
@@ -1353,7 +789,7 @@ TOptional< UE::Interchange::FImportImage > UInterchangeGLTFTranslator::GetTextur
 
 		if (!PayloadSourceData)
 		{
-			return TOptional<UE::Interchange::FImportImage>();
+			return TexturePayloadData;
 		}
 
 		UInterchangeTranslatorBase* SourceTranslator = UInterchangeManager::GetInterchangeManager().GetTranslatorForSourceData(PayloadSourceData);
@@ -1361,14 +797,24 @@ TOptional< UE::Interchange::FImportImage > UInterchangeGLTFTranslator::GetTextur
 		const IInterchangeTexturePayloadInterface* TextureTranslator = Cast< IInterchangeTexturePayloadInterface >(SourceTranslator);
 		if (!ensure(TextureTranslator))
 		{
-			return TOptional<UE::Interchange::FImportImage>();
+			return TexturePayloadData;
 		}
 		SourceTranslator->SetResultsContainer(Results);
 
 		AlternateTexturePath = TextureFilePath;
 
-		return TextureTranslator->GetTexturePayloadData(PayloadKey, AlternateTexturePath);
+		TexturePayloadData = TextureTranslator->GetTexturePayloadData(PayloadKey, AlternateTexturePath);
+	}
+
+	if (PayloadKeys.Num() == 2 && TexturePayloadData.IsSet())
+	{
+		bool TextureUsedAsNormal = false;
+		LexFromString(TextureUsedAsNormal, *PayloadKeys[1]);
+
+		TexturePayloadData.GetValue().CompressionSettings = TextureUsedAsNormal ? TC_Normalmap : TC_Default;
 	}
+
+	return TexturePayloadData;
 }
 
 TFuture<TOptional<UE::Interchange::FAnimationPayloadData>> UInterchangeGLTFTranslator::GetAnimationPayloadData(const FInterchangeAnimationPayLoadKey& PayLoadKey, const double BakeFrequency, const double RangeStartSecond, const double RangeStopSecond) const
@@ -1725,17 +1171,6 @@ void UInterchangeGLTFTranslator::SetTextureSRGB(UInterchangeBaseNodeContainer& N
 		}
 	}
 }
-void UInterchangeGLTFTranslator::SetTextureFlipGreenChannel(UInterchangeBaseNodeContainer& NodeContainer, const GLTF::FTextureMap& TextureMap) const
-{
-	if (GltfAsset.Textures.IsValidIndex(TextureMap.TextureIndex))
-	{
-		const FString TextureUid = UInterchangeTextureNode::MakeNodeUid(GltfAsset.Textures[TextureMap.TextureIndex].UniqueId);
-		if (UInterchangeTextureNode* TextureNode = const_cast<UInterchangeTextureNode*>(Cast<UInterchangeTextureNode>(NodeContainer.GetNode(TextureUid))))
-		{
-			TextureNode->SetCustombFlipGreenChannel(true);
-		}
-	}
-}
 
 TFuture<TOptional<UE::Interchange::FVariantSetPayloadData>> UInterchangeGLTFTranslator::GetVariantSetPayloadData(const FString& PayloadKey) const
 {
@@ -2189,6 +1624,11 @@ UInterchangeGLTFTranslator::UInterchangeGLTFTranslator()
 	{
 		bRenderSettingsClearCoatEnableSecondNormal = GetDefault<URendererSettings>()->bClearCoatEnableSecondNormal != 0;
 	}
+
+	if (!UE::Interchange::GLTFMaterials::AreRequiredPackagesLoaded())
+	{
+		UE_LOG(LogInterchangeImport, Warning, TEXT("UInterchangeGLTFPipeline: Some required packages are missing. Material import might be wrong"));
+	}
 }
 
 #undef LOCTEXT_NAMESPACE
diff --git a/Engine/Plugins/Interchange/Runtime/Source/Import/Private/Material/InterchangeMaterialFactory.cpp b/Engine/Plugins/Interchange/Runtime/Source/Import/Private/Material/InterchangeMaterialFactory.cpp
--- a/Engine/Plugins/Interchange/Runtime/Source/Import/Private/Material/InterchangeMaterialFactory.cpp
+++ b/Engine/Plugins/Interchange/Runtime/Source/Import/Private/Material/InterchangeMaterialFactory.cpp
@@ -774,6 +774,13 @@ void UInterchangeMaterialFactory::SetupMaterial(UMaterial* Material, const FImpo
 
 	const UInterchangeMaterialFactoryNode* MaterialFactoryNode = Cast<UInterchangeMaterialFactoryNode>(BaseMaterialFactoryNode);
 
+	{//Screen Space Reflections:
+		if (bool bScreenSpaceReflections; MaterialFactoryNode->GetCustomScreenSpaceReflections(bScreenSpaceReflections))
+		{
+			Material->bScreenSpaceReflections = bScreenSpaceReflections;
+		}
+	}
+
 	FMaterialExpressionBuilder Builder(Material, nullptr, Arguments);
 
 	if (UInterchangeShaderPortsAPI::HasInput(MaterialFactoryNode, Common::Parameters::BxDF))
diff --git a/Engine/Plugins/Interchange/Runtime/Source/Import/Public/Gltf/InterchangeGLTFMaterial.h b/Engine/Plugins/Interchange/Runtime/Source/Import/Public/Gltf/InterchangeGLTFMaterial.h
new file mode 100644
--- /dev/null
+++ b/Engine/Plugins/Interchange/Runtime/Source/Import/Public/Gltf/InterchangeGLTFMaterial.h
@@ -0,0 +1,216 @@
+// Copyright Epic Games, Inc. All Rights Reserved.
+#pragma once
+
+#include "CoreMinimal.h"
+
+namespace GLTF {
+	struct FMaterial;
+	struct FTexture;
+}
+class UInterchangeShaderGraphNode;
+class UInterchangeBaseNodeContainer;
+
+const FString InterchangeGltfMaterialAttributeIdentifier = TEXT("Gltf_MI_AttributeIdentifier_");
+
+//OffsetScale for MaterialInstances
+//Offset_X Offset_Y Scale_X Scale_Y for Materials
+#define INTERCHANGE_GLTF_STRINGIFY(x) #x
+#define DECLARE_INTERCHANGE_GLTF_MI_MAP(MapName) const FString MapName##Texture = TEXT(INTERCHANGE_GLTF_STRINGIFY(MapName ## Texture)); \
+												 const FString MapName##Texture_OffsetScale = TEXT(INTERCHANGE_GLTF_STRINGIFY(MapName ## Texture_OffsetScale)); \
+												 const FString MapName##Texture_Offset_X = TEXT(INTERCHANGE_GLTF_STRINGIFY(MapName ## Texture_Offset_X)); \
+												 const FString MapName##Texture_Offset_Y = TEXT(INTERCHANGE_GLTF_STRINGIFY(MapName ## Texture_Offset_Y)); \
+												 const FString MapName##Texture_Scale_X = TEXT(INTERCHANGE_GLTF_STRINGIFY(MapName ## Texture_Scale_X)); \
+												 const FString MapName##Texture_Scale_Y = TEXT(INTERCHANGE_GLTF_STRINGIFY(MapName ## Texture_Scale_Y)); \
+												 const FString MapName##Texture_Rotation = TEXT(INTERCHANGE_GLTF_STRINGIFY(MapName ## Texture_Rotation)); \
+												 const FString MapName##Texture_TexCoord = TEXT(INTERCHANGE_GLTF_STRINGIFY(MapName ## Texture_TexCoord));
+
+namespace UE::Interchange::GLTFMaterials
+{
+	//Inputs/Parameters (Materials/MaterialInstances)
+	namespace Inputs
+	{
+		const FString Color_RGB = TEXT("_RGB");
+		const FString Color_A = TEXT("_A");
+
+		//MetalRoughness specific:
+		DECLARE_INTERCHANGE_GLTF_MI_MAP(BaseColor)
+		const FString BaseColorFactor = TEXT("BaseColorFactor");
+		const FString BaseColorFactor_RGB = BaseColorFactor + Color_RGB; //Connection to inputs from BaseColorFactor.RGB
+		const FString BaseColorFactor_A = BaseColorFactor + Color_A; //Connection to inputs from BaseColorFactor.A
+
+		DECLARE_INTERCHANGE_GLTF_MI_MAP(MetallicRoughness)
+		const FString MetallicFactor = TEXT("MetallicFactor");
+		const FString RoughnessFactor = TEXT("RoughnessFactor");
+
+		DECLARE_INTERCHANGE_GLTF_MI_MAP(Specular)
+		const FString SpecularFactor = TEXT("SpecularFactor");
+
+
+		//SpecularGlossiness specific
+		DECLARE_INTERCHANGE_GLTF_MI_MAP(Diffuse)
+		const FString DiffuseFactor = TEXT("DiffuseFactor");
+		const FString DiffuseFactor_RGB = DiffuseFactor + Color_RGB; //Connection to inputs from BaseColorFactor.RGB
+		const FString DiffuseFactor_A = DiffuseFactor + Color_A; //Connection to inputs from BaseColorFactor.A
+
+		DECLARE_INTERCHANGE_GLTF_MI_MAP(SpecularGlossiness)
+		const FString SpecFactor = TEXT("SpecFactor");
+		const FString GlossinessFactor = TEXT("GlossinessFactor");
+
+
+		//Generic:
+		DECLARE_INTERCHANGE_GLTF_MI_MAP(Normal)
+		const FString NormalScale = TEXT("NormalScale");
+
+		DECLARE_INTERCHANGE_GLTF_MI_MAP(Emissive)
+		const FString EmissiveFactor = TEXT("EmissiveFactor");
+		const FString EmissiveStrength = TEXT("EmissiveStrength");
+
+		DECLARE_INTERCHANGE_GLTF_MI_MAP(Occlusion)
+		const FString OcclusionStrength = TEXT("OcclusionStrength");
+
+		const FString IOR = TEXT("IOR");
+
+		const FString AlphaCutoff = TEXT("AlphaCutoff");
+
+
+		//ClearCoat specific:
+		DECLARE_INTERCHANGE_GLTF_MI_MAP(ClearCoat)
+		const FString ClearCoatFactor = TEXT("ClearCoatFactor");
+
+		DECLARE_INTERCHANGE_GLTF_MI_MAP(ClearCoatRoughness)
+		const FString ClearCoatRoughnessFactor = TEXT("ClearCoatRoughnessFactor");
+
+		DECLARE_INTERCHANGE_GLTF_MI_MAP(ClearCoatNormal)
+		const FString ClearCoatNormalScale = TEXT("ClearCoatNormalScale");
+
+
+		//Sheen specific:
+		DECLARE_INTERCHANGE_GLTF_MI_MAP(SheenColor)
+		const FString SheenColorFactor = TEXT("SheenColorFactor");
+
+		DECLARE_INTERCHANGE_GLTF_MI_MAP(SheenRoughness)
+		const FString SheenRoughnessFactor = TEXT("SheenRoughnessFactor");
+
+
+		//Transmission specific:
+		DECLARE_INTERCHANGE_GLTF_MI_MAP(Transmission)
+		const FString TransmissionFactor = TEXT("TransmissionFactor");
+	}
+
+	enum EShadingModel : uint8
+	{
+		DEFAULT = 0, //MetalRoughness
+		UNLIT,
+		CLEARCOAT,
+		SHEEN,
+		TRANSMISSION,
+		SPECULARGLOSSINESS,
+		SHADINGMODELCOUNT
+	};
+
+	enum EAlphaMode : uint8
+	{
+		Opaque = 0,
+		Mask,
+		Blend
+	};
+
+	static TMap<EShadingModel, TPair<FString, TArray<FString>>> ShadingModelToMaterialFunctions = {
+		{EShadingModel::DEFAULT,
+			TPair<FString, TArray<FString>>(
+				TEXT("/Interchange/gltf/MaterialBodies/MF_Default_Body.MF_Default_Body"),
+				TArray<FString>{
+					TEXT("BaseColor"),
+					TEXT("Metallic"),
+					TEXT("Specular"),
+					TEXT("Roughness"),
+					TEXT("EmissiveColor"),
+					TEXT("Opacity"),
+					TEXT("OpacityMask"),
+					TEXT("Normal"),
+					TEXT("Occlusion")})},
+
+		{EShadingModel::UNLIT,
+			TPair<FString, TArray<FString>>(
+				TEXT("/Interchange/gltf/MaterialBodies/MF_Unlit_Body.MF_Unlit_Body"),
+				TArray<FString>{
+					TEXT("UnlitColor"),
+					TEXT("Opacity"),
+					TEXT("OpacityMask")})},
+
+		{EShadingModel::CLEARCOAT,
+			TPair<FString, TArray<FString>>(
+				TEXT("/Interchange/gltf/MaterialBodies/MF_ClearCoat_Body.MF_ClearCoat_Body"),
+				TArray<FString>{
+					TEXT("ClearCoatBottomNormal"),
+					TEXT("BaseColor"),
+					TEXT("Metallic"),
+					TEXT("Specular"),
+					TEXT("Roughness"),
+					TEXT("EmissiveColor"),
+					TEXT("Opacity"),
+					TEXT("OpacityMask"),
+					TEXT("Normal"),
+					TEXT("ClearCoat"),
+					TEXT("ClearCoatRoughness"),
+					TEXT("Occlusion")})},
+
+		{EShadingModel::SHEEN,
+			TPair<FString, TArray<FString>>(
+				TEXT("/Interchange/gltf/MaterialBodies/MF_Sheen_Body.MF_Sheen_Body"),
+				TArray<FString>{
+					TEXT("BaseColor"),
+					TEXT("Metallic"),
+					TEXT("Specular"),
+					TEXT("Roughness"),
+					TEXT("EmissiveColor"),
+					TEXT("Opacity"),
+					TEXT("OpacityMask"),
+					TEXT("Normal"),
+					TEXT("SheenColor"),
+					TEXT("SheenRoughness"),
+					TEXT("Occlusion")})},
+
+		{EShadingModel::TRANSMISSION,
+			TPair<FString, TArray<FString>>(
+				TEXT("/Interchange/gltf/MaterialBodies/MF_Transmission_Body.MF_Transmission_Body"),
+				TArray<FString>{
+					TEXT("TransmissionColor"),
+					TEXT("BaseColor"),
+					TEXT("Metallic"),
+					TEXT("Specular"),
+					TEXT("Roughness"),
+					TEXT("EmissiveColor"),
+					TEXT("Opacity"),
+					TEXT("Normal"),
+					TEXT("Occlusion"),
+					TEXT("Refraction")})},
+
+		{EShadingModel::SPECULARGLOSSINESS,
+			TPair<FString, TArray<FString>>(
+				TEXT("/Interchange/gltf/MaterialBodies/MF_SpecularGlossiness_Body.MF_SpecularGlossiness_Body"),
+				TArray<FString>{
+					TEXT("BaseColor"),
+					TEXT("Metallic"),
+					TEXT("Roughness"),
+					TEXT("EmissiveColor"),
+					TEXT("Opacity"),
+					TEXT("OpacityMask"),
+					TEXT("Normal"),
+					TEXT("Occlusion")})},
+	};
+
+	inline TArray<FString> GetRequiredMaterialFunctionPaths()
+	{
+		TArray<FString> Result;
+		for (const TPair<EShadingModel, TPair<FString, TArray<FString>>>& Entry : ShadingModelToMaterialFunctions)
+		{
+			Result.Add(Entry.Value.Key);
+		}
+		return Result;
+	}
+
+	void HandleGltfMaterial(UInterchangeBaseNodeContainer& NodeContainer, const GLTF::FMaterial& GltfMaterial, const TArray<GLTF::FTexture>& Textures, UInterchangeShaderGraphNode* ShaderGraphNode);
+
+	bool AreRequiredPackagesLoaded();
+};
diff --git a/Engine/Plugins/Interchange/Runtime/Source/Import/Public/Gltf/InterchangeGLTFMaterialInstances.h b/Engine/Plugins/Interchange/Runtime/Source/Import/Public/Gltf/InterchangeGLTFMaterialInstances.h
deleted file mode 100644
diff --git a/Engine/Plugins/Interchange/Runtime/Source/Import/Public/Gltf/InterchangeGltfTranslator.h b/Engine/Plugins/Interchange/Runtime/Source/Import/Public/Gltf/InterchangeGltfTranslator.h
--- a/Engine/Plugins/Interchange/Runtime/Source/Import/Public/Gltf/InterchangeGltfTranslator.h
+++ b/Engine/Plugins/Interchange/Runtime/Source/Import/Public/Gltf/InterchangeGltfTranslator.h
@@ -66,29 +66,15 @@ protected:
 	void HandleGltfSkeletons( UInterchangeBaseNodeContainer& NodeContainer, const FString& SceneNodeUid, const TArray<int32>& SkinnedMeshNodes, TSet<int>& UnusedMeshIndices ) const;
 	void HandleGltfNode( UInterchangeBaseNodeContainer& NodeContainer, const GLTF::FNode& GltfNode, const FString& ParentNodeUid, const int32 NodeIndex, 
 		bool &bHasVariants, TArray<int32>& SkinnedMeshNodes, TSet<int>& UnusedMeshIndices ) const;
-	void HandleGltfMaterial( UInterchangeBaseNodeContainer& NodeContainer, const GLTF::FMaterial& GltfMaterial, UInterchangeShaderGraphNode& ShaderGraphNode, bool bUseVertexColor ) const;
-	void HandleGltfMaterialParameter( UInterchangeBaseNodeContainer& NodeContainer, const GLTF::FTextureMap& TextureMap, UInterchangeShaderNode& ShaderNode,
-		const FString& MapName, const TVariant< FLinearColor, float >& MapFactor, const FString& OutputChannel, const bool bInverse = false, const bool bIsNormal = false, const bool bUseVertexColor = false ) const;
 	void HandleGltfAnimation(UInterchangeBaseNodeContainer& NodeContainer, int32 AnimationIndex) const;
 	void HandleGltfVariants(UInterchangeBaseNodeContainer& NodeContainer, const FString& FileName) const;
 	UInterchangeMeshNode* HandleGltfMesh(UInterchangeBaseNodeContainer& NodeContainer, const GLTF::FMesh& GltfMesh, int MeshIndex,
 		TSet<int>& UnusedMeshIndices, const FString& SkeletalName = "" /*If set it creates the mesh even if it was already created (for Skeletals)*/, const FString& SkeletalId = "") const;
 
-	/** Support for KHR_materials_clearcoat */
-	void HandleGltfClearCoat( UInterchangeBaseNodeContainer& NodeContainer, const GLTF::FMaterial& GltfMaterial, UInterchangeShaderGraphNode& ShaderGraphNode, const bool bSwapNormalAndClearCoatNormal ) const;
-	/** Support for KHR_materials_sheen */
-	void HandleGltfSheen( UInterchangeBaseNodeContainer& NodeContainer, const GLTF::FMaterial& GltfMaterial, UInterchangeShaderGraphNode& ShaderGraphNode ) const;
-	/** Support for KHR_materials_transmission */
-	void HandleGltfTransmission( UInterchangeBaseNodeContainer& NodeContainer, const GLTF::FMaterial& GltfMaterial, UInterchangeShaderGraphNode& ShaderGraphNode ) const;
-
-	/** Support for KHR_texture_transform */
-	void HandleGltfTextureTransform( UInterchangeBaseNodeContainer& NodeContainer, const GLTF::FTextureTransform& TextureTransform, const int32 TexCoordIndex, UInterchangeShaderNode& ShaderNode ) const;
-
 	bool GetVariantSetPayloadData(UE::Interchange::FVariantSetPayloadData& PayloadData) const;
 
 private:
 	void SetTextureSRGB(UInterchangeBaseNodeContainer& NodeContainer, const GLTF::FTextureMap& TextureMap, bool bSRGB) const;
-	void SetTextureFlipGreenChannel(UInterchangeBaseNodeContainer& NodeContainer, const GLTF::FTextureMap& TextureMap) const;
 
 	GLTF::FAsset GltfAsset;
 	mutable FNodeUidMap NodeUidMap;
diff --git a/Engine/Plugins/Interchange/Runtime/Source/Nodes/Private/InterchangeShaderGraphNode.cpp b/Engine/Plugins/Interchange/Runtime/Source/Nodes/Private/InterchangeShaderGraphNode.cpp
--- a/Engine/Plugins/Interchange/Runtime/Source/Nodes/Private/InterchangeShaderGraphNode.cpp
+++ b/Engine/Plugins/Interchange/Runtime/Source/Nodes/Private/InterchangeShaderGraphNode.cpp
@@ -245,6 +245,15 @@ bool UInterchangeShaderGraphNode::SetCustomIsAShaderFunction(const bool& Attribu
 	IMPLEMENT_NODE_ATTRIBUTE_SETTER_NODELEGATE(IsAShaderFunction, bool);
 }
 
+bool UInterchangeShaderGraphNode::GetCustomScreenSpaceReflections(bool& AttributeValue) const
+{
+	IMPLEMENT_NODE_ATTRIBUTE_GETTER(ScreenSpaceReflections, bool);
+}
+bool UInterchangeShaderGraphNode::SetCustomScreenSpaceReflections(const bool& AttributeValue)
+{
+	IMPLEMENT_NODE_ATTRIBUTE_SETTER_NODELEGATE(ScreenSpaceReflections, bool);
+}
+
 FString UInterchangeFunctionCallShaderNode::GetTypeName() const
 {
 	const FString TypeName = TEXT("FunctionCallShaderNode");
@@ -259,5 +268,4 @@ bool UInterchangeFunctionCallShaderNode::GetCustomMaterialFunction(FString& Attr
 bool UInterchangeFunctionCallShaderNode::SetCustomMaterialFunction(const FString& AttributeValue)
 {
 	IMPLEMENT_NODE_ATTRIBUTE_SETTER_NODELEGATE(MaterialFunction, FString);
-}
-
+}
\ No newline at end of file
diff --git a/Engine/Plugins/Interchange/Runtime/Source/Nodes/Public/InterchangeMaterialDefinitions.h b/Engine/Plugins/Interchange/Runtime/Source/Nodes/Public/InterchangeMaterialDefinitions.h
--- a/Engine/Plugins/Interchange/Runtime/Source/Nodes/Public/InterchangeMaterialDefinitions.h
+++ b/Engine/Plugins/Interchange/Runtime/Source/Nodes/Public/InterchangeMaterialDefinitions.h
@@ -290,6 +290,7 @@ namespace Interchange
 				const FName Normal = TEXT("Normal"); // Type: vector3f
 				const FName Tangent = TEXT("Tangent"); // Type: vector3f
 				const FName Opacity = TEXT("Opacity"); // Type: float
+				const FName OpacityMask = TEXT("OpacityMask"); // Type: float
 				const FName Occlusion = TEXT("Occlusion"); // Type: float
 				const FName IndexOfRefraction = TEXT("IOR"); // Type: float
 				const FName BxDF = TEXT("BxDF"); // input/output of BSDF or BRDF or BXDF or BTDF data
diff --git a/Engine/Plugins/Interchange/Runtime/Source/Nodes/Public/InterchangeShaderGraphNode.h b/Engine/Plugins/Interchange/Runtime/Source/Nodes/Public/InterchangeShaderGraphNode.h
--- a/Engine/Plugins/Interchange/Runtime/Source/Nodes/Public/InterchangeShaderGraphNode.h
+++ b/Engine/Plugins/Interchange/Runtime/Source/Nodes/Public/InterchangeShaderGraphNode.h
@@ -229,8 +229,15 @@ public:
 	UFUNCTION(BlueprintCallable, Category = "Interchange | Node | Material")
 	bool SetCustomIsAShaderFunction(const bool& AttributeValue);
 
+	UFUNCTION(BlueprintCallable, Category = "Interchange | Node | Material")
+	bool GetCustomScreenSpaceReflections(bool& AttributeValue) const;
+
+	UFUNCTION(BlueprintCallable, Category = "Interchange | Node | Material")
+	bool SetCustomScreenSpaceReflections(const bool& AttributeValue);
+
 private:
 	const UE::Interchange::FAttributeKey Macro_CustomTwoSidedKey = UE::Interchange::FAttributeKey(TEXT("TwoSided"));
 	const UE::Interchange::FAttributeKey Macro_CustomOpacityMaskClipValueKey = UE::Interchange::FAttributeKey(TEXT("OpacityMaskClipValue"));
 	const UE::Interchange::FAttributeKey Macro_CustomIsAShaderFunctionKey = UE::Interchange::FAttributeKey(TEXT("IsAShaderFunction"));
+	const UE::Interchange::FAttributeKey Macro_CustomScreenSpaceReflectionsKey = UE::Interchange::FAttributeKey(TEXT("ScreenSpaceReflections"));
 };
\ No newline at end of file
diff --git a/Engine/Plugins/Interchange/Runtime/Source/Parsers/GLTFCore/Private/GLTF/ExtensionsHandler.cpp b/Engine/Plugins/Interchange/Runtime/Source/Parsers/GLTFCore/Private/GLTF/ExtensionsHandler.cpp
--- a/Engine/Plugins/Interchange/Runtime/Source/Parsers/GLTFCore/Private/GLTF/ExtensionsHandler.cpp
+++ b/Engine/Plugins/Interchange/Runtime/Source/Parsers/GLTFCore/Private/GLTF/ExtensionsHandler.cpp
@@ -165,7 +165,7 @@ namespace GLTF
 
 					Material.bHasClearCoat = true;
 
-					Material.ClearCoat.ClearCoatFactor = GetScalar(ClearCoat, TEXT("clearcoatFactor"), 1.0f);
+					Material.ClearCoat.ClearCoatFactor = GetScalar(ClearCoat, TEXT("clearcoatFactor"), 0.0f);
 					GLTF::SetTextureMap(ClearCoat, TEXT("clearcoatTexture"), nullptr, Asset->Textures, Material.ClearCoat.ClearCoatMap, Messages);
 
 					Material.ClearCoat.Roughness = GetScalar(ClearCoat, TEXT("clearcoatRoughnessFactor"), 0.0f);
@@ -178,10 +178,6 @@ namespace GLTF
 				break;
 				case EExtension::KHR_MaterialsTransmission:
 				{
-					if (Material.AlphaMode == FMaterial::EAlphaMode::Blend)
-					{
-						break;
-					}
 					const FJsonObject& Transm = ExtObj;
 
 					Material.bHasTransmission = true;
diff --git a/Engine/Plugins/Interchange/Runtime/Source/Parsers/GLTFCore/Private/GLTF/MaterialUtilities.cpp b/Engine/Plugins/Interchange/Runtime/Source/Parsers/GLTFCore/Private/GLTF/MaterialUtilities.cpp
--- a/Engine/Plugins/Interchange/Runtime/Source/Parsers/GLTFCore/Private/GLTF/MaterialUtilities.cpp
+++ b/Engine/Plugins/Interchange/Runtime/Source/Parsers/GLTFCore/Private/GLTF/MaterialUtilities.cpp
@@ -89,10 +89,6 @@ namespace GLTF
 				}
 			}
 		}
-		else
-		{
-			Scale = 0.f;
-		}
 
 		return Scale;
 	}
diff --git a/Engine/Plugins/Interchange/Runtime/Source/Parsers/GLTFCore/Public/GLTFMaterial.h b/Engine/Plugins/Interchange/Runtime/Source/Parsers/GLTFCore/Public/GLTFMaterial.h
--- a/Engine/Plugins/Interchange/Runtime/Source/Parsers/GLTFCore/Public/GLTFMaterial.h
+++ b/Engine/Plugins/Interchange/Runtime/Source/Parsers/GLTFCore/Public/GLTFMaterial.h
@@ -92,7 +92,7 @@ namespace GLTF
 		};
 		struct FClearCoat
 		{
-			float ClearCoatFactor = 1.0f;
+			float ClearCoatFactor = 0.0f;
 			FTextureMap ClearCoatMap;
 
 			float Roughness = 0.0f;
diff --git a/Engine/Plugins/Interchange/Runtime/Source/Pipelines/Private/InterchangeGenericMaterialPipeline.cpp b/Engine/Plugins/Interchange/Runtime/Source/Pipelines/Private/InterchangeGenericMaterialPipeline.cpp
--- a/Engine/Plugins/Interchange/Runtime/Source/Pipelines/Private/InterchangeGenericMaterialPipeline.cpp
+++ b/Engine/Plugins/Interchange/Runtime/Source/Pipelines/Private/InterchangeGenericMaterialPipeline.cpp
@@ -1045,6 +1045,11 @@ void UInterchangeGenericMaterialPipeline::HandleCommonParameters(const UIntercha
 {
 	using namespace UE::Interchange::Materials::Common;
 
+	if (bool bScreenSpaceReflections; ShaderGraphNode->GetCustomScreenSpaceReflections(bScreenSpaceReflections))
+	{
+		MaterialFactoryNode->SetCustomScreenSpaceReflections(bScreenSpaceReflections);
+	}
+
 	// Two sidedness (ignored for thin translucency as it looks wrong)
 	if (!HasThinTranslucency(ShaderGraphNode))
 	{
@@ -1123,37 +1128,55 @@ void UInterchangeGenericMaterialPipeline::HandleCommonParameters(const UIntercha
 		}
 	}
 
-	// Opacity
+	// Opacity / OpacityMask
 	{
 		TGuardValue<EMaterialInputType> InputTypeBeingProcessedGuard(MaterialCreationContext.InputTypeBeingProcessed, EMaterialInputType::Scalar);
 
-		const bool bHasOpacityInput = UInterchangeShaderPortsAPI::HasInput(ShaderGraphNode, Parameters::Opacity);
-
-		if (bHasOpacityInput)
+		bool bUpdateBlendMode = false;
+		const bool bHasOpacityMaskInput = UInterchangeShaderPortsAPI::HasInput(ShaderGraphNode, Parameters::OpacityMask);
+		if (bHasOpacityMaskInput)
 		{
-			bool bHasSomeTransparency = true;
+			TTuple<UInterchangeMaterialExpressionFactoryNode*, FString> OpacityMaskExpressionFactoryNode =
+				CreateMaterialExpressionForInput(MaterialFactoryNode, ShaderGraphNode, Parameters::OpacityMask.ToString(), MaterialFactoryNode->GetUniqueID());
 
-			float OpacityValue;
-			if (ShaderGraphNode->GetFloatAttribute(UInterchangeShaderPortsAPI::MakeInputValueKey(Parameters::Opacity.ToString()), OpacityValue))
+			if (OpacityMaskExpressionFactoryNode.Get<0>())
 			{
-				bHasSomeTransparency = !FMath::IsNearlyEqual(OpacityValue, 1.f);
+				MaterialFactoryNode->ConnectOutputToOpacity(OpacityMaskExpressionFactoryNode.Get<0>()->GetUniqueID(), OpacityMaskExpressionFactoryNode.Get<1>());
 			}
-
-			if (bHasSomeTransparency)
+			bUpdateBlendMode = true;
+		}
+		else
+		{
+			const bool bHasOpacityInput = UInterchangeShaderPortsAPI::HasInput(ShaderGraphNode, Parameters::Opacity);
+			if (bHasOpacityInput)
 			{
-				TTuple<UInterchangeMaterialExpressionFactoryNode*, FString> OpacityExpressionFactoryNode =
-				CreateMaterialExpressionForInput(MaterialFactoryNode, ShaderGraphNode, Parameters::Opacity.ToString(), MaterialFactoryNode->GetUniqueID());
+				bool bHasSomeTransparency = true;
 
-				if (OpacityExpressionFactoryNode.Get<0>())
+				float OpacityValue;
+				if (ShaderGraphNode->GetFloatAttribute(UInterchangeShaderPortsAPI::MakeInputValueKey(Parameters::Opacity.ToString()), OpacityValue))
 				{
-					MaterialFactoryNode->ConnectOutputToOpacity(OpacityExpressionFactoryNode.Get<0>()->GetUniqueID(), OpacityExpressionFactoryNode.Get<1>());
+					bHasSomeTransparency = !FMath::IsNearlyEqual(OpacityValue, 1.f);
 				}
 
-				using namespace UE::Interchange::InterchangeGenericMaterialPipeline;
+				if (bHasSomeTransparency)
+				{
+					TTuple<UInterchangeMaterialExpressionFactoryNode*, FString> OpacityExpressionFactoryNode =
+						CreateMaterialExpressionForInput(MaterialFactoryNode, ShaderGraphNode, Parameters::Opacity.ToString(), MaterialFactoryNode->GetUniqueID());
+
+					if (OpacityExpressionFactoryNode.Get<0>())
+					{
+						MaterialFactoryNode->ConnectOutputToOpacity(OpacityExpressionFactoryNode.Get<0>()->GetUniqueID(), OpacityExpressionFactoryNode.Get<1>());
+					}
 
-				Private::UpdateBlendModeBasedOnOpacityAttributes(ShaderGraphNode, MaterialFactoryNode);
+					bUpdateBlendMode = true;
+				}
 			}
 		}
+
+		if (bUpdateBlendMode)
+		{
+			UE::Interchange::InterchangeGenericMaterialPipeline::Private::UpdateBlendModeBasedOnOpacityAttributes(ShaderGraphNode, MaterialFactoryNode);
+		}
 	}
 
 	// Ambient Occlusion
@@ -2181,10 +2204,19 @@ UInterchangeMaterialExpressionFactoryNode* UInterchangeGenericMaterialPipeline::
 
 	MaterialExpression->AddTargetNodeUid(ShaderNode->GetUniqueID());
 
-	if (*ShaderType == Nodes::TextureSample::Name)
+	if (*ShaderType == Nodes::TextureSample::Name
+		|| *ShaderType == Nodes::TextureObject::Name)
 	{
 		FString TextureUid;
-		ShaderNode->GetStringAttribute(UInterchangeShaderPortsAPI::MakeInputValueKey(Nodes::TextureSample::Inputs::Texture.ToString()), TextureUid);
+		
+		if (*ShaderType == Nodes::TextureSample::Name)
+		{
+			ShaderNode->GetStringAttribute(UInterchangeShaderPortsAPI::MakeInputValueKey(Nodes::TextureSample::Inputs::Texture.ToString()), TextureUid);
+		}
+		else if (*ShaderType == Nodes::TextureObject::Name)
+		{
+			ShaderNode->GetStringAttribute(UInterchangeShaderPortsAPI::MakeInputValueKey(Nodes::TextureObject::Inputs::Texture.ToString()), TextureUid);
+		}
 
 		// Make the material factory node have a dependency on the texture factory node so that the texture asset gets created first
 		if (const UInterchangeTextureNode* TextureNode = Cast<const UInterchangeTextureNode>(BaseNodeContainer->GetNode(TextureUid)))
diff --git a/Engine/Plugins/Interchange/Runtime/Source/Pipelines/Private/InterchangeglTFPipeline.cpp b/Engine/Plugins/Interchange/Runtime/Source/Pipelines/Private/InterchangeglTFPipeline.cpp
--- a/Engine/Plugins/Interchange/Runtime/Source/Pipelines/Private/InterchangeglTFPipeline.cpp
+++ b/Engine/Plugins/Interchange/Runtime/Source/Pipelines/Private/InterchangeglTFPipeline.cpp
@@ -9,7 +9,7 @@
 #include "InterchangeMaterialInstanceNode.h"
 #include "InterchangeMaterialFactoryNode.h"
 
-#include "Gltf/InterchangeGLTFMaterialInstances.h"
+#include "Gltf/InterchangeGLTFMaterial.h"
 
 #include "Materials/MaterialInstanceConstant.h"
 #include "Materials/MaterialInstanceDynamic.h"
@@ -85,54 +85,54 @@ TArray<FString> UGLTFPipelineSettings::ValidateMaterialInstancesAndParameters()
 
 TSet<FString> UGLTFPipelineSettings::GenerateExpectedParametersList(const FString& Identifier) const
 {
-	using namespace UE::Interchange::GLTFMaterialInstances;
+	using namespace UE::Interchange::GLTFMaterials;
 
 	TSet<FString> ExpectedParameters;
 
 	if (Identifier.Contains(TEXT("_Unlit")))
 	{
-		ExpectedParameters.Add(Parameters::BaseColorTexture);
-		ExpectedParameters.Add(Parameters::BaseColorTexture_OffsetScale);
-		ExpectedParameters.Add(Parameters::BaseColorTexture_Rotation);
-		ExpectedParameters.Add(Parameters::BaseColorTexture_TexCoord);
-		ExpectedParameters.Add(Parameters::BaseColorFactor);
+		ExpectedParameters.Add(Inputs::BaseColorTexture);
+		ExpectedParameters.Add(Inputs::BaseColorTexture_OffsetScale);
+		ExpectedParameters.Add(Inputs::BaseColorTexture_Rotation);
+		ExpectedParameters.Add(Inputs::BaseColorTexture_TexCoord);
+		ExpectedParameters.Add(Inputs::BaseColorFactor);
 
 		return ExpectedParameters;
 	}
 
 	//Generic ones:
 	{
-		ExpectedParameters.Add(Parameters::NormalTexture);
-		ExpectedParameters.Add(Parameters::NormalTexture_OffsetScale);
-		ExpectedParameters.Add(Parameters::NormalTexture_Rotation);
-		ExpectedParameters.Add(Parameters::NormalTexture_TexCoord);
-		ExpectedParameters.Add(Parameters::NormalScale);
+		ExpectedParameters.Add(Inputs::NormalTexture);
+		ExpectedParameters.Add(Inputs::NormalTexture_OffsetScale);
+		ExpectedParameters.Add(Inputs::NormalTexture_Rotation);
+		ExpectedParameters.Add(Inputs::NormalTexture_TexCoord);
+		ExpectedParameters.Add(Inputs::NormalScale);
 
 		if (!Identifier.Contains(TEXT("Transmission")))
 		{
-			ExpectedParameters.Add(Parameters::EmissiveTexture);
-			ExpectedParameters.Add(Parameters::EmissiveTexture_OffsetScale);
-			ExpectedParameters.Add(Parameters::EmissiveTexture_Rotation);
-			ExpectedParameters.Add(Parameters::EmissiveTexture_TexCoord);
-			ExpectedParameters.Add(Parameters::EmissiveFactor);
-			ExpectedParameters.Add(Parameters::EmissiveStrength);
+			ExpectedParameters.Add(Inputs::EmissiveTexture);
+			ExpectedParameters.Add(Inputs::EmissiveTexture_OffsetScale);
+			ExpectedParameters.Add(Inputs::EmissiveTexture_Rotation);
+			ExpectedParameters.Add(Inputs::EmissiveTexture_TexCoord);
+			ExpectedParameters.Add(Inputs::EmissiveFactor);
+			ExpectedParameters.Add(Inputs::EmissiveStrength);
 		}
 
-		ExpectedParameters.Add(Parameters::OcclusionTexture);
-		ExpectedParameters.Add(Parameters::OcclusionTexture_OffsetScale);
-		ExpectedParameters.Add(Parameters::OcclusionTexture_Rotation);
-		ExpectedParameters.Add(Parameters::OcclusionTexture_TexCoord);
-		ExpectedParameters.Add(Parameters::OcclusionStrength);
+		ExpectedParameters.Add(Inputs::OcclusionTexture);
+		ExpectedParameters.Add(Inputs::OcclusionTexture_OffsetScale);
+		ExpectedParameters.Add(Inputs::OcclusionTexture_Rotation);
+		ExpectedParameters.Add(Inputs::OcclusionTexture_TexCoord);
+		ExpectedParameters.Add(Inputs::OcclusionStrength);
 
 		if (!Identifier.Contains(TEXT("SpecularGlossiness")))
 		{
-			ExpectedParameters.Add(Parameters::IOR);
+			ExpectedParameters.Add(Inputs::IOR);
 
-			ExpectedParameters.Add(Parameters::SpecularTexture);
-			ExpectedParameters.Add(Parameters::SpecularTexture_OffsetScale);
-			ExpectedParameters.Add(Parameters::SpecularTexture_Rotation);
-			ExpectedParameters.Add(Parameters::SpecularTexture_TexCoord);
-			ExpectedParameters.Add(Parameters::SpecularFactor);
+			ExpectedParameters.Add(Inputs::SpecularTexture);
+			ExpectedParameters.Add(Inputs::SpecularTexture_OffsetScale);
+			ExpectedParameters.Add(Inputs::SpecularTexture_Rotation);
+			ExpectedParameters.Add(Inputs::SpecularTexture_TexCoord);
+			ExpectedParameters.Add(Inputs::SpecularFactor);
 		}
 	}
 
@@ -142,75 +142,75 @@ TSet<FString> UGLTFPipelineSettings::GenerateExpectedParametersList(const FStrin
 	{
 		//MetalRoughness Specific:
 
-		ExpectedParameters.Add(Parameters::BaseColorTexture);
-		ExpectedParameters.Add(Parameters::BaseColorTexture_OffsetScale);
-		ExpectedParameters.Add(Parameters::BaseColorTexture_Rotation);
-		ExpectedParameters.Add(Parameters::BaseColorTexture_TexCoord);
-		ExpectedParameters.Add(Parameters::BaseColorFactor);
-
-		ExpectedParameters.Add(Parameters::MetallicRoughnessTexture);
-		ExpectedParameters.Add(Parameters::MetallicRoughnessTexture_OffsetScale);
-		ExpectedParameters.Add(Parameters::MetallicRoughnessTexture_Rotation);
-		ExpectedParameters.Add(Parameters::MetallicRoughnessTexture_TexCoord);
-		ExpectedParameters.Add(Parameters::MetallicFactor);
-		ExpectedParameters.Add(Parameters::RoughnessFactor);
+		ExpectedParameters.Add(Inputs::BaseColorTexture);
+		ExpectedParameters.Add(Inputs::BaseColorTexture_OffsetScale);
+		ExpectedParameters.Add(Inputs::BaseColorTexture_Rotation);
+		ExpectedParameters.Add(Inputs::BaseColorTexture_TexCoord);
+		ExpectedParameters.Add(Inputs::BaseColorFactor);
+
+		ExpectedParameters.Add(Inputs::MetallicRoughnessTexture);
+		ExpectedParameters.Add(Inputs::MetallicRoughnessTexture_OffsetScale);
+		ExpectedParameters.Add(Inputs::MetallicRoughnessTexture_Rotation);
+		ExpectedParameters.Add(Inputs::MetallicRoughnessTexture_TexCoord);
+		ExpectedParameters.Add(Inputs::MetallicFactor);
+		ExpectedParameters.Add(Inputs::RoughnessFactor);
 	}
 	else if (Identifier.Contains(TEXT("ClearCoat")))
 	{
-		ExpectedParameters.Add(Parameters::ClearCoatTexture);
-		ExpectedParameters.Add(Parameters::ClearCoatTexture_OffsetScale);
-		ExpectedParameters.Add(Parameters::ClearCoatTexture_Rotation);
-		ExpectedParameters.Add(Parameters::ClearCoatTexture_TexCoord);
-		ExpectedParameters.Add(Parameters::ClearCoatFactor);
-
-		ExpectedParameters.Add(Parameters::ClearCoatRoughnessTexture);
-		ExpectedParameters.Add(Parameters::ClearCoatRoughnessTexture_OffsetScale);
-		ExpectedParameters.Add(Parameters::ClearCoatRoughnessTexture_Rotation);
-		ExpectedParameters.Add(Parameters::ClearCoatRoughnessTexture_TexCoord);
-		ExpectedParameters.Add(Parameters::ClearCoatRoughnessFactor);
-
-		ExpectedParameters.Add(Parameters::ClearCoatNormalTexture);
-		ExpectedParameters.Add(Parameters::ClearCoatNormalTexture_OffsetScale);
-		ExpectedParameters.Add(Parameters::ClearCoatNormalTexture_Rotation);
-		ExpectedParameters.Add(Parameters::ClearCoatNormalTexture_TexCoord);
-		ExpectedParameters.Add(Parameters::ClearCoatNormalScale);
+		ExpectedParameters.Add(Inputs::ClearCoatTexture);
+		ExpectedParameters.Add(Inputs::ClearCoatTexture_OffsetScale);
+		ExpectedParameters.Add(Inputs::ClearCoatTexture_Rotation);
+		ExpectedParameters.Add(Inputs::ClearCoatTexture_TexCoord);
+		ExpectedParameters.Add(Inputs::ClearCoatFactor);
+
+		ExpectedParameters.Add(Inputs::ClearCoatRoughnessTexture);
+		ExpectedParameters.Add(Inputs::ClearCoatRoughnessTexture_OffsetScale);
+		ExpectedParameters.Add(Inputs::ClearCoatRoughnessTexture_Rotation);
+		ExpectedParameters.Add(Inputs::ClearCoatRoughnessTexture_TexCoord);
+		ExpectedParameters.Add(Inputs::ClearCoatRoughnessFactor);
+
+		ExpectedParameters.Add(Inputs::ClearCoatNormalTexture);
+		ExpectedParameters.Add(Inputs::ClearCoatNormalTexture_OffsetScale);
+		ExpectedParameters.Add(Inputs::ClearCoatNormalTexture_Rotation);
+		ExpectedParameters.Add(Inputs::ClearCoatNormalTexture_TexCoord);
+		ExpectedParameters.Add(Inputs::ClearCoatNormalScale);
 	}
 	else if (Identifier.Contains(TEXT("Sheen")))
 	{
-		ExpectedParameters.Add(Parameters::SheenColorTexture);
-		ExpectedParameters.Add(Parameters::SheenColorTexture_OffsetScale);
-		ExpectedParameters.Add(Parameters::SheenColorTexture_Rotation);
-		ExpectedParameters.Add(Parameters::SheenColorTexture_TexCoord);
-		ExpectedParameters.Add(Parameters::SheenColorFactor);
-
-		ExpectedParameters.Add(Parameters::SheenRoughnessTexture);
-		ExpectedParameters.Add(Parameters::SheenRoughnessTexture_OffsetScale);
-		ExpectedParameters.Add(Parameters::SheenRoughnessTexture_Rotation);
-		ExpectedParameters.Add(Parameters::SheenRoughnessTexture_TexCoord);
-		ExpectedParameters.Add(Parameters::SheenRoughnessFactor);
+		ExpectedParameters.Add(Inputs::SheenColorTexture);
+		ExpectedParameters.Add(Inputs::SheenColorTexture_OffsetScale);
+		ExpectedParameters.Add(Inputs::SheenColorTexture_Rotation);
+		ExpectedParameters.Add(Inputs::SheenColorTexture_TexCoord);
+		ExpectedParameters.Add(Inputs::SheenColorFactor);
+
+		ExpectedParameters.Add(Inputs::SheenRoughnessTexture);
+		ExpectedParameters.Add(Inputs::SheenRoughnessTexture_OffsetScale);
+		ExpectedParameters.Add(Inputs::SheenRoughnessTexture_Rotation);
+		ExpectedParameters.Add(Inputs::SheenRoughnessTexture_TexCoord);
+		ExpectedParameters.Add(Inputs::SheenRoughnessFactor);
 	}
 	else if (Identifier.Contains(TEXT("Transmission")))
 	{
-		ExpectedParameters.Add(Parameters::TransmissionTexture);
-		ExpectedParameters.Add(Parameters::TransmissionTexture_OffsetScale);
-		ExpectedParameters.Add(Parameters::TransmissionTexture_Rotation);
-		ExpectedParameters.Add(Parameters::TransmissionTexture_TexCoord);
-		ExpectedParameters.Add(Parameters::TransmissionFactor);
+		ExpectedParameters.Add(Inputs::TransmissionTexture);
+		ExpectedParameters.Add(Inputs::TransmissionTexture_OffsetScale);
+		ExpectedParameters.Add(Inputs::TransmissionTexture_Rotation);
+		ExpectedParameters.Add(Inputs::TransmissionTexture_TexCoord);
+		ExpectedParameters.Add(Inputs::TransmissionFactor);
 	}
 	else if (Identifier.Contains(TEXT("SpecularGlossiness")))
 	{
-		ExpectedParameters.Add(Parameters::DiffuseTexture);
-		ExpectedParameters.Add(Parameters::DiffuseTexture_OffsetScale);
-		ExpectedParameters.Add(Parameters::DiffuseTexture_Rotation);
-		ExpectedParameters.Add(Parameters::DiffuseTexture_TexCoord);
-		ExpectedParameters.Add(Parameters::DiffuseFactor);
-
-		ExpectedParameters.Add(Parameters::SpecularGlossinessTexture);
-		ExpectedParameters.Add(Parameters::SpecularGlossinessTexture_OffsetScale);
-		ExpectedParameters.Add(Parameters::SpecularGlossinessTexture_Rotation);
-		ExpectedParameters.Add(Parameters::SpecularGlossinessTexture_TexCoord);
-		ExpectedParameters.Add(Parameters::SpecFactor);
-		ExpectedParameters.Add(Parameters::GlossinessFactor);
+		ExpectedParameters.Add(Inputs::DiffuseTexture);
+		ExpectedParameters.Add(Inputs::DiffuseTexture_OffsetScale);
+		ExpectedParameters.Add(Inputs::DiffuseTexture_Rotation);
+		ExpectedParameters.Add(Inputs::DiffuseTexture_TexCoord);
+		ExpectedParameters.Add(Inputs::DiffuseFactor);
+
+		ExpectedParameters.Add(Inputs::SpecularGlossinessTexture);
+		ExpectedParameters.Add(Inputs::SpecularGlossinessTexture_OffsetScale);
+		ExpectedParameters.Add(Inputs::SpecularGlossinessTexture_Rotation);
+		ExpectedParameters.Add(Inputs::SpecularGlossinessTexture_TexCoord);
+		ExpectedParameters.Add(Inputs::SpecFactor);
+		ExpectedParameters.Add(Inputs::GlossinessFactor);
 	}
 
 	return ExpectedParameters;
@@ -218,7 +218,7 @@ TSet<FString> UGLTFPipelineSettings::GenerateExpectedParametersList(const FStrin
 
 void UGLTFPipelineSettings::BuildMaterialInstance(const UInterchangeShaderGraphNode* ShaderGraphNode, UInterchangeMaterialInstanceFactoryNode* MaterialInstanceFactoryNode)
 {
-	using namespace UE::Interchange::GLTFMaterialInstances;
+	using namespace UE::Interchange::GLTFMaterials;
 
 	TArray<UE::Interchange::FAttributeKey> AttributeKeys;
 	ShaderGraphNode->GetAttributeKeys(AttributeKeys);
