Moved intialization of ImterchangeImport module to PreDefault instead of Default
See UDN https://epicgames.lightning.force.com/lightning/r/Case/5004z00001mDByrAAG/view for details
Note that the declaration order in the uplugin file must be kept.

diff --git a/Engine/Plugins/Interchange/Runtime/Interchange.uplugin b/Engine/Plugins/Interchange/Runtime/Interchange.uplugin
--- a/Engine/Plugins/Interchange/Runtime/Interchange.uplugin
+++ b/Engine/Plugins/Interchange/Runtime/Interchange.uplugin
@@ -34,7 +34,7 @@
 		{
 			"Name": "InterchangeImport",
 			"Type": "Runtime",
-			"LoadingPhase": "Default"
+			"LoadingPhase": "PreDefault"
 		},
 		{
 			"Name": "InterchangeMessages",
