Fixed miscellanous issues
  - In FBX, only connect opacity when transparency is not 0
 - In FBX, add typed attribute to shader node if there is no texture on a property
  - In OBJ, removed multiplication by static value when a texture is present
  - Restored use of DefaultSceneAssetPipeline for scene import

diff --git a/Engine/Plugins/Interchange/Runtime/Config/BaseInterchange.ini b/Engine/Plugins/Interchange/Runtime/Config/BaseInterchange.ini
--- a/Engine/Plugins/Interchange/Runtime/Config/BaseInterchange.ini
+++ b/Engine/Plugins/Interchange/Runtime/Config/BaseInterchange.ini
@@ -45,7 +45,7 @@
 
 [/Script/InterchangeEngine.InterchangeProjectSettings]
 ContentImportSettings=(DefaultPipelineStackOverride=((Materials, "Materials"),(Textures, "Textures")),ShowPipelineStacksConfigurationDialogOverride=((Textures, False)),PipelineStacks=(("Assets", (Pipelines=("/Interchange/Pipelines/DefaultAssetsPipeline.DefaultAssetsPipeline"),PerTranslatorPipelines=((Translator="/Script/InterchangeImport.InterchangeMaterialXTranslator",Pipelines=("/Interchange/Pipelines/DefaultAssetsPipeline.DefaultAssetsPipeline","/Interchange/Pipelines/DefaultMaterialXPipeline.DefaultMaterialXPipeline")),(Translator="/Script/InterchangeImport.InterchangeGltfTranslator",Pipelines=("/Interchange/Pipelines/DefaultGLTFAssetsPipeline.DefaultGLTFAssetsPipeline","/Interchange/Pipelines/DefaultGLTFPipeline.DefaultGLTFPipeline"))))),("Materials", (Pipelines=("/Interchange/Pipelines/DefaultMaterialPipeline.DefaultMaterialPipeline"),PerTranslatorPipelines=((Translator="/Script/InterchangeImport.InterchangeMaterialXTranslator",Pipelines=("/Interchange/Pipelines/DefaultMaterialPipeline.DefaultMaterialPipeline","/Interchange/Pipelines/DefaultMaterialXPipeline.DefaultMaterialXPipeline"))))),("Textures", (Pipelines=("/Interchange/Pipelines/DefaultTexturePipeline.DefaultTexturePipeline")))),DefaultPipelineStack="Assets",PipelineConfigurationDialogClass="/Script/InterchangeEditorPipelines.InterchangePipelineConfigurationGeneric",bShowPipelineStacksConfigurationDialog=True)
-SceneImportSettings=(PipelineStacks=(("Scene", (Pipelines=("/Interchange/Pipelines/DefaultAssetsPipeline.DefaultAssetsPipeline","/Interchange/Pipelines/DefaultSceneLevelPipeline.DefaultSceneLevelPipeline"),PerTranslatorPipelines=((Translator="/Script/InterchangeImport.InterchangeMaterialXTranslator",Pipelines=("/Interchange/Pipelines/DefaultAssetsPipeline.DefaultAssetsPipeline","/Interchange/Pipelines/DefaultSceneLevelPipeline.DefaultSceneLevelPipeline","/Interchange/Pipelines/DefaultMaterialXPipeline.DefaultMaterialXPipeline")),(Translator="/Script/InterchangeImport.InterchangeGltfTranslator",Pipelines=("/Interchange/Pipelines/DefaultGLTFSceneAssetsPipeline.DefaultGLTFSceneAssetsPipeline","/Interchange/Pipelines/DefaultSceneLevelPipeline.DefaultSceneLevelPipeline","/Interchange/Pipelines/DefaultGLTFPipeline.DefaultGLTFPipeline")))))),DefaultPipelineStack="Scene",PipelineConfigurationDialogClass="/Script/InterchangeEditorPipelines.InterchangePipelineConfigurationGeneric",bShowPipelineStacksConfigurationDialog=True)
+SceneImportSettings=(PipelineStacks=(("Scene", (Pipelines=("/Interchange/Pipelines/DefaultSceneAssetsPipeline.DefaultSceneAssetsPipeline","/Interchange/Pipelines/DefaultSceneLevelPipeline.DefaultSceneLevelPipeline"),PerTranslatorPipelines=((Translator="/Script/InterchangeImport.InterchangeMaterialXTranslator",Pipelines=("/Interchange/Pipelines/DefaultSceneAssetsPipeline.DefaultSceneAssetsPipeline","/Interchange/Pipelines/DefaultSceneLevelPipeline.DefaultSceneLevelPipeline","/Interchange/Pipelines/DefaultMaterialXPipeline.DefaultMaterialXPipeline")),(Translator="/Script/InterchangeImport.InterchangeGltfTranslator",Pipelines=("/Interchange/Pipelines/DefaultGLTFSceneAssetsPipeline.DefaultGLTFSceneAssetsPipeline","/Interchange/Pipelines/DefaultSceneLevelPipeline.DefaultSceneLevelPipeline","/Interchange/Pipelines/DefaultGLTFPipeline.DefaultGLTFPipeline")))))),DefaultPipelineStack="Scene",PipelineConfigurationDialogClass="/Script/InterchangeEditorPipelines.InterchangePipelineConfigurationGeneric",bShowPipelineStacksConfigurationDialog=True)
 FilePickerClass=/Script/InterchangeEditorUtilities.InterchangeFilePickerGeneric
 bStaticMeshUseSmoothEdgesIfSmoothingInformationIsMissing=True
 GenericPipelineClass=/Script/InterchangePipelines.InterchangeGenericAssetsPipeline
diff --git a/Engine/Plugins/Interchange/Runtime/Source/Import/Private/Mesh/InterchangeOBJTranslator.cpp b/Engine/Plugins/Interchange/Runtime/Source/Import/Private/Mesh/InterchangeOBJTranslator.cpp
--- a/Engine/Plugins/Interchange/Runtime/Source/Import/Private/Mesh/InterchangeOBJTranslator.cpp
+++ b/Engine/Plugins/Interchange/Runtime/Source/Import/Private/Mesh/InterchangeOBJTranslator.cpp
@@ -1073,22 +1073,7 @@ namespace ObjTranslatorUtils
 	{
 		UInterchangeShaderNode* TextureSampleShader = CreateMaterialTextureSampleNode(BaseNodeContainer, ShaderGraphNode, TexturePath);
 
-		if (TextureSampleShader && IsColorInitialized(Color))
-		{
-			FString MapName = Materials::Phong::Parameters::SpecularColor.ToString();
-
-			using namespace UE::Interchange::Materials::Standard::Nodes;
-
-			const FString MultiplierNodeName = MapName + TEXT("Multiply");
-			UInterchangeShaderNode* MultiplierNode = UInterchangeShaderNode::Create(&BaseNodeContainer, MultiplierNodeName, ShaderGraphNode->GetUniqueID());
-			MultiplierNode->SetCustomShaderType(Multiply::Name.ToString());
-
-			UInterchangeShaderPortsAPI::ConnectDefaultOuputToInput(MultiplierNode, Multiply::Inputs::A.ToString(), TextureSampleShader->GetUniqueID());
-			MultiplierNode->AddLinearColorAttribute(UInterchangeShaderPortsAPI::MakeInputValueKey( Multiply::Inputs::B.ToString() ), Color);
-
-			UInterchangeShaderPortsAPI::ConnectDefaultOuputToInput(ShaderGraphNode, InputType, MultiplierNode->GetUniqueID());
-		}
-		else if (TextureSampleShader)
+		if (TextureSampleShader)
 		{
 			UInterchangeShaderPortsAPI::ConnectDefaultOuputToInput(ShaderGraphNode, InputType, TextureSampleShader->GetUniqueID());
 		}
diff --git a/Engine/Plugins/Interchange/Runtime/Source/Parsers/Fbx/Private/FbxMaterial.cpp b/Engine/Plugins/Interchange/Runtime/Source/Parsers/Fbx/Private/FbxMaterial.cpp
--- a/Engine/Plugins/Interchange/Runtime/Source/Parsers/Fbx/Private/FbxMaterial.cpp
+++ b/Engine/Plugins/Interchange/Runtime/Source/Parsers/Fbx/Private/FbxMaterial.cpp
@@ -109,11 +109,45 @@ namespace UE
 			{
 				using namespace Materials::Standard::Nodes;
 
-				UInterchangeShaderNode* NodeToConnectTo = ShaderGraphNode;
-				FString InputToConnectTo = InputName.ToString();
-
 				const int32 TextureCount = Property.GetSrcObjectCount<FbxFileTexture>();
 				const FbxDataType DataType = Property.GetPropertyDataType();
+				FString InputToConnectTo = InputName.ToString();
+
+				if (TextureCount == 0)
+				{
+					if (DataType.GetType() == eFbxDouble || DataType.GetType() == eFbxFloat || DataType.GetType() == eFbxInt)
+					{
+						float PropertyValue = Property.Get<float>() * Factor;
+						if (bInverse)
+						{
+							PropertyValue = 1.f - PropertyValue;
+						}
+
+						ShaderGraphNode->AddFloatAttribute(UInterchangeShaderPortsAPI::MakeInputValueKey(InputToConnectTo), PropertyValue);
+					}
+					else if (DataType.GetType() == eFbxDouble3)
+					{
+						FLinearColor PropertyValue = FFbxConvert::ConvertColor(Property.Get<FbxDouble3>()) * Factor;
+						if (bInverse)
+						{
+							PropertyValue = FLinearColor::White - PropertyValue;
+						}
+
+						if (DefaultValue.IsType<FLinearColor>())
+						{
+							ShaderGraphNode->AddLinearColorAttribute(UInterchangeShaderPortsAPI::MakeInputValueKey(InputToConnectTo), PropertyValue);
+						}
+						else if (DefaultValue.IsType<float>())
+						{
+							// We're connecting a linear color to a float input. Ideally, we'd go through a desaturate, but for now we'll just take the red channel and ignore the rest.
+							ShaderGraphNode->AddFloatAttribute(UInterchangeShaderPortsAPI::MakeInputValueKey(InputToConnectTo), PropertyValue.R);
+						}
+					}
+
+					return true;
+				}
+
+				UInterchangeShaderNode* NodeToConnectTo = ShaderGraphNode;
 
 				if (bInverse)
 				{
@@ -152,39 +186,18 @@ namespace UE
 				}
 
 				// Handles max one texture per property.
-				if (TextureCount > 0)
-				{
-					FbxFileTexture* FbxTexture = Property.GetSrcObject<FbxFileTexture>(0);
-					if (UInterchangeShaderNode* TextureSampleShader = CreateTextureSampler(FbxTexture, NodeContainer, ShaderGraphNode->GetUniqueID()))
-					{
-						UInterchangeShaderPortsAPI::ConnectDefaultOuputToInput(NodeToConnectTo, InputToConnectTo, TextureSampleShader->GetUniqueID());
-					}
-					else
-					{
-						UInterchangeResultTextureWarning_TextureFileDoNotExist* Message = Parser.AddMessage<UInterchangeResultTextureWarning_TextureFileDoNotExist>();
-						Message->TextureName = FbxTexture ? UTF8_TO_TCHAR(FbxTexture->GetFileName()) : TEXT("Undefined");
-						Message->MaterialName = ShaderGraphNode->GetDisplayLabel();
-
-						return false;
-					}
-				}
-				else if (DataType.GetType() == eFbxDouble || DataType.GetType() == eFbxFloat || DataType.GetType() == eFbxInt)
+				FbxFileTexture* FbxTexture = Property.GetSrcObject<FbxFileTexture>(0);
+				if (UInterchangeShaderNode* TextureSampleShader = CreateTextureSampler(FbxTexture, NodeContainer, ShaderGraphNode->GetUniqueID()))
 				{
-					NodeToConnectTo->AddFloatAttribute(UInterchangeShaderPortsAPI::MakeInputValueKey(InputToConnectTo), Property.Get<float>());
+					UInterchangeShaderPortsAPI::ConnectDefaultOuputToInput(NodeToConnectTo, InputToConnectTo, TextureSampleShader->GetUniqueID());
 				}
-				else if (DataType.GetType() == eFbxDouble3)
+				else
 				{
-					const FLinearColor PropertyValue = FFbxConvert::ConvertColor(Property.Get<FbxDouble3>());
+					UInterchangeResultTextureWarning_TextureFileDoNotExist* Message = Parser.AddMessage<UInterchangeResultTextureWarning_TextureFileDoNotExist>();
+					Message->TextureName = FbxTexture ? UTF8_TO_TCHAR(FbxTexture->GetFileName()) : TEXT("Undefined");
+					Message->MaterialName = ShaderGraphNode->GetDisplayLabel();
 
-					if (DefaultValue.IsType<FLinearColor>())
-					{
-						NodeToConnectTo->AddLinearColorAttribute(UInterchangeShaderPortsAPI::MakeInputValueKey(InputToConnectTo), PropertyValue);
-					}
-					else if (DefaultValue.IsType<float>())
-					{
-						// We're connecting a linear color to a float input. Ideally, we'd go through a desaturate, but for now we'll just take the red channel and ignore the rest.
-						NodeToConnectTo->AddFloatAttribute(UInterchangeShaderPortsAPI::MakeInputValueKey(InputToConnectTo), PropertyValue.R);
-					}
+					return false;
 				}
 
 				return true;
@@ -377,11 +390,29 @@ namespace UE
 				}
 
 				// Opacity
+				// Connect only if transparency is either a texture or different from 0.f
 				{
-					const float TransparencyFactor = (float)((FbxSurfaceLambert*)SurfaceMaterial)->TransparencyFactor.Get();
-					TVariant<FLinearColor, float> DefaultValue;
-					DefaultValue.Set<float>(0.f); // Opaque
-					ConnectInput(Phong::Parameters::Opacity, FbxSurfaceMaterial::sTransparentColor, TransparencyFactor, DefaultValue, true);
+					FBXSDK_NAMESPACE::FbxProperty MaterialProperty = SurfaceMaterial->FindProperty(FbxSurfaceMaterial::sTransparentColor);
+					if (ShouldConvertProperty(MaterialProperty))
+					{
+						const float TransparencyFactor = (float)((FbxSurfaceLambert*)SurfaceMaterial)->TransparencyFactor.Get();
+
+						if (MaterialProperty.GetSrcObjectCount<FbxFileTexture>() > 0)
+						{
+							TVariant<FLinearColor, float> DefaultValue;
+							DefaultValue.Set<float>(0.f); // Opaque
+							ConvertPropertyToShaderNode(NodeContainer, ShaderGraphNode, MaterialProperty, TransparencyFactor, Phong::Parameters::Opacity, DefaultValue, true);
+						}
+						else
+						{
+							const float OpacityScalar = 1.f - (MaterialProperty.Get<float>() * TransparencyFactor);
+							if (OpacityScalar < 1.f)
+							{
+								const FString InputName = UInterchangeShaderPortsAPI::MakeInputValueKey(Phong::Parameters::Opacity.ToString());
+								ShaderGraphNode->AddFloatAttribute(InputName, OpacityScalar);
+							}
+						}
+					}
 				}
 
 				// Get specular color and shininess if FBX material is a Phong one
