Allowing Translators to import only materials when selecting only the material pipeline in the Interchange import dialog

diff --git a/Engine/Plugins/Interchange/Runtime/Source/Pipelines/Private/InterchangeGenericMaterialPipeline.cpp b/Engine/Plugins/Interchange/Runtime/Source/Pipelines/Private/InterchangeGenericMaterialPipeline.cpp
--- a/Engine/Plugins/Interchange/Runtime/Source/Pipelines/Private/InterchangeGenericMaterialPipeline.cpp
+++ b/Engine/Plugins/Interchange/Runtime/Source/Pipelines/Private/InterchangeGenericMaterialPipeline.cpp
@@ -308,11 +308,12 @@ void UInterchangeGenericMaterialPipeline::ExecutePipeline(UInterchangeBaseNodeCo
 	});
 
 	// Check to see whether materials should be created even if unused
-	// By default we do not create the materials, every node with mesh attribute can enable them. So we won't create unused materials.
-	bool bImportUnusedMaterial = false;
+	// By default we let the setting of the pipeline to decide if we create the materials, every node with mesh attribute can enable/disable them, depending on the pipeline stack chosen.
+	bool bImportUnusedMaterial = bImportMaterials;
 	if (const UInterchangeSourceNode* SourceNode = UInterchangeSourceNode::GetUniqueInstance(BaseNodeContainer))
 	{
 		SourceNode->GetCustomImportUnusedMaterial(bImportUnusedMaterial);
+		bImportUnusedMaterial |= bImportMaterials ;
 	}
 
 	// Can't import materials at runtime, fallback to instances
