MaterialX:
- Add the support of StandardSurface shader with Substrate
- Add two new assets: MX_StandardSurface_Substrate and MX_TransmissionSurface_Substrate, based on the same assets in Engine/Strata, just modifying a bit the name of the inputs to be compliant with Autodesk and MaterialX spec, and also setting the default values correctly from the spec
- Initializing once the PredefinedSurfaceShaders in the InterchangeMaterialXPipelineCustomizations if substrate is enabled, otherwise by default it'll take the default MX_StandardSurface asset

InterchangeShaderGraphNode:
- Add a method to set the Blend mode, because with Substrate with have no way of deciding it since there's only input of the main material (Front Material)

-InterchangeGenericMaterialPipeline and InterchangeMaterialFactory handles now Substrate

diff --git a/Engine/Plugins/Interchange/Editor/Source/Pipelines/Private/InterchangeMaterialXPipelineCustomizations.cpp b/Engine/Plugins/Interchange/Editor/Source/Pipelines/Private/InterchangeMaterialXPipelineCustomizations.cpp
--- a/Engine/Plugins/Interchange/Editor/Source/Pipelines/Private/InterchangeMaterialXPipelineCustomizations.cpp
+++ b/Engine/Plugins/Interchange/Editor/Source/Pipelines/Private/InterchangeMaterialXPipelineCustomizations.cpp
@@ -4,9 +4,10 @@
 
 #include "InterchangeMaterialDefinitions.h"
 #include "InterchangeMaterialXPipeline.h"
-#include "Materials/MaterialFunction.h"
 
 #include "DetailLayoutBuilder.h"
+#include "Engine/RendererSettings.h"
+#include "Materials/MaterialFunction.h"
 #include "Modules/ModuleManager.h"
 #include "PropertyCustomizationHelpers.h"
 #include "Widgets/Text/STextBlock.h"
@@ -73,6 +74,25 @@ void FInterchangeMaterialXPipelineSettingsCustomization::CustomizeDetails(IDetai
 		return;
 	}
 
+	// Init the PredefinedSurfaceShaders with Substrate assets, since the default value is set in BaseInterchange.ini and we have no way in the config file to conditionally init a property
+	if(GetDefault<URendererSettings>()->bEnableStrata)
+	{
+		static bool bFirstInit = true;
+		if(bFirstInit)
+		{
+			bFirstInit = false;
+			if(FString StandardSurfacePath = MaterialXSettings->GetAssetPathString(EInterchangeMaterialXShaders::StandardSurface); StandardSurfacePath == TEXT("/Interchange/Functions/MX_StandardSurface.MX_StandardSurface"))
+			{
+				MaterialXSettings->PredefinedSurfaceShaders.Add(EInterchangeMaterialXShaders::StandardSurface, FSoftObjectPath{ TEXT("/Interchange/Functions/MX_StandardSurface_Substrate.MX_StandardSurface_Substrate") });
+			}
+
+			if(FString TransmissionSurfacePath = MaterialXSettings->GetAssetPathString(EInterchangeMaterialXShaders::StandardSurfaceTransmission); TransmissionSurfacePath == TEXT("/Interchange/Functions/MX_TransmissionSurface.MX_TransmissionSurface"))
+			{
+				MaterialXSettings->PredefinedSurfaceShaders.Add(EInterchangeMaterialXShaders::StandardSurfaceTransmission, FSoftObjectPath{ TEXT("/Interchange/Functions/MX_TransmissionSurface_Substrate.MX_TransmissionSurface_Substrate") });
+			}
+		}
+	}
+
 	TSharedRef< IPropertyHandle > PairingsHandle = DetailBuilder.GetProperty(GET_MEMBER_NAME_CHECKED(UMaterialXPipelineSettings, PredefinedSurfaceShaders));
 	if (!PairingsHandle->IsValidHandle())
 	{
diff --git a/Engine/Plugins/Interchange/Runtime/Source/Import/Private/Material/InterchangeMaterialFactory.cpp b/Engine/Plugins/Interchange/Runtime/Source/Import/Private/Material/InterchangeMaterialFactory.cpp
--- a/Engine/Plugins/Interchange/Runtime/Source/Import/Private/Material/InterchangeMaterialFactory.cpp
+++ b/Engine/Plugins/Interchange/Runtime/Source/Import/Private/Material/InterchangeMaterialFactory.cpp
@@ -783,6 +783,57 @@ void UInterchangeMaterialFactory::SetupMaterial(UMaterial* Material, const FImpo
 
 	FMaterialExpressionBuilder Builder(Material, nullptr, Arguments);
 
+	// Substrate
+	if(UInterchangeShaderPortsAPI::HasInput(MaterialFactoryNode, Substrate::Parameters::FrontMaterial))
+	{
+		// Front Material
+		{
+			FString ExpressionNodeUid;
+			FString OutputName;
+
+			UInterchangeShaderPortsAPI::GetInputConnection(MaterialFactoryNode, Substrate::Parameters::FrontMaterial.ToString(), ExpressionNodeUid, OutputName);
+
+			const UInterchangeMaterialExpressionFactoryNode* FrontMaterial = Cast<UInterchangeMaterialExpressionFactoryNode>(Arguments.NodeContainer->GetNode(ExpressionNodeUid));
+
+			if(FrontMaterial)
+			{
+				if(UMaterialExpression* FrontMaterialExpression = Builder.CreateExpressionsForNode(*FrontMaterial))
+				{
+					if(FExpressionInput* FrontMaterialInput = Material->GetExpressionInputForProperty(MP_FrontMaterial))
+					{
+						FrontMaterialExpression->ConnectExpression(FrontMaterialInput, GetOutputIndex(*FrontMaterialExpression, OutputName));
+					}
+				}
+			}
+		}
+
+		// Opacity Mask
+		if(UInterchangeShaderPortsAPI::HasInput(MaterialFactoryNode, Substrate::Parameters::OpacityMask))
+		{
+			FString ExpressionNodeUid;
+			FString OutputName;
+
+			UInterchangeShaderPortsAPI::GetInputConnection(MaterialFactoryNode, Substrate::Parameters::FrontMaterial.ToString(), ExpressionNodeUid, OutputName);
+
+			const UInterchangeMaterialExpressionFactoryNode* OpacityMask = Cast<UInterchangeMaterialExpressionFactoryNode>(Arguments.NodeContainer->GetNode(ExpressionNodeUid));
+
+			if(OpacityMask)
+			{
+				if(UMaterialExpression* OpacityMaskExpression = Builder.CreateExpressionsForNode(*OpacityMask))
+				{
+					if(FExpressionInput* OpacityMaskInput = Material->GetExpressionInputForProperty(MP_OpacityMask))
+					{
+						OpacityMaskExpression->ConnectExpression(OpacityMaskInput, GetOutputIndex(*OpacityMaskExpression, OutputName));
+					}
+				}
+			}
+		}
+
+		UMaterialEditingLibrary::LayoutMaterialExpressions(Material);
+
+		return;
+	}
+
 	if (UInterchangeShaderPortsAPI::HasInput(MaterialFactoryNode, Common::Parameters::BxDF))
 	{
 		FString ExpressionNodeUid;
diff --git a/Engine/Plugins/Interchange/Runtime/Source/Import/Private/MaterialX/MaterialXUtils/MaterialXStandardSurfaceShader.cpp b/Engine/Plugins/Interchange/Runtime/Source/Import/Private/MaterialX/MaterialXUtils/MaterialXStandardSurfaceShader.cpp
--- a/Engine/Plugins/Interchange/Runtime/Source/Import/Private/MaterialX/MaterialXUtils/MaterialXStandardSurfaceShader.cpp
+++ b/Engine/Plugins/Interchange/Runtime/Source/Import/Private/MaterialX/MaterialXUtils/MaterialXStandardSurfaceShader.cpp
@@ -2,6 +2,7 @@
 
 #if WITH_EDITOR
 #include "MaterialXStandardSurfaceShader.h"
+#include "Engine/RendererSettings.h"
 
 namespace mx = MaterialX;
 
@@ -24,118 +25,133 @@ void FMaterialXStandardSurfaceShader::Translate(mx::NodePtr StandardSurfaceNode)
 	
 	using namespace UE::Interchange::Materials;
 
-	UInterchangeFunctionCallShaderNode* StandardSurfaceShaderNode;
-	
-	const FString NodeUID = UInterchangeShaderNode::MakeNodeUid(ANSI_TO_TCHAR(StandardSurfaceNode->getName().c_str()), FStringView{});
+	bool bIsSubstrateEnabled = GetDefault<URendererSettings>()->bEnableStrata;
 
+	if(bIsSubstrateEnabled)
+	{
+		ConnectToSubstrateStandardSurface();
+	}
+	else
+	{
+		ConnectToStandardSurface();
+	}
+}
+
+void FMaterialXStandardSurfaceShader::ConnectToStandardSurface()
+{
+	using namespace UE::Interchange::Materials;
+	using namespace mx::StandardSurface;
+
+	UInterchangeFunctionCallShaderNode* StandardSurfaceShaderNode;
+	const FString NodeUID = UInterchangeShaderNode::MakeNodeUid(ANSI_TO_TCHAR(SurfaceShaderNode->getName().c_str()), FStringView{});
 	if(StandardSurfaceShaderNode = const_cast<UInterchangeFunctionCallShaderNode*>(Cast<UInterchangeFunctionCallShaderNode>(NodeContainer.GetNode(NodeUID))); !StandardSurfaceShaderNode)
 	{
-		const FString NodeName = StandardSurfaceNode->getName().c_str();
+		const FString NodeName = SurfaceShaderNode->getName().c_str();
 		StandardSurfaceShaderNode = NewObject<UInterchangeFunctionCallShaderNode>(&NodeContainer);
 		StandardSurfaceShaderNode->InitializeNode(NodeUID, NodeName, EInterchangeNodeContainerType::TranslatedAsset);
 
 		StandardSurfaceShaderNode->SetCustomMaterialFunction(TEXT("/Interchange/Functions/MX_StandardSurface.MX_StandardSurface"));
 		NodeContainer.AddNode(StandardSurfaceShaderNode);
 
-		ShaderNodes.Add({ NodeName, StandardSurfaceNode->getNodeDef(mx::EMPTY_STRING, true)->getActiveOutputs()[0]->getName().c_str() }, StandardSurfaceShaderNode);
+		ShaderNodes.Add({ NodeName, SurfaceShaderNode->getNodeDef(mx::EMPTY_STRING, true)->getActiveOutputs()[0]->getName().c_str() }, StandardSurfaceShaderNode);
 	}
 
 	// Inputs
 	//Base
-	ConnectNodeOutputToInput(mx::StandardSurface::Input::Base, StandardSurfaceShaderNode, StandardSurface::Parameters::Base.ToString(), mx::StandardSurface::DefaultValue::Float::Base);
+	ConnectNodeOutputToInput(Input::Base, StandardSurfaceShaderNode, StandardSurface::Parameters::Base.ToString(), DefaultValue::Float::Base);
 
 	//Base Color
-	ConnectNodeOutputToInput(mx::StandardSurface::Input::BaseColor, StandardSurfaceShaderNode, StandardSurface::Parameters::BaseColor.ToString(), mx::StandardSurface::DefaultValue::Color3::BaseColor);
+	ConnectNodeOutputToInput(Input::BaseColor, StandardSurfaceShaderNode, StandardSurface::Parameters::BaseColor.ToString(), DefaultValue::Color3::BaseColor);
 
 	//Diffuse Roughness
-	ConnectNodeOutputToInput(mx::StandardSurface::Input::DiffuseRoughness, StandardSurfaceShaderNode, StandardSurface::Parameters::DiffuseRoughness.ToString(), mx::StandardSurface::DefaultValue::Float::DiffuseRoughness);
+	ConnectNodeOutputToInput(Input::DiffuseRoughness, StandardSurfaceShaderNode, StandardSurface::Parameters::DiffuseRoughness.ToString(), DefaultValue::Float::DiffuseRoughness);
 
 	//Specular
-	ConnectNodeOutputToInput(mx::StandardSurface::Input::Specular, StandardSurfaceShaderNode, StandardSurface::Parameters::Specular.ToString(), mx::StandardSurface::DefaultValue::Float::Specular);
+	ConnectNodeOutputToInput(Input::Specular, StandardSurfaceShaderNode, StandardSurface::Parameters::Specular.ToString(), DefaultValue::Float::Specular);
 
 	//Specular Roughness
-	ConnectNodeOutputToInput(mx::StandardSurface::Input::SpecularRoughness, StandardSurfaceShaderNode, StandardSurface::Parameters::SpecularRoughness.ToString(), mx::StandardSurface::DefaultValue::Float::SpecularRoughness);
+	ConnectNodeOutputToInput(Input::SpecularRoughness, StandardSurfaceShaderNode, StandardSurface::Parameters::SpecularRoughness.ToString(), DefaultValue::Float::SpecularRoughness);
 
 	//Specular IOR
-	ConnectNodeOutputToInput(mx::StandardSurface::Input::SpecularIOR, StandardSurfaceShaderNode, StandardSurface::Parameters::SpecularIOR.ToString(), mx::StandardSurface::DefaultValue::Float::SpecularIOR);
+	ConnectNodeOutputToInput(Input::SpecularIOR, StandardSurfaceShaderNode, StandardSurface::Parameters::SpecularIOR.ToString(), DefaultValue::Float::SpecularIOR);
 
 	//Specular Anisotropy
-	ConnectNodeOutputToInput(mx::StandardSurface::Input::SpecularAnisotropy, StandardSurfaceShaderNode, StandardSurface::Parameters::SpecularAnisotropy.ToString(), mx::StandardSurface::DefaultValue::Float::SpecularAnisotropy);
+	ConnectNodeOutputToInput(Input::SpecularAnisotropy, StandardSurfaceShaderNode, StandardSurface::Parameters::SpecularAnisotropy.ToString(), DefaultValue::Float::SpecularAnisotropy);
 
 	//Specular Rotation
-	ConnectNodeOutputToInput(mx::StandardSurface::Input::SpecularRotation, StandardSurfaceShaderNode, StandardSurface::Parameters::SpecularRotation.ToString(), mx::StandardSurface::DefaultValue::Float::SpecularRotation);
+	ConnectNodeOutputToInput(Input::SpecularRotation, StandardSurfaceShaderNode, StandardSurface::Parameters::SpecularRotation.ToString(), DefaultValue::Float::SpecularRotation);
 
 	//Metallic
-	ConnectNodeOutputToInput(mx::StandardSurface::Input::Metalness, StandardSurfaceShaderNode, StandardSurface::Parameters::Metalness.ToString(), mx::StandardSurface::DefaultValue::Float::Metalness);
+	ConnectNodeOutputToInput(Input::Metalness, StandardSurfaceShaderNode, StandardSurface::Parameters::Metalness.ToString(), DefaultValue::Float::Metalness);
 
 	//Subsurface
-	ConnectNodeOutputToInput(mx::StandardSurface::Input::Subsurface, StandardSurfaceShaderNode, StandardSurface::Parameters::Subsurface.ToString(), mx::StandardSurface::DefaultValue::Float::Subsurface);
+	ConnectNodeOutputToInput(Input::Subsurface, StandardSurfaceShaderNode, StandardSurface::Parameters::Subsurface.ToString(), DefaultValue::Float::Subsurface);
 
 	//Subsurface Color
-	ConnectNodeOutputToInput(mx::StandardSurface::Input::SubsurfaceColor, StandardSurfaceShaderNode, StandardSurface::Parameters::SubsurfaceColor.ToString(), mx::StandardSurface::DefaultValue::Color3::SubsurfaceColor);
+	ConnectNodeOutputToInput(Input::SubsurfaceColor, StandardSurfaceShaderNode, StandardSurface::Parameters::SubsurfaceColor.ToString(), DefaultValue::Color3::SubsurfaceColor);
 
 	//Subsurface Radius
-	ConnectNodeOutputToInput(mx::StandardSurface::Input::SubsurfaceRadius, StandardSurfaceShaderNode, StandardSurface::Parameters::SubsurfaceRadius.ToString(), mx::StandardSurface::DefaultValue::Color3::SubsurfaceRadius);
+	ConnectNodeOutputToInput(Input::SubsurfaceRadius, StandardSurfaceShaderNode, StandardSurface::Parameters::SubsurfaceRadius.ToString(), DefaultValue::Color3::SubsurfaceRadius);
 
 	//Subsurface Scale
-	ConnectNodeOutputToInput(mx::StandardSurface::Input::SubsurfaceScale, StandardSurfaceShaderNode, StandardSurface::Parameters::SubsurfaceScale.ToString(), mx::StandardSurface::DefaultValue::Float::SubsurfaceScale);
+	ConnectNodeOutputToInput(Input::SubsurfaceScale, StandardSurfaceShaderNode, StandardSurface::Parameters::SubsurfaceScale.ToString(), DefaultValue::Float::SubsurfaceScale);
 
 	//Sheen
-	ConnectNodeOutputToInput(mx::StandardSurface::Input::Sheen, StandardSurfaceShaderNode, StandardSurface::Parameters::Sheen.ToString(), mx::StandardSurface::DefaultValue::Float::Sheen);
+	ConnectNodeOutputToInput(Input::Sheen, StandardSurfaceShaderNode, StandardSurface::Parameters::Sheen.ToString(), DefaultValue::Float::Sheen);
 
 	//Sheen Color
-	ConnectNodeOutputToInput(mx::StandardSurface::Input::SheenColor, StandardSurfaceShaderNode, StandardSurface::Parameters::SheenColor.ToString(), mx::StandardSurface::DefaultValue::Color3::SheenColor);
+	ConnectNodeOutputToInput(Input::SheenColor, StandardSurfaceShaderNode, StandardSurface::Parameters::SheenColor.ToString(), DefaultValue::Color3::SheenColor);
 
 	//Sheen Roughness
-	ConnectNodeOutputToInput(mx::StandardSurface::Input::SheenRoughness, StandardSurfaceShaderNode, StandardSurface::Parameters::SheenRoughness.ToString(), mx::StandardSurface::DefaultValue::Float::SheenRoughness);
+	ConnectNodeOutputToInput(Input::SheenRoughness, StandardSurfaceShaderNode, StandardSurface::Parameters::SheenRoughness.ToString(), DefaultValue::Float::SheenRoughness);
 
 	//Coat
-	ConnectNodeOutputToInput(mx::StandardSurface::Input::Coat, StandardSurfaceShaderNode, StandardSurface::Parameters::Coat.ToString(), mx::StandardSurface::DefaultValue::Float::Coat);
+	ConnectNodeOutputToInput(Input::Coat, StandardSurfaceShaderNode, StandardSurface::Parameters::Coat.ToString(), DefaultValue::Float::Coat);
 
 	//Coat Color
-	ConnectNodeOutputToInput(mx::StandardSurface::Input::CoatColor, StandardSurfaceShaderNode, StandardSurface::Parameters::CoatColor.ToString(), mx::StandardSurface::DefaultValue::Color3::CoatColor);
+	ConnectNodeOutputToInput(Input::CoatColor, StandardSurfaceShaderNode, StandardSurface::Parameters::CoatColor.ToString(), DefaultValue::Color3::CoatColor);
 
 	//Coat Roughness
-	ConnectNodeOutputToInput(mx::StandardSurface::Input::CoatRoughness, StandardSurfaceShaderNode, StandardSurface::Parameters::CoatRoughness.ToString(), mx::StandardSurface::DefaultValue::Float::CoatRoughness);
+	ConnectNodeOutputToInput(Input::CoatRoughness, StandardSurfaceShaderNode, StandardSurface::Parameters::CoatRoughness.ToString(), DefaultValue::Float::CoatRoughness);
 
 	//Coat Normal: No need to take the default input if there is no CoatNormal input
-	ConnectNodeOutputToInput(mx::StandardSurface::Input::CoatNormal, StandardSurfaceShaderNode, StandardSurface::Parameters::CoatNormal.ToString(), nullptr, TextureCompressionSettings::TC_Normalmap);
+	ConnectNodeOutputToInput(Input::CoatNormal, StandardSurfaceShaderNode, StandardSurface::Parameters::CoatNormal.ToString(), nullptr, TextureCompressionSettings::TC_Normalmap);
 
 	//Thin Film Thickness
-	ConnectNodeOutputToInput(mx::StandardSurface::Input::ThinFilmThickness, StandardSurfaceShaderNode, StandardSurface::Parameters::ThinFilmThickness.ToString(), mx::StandardSurface::DefaultValue::Float::ThinFilmThickness);
+	ConnectNodeOutputToInput(Input::ThinFilmThickness, StandardSurfaceShaderNode, StandardSurface::Parameters::ThinFilmThickness.ToString(), DefaultValue::Float::ThinFilmThickness);
 
 	//Emission
-	ConnectNodeOutputToInput(mx::StandardSurface::Input::Emission, StandardSurfaceShaderNode, StandardSurface::Parameters::Emission.ToString(), mx::StandardSurface::DefaultValue::Float::Emission);
+	ConnectNodeOutputToInput(Input::Emission, StandardSurfaceShaderNode, StandardSurface::Parameters::Emission.ToString(), DefaultValue::Float::Emission);
 
 	//Emission Color
-	ConnectNodeOutputToInput(mx::StandardSurface::Input::EmissionColor, StandardSurfaceShaderNode, StandardSurface::Parameters::EmissionColor.ToString(), mx::StandardSurface::DefaultValue::Color3::EmissionColor);
+	ConnectNodeOutputToInput(Input::EmissionColor, StandardSurfaceShaderNode, StandardSurface::Parameters::EmissionColor.ToString(), DefaultValue::Color3::EmissionColor);
 
 	//Normal: No need to take the default input if there is no Normal input
-	ConnectNodeOutputToInput(mx::StandardSurface::Input::Normal, StandardSurfaceShaderNode, StandardSurface::Parameters::Normal.ToString(), nullptr, TextureCompressionSettings::TC_Normalmap);
+	ConnectNodeOutputToInput(Input::Normal, StandardSurfaceShaderNode, StandardSurface::Parameters::Normal.ToString(), nullptr, TextureCompressionSettings::TC_Normalmap);
 
 	//Tangent: No need to take the default input if there is no Tangent input
-	ConnectNodeOutputToInput(mx::StandardSurface::Input::Tangent, StandardSurfaceShaderNode, StandardSurface::Parameters::Tangent.ToString(), nullptr, TextureCompressionSettings::TC_Normalmap);
+	ConnectNodeOutputToInput(Input::Tangent, StandardSurfaceShaderNode, StandardSurface::Parameters::Tangent.ToString(), nullptr, TextureCompressionSettings::TC_Normalmap);
 
 	//Transmission
-	ConnectNodeOutputToInput(mx::StandardSurface::Input::Transmission, StandardSurfaceShaderNode, StandardSurface::Parameters::Transmission.ToString(), mx::StandardSurface::DefaultValue::Float::Transmission);
+	ConnectNodeOutputToInput(Input::Transmission, StandardSurfaceShaderNode, StandardSurface::Parameters::Transmission.ToString(), DefaultValue::Float::Transmission);
 
 	//Transmission Color
-	ConnectNodeOutputToInput(mx::StandardSurface::Input::TransmissionColor, StandardSurfaceShaderNode, StandardSurface::Parameters::TransmissionColor.ToString(), mx::StandardSurface::DefaultValue::Color3::TransmissionColor);
+	ConnectNodeOutputToInput(Input::TransmissionColor, StandardSurfaceShaderNode, StandardSurface::Parameters::TransmissionColor.ToString(), DefaultValue::Color3::TransmissionColor);
 
 	//Transmission Depth
-	ConnectNodeOutputToInput(mx::StandardSurface::Input::TransmissionDepth, StandardSurfaceShaderNode, StandardSurface::Parameters::TransmissionDepth.ToString(), mx::StandardSurface::DefaultValue::Float::TransmissionDepth);
+	ConnectNodeOutputToInput(Input::TransmissionDepth, StandardSurfaceShaderNode, StandardSurface::Parameters::TransmissionDepth.ToString(), DefaultValue::Float::TransmissionDepth);
 
 	//Transmission Scatter
-	ConnectNodeOutputToInput(mx::StandardSurface::Input::TransmissionScatter, StandardSurfaceShaderNode, StandardSurface::Parameters::TransmissionScatter.ToString(), mx::StandardSurface::DefaultValue::Color3::TransmissionScatter);
+	ConnectNodeOutputToInput(Input::TransmissionScatter, StandardSurfaceShaderNode, StandardSurface::Parameters::TransmissionScatter.ToString(), DefaultValue::Color3::TransmissionScatter);
 
 	//Transmission Scatter Anisotropy
-	ConnectNodeOutputToInput(mx::StandardSurface::Input::TransmissionScatterAnisotropy, StandardSurfaceShaderNode, StandardSurface::Parameters::TransmissionScatterAnisotropy.ToString(), mx::StandardSurface::DefaultValue::Float::TransmissionScatterAnisotropy);
+	ConnectNodeOutputToInput(Input::TransmissionScatterAnisotropy, StandardSurfaceShaderNode, StandardSurface::Parameters::TransmissionScatterAnisotropy.ToString(), DefaultValue::Float::TransmissionScatterAnisotropy);
 
 	//Transmission Dispersion
-	ConnectNodeOutputToInput(mx::StandardSurface::Input::TransmissionDispersion, StandardSurfaceShaderNode, StandardSurface::Parameters::TransmissionDispersion.ToString(), mx::StandardSurface::DefaultValue::Float::TransmissionDispersion);
+	ConnectNodeOutputToInput(Input::TransmissionDispersion, StandardSurfaceShaderNode, StandardSurface::Parameters::TransmissionDispersion.ToString(), DefaultValue::Float::TransmissionDispersion);
 
 	//Transmission Extra Roughness
-	ConnectNodeOutputToInput(mx::StandardSurface::Input::TransmissionExtraRoughness, StandardSurfaceShaderNode, StandardSurface::Parameters::TransmissionExtraRoughness.ToString(), mx::StandardSurface::DefaultValue::Float::TransmissionExtraRoughness);
+	ConnectNodeOutputToInput(Input::TransmissionExtraRoughness, StandardSurfaceShaderNode, StandardSurface::Parameters::TransmissionExtraRoughness.ToString(), DefaultValue::Float::TransmissionExtraRoughness);
 
 	// Outputs
 	UInterchangeShaderPortsAPI::ConnectOuputToInputByName(ShaderGraphNode, PBRMR::Parameters::BaseColor.ToString(), StandardSurfaceShaderNode->GetUniqueID(), TEXT("Base Color"));
@@ -172,4 +188,153 @@ void FMaterialXStandardSurfaceShader::Translate(mx::NodePtr StandardSurfaceNode)
 		UInterchangeShaderPortsAPI::ConnectOuputToInputByName(ShaderGraphNode, Subsurface::Parameters::SubsurfaceColor.ToString(), StandardSurfaceShaderNode->GetUniqueID(), Subsurface::Parameters::SubsurfaceColor.ToString());
 	}
 }
+
+void FMaterialXStandardSurfaceShader::ConnectToSubstrateStandardSurface()
+{
+	using namespace UE::Interchange::Materials;
+	using namespace mx::StandardSurface;
+
+	UInterchangeFunctionCallShaderNode* StandardSurfaceShaderNode;
+	const FString NodeUID = UInterchangeShaderNode::MakeNodeUid(ANSI_TO_TCHAR(SurfaceShaderNode->getName().c_str()), FStringView{});
+	if(StandardSurfaceShaderNode = const_cast<UInterchangeFunctionCallShaderNode*>(Cast<UInterchangeFunctionCallShaderNode>(NodeContainer.GetNode(NodeUID))); !StandardSurfaceShaderNode)
+	{
+		const FString NodeName = SurfaceShaderNode->getName().c_str();
+		StandardSurfaceShaderNode = NewObject<UInterchangeFunctionCallShaderNode>(&NodeContainer);
+		StandardSurfaceShaderNode->InitializeNode(NodeUID, NodeName, EInterchangeNodeContainerType::TranslatedAsset);
+
+		StandardSurfaceShaderNode->SetCustomMaterialFunction(TEXT("/Interchange/Functions/MX_StandardSurface_Substrate.MX_StandardSurface_Substrate"));
+		NodeContainer.AddNode(StandardSurfaceShaderNode);
+
+		ShaderNodes.Add({ NodeName, SurfaceShaderNode->getNodeDef(mx::EMPTY_STRING, true)->getActiveOutputs()[0]->getName().c_str() }, StandardSurfaceShaderNode);
+	}
+
+	// Inputs
+	//Base
+	ConnectNodeOutputToInput(Input::Base, StandardSurfaceShaderNode, StandardSurface::Parameters::Base.ToString(), DefaultValue::Float::Base);
+
+	//Base Color
+	ConnectNodeOutputToInput(Input::BaseColor, StandardSurfaceShaderNode, StandardSurface::Parameters::BaseColor.ToString(), DefaultValue::Color3::BaseColor);
+
+	//Diffuse Roughness
+	ConnectNodeOutputToInput(Input::DiffuseRoughness, StandardSurfaceShaderNode, StandardSurface::Parameters::DiffuseRoughness.ToString(), DefaultValue::Float::DiffuseRoughness);
+
+	//Metallic
+	ConnectNodeOutputToInput(Input::Metalness, StandardSurfaceShaderNode, StandardSurface::Parameters::Metalness.ToString(), DefaultValue::Float::Metalness);
+
+	//Specular
+	ConnectNodeOutputToInput(Input::Specular, StandardSurfaceShaderNode, StandardSurface::Parameters::Specular.ToString(), DefaultValue::Float::Specular);
+
+	//Specular Color
+	ConnectNodeOutputToInput(Input::SpecularColor, StandardSurfaceShaderNode, StandardSurface::Parameters::SpecularColor.ToString(), DefaultValue::Color3::SpecularColor);
+
+	//Specular Roughness
+	ConnectNodeOutputToInput(Input::SpecularRoughness, StandardSurfaceShaderNode, StandardSurface::Parameters::SpecularRoughness.ToString(), DefaultValue::Float::SpecularRoughness);
+
+	//Specular IOR
+	ConnectNodeOutputToInput(Input::SpecularIOR, StandardSurfaceShaderNode, StandardSurface::Parameters::SpecularIOR.ToString(), DefaultValue::Float::SpecularIOR);
+
+	//Specular Anisotropy
+	ConnectNodeOutputToInput(Input::SpecularAnisotropy, StandardSurfaceShaderNode, StandardSurface::Parameters::SpecularAnisotropy.ToString(), DefaultValue::Float::SpecularAnisotropy);
+
+	//Specular Rotation
+	ConnectNodeOutputToInput(Input::SpecularRotation, StandardSurfaceShaderNode, StandardSurface::Parameters::SpecularRotation.ToString(), DefaultValue::Float::SpecularRotation);
+
+	//Subsurface
+	ConnectNodeOutputToInput(Input::Subsurface, StandardSurfaceShaderNode, StandardSurface::Parameters::Subsurface.ToString(), DefaultValue::Float::Subsurface);
+
+	//Subsurface Color
+	ConnectNodeOutputToInput(Input::SubsurfaceColor, StandardSurfaceShaderNode, StandardSurface::Parameters::SubsurfaceColor.ToString(), DefaultValue::Color3::SubsurfaceColor);
+
+	//Subsurface Radius
+	ConnectNodeOutputToInput(Input::SubsurfaceRadius, StandardSurfaceShaderNode, StandardSurface::Parameters::SubsurfaceRadius.ToString(), DefaultValue::Color3::SubsurfaceRadius);
+
+	//Subsurface Scale
+	ConnectNodeOutputToInput(Input::SubsurfaceScale, StandardSurfaceShaderNode, StandardSurface::Parameters::SubsurfaceScale.ToString(), DefaultValue::Float::SubsurfaceScale);
+
+	//Subsurface Anisotropy
+	ConnectNodeOutputToInput(Input::SubsurfaceAnisotropy, StandardSurfaceShaderNode, StandardSurface::Parameters::SubsurfaceAnisotropy.ToString(), DefaultValue::Float::SubsurfaceAnisotropy);
+
+	//Sheen
+	ConnectNodeOutputToInput(Input::Sheen, StandardSurfaceShaderNode, StandardSurface::Parameters::Sheen.ToString(), DefaultValue::Float::Sheen);
+
+	//Sheen Color
+	ConnectNodeOutputToInput(Input::SheenColor, StandardSurfaceShaderNode, StandardSurface::Parameters::SheenColor.ToString(), DefaultValue::Color3::SheenColor);
+
+	//Sheen Roughness
+	ConnectNodeOutputToInput(Input::SheenRoughness, StandardSurfaceShaderNode, StandardSurface::Parameters::SheenRoughness.ToString(), DefaultValue::Float::SheenRoughness);
+
+	//Coat
+	ConnectNodeOutputToInput(Input::Coat, StandardSurfaceShaderNode, StandardSurface::Parameters::Coat.ToString(), DefaultValue::Float::Coat);
+
+	//Coat Color
+	ConnectNodeOutputToInput(Input::CoatColor, StandardSurfaceShaderNode, StandardSurface::Parameters::CoatColor.ToString(), DefaultValue::Color3::CoatColor);
+
+	//Coat Roughness
+	ConnectNodeOutputToInput(Input::CoatRoughness, StandardSurfaceShaderNode, StandardSurface::Parameters::CoatRoughness.ToString(), DefaultValue::Float::CoatRoughness);
+
+	//Coat IOR
+	ConnectNodeOutputToInput(Input::CoatIOR, StandardSurfaceShaderNode, StandardSurface::Parameters::CoatIOR.ToString(), DefaultValue::Float::CoatIOR);
+
+	//Coat Anisotropy
+	ConnectNodeOutputToInput(Input::CoatAnisotropy, StandardSurfaceShaderNode, StandardSurface::Parameters::CoatAnisotropy.ToString(), DefaultValue::Float::CoatAnisotropy);
+
+	//Coat Rotation
+	ConnectNodeOutputToInput(Input::CoatRotation, StandardSurfaceShaderNode, StandardSurface::Parameters::CoatRotation.ToString(), DefaultValue::Float::CoatAnisotropy);
+
+	//Coat Normal: No need to take the default input if there is no CoatNormal input
+	ConnectNodeOutputToInput(Input::CoatNormal, StandardSurfaceShaderNode, StandardSurface::Parameters::CoatNormal.ToString(), nullptr, TextureCompressionSettings::TC_Normalmap);
+
+	//Thin Film Thickness
+	ConnectNodeOutputToInput(Input::ThinFilmThickness, StandardSurfaceShaderNode, StandardSurface::Parameters::ThinFilmThickness.ToString(), DefaultValue::Float::ThinFilmThickness);
+
+	//Thin Film IOR
+	ConnectNodeOutputToInput(Input::ThinFilmIOR, StandardSurfaceShaderNode, StandardSurface::Parameters::ThinFilmIOR.ToString(), DefaultValue::Float::ThinFilmIOR);
+
+	//Emission
+	ConnectNodeOutputToInput(Input::Emission, StandardSurfaceShaderNode, StandardSurface::Parameters::Emission.ToString(), DefaultValue::Float::Emission);
+
+	//Emission Color
+	ConnectNodeOutputToInput(Input::EmissionColor, StandardSurfaceShaderNode, StandardSurface::Parameters::EmissionColor.ToString(), DefaultValue::Color3::EmissionColor);
+
+	//Normal: No need to take the default input if there is no Normal input
+	ConnectNodeOutputToInput(Input::Normal, StandardSurfaceShaderNode, StandardSurface::Parameters::Normal.ToString(), nullptr, TextureCompressionSettings::TC_Normalmap);
+
+	//Tangent: No need to take the default input if there is no Tangent input
+	ConnectNodeOutputToInput(Input::Tangent, StandardSurfaceShaderNode, StandardSurface::Parameters::Tangent.ToString(), nullptr, TextureCompressionSettings::TC_Normalmap);
+
+	//Opacity
+	ConnectNodeOutputToInput(Input::Opacity, StandardSurfaceShaderNode, StandardSurface::Parameters::Opacity.ToString(), DefaultValue::Color3::Opacity);
+
+	//Transmission
+	ConnectNodeOutputToInput(Input::Transmission, StandardSurfaceShaderNode, StandardSurface::Parameters::Transmission.ToString(), DefaultValue::Float::Transmission);
+
+	//Transmission Color
+	ConnectNodeOutputToInput(Input::TransmissionColor, StandardSurfaceShaderNode, StandardSurface::Parameters::TransmissionColor.ToString(), DefaultValue::Color3::TransmissionColor);
+
+	//Transmission Depth
+	ConnectNodeOutputToInput(Input::TransmissionDepth, StandardSurfaceShaderNode, StandardSurface::Parameters::TransmissionDepth.ToString(), DefaultValue::Float::TransmissionDepth);
+
+	//Transmission Scatter
+	ConnectNodeOutputToInput(Input::TransmissionScatter, StandardSurfaceShaderNode, StandardSurface::Parameters::TransmissionScatter.ToString(), DefaultValue::Color3::TransmissionScatter);
+
+	//Transmission Extra Roughness
+	ConnectNodeOutputToInput(Input::TransmissionExtraRoughness, StandardSurfaceShaderNode, StandardSurface::Parameters::TransmissionExtraRoughness.ToString(), DefaultValue::Float::TransmissionExtraRoughness);
+
+	// Outputs
+	if(UInterchangeShaderPortsAPI::HasInput(StandardSurfaceShaderNode, StandardSurface::Parameters::Transmission))
+	{
+		StandardSurfaceShaderNode->SetCustomMaterialFunction(TEXT("/Interchange/Functions/MX_TransmissionSurface_Substrate.MX_TransmissionSurface_Substrate"));
+		ShaderGraphNode->SetCustomBlendMode(EBlendMode::BLEND_TranslucentColoredTransmittance);
+		UInterchangeShaderPortsAPI::ConnectOuputToInputByName(ShaderGraphNode, PBRMR::Parameters::Opacity.ToString(), StandardSurfaceShaderNode->GetUniqueID(), PBRMR::Parameters::Opacity.ToString());
+		UInterchangeShaderPortsAPI::ConnectOuputToInputByName(ShaderGraphNode, Substrate::Parameters::FrontMaterial.ToString(), StandardSurfaceShaderNode->GetUniqueID(), StandardSurface::Substrate::Outputs::Translucent.ToString());
+	}
+	else
+	{
+		UInterchangeShaderPortsAPI::ConnectOuputToInputByName(ShaderGraphNode, Substrate::Parameters::FrontMaterial.ToString(), StandardSurfaceShaderNode->GetUniqueID(), StandardSurface::Substrate::Outputs::Opaque.ToString());
+		if(UInterchangeShaderPortsAPI::HasInput(StandardSurfaceShaderNode, StandardSurface::Parameters::Opacity))
+		{
+			UInterchangeShaderPortsAPI::ConnectOuputToInputByName(ShaderGraphNode, Substrate::Parameters::OpacityMask.ToString(), StandardSurfaceShaderNode->GetUniqueID(), StandardSurface::Substrate::Outputs::Opacity.ToString());
+		}
+	}
+}
 #endif
\ No newline at end of file
diff --git a/Engine/Plugins/Interchange/Runtime/Source/Import/Private/MaterialX/MaterialXUtils/MaterialXStandardSurfaceShader.h b/Engine/Plugins/Interchange/Runtime/Source/Import/Private/MaterialX/MaterialXUtils/MaterialXStandardSurfaceShader.h
--- a/Engine/Plugins/Interchange/Runtime/Source/Import/Private/MaterialX/MaterialXUtils/MaterialXStandardSurfaceShader.h
+++ b/Engine/Plugins/Interchange/Runtime/Source/Import/Private/MaterialX/MaterialXUtils/MaterialXStandardSurfaceShader.h
@@ -19,5 +19,11 @@ public:
 	static TSharedRef<FMaterialXBase> MakeInstance(UInterchangeBaseNodeContainer& BaseNodeContainer);
 
 	virtual void Translate(MaterialX::NodePtr StandardSurfaceNode) override;
+
+private:
+
+	void ConnectToStandardSurface();
+
+	void ConnectToSubstrateStandardSurface();
 };
 #endif
\ No newline at end of file
diff --git a/Engine/Plugins/Interchange/Runtime/Source/Nodes/Private/InterchangeShaderGraphNode.cpp b/Engine/Plugins/Interchange/Runtime/Source/Nodes/Private/InterchangeShaderGraphNode.cpp
--- a/Engine/Plugins/Interchange/Runtime/Source/Nodes/Private/InterchangeShaderGraphNode.cpp
+++ b/Engine/Plugins/Interchange/Runtime/Source/Nodes/Private/InterchangeShaderGraphNode.cpp
@@ -254,6 +254,16 @@ bool UInterchangeShaderGraphNode::SetCustomScreenSpaceReflections(const bool& At
 	IMPLEMENT_NODE_ATTRIBUTE_SETTER_NODELEGATE(ScreenSpaceReflections, bool);
 }
 
+bool UInterchangeShaderGraphNode::GetCustomBlendMode(int& AttributeValue) const
+{
+	IMPLEMENT_NODE_ATTRIBUTE_GETTER(BlendMode, int);
+}
+
+bool UInterchangeShaderGraphNode::SetCustomBlendMode(int AttributeValue)
+{
+	IMPLEMENT_NODE_ATTRIBUTE_SETTER_NODELEGATE(BlendMode, int);
+}
+
 FString UInterchangeFunctionCallShaderNode::GetTypeName() const
 {
 	const FString TypeName = TEXT("FunctionCallShaderNode");
diff --git a/Engine/Plugins/Interchange/Runtime/Source/Nodes/Public/InterchangeMaterialDefinitions.h b/Engine/Plugins/Interchange/Runtime/Source/Nodes/Public/InterchangeMaterialDefinitions.h
--- a/Engine/Plugins/Interchange/Runtime/Source/Nodes/Public/InterchangeMaterialDefinitions.h
+++ b/Engine/Plugins/Interchange/Runtime/Source/Nodes/Public/InterchangeMaterialDefinitions.h
@@ -432,6 +432,17 @@ namespace Interchange
 				const FName Normal = TEXT("normal");
 				const FName Tangent = TEXT("tangent");
 			}
+
+			// These outputs are only used for Substrate
+			namespace Substrate
+			{
+				namespace Outputs
+				{
+					const FName Opaque = TEXT("Substrate StandardSurface Opaque");
+					const FName Translucent = TEXT("Substrate StandardSurface Translucent");
+					const FName Opacity = TEXT("Geometry Opacity");
+				}
+			}
 		}
 
 		namespace SurfaceUnlit
@@ -477,6 +488,16 @@ namespace Interchange
 				const FName UnlitColor = TEXT("UnlitColor"); // Type: linear color
 			}
 		}
+
+		namespace Substrate
+		{
+			namespace Parameters
+			{
+				const FName FrontMaterial = TEXT("Front Material");
+				const FName OpacityMask = TEXT("Opacity Mask");
+			}
+			
+		}
 	}
 }
 }
diff --git a/Engine/Plugins/Interchange/Runtime/Source/Nodes/Public/InterchangeShaderGraphNode.h b/Engine/Plugins/Interchange/Runtime/Source/Nodes/Public/InterchangeShaderGraphNode.h
--- a/Engine/Plugins/Interchange/Runtime/Source/Nodes/Public/InterchangeShaderGraphNode.h
+++ b/Engine/Plugins/Interchange/Runtime/Source/Nodes/Public/InterchangeShaderGraphNode.h
@@ -1,7 +1,7 @@
 // Copyright Epic Games, Inc. All Rights Reserved.
  
 #pragma once
- 
+
 #include "Nodes/InterchangeBaseNode.h"
 #include "UObject/Object.h"
  
@@ -235,9 +235,19 @@ public:
 	UFUNCTION(BlueprintCallable, Category = "Interchange | Node | Material")
 	bool SetCustomScreenSpaceReflections(const bool& AttributeValue);
 
+	/**
+	 * Set the Blend Mode using EBlendMode, avoid a dependency on the Engine
+	 */
+	UFUNCTION(BlueprintCallable, Category = "Interchange | Node | Material")
+	bool GetCustomBlendMode(int& AttributeValue) const;
+
+	UFUNCTION(BlueprintCallable, Category = "Interchange | Node | Material")
+	bool SetCustomBlendMode(int AttributeValue);
+
 private:
-	const UE::Interchange::FAttributeKey Macro_CustomTwoSidedKey = UE::Interchange::FAttributeKey(TEXT("TwoSided"));
-	const UE::Interchange::FAttributeKey Macro_CustomOpacityMaskClipValueKey = UE::Interchange::FAttributeKey(TEXT("OpacityMaskClipValue"));
-	const UE::Interchange::FAttributeKey Macro_CustomIsAShaderFunctionKey = UE::Interchange::FAttributeKey(TEXT("IsAShaderFunction"));
-	const UE::Interchange::FAttributeKey Macro_CustomScreenSpaceReflectionsKey = UE::Interchange::FAttributeKey(TEXT("ScreenSpaceReflections"));
+	IMPLEMENT_NODE_ATTRIBUTE_KEY(TwoSided)
+	IMPLEMENT_NODE_ATTRIBUTE_KEY(OpacityMaskClipValue)
+	IMPLEMENT_NODE_ATTRIBUTE_KEY(IsAShaderFunction)
+	IMPLEMENT_NODE_ATTRIBUTE_KEY(ScreenSpaceReflections)
+	IMPLEMENT_NODE_ATTRIBUTE_KEY(BlendMode)
 };
\ No newline at end of file
diff --git a/Engine/Plugins/Interchange/Runtime/Source/Pipelines/Private/InterchangeGenericMaterialPipeline.cpp b/Engine/Plugins/Interchange/Runtime/Source/Pipelines/Private/InterchangeGenericMaterialPipeline.cpp
--- a/Engine/Plugins/Interchange/Runtime/Source/Pipelines/Private/InterchangeGenericMaterialPipeline.cpp
+++ b/Engine/Plugins/Interchange/Runtime/Source/Pipelines/Private/InterchangeGenericMaterialPipeline.cpp
@@ -606,20 +606,6 @@ bool UInterchangeGenericMaterialPipeline::IsLambertModel(const UInterchangeShade
 	return bHasDiffuseInput;
 }
 
-bool UInterchangeGenericMaterialPipeline::IsStandardSurfaceModel(const UInterchangeShaderGraphNode* ShaderGraphNode) const
-{
-	using namespace UE::Interchange::Materials;
-	FString ShaderType;
-	ShaderGraphNode->GetCustomShaderType(ShaderType);
-
-	if(ShaderType == StandardSurface::Name.ToString())
-	{
-		return true;
-	}
-
-	return false;
-}
-
 bool UInterchangeGenericMaterialPipeline::IsSurfaceUnlitModel(const UInterchangeShaderGraphNode* ShaderGraphNode) const
 {
 	using namespace UE::Interchange::Materials;
@@ -2363,6 +2349,11 @@ UInterchangeMaterialFactoryNode* UInterchangeGenericMaterialPipeline::CreateMate
 {
 	UInterchangeMaterialFactoryNode* MaterialFactoryNode = Cast<UInterchangeMaterialFactoryNode>( CreateBaseMaterialFactoryNode(ShaderGraphNode, UInterchangeMaterialFactoryNode::StaticClass()) );
 
+	if(HandleSubstrate(ShaderGraphNode, MaterialFactoryNode))
+	{
+		return MaterialFactoryNode;
+	}
+
 	// Handle the case where the material will be connected through the material attributes input
 	if (HandleBxDFInput(ShaderGraphNode, MaterialFactoryNode))
 	{
@@ -2880,6 +2871,50 @@ bool UInterchangeGenericMaterialPipeline::HandleUnlitModel(const UInterchangeSha
 	return bShadingModelHandled;
 }
 
+bool UInterchangeGenericMaterialPipeline::HandleSubstrate(const UInterchangeShaderGraphNode* ShaderGraphNode, UInterchangeMaterialFactoryNode* MaterialFactoryNode)
+{
+	using namespace UE::Interchange::Materials;
+	bool bShadingModelHandled = false;
+
+	if(UInterchangeShaderPortsAPI::HasInput(ShaderGraphNode, Substrate::Parameters::FrontMaterial))
+	{
+		TTuple<UInterchangeMaterialExpressionFactoryNode*, FString> FrontMaterialFactoryNode =
+			CreateMaterialExpressionForInput(MaterialFactoryNode, ShaderGraphNode, Substrate::Parameters::FrontMaterial.ToString(), MaterialFactoryNode->GetUniqueID());
+		ensure(FrontMaterialFactoryNode.Get<0>());
+
+		if(FrontMaterialFactoryNode.Get<0>())
+		{
+			UInterchangeShaderPortsAPI::ConnectOuputToInputByName(MaterialFactoryNode, Substrate::Parameters::FrontMaterial.ToString(), FrontMaterialFactoryNode.Get<0>()->GetUniqueID(), FrontMaterialFactoryNode.Get<1>());
+		}
+
+		if(UInterchangeShaderPortsAPI::HasInput(ShaderGraphNode, Substrate::Parameters::OpacityMask))
+		{
+			TTuple<UInterchangeMaterialExpressionFactoryNode*, FString> OpacityMaskFactoryNode =
+				CreateMaterialExpressionForInput(MaterialFactoryNode, ShaderGraphNode, Substrate::Parameters::OpacityMask.ToString(), MaterialFactoryNode->GetUniqueID());
+			ensure(OpacityMaskFactoryNode.Get<0>());
+
+			if(OpacityMaskFactoryNode.Get<0>())
+			{
+				UInterchangeShaderPortsAPI::ConnectOuputToInputByName(MaterialFactoryNode, Substrate::Parameters::OpacityMask.ToString(), OpacityMaskFactoryNode.Get<0>()->GetUniqueID(), OpacityMaskFactoryNode.Get<1>());
+			}
+		}
+
+		if(EBlendMode BlendMode; ShaderGraphNode->GetCustomBlendMode(reinterpret_cast<int&>(BlendMode)))
+		{
+			MaterialFactoryNode->SetCustomBlendMode(BlendMode);
+			if(BlendMode == BLEND_TranslucentColoredTransmittance)
+			{
+				MaterialFactoryNode->SetCustomTranslucencyLightingMode(ETranslucencyLightingMode::TLM_SurfacePerPixelLighting);
+				MaterialFactoryNode->SetCustomRefractionMethod(ERefractionMode::RM_IndexOfRefraction);
+			}
+		}
+
+		bShadingModelHandled = true;
+	}
+
+	return bShadingModelHandled;
+}
+
 UInterchangeMaterialExpressionFactoryNode* UInterchangeGenericMaterialPipeline::CreateFunctionCallExpression(const UInterchangeShaderNode* ShaderNode, const FString& MaterialExpressionUid, UInterchangeBaseMaterialFactoryNode* MaterialFactoryNode)
 {
 	using namespace UE::Interchange::InterchangeGenericMaterialPipeline::Private;
diff --git a/Engine/Plugins/Interchange/Runtime/Source/Pipelines/Private/InterchangeMaterialXPipeline.cpp b/Engine/Plugins/Interchange/Runtime/Source/Pipelines/Private/InterchangeMaterialXPipeline.cpp
--- a/Engine/Plugins/Interchange/Runtime/Source/Pipelines/Private/InterchangeMaterialXPipeline.cpp
+++ b/Engine/Plugins/Interchange/Runtime/Source/Pipelines/Private/InterchangeMaterialXPipeline.cpp
@@ -1,21 +1,22 @@
 // Copyright Epic Games, Inc. All Rights Reserved. 
 
 #include "InterchangeMaterialXPipeline.h"
-
-#include "InterchangePipelineLog.h"
-
 #include "InterchangeManager.h"
 #include "InterchangeMaterialDefinitions.h"
 #include "InterchangeMaterialFactoryNode.h"
+#include "InterchangePipelineLog.h"
 
+#include "Engine/RendererSettings.h"
 #include "Materials/MaterialExpressionMaterialFunctionCall.h"
 #include "Materials/MaterialFunction.h"
 #include "Misc/PackageName.h"
 
+static const bool bIsSubstrateEnabled = GetDefault<URendererSettings>()->bEnableStrata;
+
 TMap<FString, EInterchangeMaterialXShaders> UInterchangeMaterialXPipeline::PathToEnumMapping
 {
-	{TEXT("/Interchange/Functions/MX_StandardSurface.MX_StandardSurface"),EInterchangeMaterialXShaders::StandardSurface},
-	{TEXT("/Interchange/Functions/MX_TransmissionSurface.MX_TransmissionSurface"),EInterchangeMaterialXShaders::StandardSurfaceTransmission},
+	{bIsSubstrateEnabled ? TEXT("/Interchange/Functions/MX_StandardSurface_Substrate.MX_StandardSurface_Substrate") : TEXT("/Interchange/Functions/MX_StandardSurface.MX_StandardSurface"), EInterchangeMaterialXShaders::StandardSurface},
+	{bIsSubstrateEnabled ? TEXT("/Interchange/Functions/MX_TransmissionSurface_Substrate.MX_TransmissionSurface_Substrate") : TEXT("/Interchange/Functions/MX_TransmissionSurface.MX_TransmissionSurface"), EInterchangeMaterialXShaders::StandardSurfaceTransmission},
 	{TEXT("/Interchange/Functions/MX_SurfaceUnlit.MX_SurfaceUnlit"),EInterchangeMaterialXShaders::SurfaceUnlit},
 	{TEXT("/Interchange/Functions/MX_UsdPreviewSurface.MX_UsdPreviewSurface"),EInterchangeMaterialXShaders::UsdPreviewSurface},
 };
@@ -157,26 +158,40 @@ TSet<FName> UMaterialXPipelineSettings::StandardSurfaceInputs
 	UE::Interchange::Materials::StandardSurface::Parameters::Emission,
 	UE::Interchange::Materials::StandardSurface::Parameters::EmissionColor,
 	UE::Interchange::Materials::StandardSurface::Parameters::Normal,
-	UE::Interchange::Materials::StandardSurface::Parameters::Tangent
+	UE::Interchange::Materials::StandardSurface::Parameters::Tangent,
+	bIsSubstrateEnabled ? UE::Interchange::Materials::StandardSurface::Parameters::SpecularColor : FName{},
+	bIsSubstrateEnabled ? UE::Interchange::Materials::StandardSurface::Parameters::CoatIOR : FName{},
+	bIsSubstrateEnabled ? UE::Interchange::Materials::StandardSurface::Parameters::CoatAnisotropy : FName{},
+	bIsSubstrateEnabled ? UE::Interchange::Materials::StandardSurface::Parameters::CoatRotation : FName{},
+	bIsSubstrateEnabled ? UE::Interchange::Materials::StandardSurface::Parameters::ThinFilmIOR : FName{},
+	bIsSubstrateEnabled ? UE::Interchange::Materials::StandardSurface::Parameters::Opacity : FName{},
+	bIsSubstrateEnabled ? UE::Interchange::Materials::StandardSurface::Parameters::SubsurfaceAnisotropy : FName{}
 };
 
 TSet<FName> UMaterialXPipelineSettings::StandardSurfaceOutputs
 {
-	TEXT("Base Color"), // MX_StandardSurface has BaseColor with a whitespace, this should be fixed in further release
-	UE::Interchange::Materials::PBRMR::Parameters::Metallic,
-	UE::Interchange::Materials::PBRMR::Parameters::Specular,
-	UE::Interchange::Materials::PBRMR::Parameters::Roughness,
-	UE::Interchange::Materials::PBRMR::Parameters::Anisotropy,
-	UE::Interchange::Materials::PBRMR::Parameters::EmissiveColor,
-	UE::Interchange::Materials::PBRMR::Parameters::Opacity,
-	UE::Interchange::Materials::PBRMR::Parameters::Normal,
-	UE::Interchange::Materials::PBRMR::Parameters::Tangent,
-	UE::Interchange::Materials::Sheen::Parameters::SheenRoughness,
-	UE::Interchange::Materials::Sheen::Parameters::SheenColor,
-	UE::Interchange::Materials::Subsurface::Parameters::SubsurfaceColor,
-	UE::Interchange::Materials::ClearCoat::Parameters::ClearCoat,
-	UE::Interchange::Materials::ClearCoat::Parameters::ClearCoatRoughness,
-	UE::Interchange::Materials::ClearCoat::Parameters::ClearCoatNormal
+	!bIsSubstrateEnabled ?
+	TSet<FName>{
+		TEXT("Base Color"), // MX_StandardSurface has BaseColor with a whitespace, this should be fixed in further release
+		UE::Interchange::Materials::PBRMR::Parameters::Metallic,
+		UE::Interchange::Materials::PBRMR::Parameters::Specular,
+		UE::Interchange::Materials::PBRMR::Parameters::Roughness,
+		UE::Interchange::Materials::PBRMR::Parameters::Anisotropy,
+		UE::Interchange::Materials::PBRMR::Parameters::EmissiveColor,
+		UE::Interchange::Materials::PBRMR::Parameters::Opacity,
+		UE::Interchange::Materials::PBRMR::Parameters::Normal,
+		UE::Interchange::Materials::PBRMR::Parameters::Tangent,
+		UE::Interchange::Materials::Sheen::Parameters::SheenRoughness,
+		UE::Interchange::Materials::Sheen::Parameters::SheenColor,
+		UE::Interchange::Materials::Subsurface::Parameters::SubsurfaceColor,
+		UE::Interchange::Materials::ClearCoat::Parameters::ClearCoat,
+		UE::Interchange::Materials::ClearCoat::Parameters::ClearCoatRoughness,
+		UE::Interchange::Materials::ClearCoat::Parameters::ClearCoatNormal
+	}	:
+	TSet<FName>{
+		UE::Interchange::Materials::StandardSurface::Substrate::Outputs::Opaque,
+		UE::Interchange::Materials::StandardSurface::Substrate::Outputs::Opacity
+	}
 };
 
 TSet<FName> UMaterialXPipelineSettings::TransmissionSurfaceInputs
@@ -186,6 +201,7 @@ TSet<FName> UMaterialXPipelineSettings::TransmissionSurfaceInputs
 	UE::Interchange::Materials::StandardSurface::Parameters::DiffuseRoughness,
 	UE::Interchange::Materials::StandardSurface::Parameters::Metalness,
 	UE::Interchange::Materials::StandardSurface::Parameters::Specular,
+	bIsSubstrateEnabled ? UE::Interchange::Materials::StandardSurface::Parameters::SpecularColor : FName{},
 	UE::Interchange::Materials::StandardSurface::Parameters::SpecularRoughness,
 	UE::Interchange::Materials::StandardSurface::Parameters::SpecularIOR,
 	UE::Interchange::Materials::StandardSurface::Parameters::SpecularAnisotropy,
@@ -207,26 +223,36 @@ TSet<FName> UMaterialXPipelineSettings::TransmissionSurfaceInputs
 	UE::Interchange::Materials::StandardSurface::Parameters::Coat,
 	UE::Interchange::Materials::StandardSurface::Parameters::CoatColor,
 	UE::Interchange::Materials::StandardSurface::Parameters::CoatRoughness,
+	bIsSubstrateEnabled ? UE::Interchange::Materials::StandardSurface::Parameters::CoatAnisotropy : FName{},
+	bIsSubstrateEnabled ? UE::Interchange::Materials::StandardSurface::Parameters::CoatRotation : FName{},
+	bIsSubstrateEnabled ? UE::Interchange::Materials::StandardSurface::Parameters::CoatIOR : FName{},
 	UE::Interchange::Materials::StandardSurface::Parameters::CoatNormal,
 	UE::Interchange::Materials::StandardSurface::Parameters::ThinFilmThickness,
+	bIsSubstrateEnabled ? UE::Interchange::Materials::StandardSurface::Parameters::ThinFilmIOR : FName{},
 	UE::Interchange::Materials::StandardSurface::Parameters::Emission,
 	UE::Interchange::Materials::StandardSurface::Parameters::EmissionColor,
+	bIsSubstrateEnabled ? UE::Interchange::Materials::StandardSurface::Parameters::Opacity : FName{},
 	UE::Interchange::Materials::StandardSurface::Parameters::Normal,
+	UE::Interchange::Materials::StandardSurface::Parameters::Tangent
 };
 
 TSet<FName> UMaterialXPipelineSettings::TransmissionSurfaceOutputs
 {
-	UE::Interchange::Materials::PBRMR::Parameters::BaseColor,
-	UE::Interchange::Materials::PBRMR::Parameters::Metallic,
-	UE::Interchange::Materials::PBRMR::Parameters::Specular,
-	UE::Interchange::Materials::PBRMR::Parameters::Roughness,
-	UE::Interchange::Materials::PBRMR::Parameters::Anisotropy,
-	UE::Interchange::Materials::PBRMR::Parameters::EmissiveColor,
-	UE::Interchange::Materials::PBRMR::Parameters::Opacity,
-	UE::Interchange::Materials::PBRMR::Parameters::Normal,
-	UE::Interchange::Materials::PBRMR::Parameters::Tangent,
-	UE::Interchange::Materials::PBRMR::Parameters::Refraction,
-	UE::Interchange::Materials::ThinTranslucent::Parameters::TransmissionColor
+	!bIsSubstrateEnabled ?
+	TSet<FName>{
+		UE::Interchange::Materials::PBRMR::Parameters::BaseColor,
+		UE::Interchange::Materials::PBRMR::Parameters::Metallic,
+		UE::Interchange::Materials::PBRMR::Parameters::Specular,
+		UE::Interchange::Materials::PBRMR::Parameters::Roughness,
+		UE::Interchange::Materials::PBRMR::Parameters::Anisotropy,
+		UE::Interchange::Materials::PBRMR::Parameters::EmissiveColor,
+		UE::Interchange::Materials::PBRMR::Parameters::Opacity,
+		UE::Interchange::Materials::PBRMR::Parameters::Normal,
+		UE::Interchange::Materials::PBRMR::Parameters::Tangent,
+		UE::Interchange::Materials::PBRMR::Parameters::Refraction,
+		UE::Interchange::Materials::ThinTranslucent::Parameters::TransmissionColor
+	}	:
+	TSet<FName>{UE::Interchange::Materials::StandardSurface::Substrate::Outputs::Translucent}
 };
 
 TSet<FName> UMaterialXPipelineSettings::SurfaceUnlitInputs
diff --git a/Engine/Plugins/Interchange/Runtime/Source/Pipelines/Public/InterchangeGenericMaterialPipeline.h b/Engine/Plugins/Interchange/Runtime/Source/Pipelines/Public/InterchangeGenericMaterialPipeline.h
--- a/Engine/Plugins/Interchange/Runtime/Source/Pipelines/Public/InterchangeGenericMaterialPipeline.h
+++ b/Engine/Plugins/Interchange/Runtime/Source/Pipelines/Public/InterchangeGenericMaterialPipeline.h
@@ -132,9 +132,6 @@ private:
 	/** True if the shader graph has a diffuse color input. */
 	bool IsLambertModel(const UInterchangeShaderGraphNode* ShaderGraphNode) const;
 
-	/** True if the shader graph has the standard surface's shader type name */
-	bool IsStandardSurfaceModel(const UInterchangeShaderGraphNode* ShaderGraphNode) const;
-
 	/** True if the shader graph has the surface unlit's shader type name */
 	bool IsSurfaceUnlitModel(const UInterchangeShaderGraphNode* ShaderGraphNode) const;
 
@@ -155,6 +152,7 @@ private:
 	bool HandleBxDFInput(const UInterchangeShaderGraphNode* ShaderGraphNode, UInterchangeMaterialFactoryNode* MaterialFactoryNode);
 	bool HandleUnlitModel(const UInterchangeShaderGraphNode* ShaderGraphNode, UInterchangeMaterialFactoryNode* MaterialFactoryNode);
 	bool HandleSpecGlossModel(const UInterchangeShaderGraphNode* ShaderGraphNode, UInterchangeMaterialFactoryNode* MaterialFactoryNode);
+	bool HandleSubstrate(const UInterchangeShaderGraphNode* ShaderGraphNode, UInterchangeMaterialFactoryNode* MaterialFactoryNode);
 
 	void HandleFlattenNormalNode(const UInterchangeShaderNode* ShaderNode, UInterchangeBaseMaterialFactoryNode* MaterialFactoryNode, UInterchangeMaterialExpressionFactoryNode* FlattenNormalFactoryNode);
 	void HandleNormalFromHeightMapNode(const UInterchangeShaderNode* ShaderNode, UInterchangeBaseMaterialFactoryNode* MaterialFactoryNode, UInterchangeMaterialExpressionFactoryNode* NormalFromHeightMapFactoryNode);
