Fixed crash when loading InterchangePipeline module

diff --git a/Engine/Plugins/Interchange/Runtime/Source/Pipelines/Private/InterchangeMaterialXPipeline.cpp b/Engine/Plugins/Interchange/Runtime/Source/Pipelines/Private/InterchangeMaterialXPipeline.cpp
--- a/Engine/Plugins/Interchange/Runtime/Source/Pipelines/Private/InterchangeMaterialXPipeline.cpp
+++ b/Engine/Plugins/Interchange/Runtime/Source/Pipelines/Private/InterchangeMaterialXPipeline.cpp
@@ -11,7 +11,7 @@
 #include "Materials/MaterialFunction.h"
 #include "Misc/PackageName.h"
 
-static const bool bIsSubstrateEnabled = GetDefault<URendererSettings>()->bEnableStrata;
+static const bool bIsSubstrateEnabled = GetDefault<URendererSettings>() ? GetDefault<URendererSettings>()->bEnableStrata : false;
 
 TMap<FString, EInterchangeMaterialXShaders> UInterchangeMaterialXPipeline::PathToEnumMapping
 {
