[Interchange] Animation doesn't load on a specific GLTF file

diff --git a/Engine/Plugins/Interchange/Runtime/Source/Import/Private/Gltf/InterchangeGltfAnimation.cpp b/Engine/Plugins/Interchange/Runtime/Source/Import/Private/Gltf/InterchangeGltfAnimation.cpp
--- a/Engine/Plugins/Interchange/Runtime/Source/Import/Private/Gltf/InterchangeGltfAnimation.cpp
+++ b/Engine/Plugins/Interchange/Runtime/Source/Import/Private/Gltf/InterchangeGltfAnimation.cpp
@@ -462,7 +462,8 @@ namespace UE::Interchange::Gltf::Private
 		
 		const double BakeInterval = 1.0 / PayloadData.BakeFrequency;
 		const double SequenceLength = FMath::Max<double>(PayloadData.RangeEndTime - PayloadData.RangeStartTime, MINIMUM_ANIMATION_LENGTH);
-		int32 BakeKeyCount = (SequenceLength / BakeInterval) + 1;
+		int32 FrameCount = FMath::RoundToInt32(SequenceLength * PayloadData.BakeFrequency);
+		int32 BakeKeyCount = FrameCount + 1;
 
 		//buffers to use for final Transform:
 		TArray<FVector3f> TranslationData;
diff --git a/Engine/Plugins/Interchange/Runtime/Source/Import/Private/Mesh/InterchangeSkeletonHelper.cpp b/Engine/Plugins/Interchange/Runtime/Source/Import/Private/Mesh/InterchangeSkeletonHelper.cpp
--- a/Engine/Plugins/Interchange/Runtime/Source/Import/Private/Mesh/InterchangeSkeletonHelper.cpp
+++ b/Engine/Plugins/Interchange/Runtime/Source/Import/Private/Mesh/InterchangeSkeletonHelper.cpp
@@ -17,13 +17,6 @@ namespace UE::Interchange::Private
 
 	bool FSkeletonHelper::ProcessImportMeshSkeleton(const USkeleton* SkeletonAsset, FReferenceSkeleton& RefSkeleton, int32& SkeletalDepth, const UInterchangeBaseNodeContainer* NodeContainer, const FString& RootJointNodeId, TArray<SkeletalMeshImportData::FBone>& RefBonesBinary, const bool bUseTimeZeroAsBindPose, bool& bOutDiffPose)
 	{
-		auto FixupBoneName = [](FString BoneName)
-		{
-			BoneName.TrimStartAndEndInline();
-			BoneName.ReplaceInline(TEXT(" "), TEXT("-"), ESearchCase::IgnoreCase);
-			return BoneName;
-		};
-
 		RefBonesBinary.Empty();
 		// Setup skeletal hierarchy + names structure.
 		RefSkeleton.Empty();
@@ -41,7 +34,7 @@ namespace UE::Interchange::Private
 		{
 			const FJointInfo& BinaryBone = JointInfos[b];
 
-			const FString BoneName = FixupBoneName(BinaryBone.Name);
+			const FString BoneName = BinaryBone.Name;
 			const FMeshBoneInfo BoneInfo(FName(*BoneName, FNAME_Add), BinaryBone.Name, BinaryBone.ParentIndex);
 			const FTransform BoneTransform(BinaryBone.LocalTransform);
 			if (RefSkeleton.FindRawBoneIndex(BoneInfo.Name) != INDEX_NONE)
diff --git a/Engine/Plugins/Interchange/Runtime/Source/Parsers/Fbx/Private/FbxConvert.cpp b/Engine/Plugins/Interchange/Runtime/Source/Parsers/Fbx/Private/FbxConvert.cpp
--- a/Engine/Plugins/Interchange/Runtime/Source/Parsers/Fbx/Private/FbxConvert.cpp
+++ b/Engine/Plugins/Interchange/Runtime/Source/Parsers/Fbx/Private/FbxConvert.cpp
@@ -184,30 +184,6 @@ namespace UE
 				return LinearColor;
 			}
 
-			FString FFbxConvert::MakeName(const ANSICHAR* Name)
-			{
-				const TCHAR SpecialChars[] = { TEXT('.'), TEXT(','), TEXT('/'), TEXT('`'), TEXT('%') };
-
-				FString TmpName = MakeString(Name);
-
-				// Remove namespaces
-				int32 LastNamespaceTokenIndex = INDEX_NONE;
-				if (TmpName.FindLastChar(TEXT(':'), LastNamespaceTokenIndex))
-				{
-					const bool bAllowShrinking = true;
-					//+1 to remove the ':' character we found
-					TmpName.RightChopInline(LastNamespaceTokenIndex + 1, bAllowShrinking);
-				}
-
-				//Remove the special chars
-				for (int32 i = 0; i < UE_ARRAY_COUNT(SpecialChars); i++)
-				{
-					TmpName.ReplaceCharInline(SpecialChars[i], TEXT('_'), ESearchCase::CaseSensitive);
-				}
-
-				return TmpName;
-			}
-
 			/**
 			 * Convert ANSI char to a FString using ANSI_TO_TCHAR macro
 			 */
diff --git a/Engine/Plugins/Interchange/Runtime/Source/Parsers/Fbx/Private/FbxConvert.h b/Engine/Plugins/Interchange/Runtime/Source/Parsers/Fbx/Private/FbxConvert.h
--- a/Engine/Plugins/Interchange/Runtime/Source/Parsers/Fbx/Private/FbxConvert.h
+++ b/Engine/Plugins/Interchange/Runtime/Source/Parsers/Fbx/Private/FbxConvert.h
@@ -53,12 +53,6 @@ namespace UE
 				//////////////////////////////////////////////////////////////////////////
 				/** String Conversion API Begin */
 
-				/**
-				 * Replace all special characters with '_', then remove all namespace
-				 * Special characters are . , / ` %
-				 */
-				static FString MakeName(const ANSICHAR* Name);
-
 				/**
 				 * Convert ANSI char to a FString using ANSI_TO_TCHAR macro
 				 */
diff --git a/Engine/Plugins/Interchange/Runtime/Source/Parsers/Fbx/Private/FbxHelper.cpp b/Engine/Plugins/Interchange/Runtime/Source/Parsers/Fbx/Private/FbxHelper.cpp
--- a/Engine/Plugins/Interchange/Runtime/Source/Parsers/Fbx/Private/FbxHelper.cpp
+++ b/Engine/Plugins/Interchange/Runtime/Source/Parsers/Fbx/Private/FbxHelper.cpp
@@ -5,6 +5,7 @@
 #include "CoreMinimal.h"
 #include "FbxConvert.h"
 #include "FbxInclude.h"
+#include "InterchangeHelper.h"
 
 #define GeneratedLODNameSuffix "_GeneratedLOD_"
 
@@ -101,7 +102,7 @@ namespace UE
 
 			FString FFbxHelper::GetFbxPropertyName(const FbxProperty& Property) const
 			{
-				FString PropertyName = FFbxConvert::MakeName(Property.GetName());
+				FString PropertyName = UE::Interchange::MakeName(FFbxConvert::MakeString(Property.GetName()));
 				if (PropertyName.Equals(TEXT("none"), ESearchCase::IgnoreCase))
 				{
 					//Replace None by Null because None clash with NAME_None and the create asset will instead call the object ClassName_X
@@ -110,13 +111,13 @@ namespace UE
 				return PropertyName;
 			}
 
-			FString FFbxHelper::GetFbxObjectName(const FbxObject* Object) const
+			FString FFbxHelper::GetFbxObjectName(const FbxObject* Object, bool bIsJoint) const
 			{
 				if (!Object)
 				{
 					return FString();
 				}
-				FString ObjName = FFbxConvert::MakeName(Object->GetName());
+				FString ObjName = UE::Interchange::MakeName(FFbxConvert::MakeString(Object->GetName()), bIsJoint);
 				if (ObjName.Equals(TEXT("none"), ESearchCase::IgnoreCase))
 				{
 					//Replace None by Null because None clash with NAME_None and the create asset will instead call the object ClassName_X
diff --git a/Engine/Plugins/Interchange/Runtime/Source/Parsers/Fbx/Private/FbxHelper.h b/Engine/Plugins/Interchange/Runtime/Source/Parsers/Fbx/Private/FbxHelper.h
--- a/Engine/Plugins/Interchange/Runtime/Source/Parsers/Fbx/Private/FbxHelper.h
+++ b/Engine/Plugins/Interchange/Runtime/Source/Parsers/Fbx/Private/FbxHelper.h
@@ -37,7 +37,7 @@ namespace UE::Interchange::Private
 		/**
 			* Return the name of an FbxObject, return empty string if the object is null.
 			*/
-		FString GetFbxObjectName(const FbxObject* Object) const;
+		FString GetFbxObjectName(const FbxObject* Object, bool bIsJoint = false) const;
 
 		/**
 			* Return a string with the name of all the parent in the hierarchy separate by a dot( . ) from the fbx root node to the specified node.
diff --git a/Engine/Plugins/Interchange/Runtime/Source/Parsers/Fbx/Private/FbxMesh.cpp b/Engine/Plugins/Interchange/Runtime/Source/Parsers/Fbx/Private/FbxMesh.cpp
--- a/Engine/Plugins/Interchange/Runtime/Source/Parsers/Fbx/Private/FbxMesh.cpp
+++ b/Engine/Plugins/Interchange/Runtime/Source/Parsers/Fbx/Private/FbxMesh.cpp
@@ -1094,7 +1094,7 @@ bool FMeshDescriptionImporter::FillMeshDescriptionFromFbxMesh(FbxMesh* Mesh, TAr
 				if (!SortedJoints.Find(Link, BoneIndex))
 				{
 					BoneIndex = SortedJoints.Add(Link);
-					FString JointNodeUniqueID = Parser.GetFbxHelper()->GetFbxObjectName(Link);
+					FString JointNodeUniqueID = Parser.GetFbxHelper()->GetFbxObjectName(Link, true);
 					OutJointUniqueNames.Add(JointNodeUniqueID);
 				}
 
diff --git a/Engine/Plugins/Interchange/Runtime/Source/Parsers/GLTFCore/GLTFCore.Build.cs b/Engine/Plugins/Interchange/Runtime/Source/Parsers/GLTFCore/GLTFCore.Build.cs
--- a/Engine/Plugins/Interchange/Runtime/Source/Parsers/GLTFCore/GLTFCore.Build.cs
+++ b/Engine/Plugins/Interchange/Runtime/Source/Parsers/GLTFCore/GLTFCore.Build.cs
@@ -23,7 +23,8 @@ namespace UnrealBuildTool.Rules
                     "Json",
                     "RenderCore",
 					"HTTP",
-                }
+					"InterchangeCore",
+				}
                 );
         }
     }
diff --git a/Engine/Plugins/Interchange/Runtime/Source/Parsers/GLTFCore/Private/GLTF/GLTFAsset.cpp b/Engine/Plugins/Interchange/Runtime/Source/Parsers/GLTFCore/Private/GLTF/GLTFAsset.cpp
--- a/Engine/Plugins/Interchange/Runtime/Source/Parsers/GLTFCore/Private/GLTF/GLTFAsset.cpp
+++ b/Engine/Plugins/Interchange/Runtime/Source/Parsers/GLTFCore/Private/GLTF/GLTFAsset.cpp
@@ -8,6 +8,7 @@
 #include "GLTFMesh.h"
 #include "GLTFNode.h"
 #include "GLTFTexture.h"
+#include "InterchangeHelper.h"
 
 namespace GLTF
 {
@@ -26,6 +27,8 @@ namespace GLTF
 				{
 					Obj.Name = Obj.UniqueId;
 				}
+
+				Obj.Name = UE::Interchange::MakeName(Obj.Name);
 			}
 		}
 
@@ -112,6 +115,8 @@ namespace GLTF
 				{
 					Scene.Name = Scene.UniqueId;
 				}
+
+				Scene.Name = UE::Interchange::MakeName(Scene.Name);
 			}
 		}
 
@@ -131,6 +136,8 @@ namespace GLTF
 				{
 					Node.Name = Node.UniqueId;
 				}
+
+				Node.Name = UE::Interchange::MakeName(Node.Name, bIsJoint);
 			}
 		}
 
@@ -167,6 +174,7 @@ namespace GLTF
 					Tex.Name = Tex.UniqueId;
 				}
 
+				Tex.Name = UE::Interchange::MakeName(Tex.Name);
 			}
 		}
 
@@ -189,6 +197,8 @@ namespace GLTF
 						Camera.Name = Camera.UniqueId;
 					}
 				}
+
+				Camera.Name = UE::Interchange::MakeName(Camera.Name);
 			}
 		}
 
@@ -211,6 +221,8 @@ namespace GLTF
 						Light.Name = Light.UniqueId;
 					}
 				}
+
+				Light.Name = UE::Interchange::MakeName(Light.Name);
 			}
 		}
 
diff --git a/Engine/Source/Runtime/Interchange/Core/Public/InterchangeHelper.h b/Engine/Source/Runtime/Interchange/Core/Public/InterchangeHelper.h
new file mode 100644
--- /dev/null
+++ b/Engine/Source/Runtime/Interchange/Core/Public/InterchangeHelper.h
@@ -0,0 +1,38 @@
+// Copyright Epic Games, Inc. All Rights Reserved.
+
+#pragma once
+
+#include "Containers/UnrealString.h"
+#include "CoreMinimal.h"
+
+namespace UE::Interchange
+{
+	const TCHAR UnsupportedNameCharacters[] = { TEXT('.'), TEXT(','), TEXT('/'), TEXT('`'), TEXT('%') };
+	const TCHAR UnsupportedJointNameCharacters[] = { TEXT('.'), TEXT(','), TEXT('/'), TEXT('`'), TEXT('%'), TEXT('+'), TEXT(' ') };
+
+	/**
+	* Replaces any unsupported characters with "_" character, and removes namespace indicator ":" character
+	*/
+	FString MakeName(const FString& InName, bool bIsJoint = false)
+	{
+		FString SpecialChars = bIsJoint ? UnsupportedJointNameCharacters : UnsupportedNameCharacters;
+		FString TmpName = InName;
+
+		// Remove namespaces
+		int32 LastNamespaceTokenIndex = INDEX_NONE;
+		if (TmpName.FindLastChar(TEXT(':'), LastNamespaceTokenIndex))
+		{
+			const bool bAllowShrinking = true;
+			//+1 to remove the ':' character we found
+			TmpName.RightChopInline(LastNamespaceTokenIndex + 1, bAllowShrinking);
+		}
+
+		//Remove the special chars
+		for (size_t CharIndex = 0; CharIndex < SpecialChars.Len(); CharIndex++)
+		{
+			TmpName.ReplaceCharInline(SpecialChars[CharIndex], TEXT('_'), ESearchCase::CaseSensitive);
+		}
+
+		return TmpName;
+	}
+};
\ No newline at end of file
