[Backout] - CL26417344
Original CL Desc
-----------------------------------------------------------------
[Interchange] Animation doesn't load on a specific GLTF file

diff --git a/Engine/Plugins/Interchange/Runtime/Source/Import/Private/Gltf/InterchangeGltfAnimation.cpp b/Engine/Plugins/Interchange/Runtime/Source/Import/Private/Gltf/InterchangeGltfAnimation.cpp
--- a/Engine/Plugins/Interchange/Runtime/Source/Import/Private/Gltf/InterchangeGltfAnimation.cpp
+++ b/Engine/Plugins/Interchange/Runtime/Source/Import/Private/Gltf/InterchangeGltfAnimation.cpp
@@ -462,8 +462,7 @@ namespace UE::Interchange::Gltf::Private
 		
 		const double BakeInterval = 1.0 / PayloadData.BakeFrequency;
 		const double SequenceLength = FMath::Max<double>(PayloadData.RangeEndTime - PayloadData.RangeStartTime, MINIMUM_ANIMATION_LENGTH);
-		int32 FrameCount = FMath::RoundToInt32(SequenceLength * PayloadData.BakeFrequency);
-		int32 BakeKeyCount = FrameCount + 1;
+		int32 BakeKeyCount = (SequenceLength / BakeInterval) + 1;
 
 		//buffers to use for final Transform:
 		TArray<FVector3f> TranslationData;
diff --git a/Engine/Plugins/Interchange/Runtime/Source/Import/Private/Mesh/InterchangeSkeletonHelper.cpp b/Engine/Plugins/Interchange/Runtime/Source/Import/Private/Mesh/InterchangeSkeletonHelper.cpp
--- a/Engine/Plugins/Interchange/Runtime/Source/Import/Private/Mesh/InterchangeSkeletonHelper.cpp
+++ b/Engine/Plugins/Interchange/Runtime/Source/Import/Private/Mesh/InterchangeSkeletonHelper.cpp
@@ -17,6 +17,13 @@ namespace UE::Interchange::Private
 
 	bool FSkeletonHelper::ProcessImportMeshSkeleton(const USkeleton* SkeletonAsset, FReferenceSkeleton& RefSkeleton, int32& SkeletalDepth, const UInterchangeBaseNodeContainer* NodeContainer, const FString& RootJointNodeId, TArray<SkeletalMeshImportData::FBone>& RefBonesBinary, const bool bUseTimeZeroAsBindPose, bool& bOutDiffPose)
 	{
+		auto FixupBoneName = [](FString BoneName)
+		{
+			BoneName.TrimStartAndEndInline();
+			BoneName.ReplaceInline(TEXT(" "), TEXT("-"), ESearchCase::IgnoreCase);
+			return BoneName;
+		};
+
 		RefBonesBinary.Empty();
 		// Setup skeletal hierarchy + names structure.
 		RefSkeleton.Empty();
@@ -34,7 +41,7 @@ namespace UE::Interchange::Private
 		{
 			const FJointInfo& BinaryBone = JointInfos[b];
 
-			const FString BoneName = BinaryBone.Name;
+			const FString BoneName = FixupBoneName(BinaryBone.Name);
 			const FMeshBoneInfo BoneInfo(FName(*BoneName, FNAME_Add), BinaryBone.Name, BinaryBone.ParentIndex);
 			const FTransform BoneTransform(BinaryBone.LocalTransform);
 			if (RefSkeleton.FindRawBoneIndex(BoneInfo.Name) != INDEX_NONE)
diff --git a/Engine/Plugins/Interchange/Runtime/Source/Parsers/Fbx/Private/FbxConvert.cpp b/Engine/Plugins/Interchange/Runtime/Source/Parsers/Fbx/Private/FbxConvert.cpp
--- a/Engine/Plugins/Interchange/Runtime/Source/Parsers/Fbx/Private/FbxConvert.cpp
+++ b/Engine/Plugins/Interchange/Runtime/Source/Parsers/Fbx/Private/FbxConvert.cpp
@@ -184,6 +184,30 @@ namespace UE
 				return LinearColor;
 			}
 
+			FString FFbxConvert::MakeName(const ANSICHAR* Name)
+			{
+				const TCHAR SpecialChars[] = { TEXT('.'), TEXT(','), TEXT('/'), TEXT('`'), TEXT('%') };
+
+				FString TmpName = MakeString(Name);
+
+				// Remove namespaces
+				int32 LastNamespaceTokenIndex = INDEX_NONE;
+				if (TmpName.FindLastChar(TEXT(':'), LastNamespaceTokenIndex))
+				{
+					const bool bAllowShrinking = true;
+					//+1 to remove the ':' character we found
+					TmpName.RightChopInline(LastNamespaceTokenIndex + 1, bAllowShrinking);
+				}
+
+				//Remove the special chars
+				for (int32 i = 0; i < UE_ARRAY_COUNT(SpecialChars); i++)
+				{
+					TmpName.ReplaceCharInline(SpecialChars[i], TEXT('_'), ESearchCase::CaseSensitive);
+				}
+
+				return TmpName;
+			}
+
 			/**
 			 * Convert ANSI char to a FString using ANSI_TO_TCHAR macro
 			 */
diff --git a/Engine/Plugins/Interchange/Runtime/Source/Parsers/Fbx/Private/FbxConvert.h b/Engine/Plugins/Interchange/Runtime/Source/Parsers/Fbx/Private/FbxConvert.h
--- a/Engine/Plugins/Interchange/Runtime/Source/Parsers/Fbx/Private/FbxConvert.h
+++ b/Engine/Plugins/Interchange/Runtime/Source/Parsers/Fbx/Private/FbxConvert.h
@@ -53,6 +53,12 @@ namespace UE
 				//////////////////////////////////////////////////////////////////////////
 				/** String Conversion API Begin */
 
+				/**
+				 * Replace all special characters with '_', then remove all namespace
+				 * Special characters are . , / ` %
+				 */
+				static FString MakeName(const ANSICHAR* Name);
+
 				/**
 				 * Convert ANSI char to a FString using ANSI_TO_TCHAR macro
 				 */
diff --git a/Engine/Plugins/Interchange/Runtime/Source/Parsers/Fbx/Private/FbxHelper.cpp b/Engine/Plugins/Interchange/Runtime/Source/Parsers/Fbx/Private/FbxHelper.cpp
--- a/Engine/Plugins/Interchange/Runtime/Source/Parsers/Fbx/Private/FbxHelper.cpp
+++ b/Engine/Plugins/Interchange/Runtime/Source/Parsers/Fbx/Private/FbxHelper.cpp
@@ -5,7 +5,6 @@
 #include "CoreMinimal.h"
 #include "FbxConvert.h"
 #include "FbxInclude.h"
-#include "InterchangeHelper.h"
 
 #define GeneratedLODNameSuffix "_GeneratedLOD_"
 
@@ -102,7 +101,7 @@ namespace UE
 
 			FString FFbxHelper::GetFbxPropertyName(const FbxProperty& Property) const
 			{
-				FString PropertyName = UE::Interchange::MakeName(FFbxConvert::MakeString(Property.GetName()));
+				FString PropertyName = FFbxConvert::MakeName(Property.GetName());
 				if (PropertyName.Equals(TEXT("none"), ESearchCase::IgnoreCase))
 				{
 					//Replace None by Null because None clash with NAME_None and the create asset will instead call the object ClassName_X
@@ -111,13 +110,13 @@ namespace UE
 				return PropertyName;
 			}
 
-			FString FFbxHelper::GetFbxObjectName(const FbxObject* Object, bool bIsJoint) const
+			FString FFbxHelper::GetFbxObjectName(const FbxObject* Object) const
 			{
 				if (!Object)
 				{
 					return FString();
 				}
-				FString ObjName = UE::Interchange::MakeName(FFbxConvert::MakeString(Object->GetName()), bIsJoint);
+				FString ObjName = FFbxConvert::MakeName(Object->GetName());
 				if (ObjName.Equals(TEXT("none"), ESearchCase::IgnoreCase))
 				{
 					//Replace None by Null because None clash with NAME_None and the create asset will instead call the object ClassName_X
diff --git a/Engine/Plugins/Interchange/Runtime/Source/Parsers/Fbx/Private/FbxHelper.h b/Engine/Plugins/Interchange/Runtime/Source/Parsers/Fbx/Private/FbxHelper.h
--- a/Engine/Plugins/Interchange/Runtime/Source/Parsers/Fbx/Private/FbxHelper.h
+++ b/Engine/Plugins/Interchange/Runtime/Source/Parsers/Fbx/Private/FbxHelper.h
@@ -37,7 +37,7 @@ namespace UE::Interchange::Private
 		/**
 			* Return the name of an FbxObject, return empty string if the object is null.
 			*/
-		FString GetFbxObjectName(const FbxObject* Object, bool bIsJoint = false) const;
+		FString GetFbxObjectName(const FbxObject* Object) const;
 
 		/**
 			* Return a string with the name of all the parent in the hierarchy separate by a dot( . ) from the fbx root node to the specified node.
diff --git a/Engine/Plugins/Interchange/Runtime/Source/Parsers/Fbx/Private/FbxMesh.cpp b/Engine/Plugins/Interchange/Runtime/Source/Parsers/Fbx/Private/FbxMesh.cpp
--- a/Engine/Plugins/Interchange/Runtime/Source/Parsers/Fbx/Private/FbxMesh.cpp
+++ b/Engine/Plugins/Interchange/Runtime/Source/Parsers/Fbx/Private/FbxMesh.cpp
@@ -1094,7 +1094,7 @@ bool FMeshDescriptionImporter::FillMeshDescriptionFromFbxMesh(FbxMesh* Mesh, TAr
 				if (!SortedJoints.Find(Link, BoneIndex))
 				{
 					BoneIndex = SortedJoints.Add(Link);
-					FString JointNodeUniqueID = Parser.GetFbxHelper()->GetFbxObjectName(Link, true);
+					FString JointNodeUniqueID = Parser.GetFbxHelper()->GetFbxObjectName(Link);
 					OutJointUniqueNames.Add(JointNodeUniqueID);
 				}
 
diff --git a/Engine/Plugins/Interchange/Runtime/Source/Parsers/GLTFCore/GLTFCore.Build.cs b/Engine/Plugins/Interchange/Runtime/Source/Parsers/GLTFCore/GLTFCore.Build.cs
--- a/Engine/Plugins/Interchange/Runtime/Source/Parsers/GLTFCore/GLTFCore.Build.cs
+++ b/Engine/Plugins/Interchange/Runtime/Source/Parsers/GLTFCore/GLTFCore.Build.cs
@@ -23,8 +23,7 @@ namespace UnrealBuildTool.Rules
                     "Json",
                     "RenderCore",
 					"HTTP",
-					"InterchangeCore",
-				}
+                }
                 );
         }
     }
diff --git a/Engine/Plugins/Interchange/Runtime/Source/Parsers/GLTFCore/Private/GLTF/GLTFAsset.cpp b/Engine/Plugins/Interchange/Runtime/Source/Parsers/GLTFCore/Private/GLTF/GLTFAsset.cpp
--- a/Engine/Plugins/Interchange/Runtime/Source/Parsers/GLTFCore/Private/GLTF/GLTFAsset.cpp
+++ b/Engine/Plugins/Interchange/Runtime/Source/Parsers/GLTFCore/Private/GLTF/GLTFAsset.cpp
@@ -8,7 +8,6 @@
 #include "GLTFMesh.h"
 #include "GLTFNode.h"
 #include "GLTFTexture.h"
-#include "InterchangeHelper.h"
 
 namespace GLTF
 {
@@ -27,8 +26,6 @@ namespace GLTF
 				{
 					Obj.Name = Obj.UniqueId;
 				}
-
-				Obj.Name = UE::Interchange::MakeName(Obj.Name);
 			}
 		}
 
@@ -115,8 +112,6 @@ namespace GLTF
 				{
 					Scene.Name = Scene.UniqueId;
 				}
-
-				Scene.Name = UE::Interchange::MakeName(Scene.Name);
 			}
 		}
 
@@ -136,8 +131,6 @@ namespace GLTF
 				{
 					Node.Name = Node.UniqueId;
 				}
-
-				Node.Name = UE::Interchange::MakeName(Node.Name, bIsJoint);
 			}
 		}
 
@@ -174,7 +167,6 @@ namespace GLTF
 					Tex.Name = Tex.UniqueId;
 				}
 
-				Tex.Name = UE::Interchange::MakeName(Tex.Name);
 			}
 		}
 
@@ -197,8 +189,6 @@ namespace GLTF
 						Camera.Name = Camera.UniqueId;
 					}
 				}
-
-				Camera.Name = UE::Interchange::MakeName(Camera.Name);
 			}
 		}
 
@@ -221,8 +211,6 @@ namespace GLTF
 						Light.Name = Light.UniqueId;
 					}
 				}
-
-				Light.Name = UE::Interchange::MakeName(Light.Name);
 			}
 		}
 
diff --git a/Engine/Source/Runtime/Interchange/Core/Public/InterchangeHelper.h b/Engine/Source/Runtime/Interchange/Core/Public/InterchangeHelper.h
deleted file mode 100644
