[Interchange] Skeletal without material re-import fixes

diff --git a/Engine/Plugins/Interchange/Runtime/Source/Import/Private/Mesh/InterchangeSkeletalMeshFactory.cpp b/Engine/Plugins/Interchange/Runtime/Source/Import/Private/Mesh/InterchangeSkeletalMeshFactory.cpp
--- a/Engine/Plugins/Interchange/Runtime/Source/Import/Private/Mesh/InterchangeSkeletalMeshFactory.cpp
+++ b/Engine/Plugins/Interchange/Runtime/Source/Import/Private/Mesh/InterchangeSkeletalMeshFactory.cpp
@@ -1514,13 +1514,16 @@ UInterchangeFactoryBase::FImportAssetResult UInterchangeSkeletalMeshFactory::End
 		{
 			if (FSkeletalMeshLODInfo* LodInfo = SkeletalMesh->GetLODInfo(ImportAssetObjectLODData.LodIndex))
 			{
-				FLODUtilities::FSkeletalMeshMatchImportedMaterialsParameters Parameters;
-				Parameters.bIsReImport = ImportAssetObjectData.bIsReImport;
-				Parameters.LodIndex = ImportAssetObjectLODData.LodIndex;
-				Parameters.SkeletalMesh = SkeletalMesh;
-				Parameters.ImportedMaterials = &ImportAssetObjectLODData.ImportedMaterials;
-				Parameters.ExistingOriginalPerSectionMaterialImportName = &ImportAssetObjectLODData.ExistingOriginalPerSectionMaterialImportName;
-				FLODUtilities::MatchImportedMaterials(Parameters);
+				if (SkeletalMesh->GetMaterials().Num() > 0)
+				{
+					FLODUtilities::FSkeletalMeshMatchImportedMaterialsParameters Parameters;
+					Parameters.bIsReImport = ImportAssetObjectData.bIsReImport;
+					Parameters.LodIndex = ImportAssetObjectLODData.LodIndex;
+					Parameters.SkeletalMesh = SkeletalMesh;
+					Parameters.ImportedMaterials = &ImportAssetObjectLODData.ImportedMaterials;
+					Parameters.ExistingOriginalPerSectionMaterialImportName = &ImportAssetObjectLODData.ExistingOriginalPerSectionMaterialImportName;
+					FLODUtilities::MatchImportedMaterials(Parameters);
+				}
 			}
 		}
 	}
diff --git a/Engine/Plugins/Interchange/Runtime/Source/Import/Private/Mesh/InterchangeSkeletonHelper.cpp b/Engine/Plugins/Interchange/Runtime/Source/Import/Private/Mesh/InterchangeSkeletonHelper.cpp
--- a/Engine/Plugins/Interchange/Runtime/Source/Import/Private/Mesh/InterchangeSkeletonHelper.cpp
+++ b/Engine/Plugins/Interchange/Runtime/Source/Import/Private/Mesh/InterchangeSkeletonHelper.cpp
@@ -71,7 +71,7 @@ namespace UE::Interchange::Private
 		return true;
 	}
 
-	bool FSkeletonHelper::IsCompatibleSkeleton(const USkeleton* Skeleton, const FString RootJoinUid, const UInterchangeBaseNodeContainer* BaseNodeContainer)
+	bool FSkeletonHelper::IsCompatibleSkeleton(const USkeleton* Skeleton, const FString RootJoinUid, const UInterchangeBaseNodeContainer* BaseNodeContainer, bool bConvertStaticToSkeletalActive)
 	{
 		// at least % of bone should match 
 		int32 NumOfBoneMatches = 0;
@@ -81,7 +81,7 @@ namespace UE::Interchange::Private
 
 		TArray<FMeshBoneInfo> SkeletalLodRawInfos;
 		SkeletalLodRawInfos.Reserve(SkeletonBoneCount);
-		RecursiveBuildSkeletalSkeleton(RootJoinUid, INDEX_NONE, BaseNodeContainer, SkeletalLodRawInfos);
+		RecursiveBuildSkeletalSkeleton(RootJoinUid, INDEX_NONE, BaseNodeContainer, SkeletalLodRawInfos, bConvertStaticToSkeletalActive);
 		const int32 SkeletalLodBoneCount = SkeletalLodRawInfos.Num();
 
 		// first ensure the parent exists for each bone
@@ -320,10 +320,10 @@ namespace UE::Interchange::Private
 		return true;
 	}
 
-	void FSkeletonHelper::RecursiveBuildSkeletalSkeleton(const FString JoinToAddUid, const int32 ParentIndex, const UInterchangeBaseNodeContainer* BaseNodeContainer, TArray<FMeshBoneInfo>& SkeletalLodRawInfos)
+	void FSkeletonHelper::RecursiveBuildSkeletalSkeleton(const FString JoinToAddUid, const int32 ParentIndex, const UInterchangeBaseNodeContainer* BaseNodeContainer, TArray<FMeshBoneInfo>& SkeletalLodRawInfos, bool bConvertStaticToSkeletalActive)
 	{
 		const UInterchangeSceneNode* SceneNode = Cast<UInterchangeSceneNode>(BaseNodeContainer->GetNode(JoinToAddUid));
-		if (!SceneNode || !SceneNode->IsSpecializedTypeContains(FSceneNodeStaticData::GetJointSpecializeTypeString()))
+		if (!bConvertStaticToSkeletalActive && (!SceneNode || !SceneNode->IsSpecializedTypeContains(FSceneNodeStaticData::GetJointSpecializeTypeString())))
 		{
 			return;
 		}
diff --git a/Engine/Plugins/Interchange/Runtime/Source/Import/Public/InterchangeSkeletonHelper.h b/Engine/Plugins/Interchange/Runtime/Source/Import/Public/InterchangeSkeletonHelper.h
--- a/Engine/Plugins/Interchange/Runtime/Source/Import/Public/InterchangeSkeletonHelper.h
+++ b/Engine/Plugins/Interchange/Runtime/Source/Import/Public/InterchangeSkeletonHelper.h
@@ -28,7 +28,7 @@ namespace UE::Interchange::Private
 	{
 	public:
 		static bool ProcessImportMeshSkeleton(const USkeleton* SkeletonAsset, FReferenceSkeleton& RefSkeleton, int32& SkeletalDepth, const UInterchangeBaseNodeContainer* NodeContainer, const FString& RootJointNodeId, TArray<SkeletalMeshImportData::FBone>& RefBonesBinary, const bool bUseTimeZeroAsBindPose, bool& bOutDiffPose);
-		static bool IsCompatibleSkeleton(const USkeleton* Skeleton, const FString RootJoinUid, const UInterchangeBaseNodeContainer* BaseNodeContainer);
+		static bool IsCompatibleSkeleton(const USkeleton* Skeleton, const FString RootJoinUid, const UInterchangeBaseNodeContainer* BaseNodeContainer, bool bConvertStaticToSkeletalActive);
 		static void RecursiveAddSkeletonMetaDataValues(UInterchangeBaseNodeContainer* NodeContainer, UInterchangeBaseNode* DestinationNode, const FString& JointUid);
 	private:
 		static void RecursiveAddBones(const UInterchangeBaseNodeContainer* NodeContainer, const FString& JointNodeId, TArray <FJointInfo>& JointInfos, int32 ParentIndex, TArray<SkeletalMeshImportData::FBone>& RefBonesBinary, const bool bUseTimeZeroAsBindPose, bool& bOutDiffPose);
@@ -36,7 +36,7 @@ namespace UE::Interchange::Private
 		static int32 SkeletalLodFindBoneIndex(const TArray<FMeshBoneInfo>& SkeletalLodRawInfos, FName BoneName);
 		static int32 SkeletalLodGetParentIndex(const TArray<FMeshBoneInfo>& SkeletalLodRawInfos, int32 BoneIndex);
 		static bool DoesParentChainMatch(int32 StartBoneIndex, const FReferenceSkeleton& SkeletonRef, const TArray<FMeshBoneInfo>& SkeletalLodRawInfos);
-		static void RecursiveBuildSkeletalSkeleton(const FString JoinToAddUid, const int32 ParentIndex, const UInterchangeBaseNodeContainer* BaseNodeContainer, TArray<FMeshBoneInfo>& SkeletalLodRawInfos);
+		static void RecursiveBuildSkeletalSkeleton(const FString JoinToAddUid, const int32 ParentIndex, const UInterchangeBaseNodeContainer* BaseNodeContainer, TArray<FMeshBoneInfo>& SkeletalLodRawInfos, bool bConvertStaticToSkeletalActive = false);
 	};
 } //namespace UE::Interchange::Private
 
diff --git a/Engine/Plugins/Interchange/Runtime/Source/Pipelines/Private/InterchangeGenericAnimationPipeline.cpp b/Engine/Plugins/Interchange/Runtime/Source/Pipelines/Private/InterchangeGenericAnimationPipeline.cpp
--- a/Engine/Plugins/Interchange/Runtime/Source/Pipelines/Private/InterchangeGenericAnimationPipeline.cpp
+++ b/Engine/Plugins/Interchange/Runtime/Source/Pipelines/Private/InterchangeGenericAnimationPipeline.cpp
@@ -500,7 +500,7 @@ void UInterchangeGenericAnimationPipeline::CreateAnimSequenceFactoryNode(UInterc
 
 		//TODO: support skeleton helper in runtime
 #if WITH_EDITOR
-		bSkeletonCompatible = UE::Interchange::Private::FSkeletonHelper::IsCompatibleSkeleton(CommonSkeletalMeshesAndAnimationsProperties->Skeleton.Get(), RootJointUid, BaseNodeContainer);
+		bSkeletonCompatible = UE::Interchange::Private::FSkeletonHelper::IsCompatibleSkeleton(CommonSkeletalMeshesAndAnimationsProperties->Skeleton.Get(), RootJointUid, BaseNodeContainer, CommonSkeletalMeshesAndAnimationsProperties->bConvertStaticsWithMorphTargetsToSkeletals || CommonMeshesProperties->ForceAllMeshAsType == EInterchangeForceMeshType::IFMT_SkeletalMesh);
 #endif
 		if (bSkeletonCompatible)
 		{
diff --git a/Engine/Plugins/Interchange/Runtime/Source/Pipelines/Private/InterchangeGenericSkeletalMeshPipeline.cpp b/Engine/Plugins/Interchange/Runtime/Source/Pipelines/Private/InterchangeGenericSkeletalMeshPipeline.cpp
--- a/Engine/Plugins/Interchange/Runtime/Source/Pipelines/Private/InterchangeGenericSkeletalMeshPipeline.cpp
+++ b/Engine/Plugins/Interchange/Runtime/Source/Pipelines/Private/InterchangeGenericSkeletalMeshPipeline.cpp
@@ -174,7 +174,7 @@ void UInterchangeGenericMeshPipeline::ExecutePreImportPipelineSkeletalMesh()
 			return bFoundInstances;
 		};
 
-		PipelineMeshesUtilities->GetCombinedSkinnedMeshInstances(BaseNodeContainer, MeshUidsPerSkeletonRootUid, bConvertStaticMeshToSkeletalMesh, bConvertStaticsWithMorphTargetsToSkeletals);
+		PipelineMeshesUtilities->GetCombinedSkinnedMeshInstances(BaseNodeContainer, MeshUidsPerSkeletonRootUid, bConvertStaticMeshToSkeletalMesh, CommonSkeletalMeshesAndAnimationsProperties->bConvertStaticsWithMorphTargetsToSkeletals);
 		bool bUseMeshInstance = true;
 		bool bFoundMeshes = CreatePerSkeletonRootUidCombinedSkinnedMesh(bUseMeshInstance);
 
@@ -208,7 +208,7 @@ void UInterchangeGenericMeshPipeline::ExecutePreImportPipelineSkeletalMesh()
 				SkeletonRootUid = (bUseInstanceMesh ? PipelineMeshesUtilities->GetMeshInstanceSkeletonRootUid(MeshUid) : PipelineMeshesUtilities->GetMeshGeometrySkeletonRootUid(MeshUid));
 				if (SkeletonRootUid.IsEmpty())
 				{
-					if (bConvertStaticsWithMorphTargetsToSkeletals)
+					if (CommonSkeletalMeshesAndAnimationsProperties->bConvertStaticsWithMorphTargetsToSkeletals)
 					{
 						SkeletonRootUid = MeshUid;
 					}
@@ -252,7 +252,7 @@ void UInterchangeGenericMeshPipeline::ExecutePreImportPipelineSkeletalMesh()
 			return bFoundInstances;
 		};
 		
-		PipelineMeshesUtilities->GetAllSkinnedMeshInstance(MeshUids, bConvertStaticMeshToSkeletalMesh, bConvertStaticsWithMorphTargetsToSkeletals);
+		PipelineMeshesUtilities->GetAllSkinnedMeshInstance(MeshUids, bConvertStaticMeshToSkeletalMesh, CommonSkeletalMeshesAndAnimationsProperties->bConvertStaticsWithMorphTargetsToSkeletals);
 		bool bUseMeshInstance = true;
 		bool bFoundMeshes = CreatePerSkeletonRootUidSkinnedMesh(bUseMeshInstance);
 
@@ -437,7 +437,7 @@ UInterchangeSkeletalMeshFactoryNode* UInterchangeGenericMeshPipeline::CreateSkel
 
 		//TODO: support skeleton helper in runtime
 #if WITH_EDITOR
-		bSkeletonCompatible = UE::Interchange::Private::FSkeletonHelper::IsCompatibleSkeleton(CommonSkeletalMeshesAndAnimationsProperties->Skeleton.Get(), RootJointNode->GetUniqueID(), BaseNodeContainer);
+		bSkeletonCompatible = UE::Interchange::Private::FSkeletonHelper::IsCompatibleSkeleton(CommonSkeletalMeshesAndAnimationsProperties->Skeleton.Get(), RootJointNode->GetUniqueID(), BaseNodeContainer, CommonSkeletalMeshesAndAnimationsProperties->bConvertStaticsWithMorphTargetsToSkeletals || CommonMeshesProperties->ForceAllMeshAsType == EInterchangeForceMeshType::IFMT_SkeletalMesh);
 #endif
 		if (bSkeletonCompatible)
 		{
diff --git a/Engine/Plugins/Interchange/Runtime/Source/Pipelines/Private/InterchangeGenericStaticMeshPipeline.cpp b/Engine/Plugins/Interchange/Runtime/Source/Pipelines/Private/InterchangeGenericStaticMeshPipeline.cpp
--- a/Engine/Plugins/Interchange/Runtime/Source/Pipelines/Private/InterchangeGenericStaticMeshPipeline.cpp
+++ b/Engine/Plugins/Interchange/Runtime/Source/Pipelines/Private/InterchangeGenericStaticMeshPipeline.cpp
@@ -168,7 +168,7 @@ void UInterchangeGenericMeshPipeline::ExecutePreImportPipelineStaticMesh()
 			{
 				// If baking transforms, get all the static mesh instance nodes, and group them by LOD
 				TArray<FString> MeshUids;
-				PipelineMeshesUtilities->GetAllStaticMeshInstance(MeshUids, bConvertSkeletalMeshToStaticMesh, bConvertStaticsWithMorphTargetsToSkeletals);
+				PipelineMeshesUtilities->GetAllStaticMeshInstance(MeshUids, bConvertSkeletalMeshToStaticMesh, CommonSkeletalMeshesAndAnimationsProperties->bConvertStaticsWithMorphTargetsToSkeletals);
 
 				TMap<int32, TArray<FString>> MeshUidsPerLodIndex;
 
@@ -201,7 +201,7 @@ void UInterchangeGenericMeshPipeline::ExecutePreImportPipelineStaticMesh()
 			{
 				// If we haven't yet managed to build a factory node, look at static mesh geometry directly.
 				TArray<FString> MeshUids;
-				PipelineMeshesUtilities->GetAllStaticMeshGeometry(MeshUids, bConvertSkeletalMeshToStaticMesh, bConvertStaticsWithMorphTargetsToSkeletals);
+				PipelineMeshesUtilities->GetAllStaticMeshGeometry(MeshUids, bConvertSkeletalMeshToStaticMesh, CommonSkeletalMeshesAndAnimationsProperties->bConvertStaticsWithMorphTargetsToSkeletals);
 
 				TMap<int32, TArray<FString>> MeshUidsPerLodIndex;
 
@@ -229,7 +229,7 @@ void UInterchangeGenericMeshPipeline::ExecutePreImportPipelineStaticMesh()
 			bool bFoundMeshes = false;
 			{
 				TArray<FString> MeshUids;
-				PipelineMeshesUtilities->GetAllStaticMeshInstance(MeshUids, bConvertSkeletalMeshToStaticMesh, bConvertStaticsWithMorphTargetsToSkeletals);
+				PipelineMeshesUtilities->GetAllStaticMeshInstance(MeshUids, bConvertSkeletalMeshToStaticMesh, CommonSkeletalMeshesAndAnimationsProperties->bConvertStaticsWithMorphTargetsToSkeletals);
 
 				// Work out which meshes are collision meshes which correspond to another mesh
 				TMap<FString, TArray<FString>> MeshToCollisionMeshMap;
@@ -282,7 +282,7 @@ void UInterchangeGenericMeshPipeline::ExecutePreImportPipelineStaticMesh()
 			if (!bFoundMeshes)
 			{
 				TArray<FString> MeshUids;
-				PipelineMeshesUtilities->GetAllStaticMeshGeometry(MeshUids, bConvertSkeletalMeshToStaticMesh, bConvertStaticsWithMorphTargetsToSkeletals);
+				PipelineMeshesUtilities->GetAllStaticMeshGeometry(MeshUids, bConvertSkeletalMeshToStaticMesh, CommonSkeletalMeshesAndAnimationsProperties->bConvertStaticsWithMorphTargetsToSkeletals);
 
 				// Work out which meshes are collision meshes which correspond to another mesh
 				TMap<FString, TArray<FString>> MeshToCollisionMeshMap;
diff --git a/Engine/Plugins/Interchange/Runtime/Source/Pipelines/Public/InterchangeGenericAssetsPipelineSharedSettings.h b/Engine/Plugins/Interchange/Runtime/Source/Pipelines/Public/InterchangeGenericAssetsPipelineSharedSettings.h
--- a/Engine/Plugins/Interchange/Runtime/Source/Pipelines/Public/InterchangeGenericAssetsPipelineSharedSettings.h
+++ b/Engine/Plugins/Interchange/Runtime/Source/Pipelines/Public/InterchangeGenericAssetsPipelineSharedSettings.h
@@ -125,6 +125,10 @@ public:
 	UPROPERTY(EditAnywhere, BlueprintReadWrite, Category = "Common Skeletal Meshes and Animations")
 	bool bUseT0AsRefPose = false;
 
+	/** Allow to convert static mesh using morph target to skeletal mesh" */
+	UPROPERTY(EditAnywhere, BlueprintReadWrite, Category = "Static Meshes")
+	bool bConvertStaticsWithMorphTargetsToSkeletals = false;
+
 	virtual bool IsSettingsAreValid(TOptional<FText>& OutInvalidReason) const override
 	{
 		if (bImportOnlyAnimations && !Skeleton.IsValid())
diff --git a/Engine/Plugins/Interchange/Runtime/Source/Pipelines/Public/InterchangeGenericMeshPipeline.h b/Engine/Plugins/Interchange/Runtime/Source/Pipelines/Public/InterchangeGenericMeshPipeline.h
--- a/Engine/Plugins/Interchange/Runtime/Source/Pipelines/Public/InterchangeGenericMeshPipeline.h
+++ b/Engine/Plugins/Interchange/Runtime/Source/Pipelines/Public/InterchangeGenericMeshPipeline.h
@@ -44,10 +44,6 @@ public:
 	UPROPERTY(EditAnywhere, BlueprintReadWrite, Category = "Static Meshes")
 	bool bImportStaticMeshes = true;
 
-	/** Allow to convert static mesh using morph target to skeletal mesh" */
-	UPROPERTY(EditAnywhere, BlueprintReadWrite, Category = "Static Meshes")
-	bool bConvertStaticsWithMorphTargetsToSkeletals = false;
-
 	/** If enabled, all translated static mesh nodes will be imported as a single static mesh. */
 	UPROPERTY(EditAnywhere, BlueprintReadWrite, Category = "Static Meshes")
 	bool bCombineStaticMeshes = false;
