[Interchange] Imported morph play view difference with edit view

diff --git a/Engine/Plugins/Interchange/Runtime/Source/FactoryNodes/Private/InterchangeMeshActorFactoryNode.cpp b/Engine/Plugins/Interchange/Runtime/Source/FactoryNodes/Private/InterchangeMeshActorFactoryNode.cpp
--- a/Engine/Plugins/Interchange/Runtime/Source/FactoryNodes/Private/InterchangeMeshActorFactoryNode.cpp
+++ b/Engine/Plugins/Interchange/Runtime/Source/FactoryNodes/Private/InterchangeMeshActorFactoryNode.cpp
@@ -11,7 +11,6 @@
 UInterchangeMeshActorFactoryNode::UInterchangeMeshActorFactoryNode()
 {
 	SlotMaterialDependencies.Initialize(Attributes.ToSharedRef(), TEXT("__SlotMaterialDependencies__"));
-	MorphTargetCurveWeights.Initialize(Attributes.ToSharedRef(), TEXT("__MorphTargetCurves__Key"));
 }
 
 void UInterchangeMeshActorFactoryNode::GetSlotMaterialDependencies(TMap<FString, FString>& OutMaterialDependencies) const
@@ -34,20 +33,11 @@ bool UInterchangeMeshActorFactoryNode::RemoveSlotMaterialDependencyUid(const FSt
 	return SlotMaterialDependencies.RemoveKey(SlotName);
 }
 
-bool UInterchangeMeshActorFactoryNode::SetMorphTargetCurveWeight(const FString& MorphTargetName, const float& Weight)
+bool UInterchangeMeshActorFactoryNode::SetCustomAnimationAssetUidToPlay(const FString& AttributeValue)
 {
-	return MorphTargetCurveWeights.SetKeyValue(MorphTargetName, Weight);
+	IMPLEMENT_NODE_ATTRIBUTE_SETTER_NODELEGATE(AnimationAssetUidToPlay, FString);
 }
-
-void UInterchangeMeshActorFactoryNode::SetMorphTargetCurveWeights(const TMap<FString, float>& InMorphTargetCurveWeights)
-{
-	for (const TPair<FString, float>& MorphTargetCurve : InMorphTargetCurveWeights)
-	{
-		SetMorphTargetCurveWeight(MorphTargetCurve.Key, MorphTargetCurve.Value);
-	}
-}
-
-void UInterchangeMeshActorFactoryNode::GetMorphTargetCurveWeights(TMap<FString, float>& OutMorphTargetCurveWeights) const
+bool UInterchangeMeshActorFactoryNode::GetCustomAnimationAssetUidToPlay(FString& AttributeValue) const
 {
-	OutMorphTargetCurveWeights = MorphTargetCurveWeights.ToMap();
+	IMPLEMENT_NODE_ATTRIBUTE_GETTER(AnimationAssetUidToPlay, FString);
 }
\ No newline at end of file
diff --git a/Engine/Plugins/Interchange/Runtime/Source/FactoryNodes/Public/InterchangeMeshActorFactoryNode.h b/Engine/Plugins/Interchange/Runtime/Source/FactoryNodes/Public/InterchangeMeshActorFactoryNode.h
--- a/Engine/Plugins/Interchange/Runtime/Source/FactoryNodes/Public/InterchangeMeshActorFactoryNode.h
+++ b/Engine/Plugins/Interchange/Runtime/Source/FactoryNodes/Public/InterchangeMeshActorFactoryNode.h
@@ -51,19 +51,16 @@ public:
 	UFUNCTION(BlueprintCallable, Category = "Interchange | Node | StaticMesh")
 	bool RemoveSlotMaterialDependencyUid(const FString& SlotName);
 
-	/** Set MorphTarget with given weight. */
+	/** Set the Animation Asset To Play by this Scene Node. (only relevant for SkeletalMeshActors (SceneNodes that are instantiating Skeletal Meshes)) */
 	UFUNCTION(BlueprintCallable, Category = "Interchange | Node | SkeletalMesh")
-	bool SetMorphTargetCurveWeight(const FString& MorphTargetName, const float& Weight);
-
-	/** Set MorphTarget with given weight. */
-	UFUNCTION(BlueprintCallable, Category = "Interchange | Node | SkeletalMesh")
-	void SetMorphTargetCurveWeights(const TMap<FString, float>& InMorphTargetCurveWeights);
-
-	/** Get MorphTargets and their weights. */
+	bool SetCustomAnimationAssetUidToPlay(const FString& AttributeValue);
+	/** Get the Animation Asset To Play by this Scene Node. */
 	UFUNCTION(BlueprintCallable, Category = "Interchange | Node | SkeletalMesh")
-	void GetMorphTargetCurveWeights(TMap<FString, float>& OutMorphTargetCurveWeights) const;
+	bool GetCustomAnimationAssetUidToPlay(FString& AttributeValue) const;
 
 private:
 	UE::Interchange::TMapAttributeHelper<FString, FString> SlotMaterialDependencies;
-	UE::Interchange::TMapAttributeHelper<FString, float> MorphTargetCurveWeights;
+
+	//A scene node can reference an animation asset on top of base asset:
+	const UE::Interchange::FAttributeKey Macro_CustomAnimationAssetUidToPlayKey = UE::Interchange::FAttributeKey(TEXT("AnimationAssetUidToPlay"));
 };
\ No newline at end of file
diff --git a/Engine/Plugins/Interchange/Runtime/Source/Import/Private/Animation/InterchangeAnimSequenceFactory.cpp b/Engine/Plugins/Interchange/Runtime/Source/Import/Private/Animation/InterchangeAnimSequenceFactory.cpp
--- a/Engine/Plugins/Interchange/Runtime/Source/Import/Private/Animation/InterchangeAnimSequenceFactory.cpp
+++ b/Engine/Plugins/Interchange/Runtime/Source/Import/Private/Animation/InterchangeAnimSequenceFactory.cpp
@@ -561,7 +561,37 @@ namespace UE::Interchange::Private
 					}
 					if (const UInterchangeMeshNode* MorphTargetNode = Cast<UInterchangeMeshNode>(NodeContainer->GetNode(MorphTargetNodeUidAnimationPayload.Key)))
 					{
-						AnimationCurvesPayloads.Add(PayloadKey, AnimSequenceTranslatorPayloadInterface->GetAnimationPayloadData(MorphTargetNodeUidAnimationPayload.Value));
+						if (MorphTargetNodeUidAnimationPayload.Value.Type == EInterchangeAnimationPayLoadType::MORPHTARGETCURVEWEIGHTINSTANCE)
+						{
+							AnimationCurvesPayloads.Add(PayloadKey, Async(EAsyncExecution::TaskGraph, [&MorphTargetNodeUidAnimationPayload]
+								{
+									TOptional<UE::Interchange::FAnimationPayloadData> Result;
+									UE::Interchange::FAnimationPayloadData AnimationPayLoadData(MorphTargetNodeUidAnimationPayload.Value.Type);
+
+									TArray<FString> PayLoadKeys;
+									MorphTargetNodeUidAnimationPayload.Value.UniqueId.ParseIntoArray(PayLoadKeys, TEXT(":"));
+
+									if (PayLoadKeys.Num() != 2)
+									{
+										return Result;
+									}
+
+									float Weight;
+									LexFromString(Weight, *PayLoadKeys[1]);
+
+									AnimationPayLoadData.Curves.SetNum(1);
+									FRichCurve& Curve = AnimationPayLoadData.Curves[0];
+									Curve.AddKey(0, Weight);
+
+									Result.Emplace(AnimationPayLoadData);
+
+									return Result;
+								}));
+						}
+						else
+						{
+							AnimationCurvesPayloads.Add(PayloadKey, AnimSequenceTranslatorPayloadInterface->GetAnimationPayloadData(MorphTargetNodeUidAnimationPayload.Value));
+						}
 						AnimationCurveMorphTargetNodeNames.Add(PayloadKey, MorphTargetNode->GetDisplayLabel());
 					}
 				}
diff --git a/Engine/Plugins/Interchange/Runtime/Source/Import/Private/Scene/InterchangeSkeletalMeshActorFactory.cpp b/Engine/Plugins/Interchange/Runtime/Source/Import/Private/Scene/InterchangeSkeletalMeshActorFactory.cpp
--- a/Engine/Plugins/Interchange/Runtime/Source/Import/Private/Scene/InterchangeSkeletalMeshActorFactory.cpp
+++ b/Engine/Plugins/Interchange/Runtime/Source/Import/Private/Scene/InterchangeSkeletalMeshActorFactory.cpp
@@ -2,6 +2,7 @@
 #include "Scene/InterchangeSkeletalMeshActorFactory.h"
 
 #include "InterchangeMeshActorFactoryNode.h"
+#include "InterchangeAnimSequenceFactoryNode.h"
 #include "Scene/InterchangeActorHelper.h"
 #include "InterchangeSceneNode.h"
 #include "InterchangeSkeletalMeshFactoryNode.h"
@@ -59,12 +60,26 @@ void UInterchangeSkeletalMeshActorFactory::SetupObject_GameThread(const FSetupOb
 					{
 						UE::Interchange::ActorHelper::ApplySlotMaterialDependencies(*Arguments.NodeContainer, *MeshActorFactoryNode, *SkeletalMeshComponent);
 
-						//Set Morph Target Curves:
-						TMap<FString, float> MorphTargetCurveWeights;
-						MeshActorFactoryNode->GetMorphTargetCurveWeights(MorphTargetCurveWeights);
-						for (const TPair<FString, float>& MorphTargetCurve : MorphTargetCurveWeights)
+						FString AnimationAssetUidToPlay;
+						if (MeshActorFactoryNode->GetCustomAnimationAssetUidToPlay(AnimationAssetUidToPlay))
 						{
-							SkeletalMeshComponent->SetMorphTarget(*MorphTargetCurve.Key, MorphTargetCurve.Value);
+							const FString AnimSequenceFactoryNodeUid = TEXT("\\AnimSequence") + AnimationAssetUidToPlay;
+							if (const UInterchangeAnimSequenceFactoryNode* AnimSequenceFactoryNode = Cast<UInterchangeAnimSequenceFactoryNode>(Arguments.NodeContainer->GetFactoryNode(AnimSequenceFactoryNodeUid)))
+							{
+								FSoftObjectPath AnimSequenceObject;
+								if (AnimSequenceFactoryNode->GetCustomReferenceObject(AnimSequenceObject))
+								{
+									if (UAnimSequence* AnimSequence = Cast<UAnimSequence>(AnimSequenceObject.TryLoad()))
+									{
+										SkeletalMeshComponent->SetAnimationMode(EAnimationMode::AnimationSingleNode);
+										SkeletalMeshComponent->AnimationData.AnimToPlay = AnimSequence;
+										SkeletalMeshComponent->AnimationData.bSavedLooping = false;
+										SkeletalMeshComponent->AnimationData.bSavedPlaying = false;
+										SkeletalMeshComponent->SetAnimation(AnimSequence);
+										
+									}
+								}
+							}
 						}
 					}
 				}
diff --git a/Engine/Plugins/Interchange/Runtime/Source/Nodes/Private/InterchangeSceneNode.cpp b/Engine/Plugins/Interchange/Runtime/Source/Nodes/Private/InterchangeSceneNode.cpp
--- a/Engine/Plugins/Interchange/Runtime/Source/Nodes/Private/InterchangeSceneNode.cpp
+++ b/Engine/Plugins/Interchange/Runtime/Source/Nodes/Private/InterchangeSceneNode.cpp
@@ -392,4 +392,13 @@ bool UInterchangeSceneNode::SetMorphTargetCurveWeight(const FString& MorphTarget
 void UInterchangeSceneNode::GetMorphTargetCurveWeights(TMap<FString, float>& OutMorphTargetCurveWeights) const
 {
 	OutMorphTargetCurveWeights = MorphTargetCurveWeights.ToMap();
+}
+
+bool UInterchangeSceneNode::SetCustomAnimationAssetUidToPlay(const FString& AttributeValue)
+{
+	IMPLEMENT_NODE_ATTRIBUTE_SETTER_NODELEGATE(AnimationAssetUidToPlay, FString);
+}
+bool UInterchangeSceneNode::GetCustomAnimationAssetUidToPlay(FString& AttributeValue) const
+{
+	IMPLEMENT_NODE_ATTRIBUTE_GETTER(AnimationAssetUidToPlay, FString);
 }
\ No newline at end of file
diff --git a/Engine/Plugins/Interchange/Runtime/Source/Nodes/Public/InterchangeAnimationTrackSetNode.h b/Engine/Plugins/Interchange/Runtime/Source/Nodes/Public/InterchangeAnimationTrackSetNode.h
--- a/Engine/Plugins/Interchange/Runtime/Source/Nodes/Public/InterchangeAnimationTrackSetNode.h
+++ b/Engine/Plugins/Interchange/Runtime/Source/Nodes/Public/InterchangeAnimationTrackSetNode.h
@@ -14,7 +14,14 @@ enum class EInterchangeAnimationPayLoadType : uint8
 	CURVE,
 	MORPHTARGETCURVE, //handles/generates the same properties as the CURVE variation, but they way it is acquired might be different (depending on the Translator)
 	STEPCURVE,
-	BAKED
+	BAKED,
+	MORPHTARGETCURVEWEIGHTINSTANCE	//Handled within UInterchangeAnimSequenceFactory, contrary to the others which are handled in the Translators
+									// Purpose of this type is to generate a 1 frame long Animation with the Instantiated Morph Target Curve Weights.
+									//		This is needed for the special use cases where the imported 3d file format has a concept of MorphTargetWeight usages on StaticMeshes.
+									//		UE5 does not have support for this particular concept, and the workaround is to create a 1 frame long animation with the desired MorphTargetWeight settings.
+									//		This is also needed for the use cases where the import is a Level import and the 3d file format is instantiating a StaticMesh with a particular MorphTargetWeight settings.
+									//			This is done through the AnimSequnce being used on the SkeletalMeshActor's instance.
+									//		Related Ticket: UE-186102
 };
 
 
diff --git a/Engine/Plugins/Interchange/Runtime/Source/Nodes/Public/InterchangeSceneNode.h b/Engine/Plugins/Interchange/Runtime/Source/Nodes/Public/InterchangeSceneNode.h
--- a/Engine/Plugins/Interchange/Runtime/Source/Nodes/Public/InterchangeSceneNode.h
+++ b/Engine/Plugins/Interchange/Runtime/Source/Nodes/Public/InterchangeSceneNode.h
@@ -204,6 +204,14 @@ public:
 	UFUNCTION(BlueprintCallable, Category = "Interchange | Node | SkeletalMesh")
 	void GetMorphTargetCurveWeights(TMap<FString, float>& OutMorphTargetCurveWeights) const;
 
+
+	/** Set the Animation Asset To Play by this Scene Node. (only relevant for SkeletalMeshActors (SceneNodes that are instantiating Skeletal Meshes)) */
+	UFUNCTION(BlueprintCallable, Category = "Interchange | Node | SkeletalMesh")
+	bool SetCustomAnimationAssetUidToPlay(const FString& AttributeValue);
+	/** Get the Animation Asset To Play by this Scene Node. */
+	UFUNCTION(BlueprintCallable, Category = "Interchange | Node | SkeletalMesh")
+	bool GetCustomAnimationAssetUidToPlay(FString& AttributeValue) const;
+
 private:
 
 	bool GetGlobalTransformInternal(const UE::Interchange::FAttributeKey LocalTransformKey
@@ -237,6 +245,9 @@ private:
 	//A scene node can have different MorphTarget curve settings:
 	UE::Interchange::TMapAttributeHelper<FString, float> MorphTargetCurveWeights;
 
+	//A scene node can reference an animation asset on top of base asset:
+	const UE::Interchange::FAttributeKey Macro_CustomAnimationAssetUidToPlayKey = UE::Interchange::FAttributeKey(TEXT("AnimationAssetUidToPlay"));
+
 	//mutable caches for global transforms
 	mutable TOptional<FTransform> CacheGlobalTransform;
 	mutable TOptional<FTransform> CacheBindPoseGlobalTransform;
diff --git a/Engine/Plugins/Interchange/Runtime/Source/Pipelines/Private/InterchangeGenericAnimationPipeline.cpp b/Engine/Plugins/Interchange/Runtime/Source/Pipelines/Private/InterchangeGenericAnimationPipeline.cpp
--- a/Engine/Plugins/Interchange/Runtime/Source/Pipelines/Private/InterchangeGenericAnimationPipeline.cpp
+++ b/Engine/Plugins/Interchange/Runtime/Source/Pipelines/Private/InterchangeGenericAnimationPipeline.cpp
@@ -82,6 +82,50 @@ void UInterchangeGenericAnimationPipeline::ExecutePipeline(UInterchangeBaseNodeC
 			TrackSetNodes.Add(Node);
 		});
 
+
+	//Create AnimSequences(UInterchangeSkeletalAnimationTrackNode) for Mesh Instances having MorphTargetCurveWeights:
+	{
+		TArray<UInterchangeSceneNode*> SceneNodesWithMorphTargetCurveWeights;
+		BaseNodeContainer->IterateNodesOfType<UInterchangeSceneNode>([&SceneNodesWithMorphTargetCurveWeights](const FString& NodeUid, UInterchangeSceneNode* SceneNode)
+			{
+				TMap<FString, float> MorphTargetCurveWeights;
+				SceneNode->GetMorphTargetCurveWeights(MorphTargetCurveWeights);
+
+				if (MorphTargetCurveWeights.Num() > 0)
+				{
+					SceneNodesWithMorphTargetCurveWeights.Add(SceneNode);
+				}
+			});
+
+		for (UInterchangeSceneNode* SceneNode : SceneNodesWithMorphTargetCurveWeights)
+		{
+			UInterchangeSkeletalAnimationTrackNode* SkeletalAnimationNode = NewObject< UInterchangeSkeletalAnimationTrackNode >(BaseNodeContainer);
+			FString SkeletalAnimationNodeUid = "\\SkeletalAnimation\\MorphTargetCurveWeightInstantation\\" + SceneNode->GetUniqueID();
+			SkeletalAnimationNode->InitializeNode(SkeletalAnimationNodeUid, SceneNode->GetDisplayLabel(), EInterchangeNodeContainerType::TranslatedAsset);
+
+			SkeletalAnimationNode->SetCustomAnimationSampleRate(30.f);
+			SkeletalAnimationNode->SetCustomAnimationStartTime(0);
+			SkeletalAnimationNode->SetCustomAnimationStopTime(1.f / 30.f); //we want a single frame
+
+			SkeletalAnimationNode->SetCustomSkeletonNodeUid(SceneNode->GetUniqueID());
+
+			TMap<FString, float> MorphTargetCurveWeights;
+			SceneNode->GetMorphTargetCurveWeights(MorphTargetCurveWeights);
+
+			for (const TPair<FString, float>& MorphTargetCurveWeight : MorphTargetCurveWeights)
+			{
+				FString PayloadUid = MorphTargetCurveWeight.Key + TEXT(":") + LexToString(MorphTargetCurveWeight.Value);
+
+				//add the payload key:
+				SkeletalAnimationNode->SetAnimationPayloadKeyForMorphTargetNodeUid(MorphTargetCurveWeight.Key, PayloadUid, EInterchangeAnimationPayLoadType::MORPHTARGETCURVEWEIGHTINSTANCE);
+			}
+
+			BaseNodeContainer->AddNode(SkeletalAnimationNode);
+
+			SceneNode->SetCustomAnimationAssetUidToPlay(SkeletalAnimationNodeUid);
+		}
+	}
+
 	if (CommonMeshesProperties->ForceAllMeshAsType == EInterchangeForceMeshType::IFMT_SkeletalMesh)
 	{
 		for (UInterchangeAnimationTrackSetNode* TrackSetNode : TrackSetNodes)
diff --git a/Engine/Plugins/Interchange/Runtime/Source/Pipelines/Private/InterchangeGenericScenesPipeline.cpp b/Engine/Plugins/Interchange/Runtime/Source/Pipelines/Private/InterchangeGenericScenesPipeline.cpp
--- a/Engine/Plugins/Interchange/Runtime/Source/Pipelines/Private/InterchangeGenericScenesPipeline.cpp
+++ b/Engine/Plugins/Interchange/Runtime/Source/Pipelines/Private/InterchangeGenericScenesPipeline.cpp
@@ -403,9 +403,11 @@ void UInterchangeGenericLevelPipeline::SetUpFactoryNode(UInterchangeActorFactory
 
 			UE::Interchange::MeshesUtilities::ApplySlotMaterialDependencies(*MeshActorFactoryNode, SlotMaterialDependencies, *BaseNodeContainer);
 
-			TMap<FString, float> MorphTargetCurveWeights;
-			SceneNode->GetMorphTargetCurveWeights(MorphTargetCurveWeights);
-			MeshActorFactoryNode->SetMorphTargetCurveWeights(MorphTargetCurveWeights);
+			FString AnimationAssetUidToPlay;
+			if (SceneNode->GetCustomAnimationAssetUidToPlay(AnimationAssetUidToPlay))
+			{
+				MeshActorFactoryNode->SetCustomAnimationAssetUidToPlay(AnimationAssetUidToPlay);
+			}
 
 			MeshActorFactoryNode->AddFactoryDependencyUid(UInterchangeFactoryBaseNode::BuildFactoryNodeUid(MeshNode->GetUniqueID()));
 		}
