Build health fix for InterchangeAnimSequenceFactory.cpp

diff --git a/Engine/Plugins/Interchange/Runtime/Source/Import/Private/Animation/InterchangeAnimSequenceFactory.cpp b/Engine/Plugins/Interchange/Runtime/Source/Import/Private/Animation/InterchangeAnimSequenceFactory.cpp
--- a/Engine/Plugins/Interchange/Runtime/Source/Import/Private/Animation/InterchangeAnimSequenceFactory.cpp
+++ b/Engine/Plugins/Interchange/Runtime/Source/Import/Private/Animation/InterchangeAnimSequenceFactory.cpp
@@ -20,6 +20,7 @@
 #include "Nodes/InterchangeSourceNode.h"
 #include "Nodes/InterchangeUserDefinedAttribute.h"
 #include "InterchangeAnimationTrackSetNode.h"
+#include "Async/Async.h"
 
 #include UE_INLINE_GENERATED_CPP_BY_NAME(InterchangeAnimSequenceFactory)
 
