Fix box collision import rotation when box are orthonormal to axis system.

#rnx

diff --git a/Engine/Plugins/Interchange/Runtime/Source/Import/Private/Mesh/InterchangeStaticMeshFactory.cpp b/Engine/Plugins/Interchange/Runtime/Source/Import/Private/Mesh/InterchangeStaticMeshFactory.cpp
--- a/Engine/Plugins/Interchange/Runtime/Source/Import/Private/Mesh/InterchangeStaticMeshFactory.cpp
+++ b/Engine/Plugins/Interchange/Runtime/Source/Import/Private/Mesh/InterchangeStaticMeshFactory.cpp
@@ -1057,7 +1057,7 @@ bool UInterchangeStaticMeshFactory::AddBoxGeomFromTris(const FImportAssetObjectP
 					Planes[PlaneIndex].DistCount = 2;
 				}
 				// if we have a second distance, and its not that either, something is wrong.
-				else if (Planes[PlaneIndex].DistCount == 2 && !AreEqual(Dist, Planes[PlaneIndex].PlaneDist[1]))
+				else if (Planes[PlaneIndex].DistCount == 2 && !AreEqual(Dist, Planes[PlaneIndex].PlaneDist[0]) && !AreEqual(Dist, Planes[PlaneIndex].PlaneDist[1]))
 				{
 					// Error
 //					UE_LOG(LogStaticMeshImportUtils, Log, TEXT("AddBoxGeomFromTris (%s): Found more than 2 planes with different distances."), ObjName);
@@ -1110,12 +1110,37 @@ bool UInterchangeStaticMeshFactory::AddBoxGeomFromTris(const FImportAssetObjectP
 
 	// Allocate box in array
 	FKBoxElem BoxElem;
-	BoxElem.SetTransform(FTransform(FVector(Planes[0].Normal), FVector(Planes[1].Normal), FVector(Planes[2].Normal), Box.GetCenter()));
+
+	//In case we have a box oriented with the world axis system we want to reorder the plane to not introduce axis swap.
+	//If the box was turned, the order of the planes will be arbitrary and the box rotation will make the collision
+	//not playing well if the asset is built or place in a level with a non uniform scale.
+	FVector3f Axis[3] = { FVector3f::XAxisVector, FVector3f::YAxisVector, FVector3f::ZAxisVector };
+	int32 Reorder[3] = { INDEX_NONE, INDEX_NONE, INDEX_NONE };
+	for (int32 PlaneIndex = 0; PlaneIndex < 3; ++PlaneIndex)
+	{
+		for (int32 AxisIndex = 0; AxisIndex < 3; ++AxisIndex)
+		{
+			if (AreParallel(Planes[PlaneIndex].Normal, Axis[AxisIndex]))
+			{
+				Reorder[PlaneIndex] = AxisIndex;
+				break;
+			}
+		}
+	}
+
+	if (Reorder[0] == INDEX_NONE || Reorder[1] == INDEX_NONE || Reorder[2] == INDEX_NONE)
+	{
+		Reorder[0] = 0;
+		Reorder[1] = 1;
+		Reorder[2] = 2;
+	}
+
+	BoxElem.SetTransform(FTransform(FVector(Planes[Reorder[0]].Normal), FVector(Planes[Reorder[1]].Normal), FVector(Planes[Reorder[2]].Normal), Box.GetCenter()));
 
 	// distance between parallel planes is box edge lengths.
-	BoxElem.X = FMath::Abs(Planes[0].PlaneDist[0] - Planes[0].PlaneDist[1]);
-	BoxElem.Y = FMath::Abs(Planes[1].PlaneDist[0] - Planes[1].PlaneDist[1]);
-	BoxElem.Z = FMath::Abs(Planes[2].PlaneDist[0] - Planes[2].PlaneDist[1]);
+	BoxElem.X = FMath::Abs(Planes[Reorder[0]].PlaneDist[0] - Planes[Reorder[0]].PlaneDist[1]);
+	BoxElem.Y = FMath::Abs(Planes[Reorder[1]].PlaneDist[0] - Planes[Reorder[1]].PlaneDist[1]);
+	BoxElem.Z = FMath::Abs(Planes[Reorder[2]].PlaneDist[0] - Planes[Reorder[2]].PlaneDist[1]);
 
 	AggGeom.BoxElems.Add(BoxElem);
 
diff --git a/Engine/Source/Editor/UnrealEd/Private/ImportUtils/StaticMeshImportUtils.cpp b/Engine/Source/Editor/UnrealEd/Private/ImportUtils/StaticMeshImportUtils.cpp
--- a/Engine/Source/Editor/UnrealEd/Private/ImportUtils/StaticMeshImportUtils.cpp
+++ b/Engine/Source/Editor/UnrealEd/Private/ImportUtils/StaticMeshImportUtils.cpp
@@ -405,13 +405,34 @@ bool StaticMeshImportUtils::AddBoxGeomFromTris( const TArray<FPoly>& Tris, FKAgg
 
 	FMatrix BoxTM = FMatrix::Identity;
 
-	BoxTM.SetAxis(0, (FVector)Planes[0].Normal);
-	BoxTM.SetAxis(1, (FVector)Planes[1].Normal);
+	FVector3f Axis[3] = {{1.0f, 0.0f, 0.0f}, {0.0f, 1.0f, 0.0f}, {0.0f, 0.0f, 1.0f}};
+	int32 Reorder[3] = {INDEX_NONE, INDEX_NONE, INDEX_NONE };
+	for (int32 PlaneIndex = 0; PlaneIndex < 3; ++PlaneIndex)
+	{
+		for (int32 AxisIndex = 0; AxisIndex < 3; ++AxisIndex)
+		{
+			if(AreParallel(Planes[PlaneIndex].Normal, Axis[AxisIndex]))
+			{
+				Reorder[PlaneIndex] = AxisIndex;
+				break;
+			}
+		}
+	}
+	
+	if (Reorder[0] == INDEX_NONE || Reorder[1] == INDEX_NONE || Reorder[2] == INDEX_NONE)
+	{
+		Reorder[0] = 0;
+		Reorder[1] = 1;
+		Reorder[2] = 2;
+	}
+
+	BoxTM.SetAxis(0, (FVector)Planes[Reorder[0]].Normal);
+	BoxTM.SetAxis(1, (FVector)Planes[Reorder[1]].Normal);
 
 	// ensure valid TM by cross-product
-	FVector3f ZAxis = Planes[0].Normal ^ Planes[1].Normal;
+	FVector3f ZAxis = Planes[Reorder[0]].Normal ^ Planes[Reorder[1]].Normal;
 
-	if( !AreParallel(ZAxis, Planes[2].Normal) )
+	if (!AreParallel(ZAxis, Planes[Reorder[2]].Normal))
 	{
 		UE_LOG(LogStaticMeshImportUtils, Log, TEXT("AddBoxGeomFromTris (%s): Box axes are not perpendicular."), ObjName);
 		return false;
@@ -434,9 +455,9 @@ bool StaticMeshImportUtils::AddBoxGeomFromTris( const TArray<FPoly>& Tris, FKAgg
 	FKBoxElem BoxElem;
 	BoxElem.SetTransform( FTransform( BoxTM ) );	
 	// distance between parallel planes is box edge lengths.
-	BoxElem.X = FMath::Abs(Planes[0].PlaneDist[0] - Planes[0].PlaneDist[1]);
-	BoxElem.Y = FMath::Abs(Planes[1].PlaneDist[0] - Planes[1].PlaneDist[1]);
-	BoxElem.Z = FMath::Abs(Planes[2].PlaneDist[0] - Planes[2].PlaneDist[1]);
+	BoxElem.X = FMath::Abs(Planes[Reorder[0]].PlaneDist[0] - Planes[Reorder[0]].PlaneDist[1]);
+	BoxElem.Y = FMath::Abs(Planes[Reorder[1]].PlaneDist[0] - Planes[Reorder[1]].PlaneDist[1]);
+	BoxElem.Z = FMath::Abs(Planes[Reorder[2]].PlaneDist[0] - Planes[Reorder[2]].PlaneDist[1]);
 	AggGeom->BoxElems.Add(BoxElem);
 
 	return true;
