Make sure we can import animation only. A file with only a skeleton should be able to import if the specified skeleton match the imported one.

#rnx

diff --git a/Engine/Plugins/Interchange/Runtime/Source/Import/Private/Animation/InterchangeAnimSequenceFactory.cpp b/Engine/Plugins/Interchange/Runtime/Source/Import/Private/Animation/InterchangeAnimSequenceFactory.cpp
--- a/Engine/Plugins/Interchange/Runtime/Source/Import/Private/Animation/InterchangeAnimSequenceFactory.cpp
+++ b/Engine/Plugins/Interchange/Runtime/Source/Import/Private/Animation/InterchangeAnimSequenceFactory.cpp
@@ -4,14 +4,14 @@
 #include "Animation/AnimSequence.h"
 #include "Animation/BuiltInAttributeTypes.h"
 #include "Animation/InterchangeAnimationPayloadInterface.h"
-#include "InterchangeCommonPipelineDataFactoryNode.h"
+#include "InterchangeAnimSequenceFactoryNode.h"
 #include "InterchangeAssetImportData.h"
 #include "InterchangeCommonAnimationPayload.h"
+#include "InterchangeCommonPipelineDataFactoryNode.h"
 #include "InterchangeImportCommon.h"
 #include "InterchangeImportLog.h"
 #include "InterchangeMeshNode.h"
 #include "InterchangeSceneNode.h"
-#include "InterchangeAnimSequenceFactoryNode.h"
 #include "InterchangeSkeletonFactoryNode.h"
 #include "InterchangeSourceData.h"
 #include "InterchangeTranslatorBase.h"
diff --git a/Engine/Plugins/Interchange/Runtime/Source/Parsers/Fbx/Private/FbxScene.cpp b/Engine/Plugins/Interchange/Runtime/Source/Parsers/Fbx/Private/FbxScene.cpp
--- a/Engine/Plugins/Interchange/Runtime/Source/Parsers/Fbx/Private/FbxScene.cpp
+++ b/Engine/Plugins/Interchange/Runtime/Source/Parsers/Fbx/Private/FbxScene.cpp
@@ -655,6 +655,45 @@ namespace UE
 							SkeletonRootNodeUids.Add(JointNodeUid);
 						}
 					});
+				
+				//In case we import animation only and there is no meshes
+				if (SkeletonRootNodeUids.Num() == 0)
+				{
+					NodeContainer.IterateNodesOfType<UInterchangeSceneNode>([&](const FString& NodeUid, UInterchangeSceneNode* SceneNode)
+						{
+							if (SceneNode->IsSpecializedTypeContains(FSceneNodeStaticData::GetJointSpecializeTypeString()))
+							{
+								//Find the root joint for this MeshGeometry
+								FString JointNodeUid = NodeUid;
+								FString ParentNodeUid = SceneNode->GetParentUid();
+
+								while (!JointNodeUid.Equals(UInterchangeBaseNode::InvalidNodeUid()))
+								{
+									if (const UInterchangeSceneNode* Node = Cast< UInterchangeSceneNode >(NodeContainer.GetNode(ParentNodeUid)))
+									{
+										if (Node->IsSpecializedTypeContains(FSceneNodeStaticData::GetJointSpecializeTypeString()))
+										{
+											JointNodeUid = ParentNodeUid;
+											ParentNodeUid = Node->GetParentUid();
+										}
+										else
+										{
+											break;
+										}
+									}
+									else
+									{
+										break;
+									}
+								}
+
+								if (!JointNodeUid.Equals(UInterchangeBaseNode::InvalidNodeUid()))
+								{
+									SkeletonRootNodeUids.Add(JointNodeUid);
+								}
+							}
+						});
+				}
 
 				int32 NumAnimations = SDKScene->GetSrcObjectCount<FbxAnimStack>();
 
diff --git a/Engine/Plugins/Interchange/Runtime/Source/Pipelines/Private/InterchangeGenericAnimationPipeline.cpp b/Engine/Plugins/Interchange/Runtime/Source/Pipelines/Private/InterchangeGenericAnimationPipeline.cpp
--- a/Engine/Plugins/Interchange/Runtime/Source/Pipelines/Private/InterchangeGenericAnimationPipeline.cpp
+++ b/Engine/Plugins/Interchange/Runtime/Source/Pipelines/Private/InterchangeGenericAnimationPipeline.cpp
@@ -298,9 +298,49 @@ void UInterchangeGenericAnimationPipeline::CreateAnimSequenceFactoryNode(UInterc
 		// TODO: Warn something wrong happened
 		return;
 	}
+	const bool bImportOnlyAnimation = CommonSkeletalMeshesAndAnimationsProperties->bImportOnlyAnimations;
+
+	const UInterchangeSkeletonFactoryNode* SkeletonFactoryNode = nullptr;
+	const UInterchangeSkeletalMeshFactoryNode* SkeletalMeshFactoryNode = nullptr;
 
 	const FString SkeletonFactoryNodeUid = UInterchangeFactoryBaseNode::BuildFactoryNodeUid(SkeletonNodeUid);
-	const UInterchangeSkeletonFactoryNode* SkeletonFactoryNode = Cast<const UInterchangeSkeletonFactoryNode>(BaseNodeContainer->GetFactoryNode(SkeletonFactoryNodeUid));
+	SkeletonFactoryNode = Cast<const UInterchangeSkeletonFactoryNode>(BaseNodeContainer->GetFactoryNode(SkeletonFactoryNodeUid));
+
+	//If we import anim only and we do not have meshes and skeleton. We need to create a skeleton factory node
+	//base on the specified skeleton
+	if(bImportOnlyAnimation && !SkeletonFactoryNode && CommonSkeletalMeshesAndAnimationsProperties->Skeleton.IsValid())
+	{
+		const FReferenceSkeleton& ReferenceSkeleton = CommonSkeletalMeshesAndAnimationsProperties->Skeleton->GetReferenceSkeleton();
+		TArray<FString> SkeletonRootNodeUids;
+		BaseNodeContainer->IterateNodesOfType<UInterchangeSceneNode>([&SkeletonRootNodeUids, BaseNodeContainerClosure = BaseNodeContainer, &ReferenceSkeleton](const FString& NodeUid, UInterchangeSceneNode* Node)
+		{
+			if (Node->IsSpecializedTypeContains(UE::Interchange::FSceneNodeStaticData::GetJointSpecializeTypeString()))
+			{
+				if(ReferenceSkeleton.FindBoneIndex(FName(*Node->GetDisplayLabel())) != INDEX_NONE)
+				{
+					const FString ParentUid = Node->GetParentUid();
+					if (const UInterchangeSceneNode* ParentNode = Cast<UInterchangeSceneNode>(BaseNodeContainerClosure->GetNode(ParentUid)))
+					{
+						if (!ParentNode->IsSpecializedTypeContains(UE::Interchange::FSceneNodeStaticData::GetJointSpecializeTypeString()))
+						{
+							SkeletonRootNodeUids.Add(NodeUid);
+						}
+					}
+				}
+			}
+		});
+		FString SkeletonRootUid;
+		if (SkeletonRootNodeUids.Num() > 0)
+		{
+			SkeletonRootUid = SkeletonRootNodeUids[0];
+		}
+		if(!SkeletonRootUid.IsEmpty())
+		{
+			//Create a skeleton node from all the joint in the translated nodes
+			SkeletonFactoryNode = CommonSkeletalMeshesAndAnimationsProperties->CreateSkeletonFactoryNode(BaseNodeContainer, SkeletonRootUid);
+		}
+	}
+
 	if (!SkeletonFactoryNode)
 	{
 		// It can happen if we force static mesh import, in that case no skeleton will be create
@@ -308,7 +348,7 @@ void UInterchangeGenericAnimationPipeline::CreateAnimSequenceFactoryNode(UInterc
 	}
 
 	FString SkeletalMeshFactoryNodeUid;
-	const UInterchangeSkeletalMeshFactoryNode* SkeletalMeshFactoryNode = nullptr;
+
 	if (SkeletonFactoryNode->GetCustomSkeletalMeshFactoryNodeUid(SkeletalMeshFactoryNodeUid))
 	{
 		SkeletalMeshFactoryNode = Cast<const UInterchangeSkeletalMeshFactoryNode>(BaseNodeContainer->GetFactoryNode(SkeletalMeshFactoryNodeUid));
@@ -436,7 +476,10 @@ void UInterchangeGenericAnimationPipeline::CreateAnimSequenceFactoryNode(UInterc
 	UInterchangeAnimSequenceFactoryNode* AnimSequenceFactoryNode = NewObject<UInterchangeAnimSequenceFactoryNode>(BaseNodeContainer, NAME_None);
 	AnimSequenceFactoryNode->InitializeAnimSequenceNode(AnimSequenceUid, TrackNode.GetDisplayLabel());
 
-	AnimSequenceFactoryNode->SetCustomSkeletonFactoryNodeUid(SkeletonFactoryNode->GetUniqueID());
+	if (SkeletonFactoryNode)
+	{
+		AnimSequenceFactoryNode->SetCustomSkeletonFactoryNodeUid(SkeletonFactoryNode->GetUniqueID());
+	}
 	if (SkeletalMeshFactoryNode)
 	{
 		AnimSequenceFactoryNode->AddFactoryDependencyUid(SkeletalMeshFactoryNode->GetUniqueID());
@@ -465,11 +508,15 @@ void UInterchangeGenericAnimationPipeline::CreateAnimSequenceFactoryNode(UInterc
 	}
 
 	//USkeleton cannot be created without a valid skeletal mesh
-	const FString SkeletonUid = SkeletonFactoryNode->GetUniqueID();
-	AnimSequenceFactoryNode->AddFactoryDependencyUid(SkeletonUid);
+	FString SkeletonUid;
+	if(SkeletonFactoryNode)
+	{
+		SkeletonUid = SkeletonFactoryNode->GetUniqueID();
+		AnimSequenceFactoryNode->AddFactoryDependencyUid(SkeletonUid);
+	}
 
 	FString RootJointUid;
-	if (SkeletonFactoryNode->GetCustomRootJointUid(RootJointUid))
+	if (SkeletonFactoryNode && SkeletonFactoryNode->GetCustomRootJointUid(RootJointUid))
 	{
 		// NOTE: Could this be added as an array of FString attributes on the UInterchangeSkeletalAnimationTrackNode
 #if WITH_EDITOR
@@ -514,8 +561,6 @@ void UInterchangeGenericAnimationPipeline::CreateAnimSequenceFactoryNode(UInterc
 			});
 	}
 
-	const bool bImportOnlyAnimation = CommonSkeletalMeshesAndAnimationsProperties->bImportOnlyAnimations;
-
 	//Iterate dependencies
 	{
 		TArray<FString> SkeletalMeshNodeUids;
diff --git a/Engine/Plugins/Interchange/Runtime/Source/Pipelines/Private/InterchangeGenericCommonSkeletalMeshesAndAnimationsProperties.cpp b/Engine/Plugins/Interchange/Runtime/Source/Pipelines/Private/InterchangeGenericCommonSkeletalMeshesAndAnimationsProperties.cpp
new file mode 100644
--- /dev/null
+++ b/Engine/Plugins/Interchange/Runtime/Source/Pipelines/Private/InterchangeGenericCommonSkeletalMeshesAndAnimationsProperties.cpp
@@ -0,0 +1,70 @@
+// Copyright Epic Games, Inc. All Rights Reserved. 
+#include "InterchangeGenericAssetsPipelineSharedSettings.h"
+
+#include "Animation/Skeleton.h"
+#include "CoreMinimal.h"
+#include "Engine/EngineTypes.h"
+#include "InterchangePipelineLog.h"
+#include "InterchangeSkeletonFactoryNode.h"
+#include "InterchangeSkeletonHelper.h"
+#include "Nodes/InterchangeBaseNode.h"
+#include "Nodes/InterchangeBaseNodeContainer.h"
+#include "Nodes/InterchangeFactoryBaseNode.h"
+#include "UObject/Object.h"
+#include "UObject/ObjectMacros.h"
+
+UInterchangeSkeletonFactoryNode* UInterchangeGenericCommonSkeletalMeshesAndAnimationsProperties::CreateSkeletonFactoryNode(UInterchangeBaseNodeContainer* BaseNodeContainer, const FString& RootJointUid)
+{
+	const UInterchangeBaseNode* RootJointNode = BaseNodeContainer->GetNode(RootJointUid);
+	if (!RootJointNode)
+	{
+		return nullptr;
+	}
+
+	FString DisplayLabel = RootJointNode->GetDisplayLabel() + TEXT("_Skeleton");
+	FString SkeletonUid = UInterchangeFactoryBaseNode::BuildFactoryNodeUid(RootJointNode->GetUniqueID());
+
+	UInterchangeSkeletonFactoryNode* SkeletonFactoryNode = nullptr;
+	if (BaseNodeContainer->IsNodeUidValid(SkeletonUid))
+	{
+		//The node already exist, just return it
+		SkeletonFactoryNode = Cast<UInterchangeSkeletonFactoryNode>(BaseNodeContainer->GetFactoryNode(SkeletonUid));
+		if (!ensure(SkeletonFactoryNode))
+		{
+			//Log an error
+			return nullptr;
+		}
+		FString ExistingSkeletonRootJointUid;
+		SkeletonFactoryNode->GetCustomRootJointUid(ExistingSkeletonRootJointUid);
+		if (!ensure(ExistingSkeletonRootJointUid.Equals(RootJointUid)))
+		{
+			//Log an error
+			return nullptr;
+		}
+	}
+	else
+	{
+		SkeletonFactoryNode = NewObject<UInterchangeSkeletonFactoryNode>(BaseNodeContainer, NAME_None);
+		if (!ensure(SkeletonFactoryNode))
+		{
+			return nullptr;
+		}
+		SkeletonFactoryNode->InitializeSkeletonNode(SkeletonUid, DisplayLabel, USkeleton::StaticClass()->GetName());
+		SkeletonFactoryNode->SetCustomRootJointUid(RootJointNode->GetUniqueID());
+		SkeletonFactoryNode->SetCustomUseTimeZeroForBindPose(bUseT0AsRefPose);
+		BaseNodeContainer->AddNode(SkeletonFactoryNode);
+	}
+
+	//If we have a specified skeleton
+	if (Skeleton.IsValid())
+	{
+		SkeletonFactoryNode->SetEnabled(false);
+		SkeletonFactoryNode->SetCustomReferenceObject(FSoftObjectPath(Skeleton.Get()));
+	}
+#if WITH_EDITOR
+	//Iterate all joints to set the meta data value in the skeleton node
+	UE::Interchange::Private::FSkeletonHelper::RecursiveAddSkeletonMetaDataValues(BaseNodeContainer, SkeletonFactoryNode, RootJointUid);
+#endif //WITH_EDITOR
+
+	return SkeletonFactoryNode;
+}
\ No newline at end of file
diff --git a/Engine/Plugins/Interchange/Runtime/Source/Pipelines/Private/InterchangeGenericSkeletalMeshPipeline.cpp b/Engine/Plugins/Interchange/Runtime/Source/Pipelines/Private/InterchangeGenericSkeletalMeshPipeline.cpp
--- a/Engine/Plugins/Interchange/Runtime/Source/Pipelines/Private/InterchangeGenericSkeletalMeshPipeline.cpp
+++ b/Engine/Plugins/Interchange/Runtime/Source/Pipelines/Private/InterchangeGenericSkeletalMeshPipeline.cpp
@@ -132,7 +132,7 @@ void UInterchangeGenericMeshPipeline::ExecutePreImportPipelineSkeletalMesh()
 			{
 				const FString& SkeletonRootUid = SkeletonRootUidAndMeshUids.Key;
 				//Every iteration is a skeletalmesh asset that combine all MeshInstances sharing the same skeleton root node
-				UInterchangeSkeletonFactoryNode* SkeletonFactoryNode = CreateSkeletonFactoryNode(SkeletonRootUid);
+				UInterchangeSkeletonFactoryNode* SkeletonFactoryNode = CommonSkeletalMeshesAndAnimationsProperties->CreateSkeletonFactoryNode(BaseNodeContainer, SkeletonRootUid);
 				//The MeshUids can represent a SceneNode pointing on a MeshNode or directly a MeshNode;
 				TMap<int32, TArray<FString>> MeshUidsPerLodIndex;
 				const TArray<FString>& MeshUids = SkeletonRootUidAndMeshUids.Value;
@@ -218,7 +218,7 @@ void UInterchangeGenericMeshPipeline::ExecutePreImportPipelineSkeletalMesh()
 						continue;
 					}
 				}
-				UInterchangeSkeletonFactoryNode* SkeletonFactoryNode = CreateSkeletonFactoryNode(SkeletonRootUid);
+				UInterchangeSkeletonFactoryNode* SkeletonFactoryNode = CommonSkeletalMeshesAndAnimationsProperties->CreateSkeletonFactoryNode(BaseNodeContainer, SkeletonRootUid);
 				if (bUseInstanceMesh)
 				{
 					const FInterchangeMeshInstance& MeshInstance = PipelineMeshesUtilities->GetMeshInstanceByUid(MeshUid);
@@ -266,64 +266,6 @@ void UInterchangeGenericMeshPipeline::ExecutePreImportPipelineSkeletalMesh()
 	}
 }
 
-
-UInterchangeSkeletonFactoryNode* UInterchangeGenericMeshPipeline::CreateSkeletonFactoryNode(const FString& RootJointUid)
-{
-	check(CommonSkeletalMeshesAndAnimationsProperties.IsValid());
-	const UInterchangeBaseNode* RootJointNode = BaseNodeContainer->GetNode(RootJointUid);
-	if (!RootJointNode)
-	{
-		return nullptr;
-	}
-
-	FString DisplayLabel = RootJointNode->GetDisplayLabel() + TEXT("_Skeleton");
-	FString SkeletonUid = UInterchangeFactoryBaseNode::BuildFactoryNodeUid(RootJointNode->GetUniqueID());
-
-	UInterchangeSkeletonFactoryNode* SkeletonFactoryNode = nullptr;
-	if (BaseNodeContainer->IsNodeUidValid(SkeletonUid))
-	{
-		//The node already exist, just return it
-		SkeletonFactoryNode = Cast<UInterchangeSkeletonFactoryNode>(BaseNodeContainer->GetFactoryNode(SkeletonUid));
-		if (!ensure(SkeletonFactoryNode))
-		{
-			//Log an error
-			return nullptr;
-		}
-		FString ExistingSkeletonRootJointUid;
-		SkeletonFactoryNode->GetCustomRootJointUid(ExistingSkeletonRootJointUid);
-		if (!ensure(ExistingSkeletonRootJointUid.Equals(RootJointUid)))
-		{
-			//Log an error
-			return nullptr;
-		}
-	}
-	else
-	{
-		SkeletonFactoryNode = NewObject<UInterchangeSkeletonFactoryNode>(BaseNodeContainer, NAME_None);
-		if (!ensure(SkeletonFactoryNode))
-		{
-			return nullptr;
-		}
-		SkeletonFactoryNode->InitializeSkeletonNode(SkeletonUid, DisplayLabel, USkeleton::StaticClass()->GetName());
-		SkeletonFactoryNode->SetCustomRootJointUid(RootJointNode->GetUniqueID());
-		SkeletonFactoryNode->SetCustomUseTimeZeroForBindPose(CommonSkeletalMeshesAndAnimationsProperties->bUseT0AsRefPose);
-		BaseNodeContainer->AddNode(SkeletonFactoryNode);
-	}
-
-	//If we have a specified skeleton
-	if (CommonSkeletalMeshesAndAnimationsProperties->Skeleton.IsValid())
-	{
-		SkeletonFactoryNode->SetEnabled(false);
-		SkeletonFactoryNode->SetCustomReferenceObject(FSoftObjectPath(CommonSkeletalMeshesAndAnimationsProperties->Skeleton.Get()));
-	}
-#if WITH_EDITOR
-	//Iterate all joints to set the meta data value in the skeleton node
-	UE::Interchange::Private::FSkeletonHelper::RecursiveAddSkeletonMetaDataValues(BaseNodeContainer, SkeletonFactoryNode, RootJointUid);
-#endif //WITH_EDITOR
-
-	return SkeletonFactoryNode;
-}
-
 UInterchangeSkeletalMeshFactoryNode* UInterchangeGenericMeshPipeline::CreateSkeletalMeshFactoryNode(const FString& RootJointUid, const TMap<int32, TArray<FString>>& MeshUidsPerLodIndex)
 {
 	check(CommonMeshesProperties.IsValid());
diff --git a/Engine/Plugins/Interchange/Runtime/Source/Pipelines/Public/InterchangeGenericAssetsPipelineSharedSettings.h b/Engine/Plugins/Interchange/Runtime/Source/Pipelines/Public/InterchangeGenericAssetsPipelineSharedSettings.h
--- a/Engine/Plugins/Interchange/Runtime/Source/Pipelines/Public/InterchangeGenericAssetsPipelineSharedSettings.h
+++ b/Engine/Plugins/Interchange/Runtime/Source/Pipelines/Public/InterchangeGenericAssetsPipelineSharedSettings.h
@@ -7,6 +7,7 @@
 #include "Math/Rotator.h"
 #include "Math/Vector.h"
 #include "Nodes/InterchangeBaseNode.h"
+#include "Nodes/InterchangeBaseNodeContainer.h"
 #include "UObject/Class.h"
 #include "UObject/Object.h"
 #include "UObject/ObjectMacros.h"
@@ -14,6 +15,7 @@
 
 #include "InterchangeGenericAssetsPipelineSharedSettings.generated.h"
 
+class UInterchangeSkeletonFactoryNode;
 class USkeleton;
 
 /** Force mesh type, if user want to import all meshes as one type*/
@@ -138,4 +140,7 @@ public:
 		}
 		return Super::IsSettingsAreValid(OutInvalidReason);
 	}
+
+	/** Create a UInterchangeSkeletonFactorynode */
+	UInterchangeSkeletonFactoryNode* CreateSkeletonFactoryNode(UInterchangeBaseNodeContainer* BaseNodeContainer, const FString& RootJointUid);
 };
diff --git a/Engine/Plugins/Interchange/Runtime/Source/Pipelines/Public/InterchangeGenericMeshPipeline.h b/Engine/Plugins/Interchange/Runtime/Source/Pipelines/Public/InterchangeGenericMeshPipeline.h
--- a/Engine/Plugins/Interchange/Runtime/Source/Pipelines/Public/InterchangeGenericMeshPipeline.h
+++ b/Engine/Plugins/Interchange/Runtime/Source/Pipelines/Public/InterchangeGenericMeshPipeline.h
@@ -235,9 +235,6 @@ private:
 	/** Skeleton factory assets nodes */
 	TArray<UInterchangeSkeletonFactoryNode*> SkeletonFactoryNodes;
 
-	/** Create a UInterchangeSkeletonFactorynode */
-	UInterchangeSkeletonFactoryNode* CreateSkeletonFactoryNode(const FString& RootJointUid);
-
 	/** Skeletal mesh factory assets nodes */
 	TArray<UInterchangeSkeletalMeshFactoryNode*> SkeletalMeshFactoryNodes;
 
