Allow global transform to be pass to the payload data. FMeshDescription use float instead of double to store vertex position which have an impact on precision. Passing the global transform will allow the mesh to be in the engine unit (cm). Some model are modelled very small and are scale by the scene node, in this case we will apply the scale directly on the imported data to create the mesh description, that way will ensure we do not have degenerated triangles.

#rnx

diff --git a/Engine/Plugins/Experimental/Enterprise/DatasmithInterchange/Source/DatasmithInterchange/DatasmithInterchange.Build.cs b/Engine/Plugins/Experimental/Enterprise/DatasmithInterchange/Source/DatasmithInterchange/DatasmithInterchange.Build.cs
--- a/Engine/Plugins/Experimental/Enterprise/DatasmithInterchange/Source/DatasmithInterchange/DatasmithInterchange.Build.cs
+++ b/Engine/Plugins/Experimental/Enterprise/DatasmithInterchange/Source/DatasmithInterchange/DatasmithInterchange.Build.cs
@@ -21,7 +21,8 @@ namespace UnrealBuildTool.Rules
 					"InterchangeFactoryNodes",
 					"InterchangeImport",
 					"InterchangePipelines",
-					"MeshDescription"
+					"MeshDescription",
+					"StaticMeshDescription"
 				}
 			);
 
diff --git a/Engine/Plugins/Experimental/Enterprise/DatasmithInterchange/Source/DatasmithInterchange/Private/InterchangeDatasmithTranslator.cpp b/Engine/Plugins/Experimental/Enterprise/DatasmithInterchange/Source/DatasmithInterchange/Private/InterchangeDatasmithTranslator.cpp
--- a/Engine/Plugins/Experimental/Enterprise/DatasmithInterchange/Source/DatasmithInterchange/Private/InterchangeDatasmithTranslator.cpp
+++ b/Engine/Plugins/Experimental/Enterprise/DatasmithInterchange/Source/DatasmithInterchange/Private/InterchangeDatasmithTranslator.cpp
@@ -33,6 +33,8 @@
 #include "InterchangeTexture2DNode.h"
 #include "InterchangeVariantSetNode.h"
 
+#include "StaticMeshOperations.h"
+
 #include "Misc/App.h"
 
 #define LOCTEXT_NAMESPACE "DatasmithInterchange"
@@ -552,7 +554,7 @@ TOptional<UE::Interchange::FImportImage> UInterchangeDatasmithTranslator::GetTex
 	return TextureTranslator->GetTexturePayloadData(PayloadKey, AlternateTexturePath);
 }
 
-TFuture<TOptional<UE::Interchange::FMeshPayloadData>> UInterchangeDatasmithTranslator::GetMeshPayloadData(const FInterchangeMeshPayLoadKey& PayLoadKey) const
+TFuture<TOptional<UE::Interchange::FMeshPayloadData>> UInterchangeDatasmithTranslator::GetMeshPayloadData(const FInterchangeMeshPayLoadKey& PayLoadKey, const FTransform& MeshGlobalTransform) const
 {
 	TPromise<TOptional<UE::Interchange::FMeshPayloadData>> EmptyPromise;
 	EmptyPromise.SetValue(TOptional<UE::Interchange::FMeshPayloadData>());
@@ -576,7 +578,7 @@ TFuture<TOptional<UE::Interchange::FMeshPayloadData>> UInterchangeDatasmithTrans
 		return EmptyPromise.GetFuture();
 	}
 
-	return Async(EAsyncExecution::TaskGraph, [this, MeshElement = MoveTemp(MeshElement)]
+	return Async(EAsyncExecution::TaskGraph, [this, MeshElement = MoveTemp(MeshElement), &MeshGlobalTransform]
 		{
 			TOptional<UE::Interchange::FMeshPayloadData> Result;
 
@@ -587,7 +589,11 @@ TFuture<TOptional<UE::Interchange::FMeshPayloadData>> UInterchangeDatasmithTrans
 				{
 					UE::Interchange::FMeshPayloadData StaticMeshPayloadData;
 					StaticMeshPayloadData.MeshDescription = MoveTemp(DatasmithMeshPayload.LodMeshes[0]);
-
+					// Bake the payload mesh, with the provided transform
+					if (!MeshGlobalTransform.Equals(FTransform::Identity))
+					{
+						FStaticMeshOperations::ApplyTransform(StaticMeshPayloadData.MeshDescription, MeshGlobalTransform);
+					}
 					Result.Emplace(MoveTemp(StaticMeshPayloadData));
 				}
 			}
diff --git a/Engine/Plugins/Experimental/Enterprise/DatasmithInterchange/Source/DatasmithInterchange/Public/InterchangeDatasmithTranslator.h b/Engine/Plugins/Experimental/Enterprise/DatasmithInterchange/Source/DatasmithInterchange/Public/InterchangeDatasmithTranslator.h
--- a/Engine/Plugins/Experimental/Enterprise/DatasmithInterchange/Source/DatasmithInterchange/Public/InterchangeDatasmithTranslator.h
+++ b/Engine/Plugins/Experimental/Enterprise/DatasmithInterchange/Source/DatasmithInterchange/Public/InterchangeDatasmithTranslator.h
@@ -79,7 +79,7 @@ public:
 	/* IInterchangeTexturePayloadInterface End */
 
 	/* IInterchangeStaticMeshPayloadInterface Begin */
-	virtual TFuture<TOptional<UE::Interchange::FMeshPayloadData>> GetMeshPayloadData(const FInterchangeMeshPayLoadKey& PayLoadKey) const override;
+	virtual TFuture<TOptional<UE::Interchange::FMeshPayloadData>> GetMeshPayloadData(const FInterchangeMeshPayLoadKey& PayLoadKey, const FTransform& MeshGlobalTransform) const override;
 	/* IInterchangeStaticMeshPayloadInterface End */
 
 	/* IInterchangeAnimationPayloadInterface Begin */
diff --git a/Engine/Plugins/Interchange/Runtime/Source/Dispatcher/Private/InterchangeDispatcherTask.cpp b/Engine/Plugins/Interchange/Runtime/Source/Dispatcher/Private/InterchangeDispatcherTask.cpp
--- a/Engine/Plugins/Interchange/Runtime/Source/Dispatcher/Private/InterchangeDispatcherTask.cpp
+++ b/Engine/Plugins/Interchange/Runtime/Source/Dispatcher/Private/InterchangeDispatcherTask.cpp
@@ -207,6 +207,90 @@ namespace UE
 			return ResultObject->TryGetStringField(GetResultFilenameJsonKey(), ResultFilename);
 		}
 
+		
+		FString FJsonFetchMeshPayloadCmd::ToJson() const
+		{
+			//Code should not do a ToJson if the data was not set before
+			ensure(bIsDataInitialize);
+
+			TSharedPtr<FJsonObject> CmdObject = MakeShared<FJsonObject>();
+			TSharedPtr<FJsonObject> ActionDataObject = MakeShared<FJsonObject>();
+			//CmdObject
+			CmdObject->SetStringField(GetCommandIDJsonKey(), GetAction());
+			CmdObject->SetStringField(GetTranslatorIDJsonKey(), GetTranslatorID());
+			ActionDataObject->SetStringField(GetPayloadKeyJsonKey(), GetPayloadKey());
+
+			//Bake settings
+			ActionDataObject->SetStringField(GetMeshGlobalTransformJsonKey(), GetMeshGlobalTransform().ToString());
+
+			CmdObject->SetObjectField(GetCommandDataJsonKey(), ActionDataObject);
+
+			FString FetchPayloadCmd;
+			TSharedRef< TJsonWriter< TCHAR, TPrettyJsonPrintPolicy<TCHAR> > > JsonWriter = TJsonWriterFactory< TCHAR, TPrettyJsonPrintPolicy<TCHAR> >::Create(&FetchPayloadCmd);
+			if (!FJsonSerializer::Serialize(CmdObject.ToSharedRef(), JsonWriter))
+			{
+				//Error creating the json cmd string 
+				return FString();
+			}
+			return FetchPayloadCmd;
+		}
+
+		bool FJsonFetchMeshPayloadCmd::FromJson(const FString& JsonString)
+		{
+			TSharedRef<TJsonReader<TCHAR>> Reader = FJsonStringReader::Create(JsonString);
+
+			TSharedPtr<FJsonObject> CmdObject;
+			if (!FJsonSerializer::Deserialize(Reader, CmdObject) || !CmdObject.IsValid())
+			{
+				//Cannot read the json file
+				return false;
+			}
+			FString JsonActionValue;
+			if (!CmdObject->TryGetStringField(GetCommandIDJsonKey(), JsonActionValue))
+			{
+				//The json cmd id key is missing
+				return false;
+			}
+
+			if (!JsonActionValue.Equals(GetAction()))
+			{
+				//This json do not represent a load command
+				return false;
+			}
+
+			//Read the json
+			if (!CmdObject->TryGetStringField(GetTranslatorIDJsonKey(), TranslatorID))
+			{
+				//Missing Load command translator ID
+				return false;
+			}
+			const TSharedPtr<FJsonObject>* ActionDataObject = nullptr;
+			if (!CmdObject->TryGetObjectField(GetCommandDataJsonKey(), ActionDataObject))
+			{
+				//Missing Load Action data object
+				return false;
+			}
+			if (!((*ActionDataObject)->TryGetStringField(GetPayloadKeyJsonKey(), PayloadKey)))
+			{
+				return false;
+			}
+
+			//Retrieve the mesh global transform
+			FString TransformString;
+			if (!((*ActionDataObject)->TryGetStringField(GetMeshGlobalTransformJsonKey(), TransformString)))
+			{
+				return false;
+			}
+			if (!MeshGlobalTransform.InitFromString(TransformString))
+			{
+				return false;
+			}
+
+			//Since we filled the data from the json file, set the data has been initialize.
+			bIsDataInitialize = true;
+			return true;
+		}
+
 		FString FJsonFetchAnimationBakeTransformPayloadCmd::ToJson() const
 		{
 			//Code should not do a ToJson if the data was not set before
diff --git a/Engine/Plugins/Interchange/Runtime/Source/Dispatcher/Public/InterchangeCommands.h b/Engine/Plugins/Interchange/Runtime/Source/Dispatcher/Public/InterchangeCommands.h
--- a/Engine/Plugins/Interchange/Runtime/Source/Dispatcher/Public/InterchangeCommands.h
+++ b/Engine/Plugins/Interchange/Runtime/Source/Dispatcher/Public/InterchangeCommands.h
@@ -17,7 +17,7 @@ namespace UE
 		{
 		public:
 			//Major version should be updated when there is an existing API that has been change.
-			static int32 GetMajor() { return 12; }
+			static int32 GetMajor() { return 13; }
 			//Minor version must be updated when there is an addition to the API.
 			static int32 GetMinor() { return 0; }
 			//Patch version should be update if there is some bug fixes in the private code.
diff --git a/Engine/Plugins/Interchange/Runtime/Source/Dispatcher/Public/InterchangeDispatcherTask.h b/Engine/Plugins/Interchange/Runtime/Source/Dispatcher/Public/InterchangeDispatcherTask.h
--- a/Engine/Plugins/Interchange/Runtime/Source/Dispatcher/Public/InterchangeDispatcherTask.h
+++ b/Engine/Plugins/Interchange/Runtime/Source/Dispatcher/Public/InterchangeDispatcherTask.h
@@ -170,6 +170,7 @@ namespace UE
 				: TranslatorID(InTranslatorID)
 				, PayloadKey(InPayloadKey)
 			{
+				
 				bIsDataInitialize = true;
 			}
 
@@ -233,6 +234,40 @@ namespace UE
 			FString PayloadKey = FString();
 		};
 
+		//Mesh payload require a transform to bake the mesh to avoid degenerate triangle when importing a small mesh scale by a scene node.
+		class INTERCHANGEDISPATCHER_API FJsonFetchMeshPayloadCmd : public FJsonFetchPayloadCmd
+		{
+		public:
+			FJsonFetchMeshPayloadCmd()
+			{
+				check(!bIsDataInitialize);
+			}
+
+			FJsonFetchMeshPayloadCmd(const FString& InTranslatorID
+				, const FString& InPayloadKey
+				, const FTransform& InMeshGlobalTransform)
+				: FJsonFetchPayloadCmd(InTranslatorID, InPayloadKey)
+				, MeshGlobalTransform(InMeshGlobalTransform)
+			{}
+
+			virtual FString ToJson() const override;
+			virtual bool FromJson(const FString& JsonString) override;
+
+			FTransform GetMeshGlobalTransform() const
+			{
+				ensure(bIsDataInitialize);
+				return MeshGlobalTransform;
+			}
+
+			static FString GetMeshGlobalTransformJsonKey()
+			{
+				static const FString Key = TEXT("GlobalMeshTransform");
+				return Key;
+			}
+		protected:
+			FTransform MeshGlobalTransform = FTransform::Identity;
+		};
+
 		//Animation transform payload require transform to be bake by the translator
 		//The anim sequence API is not yet using curve to describe bone track animation
 		class INTERCHANGEDISPATCHER_API FJsonFetchAnimationBakeTransformPayloadCmd : public FJsonFetchPayloadCmd
diff --git a/Engine/Plugins/Interchange/Runtime/Source/Import/Private/Fbx/InterchangeFbxTranslator.cpp b/Engine/Plugins/Interchange/Runtime/Source/Import/Private/Fbx/InterchangeFbxTranslator.cpp
--- a/Engine/Plugins/Interchange/Runtime/Source/Import/Private/Fbx/InterchangeFbxTranslator.cpp
+++ b/Engine/Plugins/Interchange/Runtime/Source/Import/Private/Fbx/InterchangeFbxTranslator.cpp
@@ -205,7 +205,7 @@ TOptional<UE::Interchange::FImportImage> UInterchangeFbxTranslator::GetTexturePa
 	return TextureTranslator->GetTexturePayloadData(PayLoadKey, AlternateTexturePath);
 }
 
-TFuture<TOptional<UE::Interchange::FMeshPayloadData>> UInterchangeFbxTranslator::GetMeshPayloadData(const FInterchangeMeshPayLoadKey& PayLoadKey) const
+TFuture<TOptional<UE::Interchange::FMeshPayloadData>> UInterchangeFbxTranslator::GetMeshPayloadData(const FInterchangeMeshPayLoadKey& PayLoadKey, const FTransform& MeshGlobalTransform) const
 {
 	TSharedPtr<TPromise<TOptional<UE::Interchange::FMeshPayloadData>>> Promise = MakeShared<TPromise<TOptional<UE::Interchange::FMeshPayloadData>>>();
 
@@ -216,7 +216,7 @@ TFuture<TOptional<UE::Interchange::FMeshPayloadData>> UInterchangeFbxTranslator:
 	}
 
 	// Create a json command to read the fbx file
-	FString JsonCommand = CreateFetchPayloadFbxCommand(PayLoadKey.UniqueId);
+	FString JsonCommand = CreateFetchMeshPayloadFbxCommand(PayLoadKey.UniqueId, MeshGlobalTransform);
 	const int32 CreatedTaskIndex = Dispatcher->AddTask(JsonCommand, FInterchangeDispatcherTaskCompleted::CreateLambda([this, Promise, PayLoadKey](const int32 TaskIndex)
 	{
 		TRACE_CPUPROFILER_EVENT_SCOPE("UInterchangeFbxTranslator::GetStaticMeshPayloadData::Dispatcher->AddTaskDone")
@@ -238,7 +238,7 @@ TFuture<TOptional<UE::Interchange::FMeshPayloadData>> UInterchangeFbxTranslator:
 		}
 
 		// Grab the result file and fill the BaseNodeContainer
-		UE::Interchange::FJsonFetchPayloadCmd::JsonResultParser ResultParser;
+		UE::Interchange::FJsonFetchMeshPayloadCmd::JsonResultParser ResultParser;
 		ResultParser.FromJson(JsonResult);
 		FString StaticMeshPayloadFilename = ResultParser.GetResultFilename();
 
@@ -424,6 +424,12 @@ FString UInterchangeFbxTranslator::CreateFetchPayloadFbxCommand(const FString& F
 	return PayloadCommand.ToJson();
 }
 
+FString UInterchangeFbxTranslator::CreateFetchMeshPayloadFbxCommand(const FString& FbxPayloadKey, const FTransform& MeshGlobalTransform) const
+{
+	UE::Interchange::FJsonFetchMeshPayloadCmd PayloadCommand(TEXT("FBX"), FbxPayloadKey, MeshGlobalTransform);
+	return PayloadCommand.ToJson();
+}
+
 FString UInterchangeFbxTranslator::CreateFetchAnimationBakeTransformPayloadFbxCommand(const FString& FbxPayloadKey, const double BakeFrequency, const double RangeStartTime, const double RangeEndTime) const
 {
 	UE::Interchange::FJsonFetchAnimationBakeTransformPayloadCmd PayloadCommand(TEXT("FBX"), FbxPayloadKey, BakeFrequency, RangeStartTime, RangeEndTime);
diff --git a/Engine/Plugins/Interchange/Runtime/Source/Import/Private/Gltf/InterchangeGltfMesh.cpp b/Engine/Plugins/Interchange/Runtime/Source/Import/Private/Gltf/InterchangeGltfMesh.cpp
--- a/Engine/Plugins/Interchange/Runtime/Source/Import/Private/Gltf/InterchangeGltfMesh.cpp
+++ b/Engine/Plugins/Interchange/Runtime/Source/Import/Private/Gltf/InterchangeGltfMesh.cpp
@@ -71,7 +71,7 @@ namespace UE::Interchange::Gltf::Private
 		}
 	}
 
-	bool GetSkeletalMeshDescriptionForPayLoadKey(const GLTF::FAsset& GltfAsset, const FString& PayLoadKey, FMeshDescription& MeshDescription, TArray<FString>* OutJointUniqueNames)
+	bool GetSkeletalMeshDescriptionForPayLoadKey(const GLTF::FAsset& GltfAsset, const FString& PayLoadKey, const FTransform& MeshGlobalTransform, FMeshDescription& MeshDescription, TArray<FString>* OutJointUniqueNames)
 	{
 		TArray<FString> BakingAndPayloadKey;
 		PayLoadKey.ParseIntoArray(BakingAndPayloadKey, TEXT(";"));
@@ -121,7 +121,7 @@ namespace UE::Interchange::Gltf::Private
 		const GLTF::FMesh& GltfMesh = GltfAsset.Meshes[MeshIndex];
 		GLTF::FMeshFactory MeshFactory;
 		MeshFactory.SetUniformScale(GltfUnitConversionMultiplier); // GLTF is in meters while UE is in centimeters
-		MeshFactory.FillMeshDescription(GltfMesh, &BaseMeshDescription);
+		MeshFactory.FillMeshDescription(GltfMesh, MeshGlobalTransform, &BaseMeshDescription);
 
 		PatchPolygonGroups(BaseMeshDescription, GltfAsset);
 
@@ -278,7 +278,7 @@ namespace UE::Interchange::Gltf::Private
 		return true;
 	}
 
-	bool GetStaticMeshPayloadDataForPayLoadKey(const GLTF::FAsset& GltfAsset, const FString& PayLoadKey, FMeshDescription& MeshDescription)
+	bool GetStaticMeshPayloadDataForPayLoadKey(const GLTF::FAsset& GltfAsset, const FString& PayLoadKey, const FTransform& MeshGlobalTransform, FMeshDescription& MeshDescription)
 	{
 		TArray<FString> PayLoadKeys;
 		PayLoadKey.ParseIntoArray(PayLoadKeys, TEXT(":"));
@@ -316,7 +316,7 @@ namespace UE::Interchange::Gltf::Private
 
 		GLTF::FMeshFactory MeshFactory;
 		MeshFactory.SetUniformScale(GltfUnitConversionMultiplier); // GLTF is in meters while UE is in centimeters
-		MeshFactory.FillMeshDescription(GltfMesh, &MeshDescription, MorphTargetWeights);
+		MeshFactory.FillMeshDescription(GltfMesh, MeshGlobalTransform, &MeshDescription, MorphTargetWeights);
 
 		PatchPolygonGroups(MeshDescription, GltfAsset);
 
diff --git a/Engine/Plugins/Interchange/Runtime/Source/Import/Private/Gltf/InterchangeGltfPrivate.h b/Engine/Plugins/Interchange/Runtime/Source/Import/Private/Gltf/InterchangeGltfPrivate.h
--- a/Engine/Plugins/Interchange/Runtime/Source/Import/Private/Gltf/InterchangeGltfPrivate.h
+++ b/Engine/Plugins/Interchange/Runtime/Source/Import/Private/Gltf/InterchangeGltfPrivate.h
@@ -28,10 +28,11 @@ namespace UE::Interchange
 		// Mesh related functions
 		int32 GetRootNodeIndex(const GLTF::FAsset& GltfAsset, const TArray<int32>& NodeIndices);
 
-		bool GetSkeletalMeshDescriptionForPayLoadKey(const GLTF::FAsset& GltfAsset, const FString& PayLoadKey,
+		bool GetSkeletalMeshDescriptionForPayLoadKey(const GLTF::FAsset& GltfAsset, const FString& PayLoadKey, const FTransform& MeshGlobalTransform,
 			FMeshDescription& MeshDescription, TArray<FString>* OutJointUniqueNames);
 
-		bool GetStaticMeshPayloadDataForPayLoadKey(const GLTF::FAsset& GltfAsset, const FString& PayLoadKey, FMeshDescription& MeshDescription);
+		bool GetStaticMeshPayloadDataForPayLoadKey(const GLTF::FAsset& GltfAsset, const FString& PayLoadKey, const FTransform& MeshGlobalTransform
+		, FMeshDescription& MeshDescription);
 		//
 	}
 }
\ No newline at end of file
diff --git a/Engine/Plugins/Interchange/Runtime/Source/Import/Private/Gltf/InterchangeGltfTranslator.cpp b/Engine/Plugins/Interchange/Runtime/Source/Import/Private/Gltf/InterchangeGltfTranslator.cpp
--- a/Engine/Plugins/Interchange/Runtime/Source/Import/Private/Gltf/InterchangeGltfTranslator.cpp
+++ b/Engine/Plugins/Interchange/Runtime/Source/Import/Private/Gltf/InterchangeGltfTranslator.cpp
@@ -1390,9 +1390,9 @@ bool UInterchangeGLTFTranslator::GetVariantSetPayloadData(UE::Interchange::FVari
 	return true;
 }
 
-TFuture< TOptional< UE::Interchange::FMeshPayloadData > > UInterchangeGLTFTranslator::GetMeshPayloadData(const FInterchangeMeshPayLoadKey& PayLoadKey) const
+TFuture< TOptional< UE::Interchange::FMeshPayloadData > > UInterchangeGLTFTranslator::GetMeshPayloadData(const FInterchangeMeshPayLoadKey& PayLoadKey, const FTransform& MeshGlobalTransform) const
 {
-	return Async(EAsyncExecution::TaskGraph, [this, PayLoadKey]
+	return Async(EAsyncExecution::TaskGraph, [this, PayLoadKey, MeshGlobalTransform]
 		{
 			UE::Interchange::FMeshPayloadData MeshPayLoadData;
 			bool bSuccessfullAcquisition = false;
@@ -1400,14 +1400,14 @@ TFuture< TOptional< UE::Interchange::FMeshPayloadData > > UInterchangeGLTFTransl
 			switch (PayLoadKey.Type)
 			{
 			case EInterchangeMeshPayLoadType::STATIC:
-				bSuccessfullAcquisition = UE::Interchange::Gltf::Private::GetStaticMeshPayloadDataForPayLoadKey(GltfAsset, PayLoadKey.UniqueId, MeshPayLoadData.MeshDescription);
+				bSuccessfullAcquisition = UE::Interchange::Gltf::Private::GetStaticMeshPayloadDataForPayLoadKey(GltfAsset, PayLoadKey.UniqueId, MeshGlobalTransform, MeshPayLoadData.MeshDescription);
 				break;
 			case EInterchangeMeshPayLoadType::SKELETAL:
-				bSuccessfullAcquisition = UE::Interchange::Gltf::Private::GetSkeletalMeshDescriptionForPayLoadKey(GltfAsset, PayLoadKey.UniqueId, MeshPayLoadData.MeshDescription, &MeshPayLoadData.JointNames);
+				bSuccessfullAcquisition = UE::Interchange::Gltf::Private::GetSkeletalMeshDescriptionForPayLoadKey(GltfAsset, PayLoadKey.UniqueId, MeshGlobalTransform, MeshPayLoadData.MeshDescription, &MeshPayLoadData.JointNames);
 				break;
 			case EInterchangeMeshPayLoadType::MORPHTARGET:
 				//GLTF handles morph targets as simple Meshes
-				bSuccessfullAcquisition = UE::Interchange::Gltf::Private::GetStaticMeshPayloadDataForPayLoadKey(GltfAsset, PayLoadKey.UniqueId, MeshPayLoadData.MeshDescription);
+				bSuccessfullAcquisition = UE::Interchange::Gltf::Private::GetStaticMeshPayloadDataForPayLoadKey(GltfAsset, PayLoadKey.UniqueId, MeshGlobalTransform, MeshPayLoadData.MeshDescription);
 				break;
 			case EInterchangeMeshPayLoadType::NONE:
 			default:
diff --git a/Engine/Plugins/Interchange/Runtime/Source/Import/Private/Mesh/InterchangeOBJTranslator.cpp b/Engine/Plugins/Interchange/Runtime/Source/Import/Private/Mesh/InterchangeOBJTranslator.cpp
--- a/Engine/Plugins/Interchange/Runtime/Source/Import/Private/Mesh/InterchangeOBJTranslator.cpp
+++ b/Engine/Plugins/Interchange/Runtime/Source/Import/Private/Mesh/InterchangeOBJTranslator.cpp
@@ -158,7 +158,7 @@ struct FObjData
 	FString CurrentMaterial;
 
 	/** Generates a MeshDescription for the named group */
-	FMeshDescription MakeMeshDescriptionForGroup(const FString& GroupName) const;
+	FMeshDescription MakeMeshDescriptionForGroup(const FString& GroupName, const FTransform& MeshGlobalTransform) const;
 
 	/** Returns a bounding box fitting the vertices in the named group */
 	FBox GetGroupBoundingBox(const FString& GroupName) const;
@@ -279,7 +279,7 @@ int32 FObjData::GetGroupPolygonCount(const FString& GroupName) const
 }
 
 
-FMeshDescription FObjData::MakeMeshDescriptionForGroup(const FString& GroupName) const
+FMeshDescription FObjData::MakeMeshDescriptionForGroup(const FString& GroupName, const FTransform& MeshGlobalTransform) const
 {
 	FMeshDescription MeshDescription;
 
@@ -290,6 +290,17 @@ FMeshDescription FObjData::MakeMeshDescriptionForGroup(const FString& GroupName)
 		return MeshDescription;
 	}
 
+	FMatrix TotalMatrix = MeshGlobalTransform.ToMatrixWithScale();
+	FMatrix TotalMatrixForNormal;
+	TotalMatrixForNormal = TotalMatrix.Inverse();
+	TotalMatrixForNormal = TotalMatrixForNormal.GetTransposed();
+
+	auto TransformPosition = [](const FMatrix& Matrix, FVector3f& Position)
+	{
+		const FVector TransformedPosition = Matrix.TransformPosition(FVector(Position));
+		Position = static_cast<FVector3f>(TransformedPosition);
+	};
+
 	FStaticMeshAttributes Attributes(MeshDescription);
 	Attributes.Register();
 
@@ -318,7 +329,9 @@ FMeshDescription FObjData::MakeMeshDescriptionForGroup(const FString& GroupName)
 	for (int32 ObjVertexIndex : VertexIndexMapping)
 	{
 		FVertexID VertexIndex = MeshDescription.CreateVertex();
-		Attributes.GetVertexPositions()[VertexIndex] = PositionToUEBasis(Positions[ObjVertexIndex]);
+		FVector3f& Position = Attributes.GetVertexPositions()[VertexIndex];
+		Position = PositionToUEBasis(Positions[ObjVertexIndex]);
+		TransformPosition(TotalMatrix, Position);
 	}
 
 	// Create UVs and initialize values
@@ -364,7 +377,9 @@ FMeshDescription FObjData::MakeMeshDescriptionForGroup(const FString& GroupName)
 
 				if (VertexData.NormalIndex != INDEX_NONE && Normals.IsValidIndex(VertexData.NormalIndex))
 				{
-					Attributes.GetVertexInstanceNormals()[VertexInstanceID] = PositionToUEBasis(Normals[VertexData.NormalIndex]);
+					FVector3f& Normal = Attributes.GetVertexInstanceNormals()[VertexInstanceID];
+					Normal = PositionToUEBasis(Normals[VertexData.NormalIndex]);
+					TransformPosition(TotalMatrixForNormal, Normal);
 				}
 
 				if (VertexData.UVIndex != INDEX_NONE && UVs.IsValidIndex(VertexData.UVIndex))
@@ -1519,14 +1534,14 @@ bool UInterchangeOBJTranslator::Translate(UInterchangeBaseNodeContainer& BaseNod
 }
 
 
-TFuture<TOptional<UE::Interchange::FMeshPayloadData>> UInterchangeOBJTranslator::GetMeshPayloadData(const FInterchangeMeshPayLoadKey& PayLoadKey) const
+TFuture<TOptional<UE::Interchange::FMeshPayloadData>> UInterchangeOBJTranslator::GetMeshPayloadData(const FInterchangeMeshPayLoadKey& PayLoadKey, const FTransform& MeshGlobalTransform) const
 {
-	return Async(EAsyncExecution::TaskGraph, [this, PayLoadKey]
+	return Async(EAsyncExecution::TaskGraph, [this, PayLoadKey, MeshGlobalTransform]
 		{
 			using namespace UE::Interchange;
 
 			FMeshPayloadData Payload;
-			Payload.MeshDescription = ObjDataPtr->MakeMeshDescriptionForGroup(PayLoadKey.UniqueId);
+			Payload.MeshDescription = ObjDataPtr->MakeMeshDescriptionForGroup(PayLoadKey.UniqueId, MeshGlobalTransform);
 
 			return TOptional<FMeshPayloadData>(Payload);
 		}
diff --git a/Engine/Plugins/Interchange/Runtime/Source/Import/Private/Mesh/InterchangeSkeletalMeshFactory.cpp b/Engine/Plugins/Interchange/Runtime/Source/Import/Private/Mesh/InterchangeSkeletalMeshFactory.cpp
--- a/Engine/Plugins/Interchange/Runtime/Source/Import/Private/Mesh/InterchangeSkeletalMeshFactory.cpp
+++ b/Engine/Plugins/Interchange/Runtime/Source/Import/Private/Mesh/InterchangeSkeletalMeshFactory.cpp
@@ -94,7 +94,7 @@ namespace UE
 						FInterchangeMeshPayLoadKey& PayLoadKey = OptionalPayLoadKey.GetValue();
 
 						//Add the map entry key, the translator will be call after to bulk get all the needed payload
-						TempMorphTargetMeshDescriptionsPerMorphTargetName.Add(PayLoadKey.UniqueId, MeshTranslatorPayloadInterface->GetMeshPayloadData(PayLoadKey));
+						TempMorphTargetMeshDescriptionsPerMorphTargetName.Add(PayLoadKey.UniqueId, MeshTranslatorPayloadInterface->GetMeshPayloadData(PayLoadKey, MeshNodeContext.SceneGlobalTransform.Get(FTransform::Identity)));
 					}
 				}
 
@@ -123,15 +123,8 @@ namespace UE
 							continue;
 						}
 						MorphTargetMeshPayload->VertexOffset = VertexOffset;
-						//Use the Mesh node parent bake transform
-						if (MeshNodeContext.SceneGlobalTransform.IsSet())
-						{
-							MorphTargetMeshPayload->GlobalTransform = MeshNodeContext.SceneGlobalTransform;
-						}
-						else
-						{
-							MorphTargetMeshPayload->GlobalTransform.Reset();
-						}
+						//The Mesh node parent bake transform was pass to the payload request.
+						MorphTargetMeshPayload->GlobalTransform.Reset();
 
 						if (!MorphTargetMeshNode->GetMorphTargetName(MorphTargetMeshPayload->MorphTargetName))
 						{
@@ -451,7 +444,7 @@ namespace UE
 				for (const FMeshNodeContext& MeshNodeContext : MeshReferences)
 				{
 					//Add the payload entry key, the payload data will be fill later in bulk by the translator
-					LodMeshPayloadPerTranslatorPayloadKey.Add(MeshNodeContext.TranslatorPayloadKey.UniqueId, MeshTranslatorPayloadInterface->GetMeshPayloadData(MeshNodeContext.TranslatorPayloadKey));
+					LodMeshPayloadPerTranslatorPayloadKey.Add(MeshNodeContext.TranslatorPayloadKey.UniqueId, MeshTranslatorPayloadInterface->GetMeshPayloadData(MeshNodeContext.TranslatorPayloadKey, MeshNodeContext.SceneGlobalTransform.Get(FTransform::Identity)));
 					//Count the morph target dependencies so we can reserve the right amount
 					MorphTargetCount += (bImportMorphTarget && MeshNodeContext.MeshNode) ? MeshNodeContext.MeshNode->GetMorphTargetDependeciesCount() : 0;
 				}
@@ -545,15 +538,10 @@ namespace UE
 							}
 						}
 					}
-					//Bake the payload, with the provide transform
-					if (MeshNodeContext.SceneGlobalTransform.IsSet())
-					{
-						AppendSettings.MeshTransform = MeshNodeContext.SceneGlobalTransform;
-					}
-					else
-					{
-						AppendSettings.MeshTransform.Reset();
-					}
+					
+					//The Mesh node parent bake transform was pass to the payload request.
+					AppendSettings.MeshTransform.Reset();
+
 					FStaticMeshOperations::AppendMeshDescription(LodMeshPayload->MeshDescription, LodMeshDescription, AppendSettings);
 					if (MeshNodeContext.MeshNode->IsSkinnedMesh() || bIsRigidMesh)
 					{
diff --git a/Engine/Plugins/Interchange/Runtime/Source/Import/Private/Mesh/InterchangeStaticMeshFactory.cpp b/Engine/Plugins/Interchange/Runtime/Source/Import/Private/Mesh/InterchangeStaticMeshFactory.cpp
--- a/Engine/Plugins/Interchange/Runtime/Source/Import/Private/Mesh/InterchangeStaticMeshFactory.cpp
+++ b/Engine/Plugins/Interchange/Runtime/Source/Import/Private/Mesh/InterchangeStaticMeshFactory.cpp
@@ -368,12 +368,6 @@ UInterchangeFactoryBase::FImportAssetResult UInterchangeStaticMeshFactory::Impor
 				}
 				LodMeshDescription = MoveTemp(MeshDescription);
 				bFirstValidMoved = true;
-
-				// Bake the payload mesh, with the provided transform
-				if (!MeshPayload.Transform.Equals(FTransform::Identity))
-				{
-					FStaticMeshOperations::ApplyTransform(LodMeshDescription, MeshPayload.Transform);
-				}
 			}
 			else
 			{
@@ -381,8 +375,6 @@ UInterchangeFactoryBase::FImportAssetResult UInterchangeStaticMeshFactory::Impor
 				{
 					continue;
 				}
-				// Bake the payload mesh, with the provided transform
-				AppendSettings.MeshTransform = MeshPayload.Transform;
 				FStaticMeshOperations::AppendMeshDescription(LodMeshPayload->MeshDescription, LodMeshDescription, AppendSettings);
 			}
 		}
@@ -843,7 +835,7 @@ TArray<UInterchangeStaticMeshFactory::FMeshPayload> UInterchangeStaticMeshFactor
 		FInterchangeMeshPayLoadKey& PayLoadKey = OptionalPayLoadKey.GetValue();
 
 		Payload.MeshName = PayLoadKey.UniqueId;
-		Payload.PayloadData = MeshTranslatorPayloadInterface->GetMeshPayloadData(PayLoadKey);
+		Payload.PayloadData = MeshTranslatorPayloadInterface->GetMeshPayloadData(PayLoadKey, Payload.Transform);
 
 		Payloads.Emplace(MoveTemp(Payload));
 	}
diff --git a/Engine/Plugins/Interchange/Runtime/Source/Import/Public/Fbx/InterchangeFbxTranslator.h b/Engine/Plugins/Interchange/Runtime/Source/Import/Public/Fbx/InterchangeFbxTranslator.h
--- a/Engine/Plugins/Interchange/Runtime/Source/Import/Public/Fbx/InterchangeFbxTranslator.h
+++ b/Engine/Plugins/Interchange/Runtime/Source/Import/Public/Fbx/InterchangeFbxTranslator.h
@@ -63,7 +63,7 @@ public:
 	 * @param PayloadKey - The key to retrieve the a particular payload contain into the specified source data.
 	 * @return a PayloadData containing the imported data. The TOptional will not be set if there is an error.
 	 */
-	virtual TFuture<TOptional<UE::Interchange::FMeshPayloadData>> GetMeshPayloadData(const FInterchangeMeshPayLoadKey& PayLoadKey) const override;
+	virtual TFuture<TOptional<UE::Interchange::FMeshPayloadData>> GetMeshPayloadData(const FInterchangeMeshPayLoadKey& PayLoadKey, const FTransform& MeshGlobalTransform) const override;
 
 	///* IInterchangeMeshPayloadInterface End */
 
@@ -74,6 +74,8 @@ public:
 private:
 	FString CreateLoadFbxFileCommand(const FString& FbxFilePath) const;
 
+	FString CreateFetchMeshPayloadFbxCommand(const FString& FbxPayloadKey, const FTransform& MeshGlobalTransform) const;
+
 	FString CreateFetchPayloadFbxCommand(const FString& FbxPayloadKey) const;
 
 	FString CreateFetchAnimationBakeTransformPayloadFbxCommand(const FString& FbxPayloadKey, const double BakeFrequency, const double RangeStartTime, const double RangeEndTime) const;
diff --git a/Engine/Plugins/Interchange/Runtime/Source/Import/Public/Gltf/InterchangeGltfTranslator.h b/Engine/Plugins/Interchange/Runtime/Source/Import/Public/Gltf/InterchangeGltfTranslator.h
--- a/Engine/Plugins/Interchange/Runtime/Source/Import/Public/Gltf/InterchangeGltfTranslator.h
+++ b/Engine/Plugins/Interchange/Runtime/Source/Import/Public/Gltf/InterchangeGltfTranslator.h
@@ -42,7 +42,7 @@ public:
 
 	/* IInterchangeStaticMeshPayloadInterface Begin */
 
-	virtual TFuture<TOptional<UE::Interchange::FMeshPayloadData>> GetMeshPayloadData(const FInterchangeMeshPayLoadKey& PayLoadKey) const override;
+	virtual TFuture<TOptional<UE::Interchange::FMeshPayloadData>> GetMeshPayloadData(const FInterchangeMeshPayLoadKey& PayLoadKey, const FTransform& MeshGlobalTransform) const override;
 
 	/* IInterchangeStaticMeshPayloadInterface End */
 
diff --git a/Engine/Plugins/Interchange/Runtime/Source/Import/Public/Mesh/InterchangeMeshPayloadInterface.h b/Engine/Plugins/Interchange/Runtime/Source/Import/Public/Mesh/InterchangeMeshPayloadInterface.h
--- a/Engine/Plugins/Interchange/Runtime/Source/Import/Public/Mesh/InterchangeMeshPayloadInterface.h
+++ b/Engine/Plugins/Interchange/Runtime/Source/Import/Public/Mesh/InterchangeMeshPayloadInterface.h
@@ -33,7 +33,7 @@ public:
 	 * @param PayloadKey - The key to retrieve the a particular payload contain into the specified source data.
 	 * @return a PayloadData containing the the data ask with the key.
 	 */
-	virtual TFuture<TOptional<UE::Interchange::FMeshPayloadData>> GetMeshPayloadData(const FInterchangeMeshPayLoadKey& PayLoadKey) const = 0;
+	virtual TFuture<TOptional<UE::Interchange::FMeshPayloadData>> GetMeshPayloadData(const FInterchangeMeshPayLoadKey& PayLoadKey, const FTransform& MeshGlobalTransform) const = 0;
 };
 
 
diff --git a/Engine/Plugins/Interchange/Runtime/Source/Import/Public/Mesh/InterchangeOBJTranslator.h b/Engine/Plugins/Interchange/Runtime/Source/Import/Public/Mesh/InterchangeOBJTranslator.h
--- a/Engine/Plugins/Interchange/Runtime/Source/Import/Public/Mesh/InterchangeOBJTranslator.h
+++ b/Engine/Plugins/Interchange/Runtime/Source/Import/Public/Mesh/InterchangeOBJTranslator.h
@@ -49,7 +49,7 @@ public:
 	 * @param PayloadKey - The key to retrieve the a particular payload contain into the specified source data.
 	 * @return a PayloadData containing the the data ask with the key.
 	 */
-	virtual TFuture<TOptional<UE::Interchange::FMeshPayloadData>> GetMeshPayloadData(const FInterchangeMeshPayLoadKey& PayLoadKey) const override;
+	virtual TFuture<TOptional<UE::Interchange::FMeshPayloadData>> GetMeshPayloadData(const FInterchangeMeshPayLoadKey& PayLoadKey, const FTransform& MeshGlobalTransform) const override;
 
 	/* IInterchangeStaticMeshPayloadInterface End */
 
diff --git a/Engine/Plugins/Interchange/Runtime/Source/Parsers/Fbx/Private/FbxAPI.cpp b/Engine/Plugins/Interchange/Runtime/Source/Parsers/Fbx/Private/FbxAPI.cpp
--- a/Engine/Plugins/Interchange/Runtime/Source/Parsers/Fbx/Private/FbxAPI.cpp
+++ b/Engine/Plugins/Interchange/Runtime/Source/Parsers/Fbx/Private/FbxAPI.cpp
@@ -150,6 +150,23 @@ namespace UE
 				}
 			}
 
+			bool FFbxParser::FetchMeshPayloadData(const FString& PayloadKey, const FTransform& MeshGlobalTransform, const FString& PayloadFilepath)
+			{
+				if (!PayloadContexts.Contains(PayloadKey))
+				{
+					UInterchangeResultError_Generic* Message = AddMessage<UInterchangeResultError_Generic>();
+					Message->Text = LOCTEXT("CannotRetrievePayload", "Cannot retrieve payload; payload key doesn't have any context.");
+					return false;
+				}
+
+				{
+					//Critical section to force payload to be fetch one by one with no concurrency.
+					FScopeLock Lock(&PayloadCriticalSection);
+					TSharedPtr<FPayloadContextBase>& PayloadContext = PayloadContexts.FindChecked(PayloadKey);
+					return PayloadContext->FetchMeshPayloadToFile(*this, MeshGlobalTransform, PayloadFilepath);
+				}
+			}
+
 			bool FFbxParser::FetchAnimationBakeTransformPayload(const FString& PayloadKey, const double BakeFrequency, const double RangeStartTime, const double RangeEndTime, const FString& PayloadFilepath)
 			{
 				if (!PayloadContexts.Contains(PayloadKey))
diff --git a/Engine/Plugins/Interchange/Runtime/Source/Parsers/Fbx/Private/FbxAPI.h b/Engine/Plugins/Interchange/Runtime/Source/Parsers/Fbx/Private/FbxAPI.h
--- a/Engine/Plugins/Interchange/Runtime/Source/Parsers/Fbx/Private/FbxAPI.h
+++ b/Engine/Plugins/Interchange/Runtime/Source/Parsers/Fbx/Private/FbxAPI.h
@@ -51,6 +51,9 @@ namespace UE
 				/* Extract the fbx data from the sdk into our node container */
 				bool FetchPayloadData(const FString& PayloadKey, const FString& PayloadFilepath);
 
+				/* Extract the fbx mesh data from the sdk into our node container */
+				bool FetchMeshPayloadData(const FString& PayloadKey, const FTransform& MeshGlobalTransform, const FString& PayloadFilepath);
+
 				/* Extract the fbx data from the sdk into our node container */
 				bool FetchAnimationBakeTransformPayload(const FString& PayloadKey, const double BakeFrequency, const double RangeStartTime, const double RangeEndTime, const FString& PayloadFilepath);
 				/**
diff --git a/Engine/Plugins/Interchange/Runtime/Source/Parsers/Fbx/Private/FbxConvert.cpp b/Engine/Plugins/Interchange/Runtime/Source/Parsers/Fbx/Private/FbxConvert.cpp
--- a/Engine/Plugins/Interchange/Runtime/Source/Parsers/Fbx/Private/FbxConvert.cpp
+++ b/Engine/Plugins/Interchange/Runtime/Source/Parsers/Fbx/Private/FbxConvert.cpp
@@ -133,6 +133,33 @@ namespace UE
 				return UEMatrix;
 			}
 
+			FbxAMatrix FFbxConvert::ConvertMatrix(const FMatrix& UEMatrix)
+			{
+				FbxAMatrix FbxMatrix;
+
+				for (int i = 0; i < 4; ++i)
+				{
+					FbxVector4 Row;
+					if (i == 1)
+					{
+						Row[0] = -UEMatrix.M[i][0];
+						Row[1] = UEMatrix.M[i][1];
+						Row[2] = -UEMatrix.M[i][2];
+						Row[3] = -UEMatrix.M[i][3];
+					}
+					else
+					{
+						Row[0] = UEMatrix.M[i][0];
+						Row[1] = -UEMatrix.M[i][1];
+						Row[2] = UEMatrix.M[i][2];
+						Row[3] = UEMatrix.M[i][3];
+					}
+					FbxMatrix.SetRow(i, Row);
+				}
+
+				return FbxMatrix;
+			}
+
 			FQuat FFbxConvert::ConvertRotToQuat(FbxQuaternion Quaternion)
 			{
 				FQuat UnrealQuat;
diff --git a/Engine/Plugins/Interchange/Runtime/Source/Parsers/Fbx/Private/FbxConvert.h b/Engine/Plugins/Interchange/Runtime/Source/Parsers/Fbx/Private/FbxConvert.h
--- a/Engine/Plugins/Interchange/Runtime/Source/Parsers/Fbx/Private/FbxConvert.h
+++ b/Engine/Plugins/Interchange/Runtime/Source/Parsers/Fbx/Private/FbxConvert.h
@@ -23,6 +23,7 @@ namespace UE
 				 */
 				static FTransform ConvertTransform(const FbxAMatrix& Matrix);
 				static FMatrix ConvertMatrix(const FbxAMatrix& Matrix);
+				static FbxAMatrix ConvertMatrix(const FMatrix& UEMatrix);
 				static FQuat ConvertRotToQuat(FbxQuaternion Quaternion);
 				static FRotator ConvertEuler(FbxDouble3 Euler);
 				static FVector ConvertScale(FbxVector4 Vector);
diff --git a/Engine/Plugins/Interchange/Runtime/Source/Parsers/Fbx/Private/FbxHelper.h b/Engine/Plugins/Interchange/Runtime/Source/Parsers/Fbx/Private/FbxHelper.h
--- a/Engine/Plugins/Interchange/Runtime/Source/Parsers/Fbx/Private/FbxHelper.h
+++ b/Engine/Plugins/Interchange/Runtime/Source/Parsers/Fbx/Private/FbxHelper.h
@@ -16,6 +16,7 @@ namespace UE::Interchange::Private
 		virtual ~FPayloadContextBase() {}
 		virtual FString GetPayloadType() const { return FString(); }
 		virtual bool FetchPayloadToFile(FFbxParser& Parser, const FString& PayloadFilepath) { return false; }
+		virtual bool FetchMeshPayloadToFile(FFbxParser& Parser, const FTransform& MeshGlobalTransform, const FString& PayloadFilepath) { return false; }
 		virtual bool FetchAnimationBakeTransformPayloadToFile(FFbxParser& Parser, const double BakeFrequency, const double RangeStartTime, const double RangeEndTime, const FString& PayloadFilepath) { return false; }
 	};
 
diff --git a/Engine/Plugins/Interchange/Runtime/Source/Parsers/Fbx/Private/FbxMesh.cpp b/Engine/Plugins/Interchange/Runtime/Source/Parsers/Fbx/Private/FbxMesh.cpp
--- a/Engine/Plugins/Interchange/Runtime/Source/Parsers/Fbx/Private/FbxMesh.cpp
+++ b/Engine/Plugins/Interchange/Runtime/Source/Parsers/Fbx/Private/FbxMesh.cpp
@@ -202,27 +202,27 @@ FMeshDescriptionImporter::FMeshDescriptionImporter(FFbxParser& InParser, FMeshDe
 	bInitialized = ensure(MeshDescription) && ensure(SDKScene) && ensure(SDKGeometryConverter);
 }
 			
-bool FMeshDescriptionImporter::FillStaticMeshDescriptionFromFbxMesh(FbxMesh* Mesh)
+bool FMeshDescriptionImporter::FillStaticMeshDescriptionFromFbxMesh(FbxMesh* Mesh, const FTransform& MeshGlobalTransform)
 {
 	if (!ensure(bInitialized) || !ensure(MeshDescription) || !ensure(Mesh))
 	{
 		return false;
 	}
 	TArray<FString> UnusedArray;
-	return FillMeshDescriptionFromFbxMesh(Mesh, UnusedArray, EMeshType::Static);
+	return FillMeshDescriptionFromFbxMesh(Mesh, MeshGlobalTransform, UnusedArray, EMeshType::Static);
 }
 
-bool FMeshDescriptionImporter::FillSkinnedMeshDescriptionFromFbxMesh(FbxMesh* Mesh, TArray<FString>& OutJointUniqueNames)
+bool FMeshDescriptionImporter::FillSkinnedMeshDescriptionFromFbxMesh(FbxMesh* Mesh, const FTransform& MeshGlobalTransform, TArray<FString>& OutJointUniqueNames)
 {
 	if (!ensure(bInitialized) || !ensure(MeshDescription) || !ensure(Mesh) || !ensure(Mesh->GetDeformerCount() > 0))
 	{
 		return false;
 	}
 
-	return FillMeshDescriptionFromFbxMesh(Mesh, OutJointUniqueNames, EMeshType::Skinned);
+	return FillMeshDescriptionFromFbxMesh(Mesh, MeshGlobalTransform, OutJointUniqueNames, EMeshType::Skinned);
 }
 
-bool FMeshDescriptionImporter::FillMeshDescriptionFromFbxShape(FbxShape* Shape)
+bool FMeshDescriptionImporter::FillMeshDescriptionFromFbxShape(FbxShape* Shape, const FTransform& MeshGlobalTransform)
 {
 	//For the shape we just need the vertex positions
 	if (!ensure(bInitialized) || !ensure(MeshDescription))
@@ -243,9 +243,8 @@ bool FMeshDescriptionImporter::FillMeshDescriptionFromFbxShape(FbxShape* Shape)
 		MeshDescription->SuspendVertexIndexing();
 
 		// Construct the matrices for the conversion from right handed to left handed system
-		FbxAMatrix TotalMatrix;
+		FbxAMatrix TotalMatrix = FFbxConvert::ConvertMatrix(MeshGlobalTransform.ToInverseMatrixWithScale());
 		FbxAMatrix TotalMatrixForNormal;
-		TotalMatrix.SetIdentity(); //We use the identity since we want to bake in the interchange factory if the pipeline request it
 		TotalMatrixForNormal = TotalMatrix.Inverse();
 		TotalMatrixForNormal = TotalMatrixForNormal.Transpose();
 
@@ -287,7 +286,7 @@ bool FMeshDescriptionImporter::FillMeshDescriptionFromFbxShape(FbxShape* Shape)
 	return true;
 }
 
-bool FMeshDescriptionImporter::FillMeshDescriptionFromFbxMesh(FbxMesh* Mesh, TArray<FString>& OutJointUniqueNames, EMeshType MeshType)
+bool FMeshDescriptionImporter::FillMeshDescriptionFromFbxMesh(FbxMesh* Mesh, const FTransform& MeshGlobalTransform, TArray<FString>& OutJointUniqueNames, EMeshType MeshType)
 {
 	if (!ensure(bInitialized) || !ensure(MeshDescription) || !ensure(Mesh))
 	{
@@ -502,9 +501,8 @@ bool FMeshDescriptionImporter::FillMeshDescriptionFromFbxMesh(FbxMesh* Mesh, TAr
 		TRACE_CPUPROFILER_EVENT_SCOPE(BuildTriangles);
 
 		// Construct the matrices for the conversion from right handed to left handed system
-		FbxAMatrix TotalMatrix;
+		FbxAMatrix TotalMatrix = FFbxConvert::ConvertMatrix(MeshGlobalTransform.ToMatrixWithScale());
 		FbxAMatrix TotalMatrixForNormal;
-		TotalMatrix.SetIdentity();
 		TotalMatrixForNormal = TotalMatrix.Inverse();
 		TotalMatrixForNormal = TotalMatrixForNormal.Transpose();
 		int32 PolygonCount = Mesh->GetPolygonCount();
@@ -518,6 +516,7 @@ bool FMeshDescriptionImporter::FillMeshDescriptionFromFbxMesh(FbxMesh* Mesh, TAr
 		}
 
 		int32 VertexCount = Mesh->GetControlPointsCount();
+		//Todo implement the odd negative scale, see legacy fbx 
 		bool OddNegativeScale = IsOddNegativeScale(TotalMatrix);
 
 		TVertexAttributesRef<FVector3f> VertexPositions = Attributes.GetVertexPositions();
@@ -1198,7 +1197,7 @@ bool FMeshDescriptionImporter::IsOddNegativeScale(FbxAMatrix& TotalMatrix)
 //////////////////////////////////////////////////////////////////////////
 /// FMeshPayloadContext implementation
 
-bool FMeshPayloadContext::FetchPayloadToFile(FFbxParser& Parser, const FString& PayloadFilepath)
+bool FMeshPayloadContext::FetchMeshPayloadToFile(FFbxParser& Parser, const FTransform& MeshGlobalTransform, const FString& PayloadFilepath)
 {
 	if (!ensure(SDKScene != nullptr))
 	{
@@ -1231,7 +1230,7 @@ bool FMeshPayloadContext::FetchPayloadToFile(FFbxParser& Parser, const FString&
 		FSkeletalMeshAttributes SkeletalMeshAttribute(MeshDescription);
 		SkeletalMeshAttribute.Register();
 		FMeshDescriptionImporter MeshDescriptionImporter(Parser, &MeshDescription, SDKScene, SDKGeometryConverter);
-		if (!MeshDescriptionImporter.FillSkinnedMeshDescriptionFromFbxMesh(Mesh, JointUniqueNames))
+		if (!MeshDescriptionImporter.FillSkinnedMeshDescriptionFromFbxMesh(Mesh, MeshGlobalTransform, JointUniqueNames))
 		{
 			UInterchangeResultError_Generic* Message = Parser.AddMessage<UInterchangeResultError_Generic>();
 			Message->InterchangeKey = Parser.GetFbxHelper()->GetMeshUniqueID(Mesh);
@@ -1244,7 +1243,7 @@ bool FMeshPayloadContext::FetchPayloadToFile(FFbxParser& Parser, const FString&
 		FStaticMeshAttributes StaticMeshAttribute(MeshDescription);
 		StaticMeshAttribute.Register();
 		FMeshDescriptionImporter MeshDescriptionImporter(Parser, &MeshDescription, SDKScene, SDKGeometryConverter);
-		if (!MeshDescriptionImporter.FillStaticMeshDescriptionFromFbxMesh(Mesh))
+		if (!MeshDescriptionImporter.FillStaticMeshDescriptionFromFbxMesh(Mesh, MeshGlobalTransform))
 		{
 			UInterchangeResultError_Generic* Message = Parser.AddMessage<UInterchangeResultError_Generic>();
 			Message->InterchangeKey = Parser.GetFbxHelper()->GetMeshUniqueID(Mesh);
@@ -1275,7 +1274,7 @@ bool FMeshPayloadContext::FetchPayloadToFile(FFbxParser& Parser, const FString&
 //////////////////////////////////////////////////////////////////////////
 /// FMorphTargetPayloadContext implementation
 
-bool FMorphTargetPayloadContext::FetchPayloadToFile(FFbxParser& Parser, const FString& PayloadFilepath)
+bool FMorphTargetPayloadContext::FetchMeshPayloadToFile(FFbxParser& Parser, const FTransform& MeshGlobalTransform, const FString& PayloadFilepath)
 {
 	if (!ensure(Shape))
 	{
@@ -1287,7 +1286,7 @@ bool FMorphTargetPayloadContext::FetchPayloadToFile(FFbxParser& Parser, const FS
 	FStaticMeshAttributes StaticMeshAttribute(MorphTargetMeshDescription);
 	StaticMeshAttribute.Register();
 	FMeshDescriptionImporter MeshDescriptionImporter(Parser, &MorphTargetMeshDescription, SDKScene, SDKGeometryConverter);
-	if (!MeshDescriptionImporter.FillMeshDescriptionFromFbxShape(Shape))
+	if (!MeshDescriptionImporter.FillMeshDescriptionFromFbxShape(Shape, MeshGlobalTransform))
 	{
 		UInterchangeResultError_Generic* Message = Parser.AddMessage<UInterchangeResultError_Generic>();
 		Message->InterchangeKey = Parser.GetFbxHelper()->GetMeshUniqueID(Shape);
diff --git a/Engine/Plugins/Interchange/Runtime/Source/Parsers/Fbx/Private/FbxMesh.h b/Engine/Plugins/Interchange/Runtime/Source/Parsers/Fbx/Private/FbxMesh.h
--- a/Engine/Plugins/Interchange/Runtime/Source/Parsers/Fbx/Private/FbxMesh.h
+++ b/Engine/Plugins/Interchange/Runtime/Source/Parsers/Fbx/Private/FbxMesh.h
@@ -63,17 +63,17 @@ namespace UE
 				/*
 				 * Fill the mesh description using the Mesh parameter.
 				 */
-				bool FillStaticMeshDescriptionFromFbxMesh(FbxMesh* Mesh);
+				bool FillStaticMeshDescriptionFromFbxMesh(FbxMesh* Mesh, const FTransform& MeshGlobalTransform);
 				
 				/*
 				 * Fill the mesh description using the Mesh parameter and also fill the OutJointNodeUniqueIDs so the MeshDescription bone Index can be map to the correct interchange joint scene node.
 				 */
-				bool FillSkinnedMeshDescriptionFromFbxMesh(FbxMesh* Mesh, TArray<FString>& OutJointUniqueNames);
+				bool FillSkinnedMeshDescriptionFromFbxMesh(FbxMesh* Mesh, const FTransform& MeshGlobalTransform, TArray<FString>& OutJointUniqueNames);
 
 				/*
 				 * Fill the mesh description using the Shape parameter.
 				 */
-				bool FillMeshDescriptionFromFbxShape(FbxShape* Shape);
+				bool FillMeshDescriptionFromFbxShape(FbxShape* Shape, const FTransform& MeshGlobalTransform);
 
 				/**
 				 * Add messages to the message log
@@ -96,7 +96,7 @@ namespace UE
 					Skinned = 2, //skinned mesh with joints
 				};
 
-				bool FillMeshDescriptionFromFbxMesh(FbxMesh* Mesh, TArray<FString>& OutJointUniqueNames, EMeshType MeshType);
+				bool FillMeshDescriptionFromFbxMesh(FbxMesh* Mesh, const FTransform& MeshGlobalTransform, TArray<FString>& OutJointUniqueNames, EMeshType MeshType);
 				bool IsOddNegativeScale(FbxAMatrix& TotalMatrix);
 				
 				//TODO move the real function from RenderCore to FVector, so we do not have to add render core to compute such a simple thing
@@ -123,7 +123,7 @@ namespace UE
 			public:
 				virtual ~FMeshPayloadContext() {}
 				virtual FString GetPayloadType() const override { return TEXT("Mesh-PayloadContext"); }
-				virtual bool FetchPayloadToFile(FFbxParser& Parser, const FString& PayloadFilepath) override;
+				virtual bool FetchMeshPayloadToFile(FFbxParser& Parser, const FTransform& MeshGlobalTransform, const FString& PayloadFilepath) override;
 				bool bIsSkinnedMesh = false;
 				FbxMesh* Mesh = nullptr;
 				FbxScene* SDKScene = nullptr;
@@ -135,7 +135,7 @@ namespace UE
 			public:
 				virtual ~FMorphTargetPayloadContext() {}
 				virtual FString GetPayloadType() const override { return TEXT("MorphTarget-PayloadContext"); }
-				virtual bool FetchPayloadToFile(FFbxParser& Parser, const FString& PayloadFilepath) override;
+				virtual bool FetchMeshPayloadToFile(FFbxParser& Parser, const FTransform& MeshGlobalTransform, const FString& PayloadFilepath) override;
 				FbxShape* Shape = nullptr;
 				FbxScene* SDKScene = nullptr;
 				FbxGeometryConverter* SDKGeometryConverter = nullptr;
diff --git a/Engine/Plugins/Interchange/Runtime/Source/Parsers/Fbx/Private/FbxScene.cpp b/Engine/Plugins/Interchange/Runtime/Source/Parsers/Fbx/Private/FbxScene.cpp
--- a/Engine/Plugins/Interchange/Runtime/Source/Parsers/Fbx/Private/FbxScene.cpp
+++ b/Engine/Plugins/Interchange/Runtime/Source/Parsers/Fbx/Private/FbxScene.cpp
@@ -239,6 +239,9 @@ namespace UE
 									UnrealNode->SetCustomTimeZeroLocalTransform(&NodeContainer, GlobalTransform, bResetCache);
 								}
 							}
+
+							FString JointNodeName = Parser.GetFbxHelper()->GetFbxObjectName(Node, true);
+							UnrealNode->SetDisplayLabel(JointNodeName);
 							break;
 						}
 
diff --git a/Engine/Plugins/Interchange/Runtime/Source/Parsers/Fbx/Private/InterchangeFbxParser.cpp b/Engine/Plugins/Interchange/Runtime/Source/Parsers/Fbx/Private/InterchangeFbxParser.cpp
--- a/Engine/Plugins/Interchange/Runtime/Source/Parsers/Fbx/Private/InterchangeFbxParser.cpp
+++ b/Engine/Plugins/Interchange/Runtime/Source/Parsers/Fbx/Private/InterchangeFbxParser.cpp
@@ -104,6 +104,30 @@ namespace UE::Interchange
 		}
 	}
 
+	void FInterchangeFbxParser::FetchMeshPayload(const FString& PayloadKey, const FTransform& MeshGlobalTransform, const FString& ResultFolder)
+	{
+		check(FbxParserPrivate.IsValid());
+		ResultsContainer->Empty();
+		FString PayloadFilepathCopy;
+		{
+			FScopeLock Lock(&ResultPayloadsCriticalSection);
+			FString& PayloadFilepath = ResultPayloads.FindOrAdd(PayloadKey);
+			//To avoid file path with too many character, we hash the payloadKey so we have a deterministic length for the file path.
+			FString PayloadKeyHash = Private::HashString(PayloadKey);
+			PayloadFilepath = ResultFolder + TEXT("/") + PayloadKeyHash + TEXT(".payload");
+
+			//Copy the map filename key because we are multithreaded and the TMap can be reallocated
+			PayloadFilepathCopy = PayloadFilepath;
+		}
+		if (!FbxParserPrivate->FetchMeshPayloadData(PayloadKey, MeshGlobalTransform, PayloadFilepathCopy))
+		{
+			UInterchangeResultError_Generic* Error = AddMessage<UInterchangeResultError_Generic>();
+			Error->SourceAssetName = SourceFilename;
+			Error->Text = LOCTEXT("CantFetchPayload", "Cannot fetch FBX payload data.");
+			return;
+		}
+	}
+
 	void FInterchangeFbxParser::FetchAnimationBakeTransformPayload(const FString& PayloadKey, const double BakeFrequency, const double RangeStartTime, const double RangeEndTime, const FString& ResultFolder)
 	{
 		check(FbxParserPrivate.IsValid());
diff --git a/Engine/Plugins/Interchange/Runtime/Source/Parsers/Fbx/Public/InterchangeFbxParser.h b/Engine/Plugins/Interchange/Runtime/Source/Parsers/Fbx/Public/InterchangeFbxParser.h
--- a/Engine/Plugins/Interchange/Runtime/Source/Parsers/Fbx/Public/InterchangeFbxParser.h
+++ b/Engine/Plugins/Interchange/Runtime/Source/Parsers/Fbx/Public/InterchangeFbxParser.h
@@ -38,6 +38,14 @@ namespace UE
 			 */
 			void FetchPayload(const FString& PayloadKey, const FString& ResultFolder);
 
+			/**
+			 * Extract mesh payload data from the fbx, the key tell the translator what payload the client ask
+			 * @param - PayloadKey is the key that describe the payload data to extract from the fbx file
+			 * @param - MeshGlobalTransform is the transform we want to apply to the mesh vertex
+			 * @param - ResultFolder is the folder where we must put any result file
+			 */
+			void FetchMeshPayload(const FString& PayloadKey, const FTransform& MeshGlobalTransform, const FString& ResultFolder);
+
 			/**
 			 * Extract bake transform animation payload data from the fbx, the key tell the translator what payload the client ask
 			 * @param - PayloadKey is the key that describe the payload data to extract from the fbx file
diff --git a/Engine/Plugins/Interchange/Runtime/Source/Parsers/GLTFCore/Private/GLTFMeshFactory.cpp b/Engine/Plugins/Interchange/Runtime/Source/Parsers/GLTFCore/Private/GLTFMeshFactory.cpp
--- a/Engine/Plugins/Interchange/Runtime/Source/Parsers/GLTFCore/Private/GLTFMeshFactory.cpp
+++ b/Engine/Plugins/Interchange/Runtime/Source/Parsers/GLTFCore/Private/GLTFMeshFactory.cpp
@@ -2,7 +2,7 @@
 
 #include "GLTFMeshFactory.h"
 
-
+#include "GLTF/ConversionUtilities.h"
 #include "GLTFMesh.h"
 #include "StaticMeshAttributes.h"
 #include "Util/ColorConstants.h"
@@ -14,7 +14,7 @@ namespace GLTF
 	public:
 		FMeshFactoryImpl();
 
-		void FillMeshDescription(const GLTF::FMesh &Mesh, FMeshDescription* MeshDescription, const TArray<float>& MorphTargetWeights);
+		void FillMeshDescription(const GLTF::FMesh &Mesh, const FTransform& MeshGlobalTransform, FMeshDescription* MeshDescription, const TArray<float>& MorphTargetWeights);
 
 		void CleanUp();
 
@@ -31,6 +31,7 @@ namespace GLTF
 			const TVertexInstanceAttributesRef<FVector2f>& VertexInstanceUVs,
 			const TVertexInstanceAttributesRef<FVector4f>&  VertexInstanceColors,
 			const TEdgeAttributesRef<bool>&                EdgeHardnesses,
+			const FTransform& MeshGlobalTransform,
 			FMeshDescription* MeshDescription,
 			const TArray<float>& MorphTargetWeights);
 
@@ -308,13 +309,15 @@ namespace GLTF
 		}
 	}
 
-	void FMeshFactoryImpl::FillMeshDescription(const FMesh &Mesh, FMeshDescription* MeshDescription, const TArray<float>& MorphTargetWeights)
+	void FMeshFactoryImpl::FillMeshDescription(const FMesh &Mesh, const FTransform& MeshGlobalTransform, FMeshDescription* MeshDescription, const TArray<float>& MorphTargetWeights)
 	{
 		const int32 NumUVs = FMath::Max(1, GetNumUVs(Mesh));
 
 		FStaticMeshAttributes StaticMeshAttributes(*MeshDescription);
 		StaticMeshAttributes.Register();
 
+		FMatrix TotalMatrix = ConvertMat(MeshGlobalTransform.ToMatrixWithScale());
+
 		TVertexAttributesRef<FVector3f> VertexPositions = StaticMeshAttributes.GetVertexPositions();
 		TEdgeAttributesRef<bool>  EdgeHardnesses = StaticMeshAttributes.GetEdgeHardnesses();
 		TPolygonGroupAttributesRef<FName> PolygonGroupImportedMaterialSlotNames = StaticMeshAttributes.GetPolygonGroupMaterialSlotNames();
@@ -343,6 +346,14 @@ namespace GLTF
 
 			TArray<FVector3f>& Positions = GetVectorBuffer(PositionBufferIndex);
 			Primitive.GetPositions(Positions);
+			if (TotalMatrix != FMatrix::Identity)
+			{
+				for (FVector3f& Position : Positions)
+				{
+					const FVector TransformedPosition = TotalMatrix.TransformPosition(FVector(Position));
+					Position = static_cast<FVector3f>(TransformedPosition);
+				}
+			}
 
 			UpdateVector3fFromMorphTargets(Positions, Primitive, MorphTargetWeights, Vector3fMorphTargetAttributes::POSITIONS);
 
@@ -383,7 +394,7 @@ namespace GLTF
 				ImportPrimitive(Primitive, Index, NumUVs, Mesh.HasTangents(), Mesh.HasColors(),  //
 					VertexInstanceNormals, VertexInstanceTangents, VertexInstanceBinormalSigns, VertexInstanceUVs,
 					VertexInstanceColors,  //
-					EdgeHardnesses, MeshDescription,
+					EdgeHardnesses, MeshGlobalTransform, MeshDescription,
 					MorphTargetWeights);
 
 			bMeshUsesEmptyMaterial |= Primitive.MaterialIndex == INDEX_NONE;
@@ -419,6 +430,7 @@ namespace GLTF
 		const TVertexInstanceAttributesRef<FVector2f>& VertexInstanceUVs,
 		const TVertexInstanceAttributesRef<FVector4f>&  VertexInstanceColors,
 		const TEdgeAttributesRef<bool>&                EdgeHardnesses,
+		const FTransform& MeshGlobalTransform,
 		FMeshDescription* MeshDescription,
 		const TArray<float>& MorphTargetWeights)
 	{
@@ -432,6 +444,23 @@ namespace GLTF
 		//	When the base mesh primitive does not specify normals, client implementations MUST calculate flat normals for each morph target; the provided tangents and their displacements (if present) MUST be ignored.
 		bool bIgnoreTangents = false;
 
+		FMatrix TotalMatrix = ConvertMat(MeshGlobalTransform.ToMatrixWithScale());
+		FMatrix TotalMatrixForNormal;
+		TotalMatrixForNormal = TotalMatrix.Inverse();
+		TotalMatrixForNormal = TotalMatrixForNormal.GetTransposed();
+
+		auto TransformPosition = [](const FMatrix& Matrix, FVector3f& Position)
+		{
+			const FVector TransformedPosition = Matrix.TransformPosition(FVector(Position));
+			Position = static_cast<FVector3f>(TransformedPosition);
+		};
+
+		auto TransformPosition4f = [](const FMatrix& Matrix, FVector4f& Position)
+		{
+			const FVector4d TransformedPosition = Matrix.TransformPosition(FVector4d(Position));
+			Position = static_cast<FVector4f>(TransformedPosition);
+		};
+
 		TArray<FVector3f>& Normals = GetVectorBuffer(NormalBufferIndex);
 		// glTF does not guarantee each primitive within a mesh has the same attributes.
 		// Fill in gaps as needed:
@@ -443,6 +472,14 @@ namespace GLTF
 			TArray<FVector3f>& ReindexBuffer = GetVectorBuffer(ReindexBufferIndex);
 			Primitive.GetNormals(Normals);
 
+			if (TotalMatrixForNormal != FMatrix::Identity)
+			{
+				for (FVector3f& Normal : Normals)
+				{
+					TransformPosition(TotalMatrixForNormal, Normal);
+				}
+			}
+
 			UpdateVector3fFromMorphTargets(Normals, Primitive, MorphTargetWeights, Vector3fMorphTargetAttributes::NORMALS);
 
 			ReIndex(Normals, Indices, ReindexBuffer);
@@ -453,6 +490,14 @@ namespace GLTF
 			TArray<FVector3f>& Positions = GetVectorBuffer(PositionBufferIndex);
 			Primitive.GetPositions(Positions);
 
+			if (TotalMatrix != FMatrix::Identity)
+			{
+				for (FVector3f& Position : Positions)
+				{
+					TransformPosition(TotalMatrix, Position);
+				}
+			}
+
 			//update positions with morph targets before generating flat normals:
 			UpdateVector3fFromMorphTargets(Positions, Primitive, MorphTargetWeights, Vector3fMorphTargetAttributes::POSITIONS);
 
@@ -466,6 +511,14 @@ namespace GLTF
 			TArray<FVector4f>& ReindexBuffer = GetVector4dBuffer(Reindex4dBufferIndex);
 			Primitive.GetTangents(Tangents);
 
+			if (TotalMatrixForNormal != FMatrix::Identity)
+			{
+				for (FVector4f& Tangent : Tangents)
+				{
+					TransformPosition4f(TotalMatrixForNormal, Tangent);
+				}
+			}
+
 			UpdateVector4dFromMorphTargets(Tangents, Primitive, MorphTargetWeights, Vector4dMorphTargetAttributes::TANGENTS);
 
 			ReIndex(Tangents, Indices, ReindexBuffer);
@@ -616,9 +669,9 @@ namespace GLTF
 
 	FMeshFactory::~FMeshFactory() {}
 
-	void FMeshFactory::FillMeshDescription(const GLTF::FMesh &Mesh, FMeshDescription* MeshDescription, const TArray<float>& MorphTargetWeights)
+	void FMeshFactory::FillMeshDescription(const GLTF::FMesh &Mesh, const FTransform& MeshGlobalTransform, FMeshDescription* MeshDescription, const TArray<float>& MorphTargetWeights)
 	{
-		Impl->FillMeshDescription(Mesh, MeshDescription, MorphTargetWeights);
+		Impl->FillMeshDescription(Mesh, MeshGlobalTransform, MeshDescription, MorphTargetWeights);
 	}
 
 	const TArray<FLogMessage>& FMeshFactory::GetLogMessages() const
diff --git a/Engine/Plugins/Interchange/Runtime/Source/Parsers/GLTFCore/Public/GLTFMeshFactory.h b/Engine/Plugins/Interchange/Runtime/Source/Parsers/GLTFCore/Public/GLTFMeshFactory.h
--- a/Engine/Plugins/Interchange/Runtime/Source/Parsers/GLTFCore/Public/GLTFMeshFactory.h
+++ b/Engine/Plugins/Interchange/Runtime/Source/Parsers/GLTFCore/Public/GLTFMeshFactory.h
@@ -23,7 +23,7 @@ namespace GLTF
 
 		using FIndexVertexIdMap = TMap<int32, FVertexID>;
 
-		void FillMeshDescription(const GLTF::FMesh &Mesh, FMeshDescription* MeshDescription, const TArray<float>& MorphTargetWeights = TArray<float>());
+		void FillMeshDescription(const GLTF::FMesh &Mesh, const FTransform& MeshGlobalTransform, FMeshDescription* MeshDescription, const TArray<float>& MorphTargetWeights = TArray<float>());
 
 		float GetUniformScale() const;
 		void  SetUniformScale(float Scale);
diff --git a/Engine/Source/Programs/InterchangeWorker/Private/InterchangeWorkerImpl.cpp b/Engine/Source/Programs/InterchangeWorker/Private/InterchangeWorkerImpl.cpp
--- a/Engine/Source/Programs/InterchangeWorker/Private/InterchangeWorkerImpl.cpp
+++ b/Engine/Source/Programs/InterchangeWorker/Private/InterchangeWorkerImpl.cpp
@@ -223,6 +223,7 @@ void FInterchangeWorkerImpl::ProcessCommand(const TSharedPtr<UE::Interchange::IC
 	TArray<FString> JSonMessages;
 	FJsonLoadSourceCmd LoadSourceCommand;
 	FJsonFetchAnimationBakeTransformPayloadCmd FetchAnimationBakeTransform;
+	FJsonFetchMeshPayloadCmd FetchMeshPayloadCommand;
 	FJsonFetchPayloadCmd FetchPayloadCommand;
 	//Any command FromJson function return true if the Json descibe the command
 	if (LoadSourceCommand.FromJson(JsonToProcess))
@@ -243,6 +244,15 @@ void FInterchangeWorkerImpl::ProcessCommand(const TSharedPtr<UE::Interchange::IC
 			ProcessResult = FetchFbxPayload(FetchAnimationBakeTransform, JSonResult, JSonMessages);
 		}
 	}
+	else if (FetchMeshPayloadCommand.FromJson(JsonToProcess))
+	{
+		//Load file command
+		if (FetchMeshPayloadCommand.GetTranslatorID().Equals(TEXT("FBX"), ESearchCase::IgnoreCase))
+		{
+			//We want to load an FBX file
+			ProcessResult = FetchFbxPayload(FetchMeshPayloadCommand, JSonResult, JSonMessages);
+		}
+	}
 	else if (FetchPayloadCommand.FromJson(JsonToProcess))
 	{
 		//Load file command
@@ -302,6 +312,20 @@ ETaskState FInterchangeWorkerImpl::FetchFbxPayload(const FJsonFetchPayloadCmd& F
 	return ResultState;
 }
 
+ETaskState FInterchangeWorkerImpl::FetchFbxPayload(const FJsonFetchMeshPayloadCmd& FetchMeshPayloadCommand, FString& OutJSonResult, TArray<FString>& OutJSonMessages)
+{
+	ETaskState ResultState = ETaskState::Unknown;
+	FString PayloadKey = FetchMeshPayloadCommand.GetPayloadKey();
+	FTransform MeshGlobalTransform = FetchMeshPayloadCommand.GetMeshGlobalTransform();
+	FbxParser.FetchMeshPayload(PayloadKey, MeshGlobalTransform, ResultFolder);
+	FJsonLoadSourceCmd::JsonResultParser ResultParser;
+	ResultParser.SetResultFilename(FbxParser.GetResultPayloadFilepath(PayloadKey));
+	OutJSonMessages = FbxParser.GetJsonLoadMessages();
+	OutJSonResult = ResultParser.ToJson();
+	ResultState = ETaskState::ProcessOk;
+	return ResultState;
+}
+
 ETaskState FInterchangeWorkerImpl::FetchFbxPayload(const FJsonFetchAnimationBakeTransformPayloadCmd& FetchAnimationBakeTransformPayloadCommand, FString& OutJSonResult, TArray<FString>& OutJSonMessages)
 {
 	ETaskState ResultState = ETaskState::Unknown;
diff --git a/Engine/Source/Programs/InterchangeWorker/Private/InterchangeWorkerImpl.h b/Engine/Source/Programs/InterchangeWorker/Private/InterchangeWorkerImpl.h
--- a/Engine/Source/Programs/InterchangeWorker/Private/InterchangeWorkerImpl.h
+++ b/Engine/Source/Programs/InterchangeWorker/Private/InterchangeWorkerImpl.h
@@ -27,6 +27,7 @@ private:
 
 	UE::Interchange::ETaskState LoadFbxFile(const UE::Interchange::FJsonLoadSourceCmd& LoadSourceCommand, FString& OutJSonResult, TArray<FString>& OutJSonMessages);
 	UE::Interchange::ETaskState FetchFbxPayload(const UE::Interchange::FJsonFetchPayloadCmd& FetchPayloadCommand, FString& OutJSonResult, TArray<FString>& OutJSonMessages);
+	UE::Interchange::ETaskState FetchFbxPayload(const UE::Interchange::FJsonFetchMeshPayloadCmd& FetchMeshPayloadCommand, FString& OutJSonResult, TArray<FString>& OutJSonMessages);
 	UE::Interchange::ETaskState FetchFbxPayload(const UE::Interchange::FJsonFetchAnimationBakeTransformPayloadCmd& FetchAnimationBakeTransformPayloadCommand, FString& OutJSonResult, TArray<FString>& OutJSonMessages);
 private:
 	UE::Interchange::FNetworkClientNode NetworkInterface;
