[Interchange] GLTF.MaterialInstance.M_Transmission handles Opaque and Mask incorrectly fixes.

diff --git a/Engine/Plugins/Interchange/Runtime/Source/Import/Private/Gltf/InterchangeGLTFMaterial.cpp b/Engine/Plugins/Interchange/Runtime/Source/Import/Private/Gltf/InterchangeGLTFMaterial.cpp
--- a/Engine/Plugins/Interchange/Runtime/Source/Import/Private/Gltf/InterchangeGLTFMaterial.cpp
+++ b/Engine/Plugins/Interchange/Runtime/Source/Import/Private/Gltf/InterchangeGLTFMaterial.cpp
@@ -185,6 +185,7 @@ namespace UE::Interchange::GLTFMaterials
 			if (ShadingModel == EShadingModel::TRANSMISSION)
 			{
 				ShaderGraphNode->SetCustomScreenSpaceReflections(true);
+				ShaderGraphNode->SetCustomTwoSidedTransmission(true);
 			}
 
 			TArray<FString> Outputs = GetOutputs(ShadingModel);
@@ -215,10 +216,17 @@ namespace UE::Interchange::GLTFMaterials
 
 			//Set OpacityMaskClipValue if needed:
 			if (!GltfMaterial.bIsUnlitShadingModel &&
-				GltfMaterial.AlphaMode == GLTF::FMaterial::EAlphaMode::Mask &&
-				ShadingModel != EShadingModel::TRANSMISSION) //Transmission does not support AlphaCutoff
+				GltfMaterial.AlphaMode == GLTF::FMaterial::EAlphaMode::Mask)
 			{
-				ShaderGraphNode->SetCustomOpacityMaskClipValue(GltfMaterial.AlphaCutoff);
+				if (ShadingModel == EShadingModel::TRANSMISSION)
+				{
+					//Transmission (even for Materials) supports AlphaCutoff through the MF_Body instead of the automatic OpacityMaskClipValue
+					SetScalar(Inputs::AlphaCutoff, GltfMaterial.AlphaCutoff, 0.f, EProcessType::MATERIAL);
+				}
+				else
+				{
+					ShaderGraphNode->SetCustomOpacityMaskClipValue(GltfMaterial.AlphaCutoff);
+				}
 			}
 		}
 
@@ -226,13 +234,12 @@ namespace UE::Interchange::GLTFMaterials
 		{
 			MaterialInstanceNode = ShaderGraphNode;
 
-			FString ParentIdentifier = GetIdentifier(ShadingModel, AlphaMode, GltfMaterial.bIsDoubleSided && (ShadingModel != EShadingModel::TRANSMISSION)); //Transmission does not support two sided materials, see UInterchangeGenericMaterialPipeline::HandleCommonParameters
+			FString ParentIdentifier = GetIdentifier(ShadingModel, AlphaMode, GltfMaterial.bIsDoubleSided);
 			MaterialInstanceNode->AddStringAttribute(*(InterchangeGltfMaterialAttributeIdentifier + TEXT("ParentIdentifier")), ParentIdentifier);
 
 			//Set AlphaCutOff if needed:
 			if (!GltfMaterial.bIsUnlitShadingModel &&
-				GltfMaterial.AlphaMode == GLTF::FMaterial::EAlphaMode::Mask &&
-				ShadingModel != EShadingModel::TRANSMISSION) //Transmission does not support AlphaCutoff
+				GltfMaterial.AlphaMode == GLTF::FMaterial::EAlphaMode::Mask)
 			{
 				//AlphaCutoff
 				SetScalar(Inputs::AlphaCutoff, GltfMaterial.AlphaCutoff, 0.f, EProcessType::MATERIALINSTANCE);
@@ -526,6 +533,9 @@ namespace UE::Interchange::GLTFMaterials
 
 					//TransmissionFactor
 					SetScalar(Inputs::TransmissionFactor, GltfMaterial.Transmission.TransmissionFactor, 0.f);
+
+					//AlphaMode
+					SetScalar(Inputs::AlphaMode, AlphaMode, EAlphaMode::Blend);
 				}
 			}
 			else if (GltfMaterial.bHasClearCoat)
diff --git a/Engine/Plugins/Interchange/Runtime/Source/Import/Public/Gltf/InterchangeGLTFMaterial.h b/Engine/Plugins/Interchange/Runtime/Source/Import/Public/Gltf/InterchangeGLTFMaterial.h
--- a/Engine/Plugins/Interchange/Runtime/Source/Import/Public/Gltf/InterchangeGLTFMaterial.h
+++ b/Engine/Plugins/Interchange/Runtime/Source/Import/Public/Gltf/InterchangeGLTFMaterial.h
@@ -71,7 +71,7 @@ namespace UE::Interchange::GLTFMaterials
 		const FString IOR = TEXT("IOR");
 
 		const FString AlphaCutoff = TEXT("AlphaCutoff");
-
+		const FString AlphaMode = TEXT("AlphaMode"); //Specifically for Transmission Materials
 
 		//ClearCoat specific:
 		DECLARE_INTERCHANGE_GLTF_MI_MAP(ClearCoat)
diff --git a/Engine/Plugins/Interchange/Runtime/Source/Nodes/Private/InterchangeShaderGraphNode.cpp b/Engine/Plugins/Interchange/Runtime/Source/Nodes/Private/InterchangeShaderGraphNode.cpp
--- a/Engine/Plugins/Interchange/Runtime/Source/Nodes/Private/InterchangeShaderGraphNode.cpp
+++ b/Engine/Plugins/Interchange/Runtime/Source/Nodes/Private/InterchangeShaderGraphNode.cpp
@@ -225,6 +225,16 @@ bool UInterchangeShaderGraphNode::SetCustomTwoSided(const bool& AttributeValue)
 	IMPLEMENT_NODE_ATTRIBUTE_SETTER_NODELEGATE(TwoSided, bool);
 }
 
+bool UInterchangeShaderGraphNode::GetCustomTwoSidedTransmission(bool& AttributeValue) const
+{
+	IMPLEMENT_NODE_ATTRIBUTE_GETTER(TwoSidedTransmission, bool);
+}
+
+bool UInterchangeShaderGraphNode::SetCustomTwoSidedTransmission(const bool& AttributeValue)
+{
+	IMPLEMENT_NODE_ATTRIBUTE_SETTER_NODELEGATE(TwoSidedTransmission, bool);
+}
+
 bool UInterchangeShaderGraphNode::GetCustomOpacityMaskClipValue(float& AttributeValue) const
 {
 	IMPLEMENT_NODE_ATTRIBUTE_GETTER(OpacityMaskClipValue, float);
diff --git a/Engine/Plugins/Interchange/Runtime/Source/Nodes/Public/InterchangeShaderGraphNode.h b/Engine/Plugins/Interchange/Runtime/Source/Nodes/Public/InterchangeShaderGraphNode.h
--- a/Engine/Plugins/Interchange/Runtime/Source/Nodes/Public/InterchangeShaderGraphNode.h
+++ b/Engine/Plugins/Interchange/Runtime/Source/Nodes/Public/InterchangeShaderGraphNode.h
@@ -211,6 +211,15 @@ public:
 	UFUNCTION(BlueprintCallable, Category = "Interchange | Node | Material")
 	bool SetCustomTwoSided(const bool& AttributeValue);
 
+	/**
+	* Forces two sided even for Transmission materials.
+	*/
+	UFUNCTION(BlueprintCallable, Category = "Interchange | Node | Material")
+	bool GetCustomTwoSidedTransmission(bool& AttributeValue) const;
+
+	UFUNCTION(BlueprintCallable, Category = "Interchange | Node | Material")
+	bool SetCustomTwoSidedTransmission(const bool& AttributeValue);
+
 	UFUNCTION(BlueprintCallable, Category = "Interchange | Node | Material")
 	bool GetCustomOpacityMaskClipValue(float& AttributeValue) const;
 
@@ -237,6 +246,7 @@ public:
 
 private:
 	const UE::Interchange::FAttributeKey Macro_CustomTwoSidedKey = UE::Interchange::FAttributeKey(TEXT("TwoSided"));
+	const UE::Interchange::FAttributeKey Macro_CustomTwoSidedTransmissionKey = UE::Interchange::FAttributeKey(TEXT("TwoSidedTransmission"));
 	const UE::Interchange::FAttributeKey Macro_CustomOpacityMaskClipValueKey = UE::Interchange::FAttributeKey(TEXT("OpacityMaskClipValue"));
 	const UE::Interchange::FAttributeKey Macro_CustomIsAShaderFunctionKey = UE::Interchange::FAttributeKey(TEXT("IsAShaderFunction"));
 	const UE::Interchange::FAttributeKey Macro_CustomScreenSpaceReflectionsKey = UE::Interchange::FAttributeKey(TEXT("ScreenSpaceReflections"));
diff --git a/Engine/Plugins/Interchange/Runtime/Source/Pipelines/Private/InterchangeGenericMaterialPipeline.cpp b/Engine/Plugins/Interchange/Runtime/Source/Pipelines/Private/InterchangeGenericMaterialPipeline.cpp
--- a/Engine/Plugins/Interchange/Runtime/Source/Pipelines/Private/InterchangeGenericMaterialPipeline.cpp
+++ b/Engine/Plugins/Interchange/Runtime/Source/Pipelines/Private/InterchangeGenericMaterialPipeline.cpp
@@ -1051,8 +1051,10 @@ void UInterchangeGenericMaterialPipeline::HandleCommonParameters(const UIntercha
 		MaterialFactoryNode->SetCustomScreenSpaceReflections(bScreenSpaceReflections);
 	}
 
+	bool bTwoSidedTransmission = false;
+	ShaderGraphNode->GetCustomTwoSidedTransmission(bTwoSidedTransmission);
 	// Two sidedness (ignored for thin translucency as it looks wrong)
-	if (!HasThinTranslucency(ShaderGraphNode))
+	if (bTwoSidedTransmission || !HasThinTranslucency(ShaderGraphNode))
 	{
 		bool bTwoSided = false;
 		ShaderGraphNode->GetCustomTwoSided(bTwoSided);
