[Interchange - glTF] Mesh (transform) Baking fix

diff --git a/Engine/Plugins/Interchange/Runtime/Source/Parsers/GLTFCore/Private/GLTFMeshFactory.cpp b/Engine/Plugins/Interchange/Runtime/Source/Parsers/GLTFCore/Private/GLTFMeshFactory.cpp
--- a/Engine/Plugins/Interchange/Runtime/Source/Parsers/GLTFCore/Private/GLTFMeshFactory.cpp
+++ b/Engine/Plugins/Interchange/Runtime/Source/Parsers/GLTFCore/Private/GLTFMeshFactory.cpp
@@ -316,7 +316,7 @@ namespace GLTF
 		FStaticMeshAttributes StaticMeshAttributes(*MeshDescription);
 		StaticMeshAttributes.Register();
 
-		FMatrix TotalMatrix = ConvertMat(MeshGlobalTransform.ToMatrixWithScale());
+		FMatrix TotalMatrix = MeshGlobalTransform.ToMatrixWithScale();
 
 		TVertexAttributesRef<FVector3f> VertexPositions = StaticMeshAttributes.GetVertexPositions();
 		TEdgeAttributesRef<bool>  EdgeHardnesses = StaticMeshAttributes.GetEdgeHardnesses();
@@ -444,7 +444,7 @@ namespace GLTF
 		//	When the base mesh primitive does not specify normals, client implementations MUST calculate flat normals for each morph target; the provided tangents and their displacements (if present) MUST be ignored.
 		bool bIgnoreTangents = false;
 
-		FMatrix TotalMatrix = ConvertMat(MeshGlobalTransform.ToMatrixWithScale());
+		FMatrix TotalMatrix = MeshGlobalTransform.ToMatrixWithScale();
 		FMatrix TotalMatrixForNormal;
 		TotalMatrixForNormal = TotalMatrix.Inverse();
 		TotalMatrixForNormal = TotalMatrixForNormal.GetTransposed();
diff --git a/Engine/Plugins/Interchange/Runtime/Source/Parsers/GLTFCore/Public/GLTFMeshFactory.h b/Engine/Plugins/Interchange/Runtime/Source/Parsers/GLTFCore/Public/GLTFMeshFactory.h
--- a/Engine/Plugins/Interchange/Runtime/Source/Parsers/GLTFCore/Public/GLTFMeshFactory.h
+++ b/Engine/Plugins/Interchange/Runtime/Source/Parsers/GLTFCore/Public/GLTFMeshFactory.h
@@ -23,7 +23,7 @@ namespace GLTF
 
 		using FIndexVertexIdMap = TMap<int32, FVertexID>;
 
-		void FillMeshDescription(const GLTF::FMesh &Mesh, const FTransform& MeshGlobalTransform, FMeshDescription* MeshDescription, const TArray<float>& MorphTargetWeights = TArray<float>());
+		void FillMeshDescription(const GLTF::FMesh &Mesh, const FTransform& MeshGlobalTransform /*In UE Space*/, FMeshDescription* MeshDescription, const TArray<float>& MorphTargetWeights = TArray<float>());
 
 		float GetUniformScale() const;
 		void  SetUniformScale(float Scale);
