[Interchange - FBX] Importing model with Rigid Animation produces incorrect orientation on LevelSequence

diff --git a/Engine/Plugins/Interchange/Runtime/Source/Parsers/Fbx/Private/FbxAnimation.cpp b/Engine/Plugins/Interchange/Runtime/Source/Parsers/Fbx/Private/FbxAnimation.cpp
--- a/Engine/Plugins/Interchange/Runtime/Source/Parsers/Fbx/Private/FbxAnimation.cpp
+++ b/Engine/Plugins/Interchange/Runtime/Source/Parsers/Fbx/Private/FbxAnimation.cpp
@@ -26,7 +26,7 @@
 
 namespace UE::Interchange::Private
 {
-	bool ImportCurve(const FbxAnimCurve* SourceFloatCurves, const float ScaleValue, TArray<FInterchangeCurveKey>& DestinationFloatCurve)
+	bool ImportCurve(const FbxAnimCurve* SourceFloatCurves, const float Offset, const float ScaleValue, TArray<FInterchangeCurveKey>& DestinationFloatCurve)
 	{
 		if (!SourceFloatCurves)
 		{
@@ -41,7 +41,7 @@ namespace UE::Interchange::Private
 			FbxAnimCurveKey Key = SourceFloatCurves->KeyGet(KeyIndex);
 			FbxTime KeyTime = Key.GetTime();
 			const float KeyTimeValue = static_cast<float>(KeyTime.GetSecondDouble());
-			float Value = Key.GetValue() * ScaleValue;
+			float Value = (Key.GetValue() * ScaleValue) + Offset;
 			FInterchangeCurveKey& InterchangeCurveKey = DestinationFloatCurve.AddDefaulted_GetRef();
 			InterchangeCurveKey.Time = KeyTimeValue;
 			InterchangeCurveKey.Value = Value;
@@ -648,7 +648,7 @@ namespace UE::Interchange::Private
 					{
 						if (FbxAnimCurve* CurrentAnimCurve = FetchPayloadData.AnimCurves->GetCurve(ChannelIndex, CurveIndex))
 						{
-							ImportCurve(CurrentAnimCurve, 1.0f, InterchangeCurves.AddDefaulted_GetRef().Keys);
+							ImportCurve(CurrentAnimCurve, 0.0f, 1.0f, InterchangeCurves.AddDefaulted_GetRef().Keys);
 						}
 					}
 				}
@@ -670,7 +670,7 @@ namespace UE::Interchange::Private
 			//Fetch TArray<FInterchangeCurve> which are float curves with interpolation
 			TArray<FInterchangeCurve> InterchangeCurves;
 
-			auto ParseInterchangeCurveCurve = [&InterchangeCurves](FbxAnimCurveNode* AnimCurveNode, const char* pChannelName, float Scale = 1.0f)
+			auto ParseInterchangeCurveCurve = [&InterchangeCurves](FbxAnimCurveNode* AnimCurveNode, const char* pChannelName, float Offset = 0.0f, float Scale = 1.0f)
 			{
 				bool bInterchangeCurveAdded = false;
 				if (AnimCurveNode)
@@ -680,7 +680,7 @@ namespace UE::Interchange::Private
 					{
 						if (FbxAnimCurve* CurrentAnimCurve = AnimCurveNode->GetCurve(ChannelIndex))
 						{
-							ImportCurve(CurrentAnimCurve, Scale, InterchangeCurves.AddDefaulted_GetRef().Keys);
+							ImportCurve(CurrentAnimCurve, Offset, Scale, InterchangeCurves.AddDefaulted_GetRef().Keys);
 							bInterchangeCurveAdded = true;
 						}
 					}
@@ -689,13 +689,32 @@ namespace UE::Interchange::Private
 				if (!bInterchangeCurveAdded) InterchangeCurves.AddDefaulted_GetRef();
 			};
 
-			ParseInterchangeCurveCurve(FetchPayloadData.TranlsationCurveNode, FBXSDK_CURVENODE_COMPONENT_X);
-			ParseInterchangeCurveCurve(FetchPayloadData.TranlsationCurveNode, FBXSDK_CURVENODE_COMPONENT_Y, -1.0f); //FBX to UE conversion with scale -1.0
-			ParseInterchangeCurveCurve(FetchPayloadData.TranlsationCurveNode, FBXSDK_CURVENODE_COMPONENT_Z);
+			//As the hierarchical transforms are calculated based on 'EvaluateGlobalTransform'
+			//we calculate the difference between Evaluation and Local values
+			//As the FbxAnimCurve Values do not take into account Pre/Post rotations, offsets, pivots
+			//Note: This is not the utmost ideal scenario as Interchange does not have support for Pre/Post changes nor Pivots
+			//		However, this still prooved to provide the most precise outcome.
+			FbxDouble3 TranslationOffset;
+			FbxDouble3 RotationOffset;
+			{
+				FbxAMatrix LocalFbxMatrix = FetchPayloadData.Node->EvaluateGlobalTransform();
+				if (fbxsdk::FbxNode* ParentNode = FetchPayloadData.Node->GetParent())
+				{
+					FbxAMatrix GlobalFbxParentMatrix = ParentNode->EvaluateGlobalTransform();
+					LocalFbxMatrix = GlobalFbxParentMatrix.Inverse() * LocalFbxMatrix;
+				}
+
+				TranslationOffset = LocalFbxMatrix.GetT() - FetchPayloadData.Node->LclTranslation.Get();
+				RotationOffset = LocalFbxMatrix.GetR() - FetchPayloadData.Node->LclRotation.Get();
+			};
+
+			ParseInterchangeCurveCurve(FetchPayloadData.TranlsationCurveNode, FBXSDK_CURVENODE_COMPONENT_X, TranslationOffset[0]);
+			ParseInterchangeCurveCurve(FetchPayloadData.TranlsationCurveNode, FBXSDK_CURVENODE_COMPONENT_Y, TranslationOffset[1], -1.0f); //FBX to UE conversion with scale -1.0
+			ParseInterchangeCurveCurve(FetchPayloadData.TranlsationCurveNode, FBXSDK_CURVENODE_COMPONENT_Z, TranslationOffset[2]);
 
-			ParseInterchangeCurveCurve(FetchPayloadData.RotationCurveNode, FBXSDK_CURVENODE_COMPONENT_X);
-			ParseInterchangeCurveCurve(FetchPayloadData.RotationCurveNode, FBXSDK_CURVENODE_COMPONENT_Y, -1.0f); //FBX to UE conversion with scale -1.0
-			ParseInterchangeCurveCurve(FetchPayloadData.RotationCurveNode, FBXSDK_CURVENODE_COMPONENT_Z, -1.0f); //FBX to UE conversion with scale -1.0
+			ParseInterchangeCurveCurve(FetchPayloadData.RotationCurveNode, FBXSDK_CURVENODE_COMPONENT_X, RotationOffset[0]);
+			ParseInterchangeCurveCurve(FetchPayloadData.RotationCurveNode, FBXSDK_CURVENODE_COMPONENT_Y, RotationOffset[1], -1.0f); //FBX to UE conversion with scale -1.0
+			ParseInterchangeCurveCurve(FetchPayloadData.RotationCurveNode, FBXSDK_CURVENODE_COMPONENT_Z, RotationOffset[2], -1.0f); //FBX to UE conversion with scale -1.0
 
 			ParseInterchangeCurveCurve(FetchPayloadData.ScaleCurveNode, FBXSDK_CURVENODE_COMPONENT_X);
 			ParseInterchangeCurveCurve(FetchPayloadData.ScaleCurveNode, FBXSDK_CURVENODE_COMPONENT_Y);
@@ -741,7 +760,7 @@ namespace UE::Interchange::Private
 		//Morph target curve in fbx are between 0 and 100, in Unreal we are between 0 and 1, so we must scale
 		//The curve with 0.01
 		constexpr float ScaleCurve = 0.01f;
-		ImportCurve(AnimCurve, ScaleCurve, InterchangeCurves.AddDefaulted_GetRef().Keys);
+		ImportCurve(AnimCurve, 0.0f, ScaleCurve, InterchangeCurves.AddDefaulted_GetRef().Keys);
 		{
 			FLargeMemoryWriter Ar;
 			Ar << InterchangeCurves;
