[Interchange] Hierarchy Import fixes (for RigidMesh.fbx)

diff --git a/Engine/Plugins/Interchange/Runtime/Source/Pipelines/Private/InterchangeGenericScenesPipeline.cpp b/Engine/Plugins/Interchange/Runtime/Source/Pipelines/Private/InterchangeGenericScenesPipeline.cpp
--- a/Engine/Plugins/Interchange/Runtime/Source/Pipelines/Private/InterchangeGenericScenesPipeline.cpp
+++ b/Engine/Plugins/Interchange/Runtime/Source/Pipelines/Private/InterchangeGenericScenesPipeline.cpp
@@ -39,34 +39,41 @@
 
 namespace UE::Interchange::Private
 {
-	//Either a non-joint, or a root joint can be a parent (only those get FactoryNodes) :
+	//Either a (TransformSpecialized || !JointSpecialized || RootJoint) can be a parent (only those get FactoryNodes) :
 	FString FindFactoryParentSceneNodeUid(UInterchangeBaseNodeContainer* BaseNodeContainer, const UInterchangeSceneNode* SceneNode)
 	{
 		FString ParentUid = SceneNode->GetParentUid();
 		if (const UInterchangeSceneNode* ParentSceneNode = Cast<UInterchangeSceneNode>(BaseNodeContainer->GetNode(ParentUid)))
 		{
-			bool bParentIsJoint = ParentSceneNode->IsSpecializedTypeContains(UE::Interchange::FSceneNodeStaticData::GetJointSpecializeTypeString());
-			if (bParentIsJoint)
+			if (ParentSceneNode->IsSpecializedTypeContains(UE::Interchange::FSceneNodeStaticData::GetTransformSpecializeTypeString()))
 			{
-				//Check if its a rootjoint:
-				//	aka check parent's parent if its !joint:
-				FString ParentsParentUid = ParentSceneNode->GetParentUid();
-				if (const UInterchangeSceneNode* ParentsParentSceneNode = Cast<UInterchangeSceneNode>(BaseNodeContainer->GetNode(ParentsParentUid)))
+				return ParentUid;
+			}
+			else
+			{
+				bool bParentIsJoint = ParentSceneNode->IsSpecializedTypeContains(UE::Interchange::FSceneNodeStaticData::GetJointSpecializeTypeString());
+				if (bParentIsJoint)
 				{
-					bool bParentsParentIsJoint = ParentSceneNode->IsSpecializedTypeContains(UE::Interchange::FSceneNodeStaticData::GetJointSpecializeTypeString());
-					if (bParentsParentIsJoint)
+					//Check if its a rootjoint:
+					//	aka check parent's parent if its !joint:
+					FString ParentsParentUid = ParentSceneNode->GetParentUid();
+					if (const UInterchangeSceneNode* ParentsParentSceneNode = Cast<UInterchangeSceneNode>(BaseNodeContainer->GetNode(ParentsParentUid)))
 					{
-						return FindFactoryParentSceneNodeUid(BaseNodeContainer, ParentSceneNode);
-					}
-					else
-					{
-						return ParentUid;
+						bool bParentsParentIsJoint = ParentSceneNode->IsSpecializedTypeContains(UE::Interchange::FSceneNodeStaticData::GetJointSpecializeTypeString());
+						if (bParentsParentIsJoint)
+						{
+							return FindFactoryParentSceneNodeUid(BaseNodeContainer, ParentSceneNode);
+						}
+						else
+						{
+							return ParentUid;
+						}
 					}
 				}
-			}
-			else
-			{
-				return ParentUid;
+				else
+				{
+					return ParentUid;
+				}
 			}
 		}
 
