Interchange: Add "Import Collision" option for static mesh.

#rnx

diff --git a/Engine/Plugins/Interchange/Runtime/Source/FactoryNodes/Private/InterchangeStaticMeshLodDataNode.cpp b/Engine/Plugins/Interchange/Runtime/Source/FactoryNodes/Private/InterchangeStaticMeshLodDataNode.cpp
--- a/Engine/Plugins/Interchange/Runtime/Source/FactoryNodes/Private/InterchangeStaticMeshLodDataNode.cpp
+++ b/Engine/Plugins/Interchange/Runtime/Source/FactoryNodes/Private/InterchangeStaticMeshLodDataNode.cpp
@@ -228,6 +228,16 @@ bool UInterchangeStaticMeshLodDataNode::SetOneConvexHullPerUCX(bool AttributeVal
 	IMPLEMENT_NODE_ATTRIBUTE_SETTER_NODELEGATE(OneConvexHullPerUCX, bool)
 }
 
+bool UInterchangeStaticMeshLodDataNode::GetImportCollision(bool& AttributeValue) const
+{
+	IMPLEMENT_NODE_ATTRIBUTE_GETTER(ImportCollision, bool);
+}
+
+bool UInterchangeStaticMeshLodDataNode::SetImportCollision(bool AttributeValue)
+{
+	IMPLEMENT_NODE_ATTRIBUTE_SETTER_NODELEGATE(ImportCollision, bool)
+}
+
 bool UInterchangeStaticMeshLodDataNode::IsEditorOnlyDataDefined()
 {
 #if WITH_EDITORONLY_DATA
diff --git a/Engine/Plugins/Interchange/Runtime/Source/FactoryNodes/Public/InterchangeStaticMeshLodDataNode.h b/Engine/Plugins/Interchange/Runtime/Source/FactoryNodes/Public/InterchangeStaticMeshLodDataNode.h
--- a/Engine/Plugins/Interchange/Runtime/Source/FactoryNodes/Public/InterchangeStaticMeshLodDataNode.h
+++ b/Engine/Plugins/Interchange/Runtime/Source/FactoryNodes/Public/InterchangeStaticMeshLodDataNode.h
@@ -126,6 +126,12 @@ public:
 	UFUNCTION(BlueprintCallable, Category = "Interchange | Node | StaticMeshLodData")
 	bool SetOneConvexHullPerUCX(bool AttributeValue);
 
+	UFUNCTION(BlueprintCallable, Category = "Interchange | Node | StaticMeshLodData")
+	bool GetImportCollision(bool& AttributeValue) const;
+
+	UFUNCTION(BlueprintCallable, Category = "Interchange | Node | StaticMeshLodData")
+	bool SetImportCollision(bool AttributeValue);
+
 
 private:
 
@@ -138,4 +144,5 @@ private:
 	UE::Interchange::TArrayAttributeHelper<FString> ConvexCollisionMeshUids;
 
 	const UE::Interchange::FAttributeKey Macro_CustomOneConvexHullPerUCXKey = UE::Interchange::FAttributeKey(TEXT("__OneConvexHullPerUCX__Key"));
+	const UE::Interchange::FAttributeKey Macro_CustomImportCollisionKey = UE::Interchange::FAttributeKey(TEXT("__ImportCollision__Key"));
 };
diff --git a/Engine/Plugins/Interchange/Runtime/Source/Import/Private/Mesh/InterchangeStaticMeshFactory.cpp b/Engine/Plugins/Interchange/Runtime/Source/Import/Private/Mesh/InterchangeStaticMeshFactory.cpp
--- a/Engine/Plugins/Interchange/Runtime/Source/Import/Private/Mesh/InterchangeStaticMeshFactory.cpp
+++ b/Engine/Plugins/Interchange/Runtime/Source/Import/Private/Mesh/InterchangeStaticMeshFactory.cpp
@@ -320,6 +320,7 @@ UInterchangeFactoryBase::FImportAssetResult UInterchangeStaticMeshFactory::Impor
 	TArray<FMeshDescription>& LodMeshDescriptions = ImportAssetObjectData.LodMeshDescriptions;
 	LodMeshDescriptions.SetNum(LodCount);
 
+	bool bImportCollision = false;
 	bool bImportedCustomCollision = false;
 	int32 CurrentLodIndex = 0;
 	for (int32 LodIndex = 0; LodIndex < LodCount; ++LodIndex)
@@ -437,10 +438,14 @@ UInterchangeFactoryBase::FImportAssetResult UInterchangeStaticMeshFactory::Impor
 		// Import collision geometry
 		if (CurrentLodIndex == 0)
 		{
-			bImportedCustomCollision |= ImportBoxCollision(Arguments, StaticMesh, LodDataNode);
-			bImportedCustomCollision |= ImportCapsuleCollision(Arguments, StaticMesh, LodDataNode);
-			bImportedCustomCollision |= ImportSphereCollision(Arguments, StaticMesh, LodDataNode);
-			bImportedCustomCollision |= ImportConvexCollision(Arguments, StaticMesh, LodDataNode);
+			LodDataNode->GetImportCollision(bImportCollision);
+			if(bImportCollision)
+			{
+				bImportedCustomCollision |= ImportBoxCollision(Arguments, StaticMesh, LodDataNode);
+				bImportedCustomCollision |= ImportCapsuleCollision(Arguments, StaticMesh, LodDataNode);
+				bImportedCustomCollision |= ImportSphereCollision(Arguments, StaticMesh, LodDataNode);
+				bImportedCustomCollision |= ImportConvexCollision(Arguments, StaticMesh, LodDataNode);
+			}
 		}
 
 		CurrentLodIndex++;
@@ -462,6 +467,7 @@ UInterchangeFactoryBase::FImportAssetResult UInterchangeStaticMeshFactory::Impor
 	}
 #endif // WITH_EDITOR
 
+	ImportAssetObjectData.bImportCollision = bImportCollision;
 	ImportAssetObjectData.bImportedCustomCollision = bImportedCustomCollision;
 
 	// Getting the file Hash will cache it into the source data
@@ -638,16 +644,19 @@ UInterchangeFactoryBase::FImportAssetResult UInterchangeStaticMeshFactory::EndIm
 	}
 #endif // WITH_EDITOR
 
-	if (!ImportAssetObjectData.bImportedCustomCollision)
+	if(ImportAssetObjectData.bImportCollision)
 	{
-		GenerateKDopCollision(Arguments, StaticMesh);
-	}
+		if (!ImportAssetObjectData.bImportedCustomCollision)
+		{
+			GenerateKDopCollision(Arguments, StaticMesh);
+		}
 #if WITH_EDITORONLY_DATA
-	else
-	{
-		StaticMesh->bCustomizedCollision = true;
-	}
+		else
+		{
+			StaticMesh->bCustomizedCollision = true;
+		}
 #endif // WITH_EDITORONLY_DATA
+	}
 #if WITH_EDITOR
 	//Lod group need to use the static mesh API and cannot use the apply delegate
 	if (!Arguments.ReimportObject)
@@ -1294,7 +1303,7 @@ bool UInterchangeStaticMeshFactory::ImportBoxCollision(const FImportAssetObjectP
 	TArray<FMeshPayload> MeshPayloads = GetMeshPayloads(Arguments, BoxCollisionMeshUids);
 	for (const FMeshPayload& MeshPayload : MeshPayloads)
 	{
-		const FTransform& Transform = MeshPayload.Transform;
+		const FTransform Transform = FTransform::Identity;
 		const TOptional<FMeshPayloadData>& PayloadData = MeshPayload.PayloadData.Get();
 
 		if (!PayloadData.IsSet())
@@ -1342,7 +1351,7 @@ bool UInterchangeStaticMeshFactory::ImportCapsuleCollision(const FImportAssetObj
 	TArray<FMeshPayload> MeshPayloads = GetMeshPayloads(Arguments, CapsuleCollisionMeshUids);
 	for (const FMeshPayload& MeshPayload : MeshPayloads)
 	{
-		const FTransform& Transform = MeshPayload.Transform;
+		const FTransform& Transform = FTransform::Identity;
 		const TOptional<FMeshPayloadData>& PayloadData = MeshPayload.PayloadData.Get();
 
 		if (!PayloadData.IsSet())
@@ -1390,7 +1399,7 @@ bool UInterchangeStaticMeshFactory::ImportSphereCollision(const FImportAssetObje
 	TArray<FMeshPayload> MeshPayloads = GetMeshPayloads(Arguments, SphereCollisionMeshUids);
 	for (const FMeshPayload& MeshPayload : MeshPayloads)
 	{
-		const FTransform& Transform = MeshPayload.Transform;
+		const FTransform& Transform = FTransform::Identity;
 		const TOptional<FMeshPayloadData>& PayloadData = MeshPayload.PayloadData.Get();
 
 		if (!PayloadData.IsSet())
@@ -1441,7 +1450,7 @@ bool UInterchangeStaticMeshFactory::ImportConvexCollision(const FImportAssetObje
 	{
 		for (const FMeshPayload& MeshPayload : MeshPayloads)
 		{
-			const FTransform& Transform = MeshPayload.Transform;
+			const FTransform& Transform = FTransform::Identity;
 			const TOptional<FMeshPayloadData>& PayloadData = MeshPayload.PayloadData.Get();
 
 			if (!PayloadData.IsSet())
@@ -1466,7 +1475,7 @@ bool UInterchangeStaticMeshFactory::ImportConvexCollision(const FImportAssetObje
 
 		for (const FMeshPayload& MeshPayload : MeshPayloads)
 		{
-			const FTransform& Transform = MeshPayload.Transform;
+			const FTransform& Transform = FTransform::Identity;
 			TOptional<FMeshPayloadData> PayloadData = MeshPayload.PayloadData.Get();
 
 			if (!PayloadData.IsSet())
diff --git a/Engine/Plugins/Interchange/Runtime/Source/Import/Public/Mesh/InterchangeStaticMeshFactory.h b/Engine/Plugins/Interchange/Runtime/Source/Import/Public/Mesh/InterchangeStaticMeshFactory.h
--- a/Engine/Plugins/Interchange/Runtime/Source/Import/Public/Mesh/InterchangeStaticMeshFactory.h
+++ b/Engine/Plugins/Interchange/Runtime/Source/Import/Public/Mesh/InterchangeStaticMeshFactory.h
@@ -72,6 +72,7 @@ private:
 	struct FImportAssetObjectData
 	{
 		bool bImportedCustomCollision = false;
+		bool bImportCollision = false;
 		TArray<FMeshDescription> LodMeshDescriptions;
 	};
 	FImportAssetObjectData ImportAssetObjectData;
diff --git a/Engine/Plugins/Interchange/Runtime/Source/Pipelines/Private/InterchangeGenericStaticMeshPipeline.cpp b/Engine/Plugins/Interchange/Runtime/Source/Pipelines/Private/InterchangeGenericStaticMeshPipeline.cpp
--- a/Engine/Plugins/Interchange/Runtime/Source/Pipelines/Private/InterchangeGenericStaticMeshPipeline.cpp
+++ b/Engine/Plugins/Interchange/Runtime/Source/Pipelines/Private/InterchangeGenericStaticMeshPipeline.cpp
@@ -264,7 +264,7 @@ void UInterchangeGenericMeshPipeline::ExecutePreImportPipelineStaticMesh()
 
 					if (MeshUidsPerLodIndex.Num() > 0)
 					{
-						if (bImportCollisionAccordingToMeshName)
+						if (bImportCollision && bImportCollisionAccordingToMeshName)
 						{
 							if (const TArray<FString>* CorrespondingCollisionMeshes = MeshToCollisionMeshMap.Find(MeshUid))
 							{
@@ -308,7 +308,7 @@ void UInterchangeGenericMeshPipeline::ExecutePreImportPipelineStaticMesh()
 
 					if (MeshUidsPerLodIndex.Num() > 0)
 					{
-						if (bImportCollisionAccordingToMeshName)
+						if (bImportCollision && bImportCollisionAccordingToMeshName)
 						{
 							if (const TArray<FString>* CorrespondingCollisionMeshes = MeshToCollisionMeshMap.Find(MeshUid))
 							{
@@ -472,6 +472,7 @@ UInterchangeStaticMeshLodDataNode* UInterchangeGenericMeshPipeline::CreateStatic
 
 	StaticMeshLodDataNode->InitializeNode(NodeUID, DisplayLabel, EInterchangeNodeContainerType::FactoryData);
 	StaticMeshLodDataNode->SetOneConvexHullPerUCX(bOneConvexHullPerUCX);
+	StaticMeshLodDataNode->SetImportCollision(bImportCollision);
 	BaseNodeContainer->AddNode(StaticMeshLodDataNode);
 	return StaticMeshLodDataNode;
 }
diff --git a/Engine/Plugins/Interchange/Runtime/Source/Pipelines/Public/InterchangeGenericMeshPipeline.h b/Engine/Plugins/Interchange/Runtime/Source/Pipelines/Public/InterchangeGenericMeshPipeline.h
--- a/Engine/Plugins/Interchange/Runtime/Source/Pipelines/Public/InterchangeGenericMeshPipeline.h
+++ b/Engine/Plugins/Interchange/Runtime/Source/Pipelines/Public/InterchangeGenericMeshPipeline.h
@@ -52,6 +52,12 @@ public:
 	UPROPERTY(EditAnywhere, BlueprintReadWrite, Category = "Static Meshes")
 	FName LodGroup = NAME_None;
 
+	/** If checked, custom collision will be imported and if there is no custom collision, a generic collision will be automatically generated.
+	 * If unchecked there will be no collision created or imported.
+	 */
+	UPROPERTY(EditAnywhere, BlueprintReadWrite, Category = "Static Meshes", meta = (SubCategory = "Collision"))
+	bool bImportCollision = true;
+
 	/**
 	 * If enabled, meshes with certain prefixes will be imported as collision primitives for the mesh with the corresponding unprefixed name.
 	 * 
@@ -61,12 +67,11 @@ public:
 	 * USP_ Sphere collision
 	 * UCX_ Convex collision
 	 */
-
-	UPROPERTY(EditAnywhere, BlueprintReadWrite, Category = "Static Meshes")
+	UPROPERTY(EditAnywhere, BlueprintReadWrite, Category = "Static Meshes", meta = (SubCategory = "Collision", editcondition = "bImportCollision"))
 	bool bImportCollisionAccordingToMeshName = true;
 
 	/** If enabled, each UCX collision mesh will be imported as a single convex hull. If disabled, a UCX mesh will be decomposed into its separate pieces and a convex hull generated for each. */
-	UPROPERTY(EditAnywhere, BlueprintReadWrite, Category = "Static Meshes")
+	UPROPERTY(EditAnywhere, BlueprintReadWrite, Category = "Static Meshes", meta = (SubCategory = "Collision", editcondition = "bImportCollision && bImportCollisionAccordingToMeshName"))
 	bool bOneConvexHullPerUCX = true;
 
 	//////	Static Meshes Build settings Properties //////
