Interchange: Fix morph target fbx import, there was an error in the last change the shape payload global transform matrix was inverse and should not.

#rnx

diff --git a/Engine/Plugins/Interchange/Runtime/Source/Parsers/Fbx/Private/FbxMesh.cpp b/Engine/Plugins/Interchange/Runtime/Source/Parsers/Fbx/Private/FbxMesh.cpp
--- a/Engine/Plugins/Interchange/Runtime/Source/Parsers/Fbx/Private/FbxMesh.cpp
+++ b/Engine/Plugins/Interchange/Runtime/Source/Parsers/Fbx/Private/FbxMesh.cpp
@@ -243,7 +243,7 @@ bool FMeshDescriptionImporter::FillMeshDescriptionFromFbxShape(FbxShape* Shape,
 		MeshDescription->SuspendVertexIndexing();
 
 		// Construct the matrices for the conversion from right handed to left handed system
-		FbxAMatrix TotalMatrix = FFbxConvert::ConvertMatrix(MeshGlobalTransform.ToInverseMatrixWithScale());
+		FbxAMatrix TotalMatrix = FFbxConvert::ConvertMatrix(MeshGlobalTransform.ToMatrixWithScale());
 		FbxAMatrix TotalMatrixForNormal;
 		TotalMatrixForNormal = TotalMatrix.Inverse();
 		TotalMatrixForNormal = TotalMatrixForNormal.Transpose();
