Fixed mesh instancing in level import

#rnx

diff --git a/Engine/Plugins/Interchange/Runtime/Source/Pipelines/Private/InterchangeGenericStaticMeshPipeline.cpp b/Engine/Plugins/Interchange/Runtime/Source/Pipelines/Private/InterchangeGenericStaticMeshPipeline.cpp
--- a/Engine/Plugins/Interchange/Runtime/Source/Pipelines/Private/InterchangeGenericStaticMeshPipeline.cpp
+++ b/Engine/Plugins/Interchange/Runtime/Source/Pipelines/Private/InterchangeGenericStaticMeshPipeline.cpp
@@ -227,6 +227,7 @@ void UInterchangeGenericMeshPipeline::ExecutePreImportPipelineStaticMesh()
 			// Do not combine static meshes
 
 			bool bFoundMeshes = false;
+			if(CommonMeshesProperties->bBakeMeshes)
 			{
 				TArray<FString> MeshUids;
 				PipelineMeshesUtilities->GetAllStaticMeshInstance(MeshUids, bConvertSkeletalMeshToStaticMesh, CommonSkeletalMeshesAndAnimationsProperties->bConvertStaticsWithMorphTargetsToSkeletals);
