Interchange: Prevent adding specialize skeletal mesh source file in the asset import data if the user import both skinning and geometry.

#rnx

diff --git a/Engine/Plugins/Interchange/Runtime/Source/Import/Private/Mesh/InterchangeSkeletalMeshFactory.cpp b/Engine/Plugins/Interchange/Runtime/Source/Import/Private/Mesh/InterchangeSkeletalMeshFactory.cpp
--- a/Engine/Plugins/Interchange/Runtime/Source/Import/Private/Mesh/InterchangeSkeletalMeshFactory.cpp
+++ b/Engine/Plugins/Interchange/Runtime/Source/Import/Private/Mesh/InterchangeSkeletalMeshFactory.cpp
@@ -1734,13 +1734,18 @@ void UInterchangeSkeletalMeshFactory::SetupObject_GameThread(const FSetupObjectP
 					//NewSourceIndex should be 0, 1 or 2 (All, Geo, Skinning)
 					check(NewSourceIndex >= 0 && NewSourceIndex < 3);
 					const TArray<FString> OldFilenames = AssetImportData->ScriptExtractFilenames();
-					const FString DefaultFilename = OldFilenames.Num() > 0 ? OldFilenames[0] : NewSourceFilename;
+					const FString DefaultFilename = OldFilenames.Num() > 0 ? OldFilenames[GetSourceIndexFromContentType(EInterchangeSkeletalMeshContentType::All)] : NewSourceFilename;
 					for (int32 SourceIndex = 0; SourceIndex < 3; ++SourceIndex)
 					{
 						FString SourceLabel = GetSourceLabelFromSourceIndex(SourceIndex);
 						if (SourceIndex == NewSourceIndex)
 						{
 							AssetImportData->ScriptedAddFilename(NewSourceFilename, SourceIndex, SourceLabel);
+							//Do not add any extra path if import both geo and skinning
+							if(SourceIndex == GetSourceIndexFromContentType(EInterchangeSkeletalMeshContentType::All))
+							{
+								break;
+							}
 						}
 						else
 						{
