Interchange: Make sure the option bImportMeshesInBoneHierarchy is supported correctly.

#rnx

diff --git a/Engine/Plugins/Interchange/Runtime/Source/Dispatcher/Public/InterchangeCommands.h b/Engine/Plugins/Interchange/Runtime/Source/Dispatcher/Public/InterchangeCommands.h
--- a/Engine/Plugins/Interchange/Runtime/Source/Dispatcher/Public/InterchangeCommands.h
+++ b/Engine/Plugins/Interchange/Runtime/Source/Dispatcher/Public/InterchangeCommands.h
@@ -17,7 +17,7 @@ namespace UE
 		{
 		public:
 			//Major version should be updated when there is an existing API that has been change.
-			static int32 GetMajor() { return 13; }
+			static int32 GetMajor() { return 14; }
 			//Minor version must be updated when there is an addition to the API.
 			static int32 GetMinor() { return 0; }
 			//Patch version should be update if there is some bug fixes in the private code.
diff --git a/Engine/Plugins/Interchange/Runtime/Source/Parsers/Fbx/Private/FbxMesh.cpp b/Engine/Plugins/Interchange/Runtime/Source/Parsers/Fbx/Private/FbxMesh.cpp
--- a/Engine/Plugins/Interchange/Runtime/Source/Parsers/Fbx/Private/FbxMesh.cpp
+++ b/Engine/Plugins/Interchange/Runtime/Source/Parsers/Fbx/Private/FbxMesh.cpp
@@ -1623,19 +1623,15 @@ bool FFbxMesh::ExtractSkinnedMeshNodeJoints(FbxScene* SDKScene, UInterchangeBase
 				continue;
 			}
 			FbxNode* Link = Cluster->GetLink();
-			//Any valid link should be under the scene root node
-			if (Link->GetParent() != nullptr)
-			{
-				bFoundValidJoint = true;
+			bFoundValidJoint = true;
 
-				FString JointNodeUniqueID = Parser.GetFbxHelper()->GetFbxNodeHierarchyName(Link);
+			FString JointNodeUniqueID = Parser.GetFbxHelper()->GetFbxNodeHierarchyName(Link);
 
-				// find the bone index
-				if (!JointNodeUniqueIDs.Contains(JointNodeUniqueID))
-				{
-					JointNodeUniqueIDs.Add(JointNodeUniqueID);
-					MeshNode->SetSkeletonDependencyUid(JointNodeUniqueID);
-				}
+			// find the bone index
+			if (!JointNodeUniqueIDs.Contains(JointNodeUniqueID))
+			{
+				JointNodeUniqueIDs.Add(JointNodeUniqueID);
+				MeshNode->SetSkeletonDependencyUid(JointNodeUniqueID);
 			}
 		}
 	}
diff --git a/Engine/Plugins/Interchange/Runtime/Source/Parsers/Fbx/Private/FbxScene.cpp b/Engine/Plugins/Interchange/Runtime/Source/Parsers/Fbx/Private/FbxScene.cpp
--- a/Engine/Plugins/Interchange/Runtime/Source/Parsers/Fbx/Private/FbxScene.cpp
+++ b/Engine/Plugins/Interchange/Runtime/Source/Parsers/Fbx/Private/FbxScene.cpp
@@ -118,14 +118,8 @@ namespace UE
 				
 				auto GetConvertedTransform = [Node](FbxAMatrix& NewFbxMatrix)
 				{
-					FTransform Transform;
-					FbxVector4 NewLocalT = NewFbxMatrix.GetT();
-					FbxVector4 NewLocalS = NewFbxMatrix.GetS();
-					FbxQuaternion NewLocalQ = NewFbxMatrix.GetQ();
-					Transform.SetTranslation(FFbxConvert::ConvertPos(NewLocalT));
-					Transform.SetScale3D(FFbxConvert::ConvertScale(NewLocalS));
-					Transform.SetRotation(FFbxConvert::ConvertRotToQuat(NewLocalQ));
-
+					FTransform Transform = FTransform(FFbxConvert::ConvertMatrix(NewFbxMatrix));
+					
 					if (FbxNodeAttribute* NodeAttribute = Node->GetNodeAttribute())
 					{
 						switch (NodeAttribute->GetAttributeType())
@@ -444,9 +438,24 @@ namespace UE
 
 			void FFbxScene::AddHierarchy(FbxScene* SDKScene, UInterchangeBaseNodeContainer& NodeContainer, TMap<FString, TSharedPtr<FPayloadContextBase, ESPMode::ThreadSafe>>& PayloadContexts)
 			{
-				 FbxNode* RootNode = SDKScene->GetRootNode();
+				FbxNode* RootNode = SDKScene->GetRootNode();
 
-				 AddHierarchyRecursively(nullptr, RootNode, SDKScene, NodeContainer, PayloadContexts);
+				AddHierarchyRecursively(nullptr, RootNode, SDKScene, NodeContainer, PayloadContexts);
+
+				int32 NodeCount = SDKScene->GetNodeCount();
+				for (int32 NodeIndex = 0; NodeIndex < NodeCount; ++NodeIndex)
+				{
+					if (FbxNode* Node = SDKScene->GetNode(NodeIndex))
+					{
+						if(Node != RootNode)
+						{
+							if (Node->GetParent() == nullptr)
+							{
+								AddHierarchyRecursively(nullptr, Node, SDKScene, NodeContainer, PayloadContexts);
+							}
+						}
+					}
+				}
 			}
 
 			void FFbxScene::AddRigidAnimation(FbxNode* Node
@@ -704,6 +713,20 @@ namespace UE
 				for (int32 AnimationIndex = 0; AnimationIndex < NumAnimations; AnimationIndex++)
 				{
 					AddAnimationRecursively(RootNode, SDKScene, NodeContainer, PayloadContexts, nullptr, false, RootSceneNodeUniqueID, SkeletonRootNodeUids, AnimationIndex);
+					int32 NodeCount = SDKScene->GetNodeCount();
+					for (int32 NodeIndex = 0; NodeIndex < NodeCount; ++NodeIndex)
+					{
+						if (FbxNode* Node = SDKScene->GetNode(NodeIndex))
+						{
+							if (Node != RootNode)
+							{
+								if (Node->GetParent() == nullptr)
+								{
+									AddAnimationRecursively(Node, SDKScene, NodeContainer, PayloadContexts, nullptr, false, RootSceneNodeUniqueID, SkeletonRootNodeUids, AnimationIndex);
+								}
+							}
+						}
+					}
 				}
 
 				TArray<FString> TransformAnimTrackNodeUids;
diff --git a/Engine/Plugins/Interchange/Runtime/Source/Parsers/Fbx/Private/InterchangeFbxParser.cpp b/Engine/Plugins/Interchange/Runtime/Source/Parsers/Fbx/Private/InterchangeFbxParser.cpp
--- a/Engine/Plugins/Interchange/Runtime/Source/Parsers/Fbx/Private/InterchangeFbxParser.cpp
+++ b/Engine/Plugins/Interchange/Runtime/Source/Parsers/Fbx/Private/InterchangeFbxParser.cpp
@@ -90,7 +90,7 @@ namespace UE::Interchange
 			FString& PayloadFilepath = ResultPayloads.FindOrAdd(PayloadKey);
 			//To avoid file path with too many character, we hash the payloadKey so we have a deterministic length for the file path.
 			FString PayloadKeyHash = Private::HashString(PayloadKey);
-			PayloadFilepath = ResultFolder + TEXT("/") + PayloadKeyHash + TEXT(".payload");
+			PayloadFilepath = ResultFolder + TEXT("/") + PayloadKeyHash + FString::FromInt(UniqueIdCounter.IncrementExchange()) + TEXT(".payload");
 
 			//Copy the map filename key because we are multithreaded and the TMap can be reallocated
 			PayloadFilepathCopy = PayloadFilepath;
@@ -114,7 +114,7 @@ namespace UE::Interchange
 			FString& PayloadFilepath = ResultPayloads.FindOrAdd(PayloadKey);
 			//To avoid file path with too many character, we hash the payloadKey so we have a deterministic length for the file path.
 			FString PayloadKeyHash = Private::HashString(PayloadKey);
-			PayloadFilepath = ResultFolder + TEXT("/") + PayloadKeyHash + TEXT(".payload");
+			PayloadFilepath = ResultFolder + TEXT("/") + PayloadKeyHash + FString::FromInt(UniqueIdCounter.IncrementExchange()) + TEXT(".payload");
 
 			//Copy the map filename key because we are multithreaded and the TMap can be reallocated
 			PayloadFilepathCopy = PayloadFilepath;
@@ -138,7 +138,7 @@ namespace UE::Interchange
 			FString& PayloadFilepath = ResultPayloads.FindOrAdd(PayloadKey);
 			//To avoid file path with too many character, we hash the payloadKey so we have a deterministic length for the file path.
 			FString PayloadKeyHash = Private::HashString(PayloadKey);
-			PayloadFilepath = ResultFolder + TEXT("/") + PayloadKeyHash + TEXT(".payload");
+			PayloadFilepath = ResultFolder + TEXT("/") + PayloadKeyHash + FString::FromInt(UniqueIdCounter.IncrementExchange()) + TEXT(".payload");
 
 			//Copy the map filename key because we are multithreaded and the TMap can be reallocated
 			PayloadFilepathCopy = PayloadFilepath;
diff --git a/Engine/Plugins/Interchange/Runtime/Source/Parsers/Fbx/Public/InterchangeFbxParser.h b/Engine/Plugins/Interchange/Runtime/Source/Parsers/Fbx/Public/InterchangeFbxParser.h
--- a/Engine/Plugins/Interchange/Runtime/Source/Parsers/Fbx/Public/InterchangeFbxParser.h
+++ b/Engine/Plugins/Interchange/Runtime/Source/Parsers/Fbx/Public/InterchangeFbxParser.h
@@ -86,6 +86,7 @@ namespace UE
 			mutable FCriticalSection ResultPayloadsCriticalSection;
 			TMap<FString, FString> ResultPayloads;
 			TUniquePtr<UE::Interchange::Private::FFbxParser> FbxParserPrivate;
+			TAtomic<int64> UniqueIdCounter = 0;
 		};
 	} // ns Interchange
 }//ns UE
diff --git a/Engine/Plugins/Interchange/Runtime/Source/Pipelines/Private/InterchangeGenericAnimationPipeline.cpp b/Engine/Plugins/Interchange/Runtime/Source/Pipelines/Private/InterchangeGenericAnimationPipeline.cpp
--- a/Engine/Plugins/Interchange/Runtime/Source/Pipelines/Private/InterchangeGenericAnimationPipeline.cpp
+++ b/Engine/Plugins/Interchange/Runtime/Source/Pipelines/Private/InterchangeGenericAnimationPipeline.cpp
@@ -27,6 +27,9 @@ void UInterchangeGenericAnimationPipeline::AdjustSettingsForContext(EInterchange
 
 	check(CommonSkeletalMeshesAndAnimationsProperties.IsValid());
 	
+	bSceneImport = ImportType == EInterchangePipelineContext::SceneImport
+				|| ImportType == EInterchangePipelineContext::SceneReimport;
+
 	if (ImportType == EInterchangePipelineContext::AssetCustomLODImport
 		|| ImportType == EInterchangePipelineContext::AssetCustomLODReimport
 		|| ImportType == EInterchangePipelineContext::AssetAlternateSkinningImport
@@ -126,8 +129,9 @@ void UInterchangeGenericAnimationPipeline::ExecutePipeline(UInterchangeBaseNodeC
 		}
 	}
 
-	if (CommonMeshesProperties->ForceAllMeshAsType == EInterchangeForceMeshType::IFMT_SkeletalMesh)
+	if (!bSceneImport)
 	{
+		//Support rigid mesh animation animation data  (UAnimSequence for rigid mesh)
 		for (UInterchangeAnimationTrackSetNode* TrackSetNode : TrackSetNodes)
 		{
 			if (!TrackSetNode) continue;
@@ -165,6 +169,7 @@ void UInterchangeGenericAnimationPipeline::ExecutePipeline(UInterchangeBaseNodeC
 								if (!bCustomSkeletonNodeUidSet)
 								{
 									FString SkeletonRootUid;
+									FString LastSceneNode = ActorNodeUid;
 									if (const UInterchangeSceneNode* SceneNode = ActorNode)
 									{
 										FString ParentUid = SceneNode->GetParentUid();
@@ -172,10 +177,18 @@ void UInterchangeGenericAnimationPipeline::ExecutePipeline(UInterchangeBaseNodeC
 										{
 											if (const UInterchangeSceneNode* ParentNode = Cast<UInterchangeSceneNode>(BaseNodeContainer->GetNode(ParentUid)))
 											{
-												SkeletonRootUid = ParentUid;
+												if(ParentNode->IsSpecializedTypeContains(UE::Interchange::FSceneNodeStaticData::GetJointSpecializeTypeString()))
+												{
+													SkeletonRootUid = ParentUid;
+												}
+												LastSceneNode = ParentUid;
 												ParentUid = ParentNode->GetParentUid();
 											}
 										}
+										if (SkeletonRootUid.IsEmpty())
+										{
+											SkeletonRootUid = LastSceneNode;
+										}
 									}
 
 									SkeletalAnimationNode->SetCustomSkeletonNodeUid(SkeletonRootUid);
@@ -195,6 +208,7 @@ void UInterchangeGenericAnimationPipeline::ExecutePipeline(UInterchangeBaseNodeC
 	}
 	else
 	{
+		//Support scene node animation (ULevelSequence, only supported when doing scene import)
 		for (UInterchangeAnimationTrackSetNode* TrackSetNode : TrackSetNodes)
 		{
 			if (TrackSetNode)
@@ -226,6 +240,7 @@ void UInterchangeGenericAnimationPipeline::ExecutePipeline(UInterchangeBaseNodeC
 			TrackNodes.Add(Node);
 		});
 
+	//Support skeletal mesh animation (UAnimSequence)
 	for (UInterchangeSkeletalAnimationTrackNode* TrackNode : TrackNodes)
 	{
 		if (TrackNode)
diff --git a/Engine/Plugins/Interchange/Runtime/Source/Pipelines/Private/InterchangeGenericMeshPipeline.cpp b/Engine/Plugins/Interchange/Runtime/Source/Pipelines/Private/InterchangeGenericMeshPipeline.cpp
--- a/Engine/Plugins/Interchange/Runtime/Source/Pipelines/Private/InterchangeGenericMeshPipeline.cpp
+++ b/Engine/Plugins/Interchange/Runtime/Source/Pipelines/Private/InterchangeGenericMeshPipeline.cpp
@@ -146,9 +146,17 @@ void UInterchangeGenericMeshPipeline::ExecutePipeline(UInterchangeBaseNodeContai
 	{
 		SourceDatas.Add(SourceData);
 	}
-
 	PipelineMeshesUtilities = UInterchangePipelineMeshesUtilities::CreateInterchangePipelineMeshesUtilities(BaseNodeContainer);
 
+	//Set the context option to use when querying the pipeline mesh utilities
+	FInterchangePipelineMeshesUtilitiesContext DataContext;
+	DataContext.bConvertStaticMeshToSkeletalMesh = (CommonMeshesProperties->ForceAllMeshAsType == EInterchangeForceMeshType::IFMT_SkeletalMesh);
+	DataContext.bConvertSkeletalMeshToStaticMesh = (CommonMeshesProperties->ForceAllMeshAsType == EInterchangeForceMeshType::IFMT_StaticMesh);
+	DataContext.bConvertStaticsWithMorphTargetsToSkeletals = CommonSkeletalMeshesAndAnimationsProperties->bConvertStaticsWithMorphTargetsToSkeletals;
+	DataContext.bImportMeshesInBoneHierarchy = CommonSkeletalMeshesAndAnimationsProperties->bImportMeshesInBoneHierarchy;
+	DataContext.bQueryGeometryOnlyIfNoInstance = true;
+	PipelineMeshesUtilities->SetContext(DataContext);
+
 	//Create skeletalmesh factory nodes
 	ExecutePreImportPipelineSkeletalMesh();
 
diff --git a/Engine/Plugins/Interchange/Runtime/Source/Pipelines/Private/InterchangeGenericSkeletalMeshPipeline.cpp b/Engine/Plugins/Interchange/Runtime/Source/Pipelines/Private/InterchangeGenericSkeletalMeshPipeline.cpp
--- a/Engine/Plugins/Interchange/Runtime/Source/Pipelines/Private/InterchangeGenericSkeletalMeshPipeline.cpp
+++ b/Engine/Plugins/Interchange/Runtime/Source/Pipelines/Private/InterchangeGenericSkeletalMeshPipeline.cpp
@@ -104,7 +104,6 @@ void UInterchangeGenericMeshPipeline::ExecutePreImportPipelineSkeletalMesh()
 		});
 #endif
 
-	const bool bConvertStaticMeshToSkeletalMesh = (CommonMeshesProperties->ForceAllMeshAsType == EInterchangeForceMeshType::IFMT_SkeletalMesh);
 	TMap<FString, TArray<FString>> SkeletalMeshFactoryDependencyOrderPerSkeletonRootNodeUid;
 
 	auto SetSkeletalMeshDependencies = [&SkeletalMeshFactoryDependencyOrderPerSkeletonRootNodeUid](const FString& JointNodeUid, UInterchangeSkeletalMeshFactoryNode* SkeletalMeshFactoryNode)
@@ -174,18 +173,10 @@ void UInterchangeGenericMeshPipeline::ExecutePreImportPipelineSkeletalMesh()
 			return bFoundInstances;
 		};
 
-		PipelineMeshesUtilities->GetCombinedSkinnedMeshInstances(BaseNodeContainer, MeshUidsPerSkeletonRootUid, bConvertStaticMeshToSkeletalMesh, CommonSkeletalMeshesAndAnimationsProperties->bConvertStaticsWithMorphTargetsToSkeletals);
-		bool bUseMeshInstance = true;
-		bool bFoundMeshes = CreatePerSkeletonRootUidCombinedSkinnedMesh(bUseMeshInstance);
 
-		if (!bFoundMeshes)
-		{
-			//Fall back to support mesh not reference by any scene node
-			MeshUidsPerSkeletonRootUid.Empty();
-			PipelineMeshesUtilities->GetCombinedSkinnedMeshGeometries(MeshUidsPerSkeletonRootUid);
-			bUseMeshInstance = false;
-			CreatePerSkeletonRootUidCombinedSkinnedMesh(bUseMeshInstance);
-		}
+		PipelineMeshesUtilities->GetCombinedSkinnedMeshInstances(MeshUidsPerSkeletonRootUid);
+		bool bUseMeshInstance = true;
+		CreatePerSkeletonRootUidCombinedSkinnedMesh(bUseMeshInstance);
 	}
 	else
 	{
@@ -252,17 +243,9 @@ void UInterchangeGenericMeshPipeline::ExecutePreImportPipelineSkeletalMesh()
 			return bFoundInstances;
 		};
 		
-		PipelineMeshesUtilities->GetAllSkinnedMeshInstance(MeshUids, bConvertStaticMeshToSkeletalMesh, CommonSkeletalMeshesAndAnimationsProperties->bConvertStaticsWithMorphTargetsToSkeletals);
+		PipelineMeshesUtilities->GetAllSkinnedMeshInstance(MeshUids);
 		bool bUseMeshInstance = true;
-		bool bFoundMeshes = CreatePerSkeletonRootUidSkinnedMesh(bUseMeshInstance);
-
-		if (!bFoundMeshes)
-		{
-			MeshUids.Empty();
-			PipelineMeshesUtilities->GetAllSkinnedMeshGeometry(MeshUids);
-			bUseMeshInstance = false;
-			CreatePerSkeletonRootUidSkinnedMesh(bUseMeshInstance);
-		}
+		CreatePerSkeletonRootUidSkinnedMesh(bUseMeshInstance);
 	}
 }
 
diff --git a/Engine/Plugins/Interchange/Runtime/Source/Pipelines/Private/InterchangeGenericStaticMeshPipeline.cpp b/Engine/Plugins/Interchange/Runtime/Source/Pipelines/Private/InterchangeGenericStaticMeshPipeline.cpp
--- a/Engine/Plugins/Interchange/Runtime/Source/Pipelines/Private/InterchangeGenericStaticMeshPipeline.cpp
+++ b/Engine/Plugins/Interchange/Runtime/Source/Pipelines/Private/InterchangeGenericStaticMeshPipeline.cpp
@@ -159,7 +159,6 @@ void UInterchangeGenericMeshPipeline::ExecutePreImportPipelineStaticMesh()
 
 	if (bImportStaticMeshes && (CommonMeshesProperties->ForceAllMeshAsType == EInterchangeForceMeshType::IFMT_None || CommonMeshesProperties->ForceAllMeshAsType == EInterchangeForceMeshType::IFMT_StaticMesh))
 	{
-		const bool bConvertSkeletalMeshToStaticMesh = (CommonMeshesProperties->ForceAllMeshAsType == EInterchangeForceMeshType::IFMT_StaticMesh);
 		if (bCombineStaticMeshes)
 		{
 			// Combine all the static meshes
@@ -168,7 +167,7 @@ void UInterchangeGenericMeshPipeline::ExecutePreImportPipelineStaticMesh()
 			{
 				// If baking transforms, get all the static mesh instance nodes, and group them by LOD
 				TArray<FString> MeshUids;
-				PipelineMeshesUtilities->GetAllStaticMeshInstance(MeshUids, bConvertSkeletalMeshToStaticMesh, CommonSkeletalMeshesAndAnimationsProperties->bConvertStaticsWithMorphTargetsToSkeletals);
+				PipelineMeshesUtilities->GetAllStaticMeshInstance(MeshUids);
 
 				TMap<int32, TArray<FString>> MeshUidsPerLodIndex;
 
@@ -201,7 +200,7 @@ void UInterchangeGenericMeshPipeline::ExecutePreImportPipelineStaticMesh()
 			{
 				// If we haven't yet managed to build a factory node, look at static mesh geometry directly.
 				TArray<FString> MeshUids;
-				PipelineMeshesUtilities->GetAllStaticMeshGeometry(MeshUids, bConvertSkeletalMeshToStaticMesh, CommonSkeletalMeshesAndAnimationsProperties->bConvertStaticsWithMorphTargetsToSkeletals);
+				PipelineMeshesUtilities->GetAllStaticMeshGeometry(MeshUids);
 
 				TMap<int32, TArray<FString>> MeshUidsPerLodIndex;
 
@@ -230,7 +229,7 @@ void UInterchangeGenericMeshPipeline::ExecutePreImportPipelineStaticMesh()
 			if(CommonMeshesProperties->bBakeMeshes)
 			{
 				TArray<FString> MeshUids;
-				PipelineMeshesUtilities->GetAllStaticMeshInstance(MeshUids, bConvertSkeletalMeshToStaticMesh, CommonSkeletalMeshesAndAnimationsProperties->bConvertStaticsWithMorphTargetsToSkeletals);
+				PipelineMeshesUtilities->GetAllStaticMeshInstance(MeshUids);
 
 				// Work out which meshes are collision meshes which correspond to another mesh
 				TMap<FString, TArray<FString>> MeshToCollisionMeshMap;
@@ -283,7 +282,7 @@ void UInterchangeGenericMeshPipeline::ExecutePreImportPipelineStaticMesh()
 			if (!bFoundMeshes)
 			{
 				TArray<FString> MeshUids;
-				PipelineMeshesUtilities->GetAllStaticMeshGeometry(MeshUids, bConvertSkeletalMeshToStaticMesh, CommonSkeletalMeshesAndAnimationsProperties->bConvertStaticsWithMorphTargetsToSkeletals);
+				PipelineMeshesUtilities->GetAllStaticMeshGeometry(MeshUids);
 
 				// Work out which meshes are collision meshes which correspond to another mesh
 				TMap<FString, TArray<FString>> MeshToCollisionMeshMap;
diff --git a/Engine/Plugins/Interchange/Runtime/Source/Pipelines/Private/InterchangePipelineMeshesUtilities.cpp b/Engine/Plugins/Interchange/Runtime/Source/Pipelines/Private/InterchangePipelineMeshesUtilities.cpp
--- a/Engine/Plugins/Interchange/Runtime/Source/Pipelines/Private/InterchangePipelineMeshesUtilities.cpp
+++ b/Engine/Plugins/Interchange/Runtime/Source/Pipelines/Private/InterchangePipelineMeshesUtilities.cpp
@@ -100,6 +100,69 @@ namespace UE::Private::InterchangeMeshPipeline
 			LodChildUids.Shrink();
 		}
 	}
+
+	void CollectAllChildren(UInterchangeBaseNodeContainer* BaseNodeContainer, const FString SceneNodeUids, TArray<FString>& AllChildrenUids)
+	{
+		AllChildrenUids.Add(SceneNodeUids);
+		TArray<FString> ChildrenUids = BaseNodeContainer->GetNodeChildrenUids(SceneNodeUids);
+		for (const FString& ChildUid : ChildrenUids)
+		{
+			CollectAllChildren(BaseNodeContainer, ChildUid, AllChildrenUids);
+		}
+	}
+
+	bool IsSceneNodeNestedInSkeleton(UInterchangeBaseNodeContainer* BaseNodeContainer, const UInterchangeSceneNode* SceneNode)
+	{
+		auto IsJointHookToDeformer = [&BaseNodeContainer](const FString JointNodeUid)->bool
+		{
+			TArray<FString> AllChildrenUids;
+			CollectAllChildren(BaseNodeContainer, JointNodeUid, AllChildrenUids);
+			bool bResult = false;
+			BaseNodeContainer->BreakableIterateNodesOfType<UInterchangeMeshNode>([&AllChildrenUids, &bResult](const FString& MeshUid, UInterchangeMeshNode* MeshNode)
+			{
+				if (MeshNode->IsSkinnedMesh())
+				{
+					TArray<FString> SkeletonDependencies;
+					MeshNode->GetSkeletonDependencies(SkeletonDependencies);
+					for (const FString& MeshJointUid : SkeletonDependencies)
+					{
+						if (AllChildrenUids.Contains(MeshJointUid))
+						{
+							bResult = true;
+							return true;
+						}
+					}
+				}
+				return false;
+			});
+
+			return bResult;
+		};
+
+		if (SceneNode->IsSpecializedTypeContains(UE::Interchange::FSceneNodeStaticData::GetJointSpecializeTypeString()))
+		{
+			if(IsJointHookToDeformer(SceneNode->GetUniqueID()))
+			{
+				return true;
+			}
+		}
+		FString ParentUid = SceneNode->GetParentUid();
+		const UInterchangeSceneNode* ParentNode = Cast<UInterchangeSceneNode>(BaseNodeContainer->GetNode(ParentUid));
+		while (ParentNode)
+		{
+			if (ParentNode->IsSpecializedTypeContains(UE::Interchange::FSceneNodeStaticData::GetJointSpecializeTypeString()))
+			{
+				if (IsJointHookToDeformer(ParentNode->GetUniqueID()))
+				{
+					return true;
+				}
+			}
+			ParentUid = ParentNode->GetParentUid();
+			ParentNode = Cast<UInterchangeSceneNode>(BaseNodeContainer->GetNode(ParentUid));
+		}
+		return false;
+	}
+
 }
 
 static bool IsSceneNodeASocket(const UInterchangeSceneNode* SceneNode)
@@ -110,12 +173,85 @@ static bool IsSceneNodeASocket(const UInterchangeSceneNode* SceneNode)
 	return NodeDisplayName.StartsWith(TEXT("SOCKET_"));
 }
 
+bool FInterchangePipelineMeshesUtilitiesContext::IsStaticMeshInstance(const FInterchangeMeshInstance& MeshInstance, UInterchangeBaseNodeContainer* BaseNodeContainer)
+{
+	return !IsSkeletalMeshInstance(MeshInstance, BaseNodeContainer);
+}
+
+bool FInterchangePipelineMeshesUtilitiesContext::IsSkeletalMeshInstance(const FInterchangeMeshInstance& MeshInstance, UInterchangeBaseNodeContainer* BaseNodeContainer)
+{
+	bool bOutIsStaticMeshNestedInSkeleton = false;
+	return IsSkeletalMeshInstance(MeshInstance, BaseNodeContainer, bOutIsStaticMeshNestedInSkeleton);
+}
+
+bool FInterchangePipelineMeshesUtilitiesContext::IsSkeletalMeshInstance(const FInterchangeMeshInstance& MeshInstance, UInterchangeBaseNodeContainer* BaseNodeContainer, bool& bOutIsStaticMeshNestedInSkeleton)
+{
+	bOutIsStaticMeshNestedInSkeleton = false;
+	if (bConvertSkeletalMeshToStaticMesh)
+	{
+		return false;
+	}
+	if ((bConvertStaticMeshToSkeletalMesh || MeshInstance.bReferenceSkinnedMesh || (bConvertStaticsWithMorphTargetsToSkeletals && MeshInstance.bHasMorphTargets)) && !MeshInstance.bReferenceMorphTarget)
+	{
+		return true;
+	}
+	else if (bImportMeshesInBoneHierarchy)
+	{
+		if (MeshInstance.SceneNodePerLodIndex.Contains(0))
+		{
+			const FInterchangeLodSceneNodeContainer& LodSceneNodeContainer = MeshInstance.SceneNodePerLodIndex.FindChecked(0);
+			for (const TObjectPtr<const UInterchangeSceneNode>& SceneNode : LodSceneNodeContainer.SceneNodes)
+			{
+				if (UE::Private::InterchangeMeshPipeline::IsSceneNodeNestedInSkeleton(BaseNodeContainer, SceneNode))
+				{
+					bOutIsStaticMeshNestedInSkeleton = true;
+					return true;
+				}
+			}
+		}
+	}
+	return false;
+}
+
+bool FInterchangePipelineMeshesUtilitiesContext::IsStaticMeshGeometry(const FInterchangeMeshGeometry& MeshGeometry)
+{
+	if (bQueryGeometryOnlyIfNoInstance && MeshGeometry.ReferencingMeshInstanceUids.Num() > 0)
+	{
+		return false;
+	}
+
+	return !IsSkeletalMeshGeometry(MeshGeometry);
+}
+
+bool FInterchangePipelineMeshesUtilitiesContext::IsSkeletalMeshGeometry(const FInterchangeMeshGeometry& MeshGeometry)
+{
+	if (bConvertSkeletalMeshToStaticMesh)
+	{
+		return false;
+	}
+
+	if (bQueryGeometryOnlyIfNoInstance && MeshGeometry.ReferencingMeshInstanceUids.Num() > 0)
+	{
+		return false;
+	}
+
+	if (bConvertStaticMeshToSkeletalMesh
+		|| MeshGeometry.MeshNode->IsSkinnedMesh()
+		|| (bConvertStaticsWithMorphTargetsToSkeletals && (MeshGeometry.MeshNode->GetMorphTargetDependeciesCount() > 0)))
+	{
+		return true;
+	}
+	return false;
+}
 
 UInterchangePipelineMeshesUtilities* UInterchangePipelineMeshesUtilities::CreateInterchangePipelineMeshesUtilities(UInterchangeBaseNodeContainer* BaseNodeContainer)
 {
 	check(BaseNodeContainer);
 	UInterchangePipelineMeshesUtilities* PipelineMeshesUtilities = NewObject<UInterchangePipelineMeshesUtilities>(GetTransientPackage(), NAME_None);
 	
+	//Set the container
+	PipelineMeshesUtilities->BaseNodeContainer = BaseNodeContainer;
+
 	TArray<FString> SkeletonRootNodeUids;
 	
 	//Find all translated node we need for this pipeline
@@ -340,50 +476,50 @@ void UInterchangePipelineMeshesUtilities::IterateAllMeshInstance(TFunctionRef<vo
 	}
 }
 
-void UInterchangePipelineMeshesUtilities::GetAllSkinnedMeshInstance(TArray<FString>& MeshInstanceUids, const bool bConvertStaticMeshToSkeletalMesh, const bool bConvertStaticsWithMorphTargetsToSkeletals) const
+void UInterchangePipelineMeshesUtilities::GetAllSkinnedMeshInstance(TArray<FString>& MeshInstanceUids) const
 {
 	MeshInstanceUids.Empty(MeshInstancesPerMeshInstanceUid.Num());
 	for (const TPair<FString, FInterchangeMeshInstance>& MeshInstanceUidAndMeshInstance : MeshInstancesPerMeshInstanceUid)
 	{
 		const FInterchangeMeshInstance& MeshInstance = MeshInstanceUidAndMeshInstance.Value;
-		if((bConvertStaticMeshToSkeletalMesh || MeshInstance.bReferenceSkinnedMesh || (bConvertStaticsWithMorphTargetsToSkeletals && MeshInstance.bHasMorphTargets)) && !MeshInstance.bReferenceMorphTarget)
+		if (CurrentDataContext.IsSkeletalMeshInstance(MeshInstance, BaseNodeContainer))
 		{
 			MeshInstanceUids.Add(MeshInstance.MeshInstanceUid);
 		}
 	}
 }
 
-void UInterchangePipelineMeshesUtilities::IterateAllSkinnedMeshInstance(TFunctionRef<void(const FInterchangeMeshInstance&)> IterationLambda, const bool bConvertStaticMeshToSkeletalMesh, const bool bConvertStaticsWithMorphTargetsToSkeletals) const
+void UInterchangePipelineMeshesUtilities::IterateAllSkinnedMeshInstance(TFunctionRef<void(const FInterchangeMeshInstance&)> IterationLambda) const
 {
 	for (const TPair<FString, FInterchangeMeshInstance>& MeshInstanceUidAndMeshInstance : MeshInstancesPerMeshInstanceUid)
 	{
 		const FInterchangeMeshInstance& MeshInstance = MeshInstanceUidAndMeshInstance.Value;
-		if ((bConvertStaticMeshToSkeletalMesh || MeshInstance.bReferenceSkinnedMesh || (bConvertStaticsWithMorphTargetsToSkeletals && MeshInstance.bHasMorphTargets)) && !MeshInstance.bReferenceMorphTarget)
+		if (CurrentDataContext.IsSkeletalMeshInstance(MeshInstance, BaseNodeContainer))
 		{
 			IterationLambda(MeshInstance);
 		}
 	}
 }
 
-void UInterchangePipelineMeshesUtilities::GetAllStaticMeshInstance(TArray<FString>& MeshInstanceUids, const bool bConvertSkeletalMeshToStaticMesh, const bool bConvertStaticsWithMorphTargetsToSkeletals) const
+void UInterchangePipelineMeshesUtilities::GetAllStaticMeshInstance(TArray<FString>& MeshInstanceUids) const
 {
 	MeshInstanceUids.Empty(MeshInstancesPerMeshInstanceUid.Num());
 	for (const TPair<FString, FInterchangeMeshInstance>& MeshInstanceUidAndMeshInstance : MeshInstancesPerMeshInstanceUid)
 	{
 		const FInterchangeMeshInstance& MeshInstance = MeshInstanceUidAndMeshInstance.Value;
-		if ((bConvertSkeletalMeshToStaticMesh || !MeshInstance.bReferenceSkinnedMesh) && !MeshInstance.bReferenceMorphTarget && !(bConvertStaticsWithMorphTargetsToSkeletals && MeshInstance.bHasMorphTargets))
+		if (CurrentDataContext.IsStaticMeshInstance(MeshInstance, BaseNodeContainer))
 		{
 			MeshInstanceUids.Add(MeshInstance.MeshInstanceUid);
 		}
 	}
 }
 
-void UInterchangePipelineMeshesUtilities::IterateAllStaticMeshInstance(TFunctionRef<void(const FInterchangeMeshInstance&)> IterationLambda, const bool bConvertSkeletalMeshToStaticMesh, const bool bConvertStaticsWithMorphTargetsToSkeletals) const
+void UInterchangePipelineMeshesUtilities::IterateAllStaticMeshInstance(TFunctionRef<void(const FInterchangeMeshInstance&)> IterationLambda) const
 {
 	for (const TPair<FString, FInterchangeMeshInstance>& MeshInstanceUidAndMeshInstance : MeshInstancesPerMeshInstanceUid)
 	{
 		const FInterchangeMeshInstance& MeshInstance = MeshInstanceUidAndMeshInstance.Value;
-		if ((bConvertSkeletalMeshToStaticMesh || !MeshInstance.bReferenceSkinnedMesh) && !MeshInstance.bReferenceMorphTarget && !(bConvertStaticsWithMorphTargetsToSkeletals && MeshInstance.bHasMorphTargets))
+		if (CurrentDataContext.IsStaticMeshInstance(MeshInstance, BaseNodeContainer))
 		{
 			IterationLambda(MeshInstance);
 		}
@@ -410,7 +546,7 @@ void UInterchangePipelineMeshesUtilities::GetAllSkinnedMeshGeometry(TArray<FStri
 	for (const TPair<FString, FInterchangeMeshGeometry>& MeshGeometryUidAndMeshGeometry : MeshGeometriesPerMeshUid)
 	{
 		const FInterchangeMeshGeometry& MeshGeometry = MeshGeometryUidAndMeshGeometry.Value;
-		if (MeshGeometry.MeshNode->IsSkinnedMesh())
+		if(CurrentDataContext.IsSkeletalMeshGeometry(MeshGeometry))
 		{
 			MeshGeometryUids.Add(MeshGeometry.MeshUid);
 		}
@@ -422,32 +558,32 @@ void UInterchangePipelineMeshesUtilities::IterateAllSkinnedMeshGeometry(TFunctio
 	for (const TPair<FString, FInterchangeMeshGeometry>& MeshGeometryUidAndMeshGeometry : MeshGeometriesPerMeshUid)
 	{
 		const FInterchangeMeshGeometry& MeshGeometry = MeshGeometryUidAndMeshGeometry.Value;
-		if (MeshGeometry.MeshNode->IsSkinnedMesh())
+		if (CurrentDataContext.IsSkeletalMeshGeometry(MeshGeometry))
 		{
 			IterationLambda(MeshGeometry);
 		}
 	}
 }
 
-void UInterchangePipelineMeshesUtilities::GetAllStaticMeshGeometry(TArray<FString>& MeshGeometryUids, const bool bConvertSkeletalMeshToStaticMesh, const bool bConvertStaticsWithMorphTargetsToSkeletals) const
+void UInterchangePipelineMeshesUtilities::GetAllStaticMeshGeometry(TArray<FString>& MeshGeometryUids) const
 {
 	MeshGeometryUids.Empty(MeshGeometriesPerMeshUid.Num());
 	for (const TPair<FString, FInterchangeMeshGeometry>& MeshGeometryUidAndMeshGeometry : MeshGeometriesPerMeshUid)
 	{
 		const FInterchangeMeshGeometry& MeshGeometry = MeshGeometryUidAndMeshGeometry.Value;
-		if ((bConvertSkeletalMeshToStaticMesh || !MeshGeometry.MeshNode->IsSkinnedMesh()) && !MeshGeometry.MeshNode->IsMorphTarget() && !(bConvertStaticsWithMorphTargetsToSkeletals && (MeshGeometry.MeshNode->GetMorphTargetDependeciesCount() > 0)))
+		if (CurrentDataContext.IsStaticMeshGeometry(MeshGeometry))
 		{
 			MeshGeometryUids.Add(MeshGeometry.MeshUid);
 		}
 	}
 }
 
-void UInterchangePipelineMeshesUtilities::IterateAllStaticMeshGeometry(TFunctionRef<void(const FInterchangeMeshGeometry&)> IterationLambda, const bool bConvertSkeletalMeshToStaticMesh, const bool bConvertStaticsWithMorphTargetsToSkeletals) const
+void UInterchangePipelineMeshesUtilities::IterateAllStaticMeshGeometry(TFunctionRef<void(const FInterchangeMeshGeometry&)> IterationLambda) const
 {
 	for (const TPair<FString, FInterchangeMeshGeometry>& MeshGeometryUidAndMeshGeometry : MeshGeometriesPerMeshUid)
 	{
 		const FInterchangeMeshGeometry& MeshGeometry = MeshGeometryUidAndMeshGeometry.Value;
-		if ((bConvertSkeletalMeshToStaticMesh || !MeshGeometry.MeshNode->IsSkinnedMesh()) && !MeshGeometry.MeshNode->IsMorphTarget() && !(bConvertStaticsWithMorphTargetsToSkeletals && (MeshGeometry.MeshNode->GetMorphTargetDependeciesCount() > 0)))
+		if (CurrentDataContext.IsStaticMeshGeometry(MeshGeometry))
 		{
 			IterationLambda(MeshGeometry);
 		}
@@ -515,20 +651,23 @@ void UInterchangePipelineMeshesUtilities::IterateAllMeshInstanceUsingMeshGeometr
 	}
 }
 
-void UInterchangePipelineMeshesUtilities::GetCombinedSkinnedMeshInstances(UInterchangeBaseNodeContainer* BaseNodeContainer, TMap<FString, TArray<FString>>& OutMeshInstanceUidsPerSkeletonRootUid, const bool bConvertStaticMeshToSkeletalMesh, const bool bConvertStaticsWithMorphTargetsToSkeletals) const
+void UInterchangePipelineMeshesUtilities::GetCombinedSkinnedMeshInstances(TMap<FString, TArray<FString>>& OutMeshInstanceUidsPerSkeletonRootUid) const
 {
 	check(BaseNodeContainer);
 
 	for (const TPair<FString, FInterchangeMeshInstance>& MeshInstanceUidAndMeshInstance : MeshInstancesPerMeshInstanceUid)
 	{
 		const FInterchangeMeshInstance& MeshInstance = MeshInstanceUidAndMeshInstance.Value;
-		if (!(bConvertStaticMeshToSkeletalMesh || bConvertStaticsWithMorphTargetsToSkeletals) && !MeshInstance.bReferenceSkinnedMesh)
+
+		bool bIsNestedIntoSkeleton = false;
+		if (!CurrentDataContext.IsSkeletalMeshInstance(MeshInstance, BaseNodeContainer, bIsNestedIntoSkeleton))
 		{
 			continue;
 		}
-		bool bIsStaticMesh = (bConvertStaticMeshToSkeletalMesh || (bConvertStaticsWithMorphTargetsToSkeletals && MeshInstance.bHasMorphTargets)) && !MeshInstance.bReferenceSkinnedMesh;
+		bool bIsStaticMesh = !MeshInstance.bReferenceSkinnedMesh;
 		//Find the root skeleton for this MeshInstance
-		FString SkeletonRootUid;
+		FString SkeletonRootUid; // = CurrentDataContext.FindSkeletalMeshSkeletonRootFromMeshInstance(MeshInstance, BaseNodeContainer);
+		
 		for (const FString& MeshGeometryUid : MeshInstance.ReferencingMeshGeometryUids)
 		{
 			if (const FString* SkeletonRootUidPtr = SkeletonRootUidPerMeshUid.Find(MeshGeometryUid))
@@ -548,21 +687,30 @@ void UInterchangePipelineMeshesUtilities::GetCombinedSkinnedMeshInstances(UInter
 			{
 				//Create a joint from the instance node (the scene node pointing on the mesh).
 				SkeletonRootUid = MeshInstance.MeshInstanceUid;
-
-				if (!MeshInstance.bHasMorphTargets)
+				if (bIsNestedIntoSkeleton || !MeshInstance.bHasMorphTargets)
 				{
-					//Since we are dealing with rigid mesh and combine we will use the outermost valid parent
+					//Find the deepest joint node
 					if (const UInterchangeSceneNode* SceneNode = Cast<UInterchangeSceneNode>(BaseNodeContainer->GetNode(MeshInstance.MeshInstanceUid)))
 					{
 						FString ParentUid = SceneNode->GetParentUid();
+						FString LastSceneNodeUid = SkeletonRootUid;
 						while (!ParentUid.Equals(UInterchangeBaseNode::InvalidNodeUid()))
 						{
 							if (const UInterchangeSceneNode* ParentNode = Cast<UInterchangeSceneNode>(BaseNodeContainer->GetNode(ParentUid)))
 							{
-								SkeletonRootUid = ParentUid;
+								if(ParentNode->IsSpecializedTypeContains(UE::Interchange::FSceneNodeStaticData::GetJointSpecializeTypeString()))
+								{
+									SkeletonRootUid = ParentUid;
+								}
+								LastSceneNodeUid = ParentUid;
 								ParentUid = ParentNode->GetParentUid();
 							}
 						}
+						//If we did not find any joint because we have non nested mesh, get the deepest scene node
+						if (!bIsNestedIntoSkeleton && SkeletonRootUid.Equals(MeshInstance.MeshInstanceUid))
+						{
+							SkeletonRootUid = LastSceneNodeUid;
+						}
 					}
 				}
 			}
@@ -582,39 +730,6 @@ void UInterchangePipelineMeshesUtilities::GetCombinedSkinnedMeshInstances(UInter
 	}
 }
 
-void UInterchangePipelineMeshesUtilities::GetCombinedSkinnedMeshGeometries(TMap<FString, TArray<FString>>& OutMeshGeometryUidsPerSkeletonRootUid) const
-{
-	for (const TPair<FString, FInterchangeMeshGeometry>& MeshGeometryUidAndMeshGeometry : MeshGeometriesPerMeshUid)
-	{
-		const FInterchangeMeshGeometry& MeshGeometry = MeshGeometryUidAndMeshGeometry.Value;
-		if (!MeshGeometry.MeshNode || !MeshGeometry.MeshNode->IsSkinnedMesh())
-		{
-			continue;
-		}
-		//Find the root skeleton for this MeshInstance
-		FString SkeletonRootUid;
-		if (const FString* SkeletonRootUidPtr = SkeletonRootUidPerMeshUid.Find(MeshGeometryUidAndMeshGeometry.Key))
-		{
-			if (SkeletonRootUid.IsEmpty())
-			{
-				SkeletonRootUid = *SkeletonRootUidPtr;
-			}
-		}
-		else
-		{
-			//every skinned geometry should have a skeleton root node ???
-		}
-
-		if (SkeletonRootUid.IsEmpty())
-		{
-			//Skip this MeshGeometry
-			continue;
-		}
-		TArray<FString>& MeshGeometryUids = OutMeshGeometryUidsPerSkeletonRootUid.FindOrAdd(SkeletonRootUid);
-		MeshGeometryUids.Add(MeshGeometryUidAndMeshGeometry.Key);
-	}
-}
-
 FString UInterchangePipelineMeshesUtilities::GetMeshInstanceSkeletonRootUid(const FString& MeshInstanceUid) const
 {
 	FString SkeletonRootUid;
diff --git a/Engine/Plugins/Interchange/Runtime/Source/Pipelines/Public/InterchangeGenericAnimationPipeline.h b/Engine/Plugins/Interchange/Runtime/Source/Pipelines/Public/InterchangeGenericAnimationPipeline.h
--- a/Engine/Plugins/Interchange/Runtime/Source/Pipelines/Public/InterchangeGenericAnimationPipeline.h
+++ b/Engine/Plugins/Interchange/Runtime/Source/Pipelines/Public/InterchangeGenericAnimationPipeline.h
@@ -131,6 +131,7 @@ private:
 
 	void CreateAnimSequenceFactoryNode(UInterchangeSkeletalAnimationTrackNode& Node);
 
+	bool bSceneImport = false;
 	UInterchangeBaseNodeContainer* BaseNodeContainer = nullptr;
 	TArray<const UInterchangeSourceData*> SourceDatas;
 
diff --git a/Engine/Plugins/Interchange/Runtime/Source/Pipelines/Public/InterchangePipelineMeshesUtilities.h b/Engine/Plugins/Interchange/Runtime/Source/Pipelines/Public/InterchangePipelineMeshesUtilities.h
--- a/Engine/Plugins/Interchange/Runtime/Source/Pipelines/Public/InterchangePipelineMeshesUtilities.h
+++ b/Engine/Plugins/Interchange/Runtime/Source/Pipelines/Public/InterchangePipelineMeshesUtilities.h
@@ -19,7 +19,7 @@ class UInterchangeSceneNode;
 /*
 * This container exist only because UPROPERTY cannot support nested container. See FInterchangeMeshInstance
 */
-USTRUCT(BlueprintType, Experimental)
+USTRUCT(BlueprintType)
 struct FInterchangeLodSceneNodeContainer
 {
 	GENERATED_BODY()
@@ -36,7 +36,7 @@ struct FInterchangeLodSceneNodeContainer
 * A mesh instance pointing on a lod group can have many lods and many scene mesh nodes per lod index.
 * A mesh instance pointing on a mesh node will have only the lod 0 and will point on one scene mesh node.
 */
-USTRUCT(BlueprintType, Experimental)
+USTRUCT(BlueprintType)
 struct FInterchangeMeshInstance
 {
 	GENERATED_BODY()
@@ -85,7 +85,7 @@ struct FInterchangeMeshInstance
 /*
 * A mesh geometry is a description of a translated mesh asset node that define a geometry.
 */
-USTRUCT(BlueprintType, Experimental)
+USTRUCT(BlueprintType)
 struct FInterchangeMeshGeometry
 {
 	GENERATED_BODY()
@@ -120,9 +120,54 @@ struct FInterchangeMeshGeometry
 	TArray<FString> AttachedSocketUids;
 };
 
+/*
+* Represent the context UInterchangePipelineMeshesUtilities will use when client query the data
+*/
+USTRUCT(BlueprintType)
+struct FInterchangePipelineMeshesUtilitiesContext
+{
+	GENERATED_BODY()
+
+	/**
+	 * Convert static mesh to skeletal mesh
+	 */
+	UPROPERTY(EditAnywhere, Category = "Interchange | Pipeline | MeshesContext")
+	bool bConvertStaticMeshToSkeletalMesh = false;
+
+	/**
+	 * Convert static mesh to skeletal mesh
+	 */
+	UPROPERTY(EditAnywhere, Category = "Interchange | Pipeline | MeshesContext")
+	bool bConvertSkeletalMeshToStaticMesh = false;
+
+	/**
+	 * Convert static mesh that has morph target to skeletal mesh
+	 */
+	UPROPERTY(EditAnywhere, Category = "Interchange | Pipeline | MeshesContext")
+	bool bConvertStaticsWithMorphTargetsToSkeletals = false;
+
+	/**
+	 * If checked, meshes nested in bone hierarchies will be imported instead of being converted to bones. If the mesh are not skinned they will
+	 * be added to skeletal mesh and remove from the static meshes.
+	 */
+	UPROPERTY(EditAnywhere, Category = "Interchange | Pipeline | MeshesContext")
+	bool bImportMeshesInBoneHierarchy = true;
+
+	/**
+	 * When querying geometry, this flag will not add MeshGeometry if there is a scene node pointing on a geometry.
+	 */
+	UPROPERTY(EditAnywhere, Category = "Interchange | Pipeline | MeshesContext")
+	bool bQueryGeometryOnlyIfNoInstance = true;
+
+	bool IsStaticMeshInstance(const FInterchangeMeshInstance& MeshInstance, UInterchangeBaseNodeContainer* BaseNodeContainer);
+	bool IsSkeletalMeshInstance(const FInterchangeMeshInstance& MeshInstance, UInterchangeBaseNodeContainer* BaseNodeContainer);
+	bool IsSkeletalMeshInstance(const FInterchangeMeshInstance& MeshInstance, UInterchangeBaseNodeContainer* BaseNodeContainer, bool& bOutIsStaticMeshNestedInSkeleton);
+	bool IsStaticMeshGeometry(const FInterchangeMeshGeometry& MeshGeometry);
+	bool IsSkeletalMeshGeometry(const FInterchangeMeshGeometry& MeshGeometry);
+};
 
 /**/
-UCLASS(BlueprintType, Experimental)
+UCLASS(BlueprintType)
 class INTERCHANGEPIPELINES_API UInterchangePipelineMeshesUtilities : public UObject
 {
 	GENERATED_BODY()
@@ -148,23 +193,23 @@ public:
 	* Get all skinned mesh instance unique ids.
 	*/
 	UFUNCTION(BlueprintCallable, Category = "Interchange | Pipeline | Meshes")
-	void GetAllSkinnedMeshInstance(TArray<FString>& MeshInstanceUids, const bool bConvertStaticMeshToSkeletalMesh, const bool bConvertStaticsWithMorphTargetsToSkeletals) const;
+	void GetAllSkinnedMeshInstance(TArray<FString>& MeshInstanceUids) const;
 
 	/**
 	* Iterate all skinned mesh instance.
 	*/
-	void IterateAllSkinnedMeshInstance(TFunctionRef<void(const FInterchangeMeshInstance&)> IterationLambda, const bool bConvertStaticMeshToSkeletalMesh, const bool bConvertStaticsWithMorphTargetsToSkeletals) const;
+	void IterateAllSkinnedMeshInstance(TFunctionRef<void(const FInterchangeMeshInstance&)> IterationLambda) const;
 
 	/**
 	* Get all static mesh instance unique ids.
 	*/
 	UFUNCTION(BlueprintCallable, Category = "Interchange | Pipeline | Meshes")
-	void GetAllStaticMeshInstance(TArray<FString>& MeshInstanceUids, const bool bConvertSkeletalMeshToStaticMesh, const bool bConvertStaticsWithMorphTargetsToSkeletals) const;
+	void GetAllStaticMeshInstance(TArray<FString>& MeshInstanceUids) const;
 
 	/**
 	* Iterate all static mesh instance.
 	*/
-	void IterateAllStaticMeshInstance(TFunctionRef<void(const FInterchangeMeshInstance&)> IterationLambda, const bool bConvertSkeletalMeshToStaticMesh, const bool bConvertStaticsWithMorphTargetsToSkeletals) const;
+	void IterateAllStaticMeshInstance(TFunctionRef<void(const FInterchangeMeshInstance&)> IterationLambda) const;
 
 	/**
 	* Get all mesh geometry unique ids.
@@ -192,12 +237,12 @@ public:
 	* Get all static mesh geometry unique ids.
 	*/
 	UFUNCTION(BlueprintCallable, Category = "Interchange | Pipeline | Meshes")
-	void GetAllStaticMeshGeometry(TArray<FString>& MeshGeometryUids, const bool bConvertSkeletalMeshToStaticMesh, const bool bConvertStaticsWithMorphTargetsToSkeletals) const;
+	void GetAllStaticMeshGeometry(TArray<FString>& MeshGeometryUids) const;
 
 	/**
 	* Iterate all static mesh geometry.
 	*/
-	void IterateAllStaticMeshGeometry(TFunctionRef<void(const FInterchangeMeshGeometry&)> IterationLambda, const bool bConvertSkeletalMeshToStaticMesh, const bool bConvertStaticsWithMorphTargetsToSkeletals) const;
+	void IterateAllStaticMeshGeometry(TFunctionRef<void(const FInterchangeMeshGeometry&)> IterationLambda) const;
 
 	/**
 	* Get all not instanced mesh geometry unique ids.
@@ -249,14 +294,8 @@ public:
 	* Return a list of skinned FInterchangeMeshInstance uid that can be combined together.
 	* We cannot create a skinned mesh with multiple skeleton root node, This function return combined MeshInstance per skeleton roots
 	*/
-	void GetCombinedSkinnedMeshInstances(UInterchangeBaseNodeContainer* BaseNodeContainer, TMap<FString, TArray<FString>>& OutMeshInstanceUidsPerSkeletonRootUid, const bool bConvertStaticMeshToSkeletalMesh, const bool bConvertStaticsWithMorphTargetsToSkeletals) const;
+	void GetCombinedSkinnedMeshInstances(TMap<FString, TArray<FString>>& OutMeshInstanceUidsPerSkeletonRootUid) const;
 	
-	/**
-	* Return a list of skinned FInterchangeMeshGeometry uid that can be combined together.
-	* We cannot create a skinned mesh with multiple skeleton root node, This function return combined FInterchangeMeshGeometry per skeleton roots
-	*/
-	void GetCombinedSkinnedMeshGeometries(TMap<FString, TArray<FString>>& OutMeshGeometryUidsPerSkeletonRootUid) const;
-
 	/**
 	* Return the skeleton root node Uid, this is the uid for a UInterchangeSceneNode that has a "Joint" specialized type.
 	* Return an empty string if the MeshInstanceUid parameter point on nothing.
@@ -275,10 +314,20 @@ public:
 
 	FString GetMeshGeometrySkeletonRootUid(const FInterchangeMeshGeometry& MeshGeometry) const;
 
+	UFUNCTION(BlueprintCallable, Category = "Interchange | Pipeline | Meshes")
+	void SetContext(const FInterchangePipelineMeshesUtilitiesContext& Context) const
+	{
+		CurrentDataContext = Context;
+	}
+	
 protected:
 	TMap<FString, FInterchangeMeshGeometry> MeshGeometriesPerMeshUid;
 	TMap<FString, FInterchangeMeshInstance> MeshInstancesPerMeshInstanceUid;
 	TMap<FString, FString> SkeletonRootUidPerMeshUid;
+
+	UInterchangeBaseNodeContainer* BaseNodeContainer;
+
+	mutable FInterchangePipelineMeshesUtilitiesContext CurrentDataContext;
 };
 
 namespace UE::Interchange::MeshesUtilities
