Fixed issue with FBX files containing incomplete shading model
 - Root cause: The importer was discarding materials which were not Lambert or Phong
 - Solution: Just process all the possible properties on a FbxSurfaceMaterial

#rnx

diff --git a/Engine/Plugins/Interchange/Runtime/Source/Parsers/Fbx/Private/FbxMaterial.cpp b/Engine/Plugins/Interchange/Runtime/Source/Parsers/Fbx/Private/FbxMaterial.cpp
--- a/Engine/Plugins/Interchange/Runtime/Source/Parsers/Fbx/Private/FbxMaterial.cpp
+++ b/Engine/Plugins/Interchange/Runtime/Source/Parsers/Fbx/Private/FbxMaterial.cpp
@@ -241,7 +241,6 @@ namespace UE
 				}
 			}
 
-
 			const UInterchangeShaderGraphNode* FFbxMaterial::AddShaderGraphNode(FbxSurfaceMaterial* SurfaceMaterial, UInterchangeBaseNodeContainer& NodeContainer)
 			{
 				using namespace UE::Interchange::Materials;
@@ -273,21 +272,6 @@ namespace UE
 					return nullptr;
 				}
 
-				// If not a Lambert nor a Phong material, match legacy FBX import
-				// Use random color because there may be multiple materials, then they can be different
-				if (!SurfaceMaterial->Is<FbxSurfaceLambert>())
-				{
-					FLinearColor BaseColor;
-					BaseColor.R = 0.5f + 0.5f * FMath::FRand();
-					BaseColor.G = 0.5f + 0.5f * FMath::FRand();
-					BaseColor.B = 0.5f + 0.5f * FMath::FRand();
-
-					const FString InputValueKey = UInterchangeShaderPortsAPI::MakeInputValueKey(Phong::Parameters::DiffuseColor.ToString());
-					ShaderGraphNode->AddLinearColorAttribute(InputValueKey, BaseColor);
-
-					return ShaderGraphNode;
-				}
-
 				TFunction<bool(FBXSDK_NAMESPACE::FbxProperty&)> ShouldConvertProperty = [&](FBXSDK_NAMESPACE::FbxProperty& MaterialProperty) -> bool
 				{
 					bool bShouldConvertProperty = false;
@@ -302,6 +286,12 @@ namespace UE
 					return bShouldConvertProperty;
 				};
 
+				TFunction<float(const char*)> GetFactor = [&](const char* FactorName)
+				{
+					FBXSDK_NAMESPACE::FbxProperty Property = SurfaceMaterial->FindProperty(FactorName);
+					return Property.IsValid() ? (float)Property.Get<FbxDouble>() : 1.;
+				};
+
 				TFunction<bool(FName, const char*, float , TVariant<FLinearColor, float>&, bool)>  ConnectInput;
 				ConnectInput = [&](FName InputName, const char* FbxPropertyName, float Factor, TVariant<FLinearColor, float>& DefaultValue, bool bInverse) -> bool
 				{
@@ -316,21 +306,15 @@ namespace UE
 
 				// Diffuse
 				{
-					const float Factor = (float)((FbxSurfaceLambert*)SurfaceMaterial)->DiffuseFactor.Get();
+					const float Factor = GetFactor(FbxSurfaceMaterial::sDiffuseFactor);
 					TVariant<FLinearColor, float> DefaultValue;
 					DefaultValue.Set<FLinearColor>(FLinearColor::Black);
-					if (!ConnectInput(Phong::Parameters::DiffuseColor, FbxSurfaceMaterial::sDiffuse, Factor, DefaultValue, false))
-					{
-						// Nothing to retrieve from the Diffuse property, get Diffuse directly from the class
-						const FbxDouble3 DiffuseColor = ((FbxSurfaceLambert*)SurfaceMaterial)->Diffuse.Get();
-						const FString InputValueKey = UInterchangeShaderPortsAPI::MakeInputValueKey(Phong::Parameters::DiffuseColor.ToString());
-						ShaderGraphNode->AddLinearColorAttribute(InputValueKey, FFbxConvert::ConvertColor(DiffuseColor));
-					}
+					ConnectInput(Phong::Parameters::DiffuseColor, FbxSurfaceMaterial::sDiffuse, Factor, DefaultValue, false);
 				}
 
 				// Ambient
 				{
-					const float Factor = (float)((FbxSurfacePhong*)SurfaceMaterial)->AmbientFactor.Get();
+					const float Factor = GetFactor(FbxSurfaceMaterial::sAmbientFactor);
 					TVariant<FLinearColor, float> DefaultValue;
 					DefaultValue.Set<FLinearColor>(FLinearColor::Black);
 					ConnectInput(Phong::Parameters::AmbientColor, FbxSurfaceMaterial::sAmbient, Factor, DefaultValue, false);
@@ -338,7 +322,7 @@ namespace UE
 
 				// Emissive
 				{
-					const float Factor = (float)((FbxSurfaceLambert*)SurfaceMaterial)->EmissiveFactor.Get();
+					const float Factor = GetFactor(FbxSurfaceMaterial::sEmissiveFactor);
 					TVariant<FLinearColor, float> DefaultValue;
 					DefaultValue.Set<FLinearColor>(FLinearColor::Black);
 					ConnectInput(Phong::Parameters::EmissiveColor, FbxSurfaceMaterial::sEmissive, Factor, DefaultValue, false);
@@ -347,17 +331,16 @@ namespace UE
 				// Normal
 				{
 					// FbxSurfaceMaterial can have either a normal map or a bump map, check for both
-					float Factor = 1.f;
 					FBXSDK_NAMESPACE::FbxProperty MaterialProperty = SurfaceMaterial->FindProperty(FbxSurfaceMaterial::sNormalMap);
 					if (MaterialProperty.IsValid() && MaterialProperty.GetSrcObjectCount<FbxTexture>() > 0)
 					{
 						TVariant<FLinearColor, float> DefaultValue;
 						DefaultValue.Set<FLinearColor>(FLinearColor(FVector::UpVector));
-						ConvertPropertyToShaderNode(NodeContainer, ShaderGraphNode, MaterialProperty, Factor, Phong::Parameters::Normal, DefaultValue);
+						ConvertPropertyToShaderNode(NodeContainer, ShaderGraphNode, MaterialProperty, 1.f, Phong::Parameters::Normal, DefaultValue);
 					}
 					else
 					{
-						Factor = (float)((FbxSurfaceLambert*)SurfaceMaterial)->BumpFactor.Get();
+						const float Factor = GetFactor(FbxSurfaceMaterial::sBumpFactor);
 						MaterialProperty = SurfaceMaterial->FindProperty(FbxSurfaceMaterial::sBump);
 						if (MaterialProperty.IsValid() && MaterialProperty.GetSrcObjectCount<FbxTexture>() > 0)
 						{
@@ -395,17 +378,24 @@ namespace UE
 					FBXSDK_NAMESPACE::FbxProperty MaterialProperty = SurfaceMaterial->FindProperty(FbxSurfaceMaterial::sTransparentColor);
 					if (ShouldConvertProperty(MaterialProperty))
 					{
-						const float TransparencyFactor = (float)((FbxSurfaceLambert*)SurfaceMaterial)->TransparencyFactor.Get();
+						const float Factor = GetFactor(FbxSurfaceMaterial::sTransparencyFactor);
 
 						if (MaterialProperty.GetSrcObjectCount<FbxFileTexture>() > 0)
 						{
 							TVariant<FLinearColor, float> DefaultValue;
 							DefaultValue.Set<float>(0.f); // Opaque
-							ConvertPropertyToShaderNode(NodeContainer, ShaderGraphNode, MaterialProperty, TransparencyFactor, Phong::Parameters::Opacity, DefaultValue, true);
+							FName InputName = Phong::Parameters::Opacity;
+							// The texture is hooked to the OpacityMask when transparency is with a texture
+							if (MaterialProperty.Get<FbxDouble>() == 0.)
+							{
+								InputName = Phong::Parameters::OpacityMask;
+								ShaderGraphNode->SetCustomOpacityMaskClipValue(0.333f);
+							}
+							ConvertPropertyToShaderNode(NodeContainer, ShaderGraphNode, MaterialProperty, Factor, InputName, DefaultValue, true);
 						}
 						else
 						{
-							const float OpacityScalar = 1.f - (MaterialProperty.Get<float>() * TransparencyFactor);
+							const float OpacityScalar = 1.f - (MaterialProperty.Get<float>() * Factor);
 							if (OpacityScalar < 1.f)
 							{
 								const FString InputName = UInterchangeShaderPortsAPI::MakeInputValueKey(Phong::Parameters::Opacity.ToString());
@@ -415,19 +405,30 @@ namespace UE
 					}
 				}
 
-				// Get specular color and shininess if FBX material is a Phong one
-				if (SurfaceMaterial->Is<FbxSurfacePhong>())
+				// Specular
 				{
-					// Specular
-					{
-						const float Factor = (float)((FbxSurfacePhong*)SurfaceMaterial)->SpecularFactor.Get();
-						TVariant<FLinearColor, float> DefaultValue;
-						DefaultValue.Set<FLinearColor>(FLinearColor::Black);
-						ConnectInput(Phong::Parameters::SpecularColor, FbxSurfaceMaterial::sSpecular, Factor, DefaultValue, false);
-					}
+					const float Factor = GetFactor(FbxSurfaceMaterial::sSpecularFactor);
+					TVariant<FLinearColor, float> DefaultValue;
+					DefaultValue.Set<FLinearColor>(FLinearColor::Black);
+					ConnectInput(Phong::Parameters::SpecularColor, FbxSurfaceMaterial::sSpecular, Factor, DefaultValue, false);
+				}
+
+				// Shininess
+				ConvertShininessToShaderNode(*SurfaceMaterial, NodeContainer, ShaderGraphNode);
+
+				// If no valid property found, create a material anyway
+				// Use random color because there may be multiple materials, then they can be different
+				TArray<FString> InputNames;
+				UInterchangeShaderPortsAPI::GatherInputs(ShaderGraphNode, InputNames);
+				if (InputNames.Num() == 0)
+				{
+					FLinearColor BaseColor;
+					BaseColor.R = 0.5f + 0.5f * FMath::FRand();
+					BaseColor.G = 0.5f + 0.5f * FMath::FRand();
+					BaseColor.B = 0.5f + 0.5f * FMath::FRand();
 
-					// Shininess
-					ConvertShininessToShaderNode(*SurfaceMaterial, NodeContainer, ShaderGraphNode);
+					const FString InputValueKey = UInterchangeShaderPortsAPI::MakeInputValueKey(Phong::Parameters::DiffuseColor.ToString());
+					ShaderGraphNode->AddLinearColorAttribute(InputValueKey, BaseColor);
 				}
 
 				return ShaderGraphNode;
@@ -460,9 +461,11 @@ namespace UE
 				int32 MaterialCount = ParentFbxNode->GetMaterialCount();
 				for (int32 MaterialIndex = 0; MaterialIndex < MaterialCount; ++MaterialIndex)
 				{
-					FbxSurfaceMaterial* SurfaceMaterial = ParentFbxNode->GetMaterial(MaterialIndex);
-					const UInterchangeShaderGraphNode* ShaderGraphNode = AddShaderGraphNode(SurfaceMaterial, NodeContainer);
-					SceneNode->SetSlotMaterialDependencyUid(Parser.GetFbxHelper()->GetFbxObjectName(SurfaceMaterial), ShaderGraphNode->GetUniqueID());
+					if (FbxSurfaceMaterial* SurfaceMaterial = ParentFbxNode->GetMaterial(MaterialIndex))
+					{
+						const UInterchangeShaderGraphNode* ShaderGraphNode = AddShaderGraphNode(SurfaceMaterial, NodeContainer);
+						SceneNode->SetSlotMaterialDependencyUid(Parser.GetFbxHelper()->GetFbxObjectName(SurfaceMaterial), ShaderGraphNode->GetUniqueID());
+					}
 				}
 			}
 
@@ -471,8 +474,10 @@ namespace UE
 				int32 MaterialCount = SDKScene->GetMaterialCount();
 				for (int32 MaterialIndex = 0; MaterialIndex < MaterialCount; ++MaterialIndex)
 				{
-					FbxSurfaceMaterial* SurfaceMaterial = SDKScene->GetMaterial(MaterialIndex);
-					AddShaderGraphNode(SurfaceMaterial, NodeContainer);
+					if (FbxSurfaceMaterial* SurfaceMaterial = SDKScene->GetMaterial(MaterialIndex))
+					{
+						AddShaderGraphNode(SurfaceMaterial, NodeContainer);
+					}
 				}
 			}
 		} //ns Private
