Changed Interchange plugin to be supported only on Win64, linux and Mac

#rnx

diff --git a/Engine/Plugins/Interchange/Runtime/Interchange.uplugin b/Engine/Plugins/Interchange/Runtime/Interchange.uplugin
--- a/Engine/Plugins/Interchange/Runtime/Interchange.uplugin
+++ b/Engine/Plugins/Interchange/Runtime/Interchange.uplugin
@@ -14,6 +14,11 @@
 	"CanContainContent": true,
 	"IsBetaVersion": false,
 	"Installed": false,
+	"SupportedTargetPlatforms": [
+		"Win64",
+		"Linux",
+		"Mac"
+	],
 	"IsHidden": true,
 	"Modules": [
 		{
