Allow shipping target for InterchangeWorker program

#rnx

diff --git a/Engine/Plugins/Interchange/Runtime/Source/Dispatcher/Private/InterchangeWorkerHandler.cpp b/Engine/Plugins/Interchange/Runtime/Source/Dispatcher/Private/InterchangeWorkerHandler.cpp
--- a/Engine/Plugins/Interchange/Runtime/Source/Dispatcher/Private/InterchangeWorkerHandler.cpp
+++ b/Engine/Plugins/Interchange/Runtime/Source/Dispatcher/Private/InterchangeWorkerHandler.cpp
@@ -22,30 +22,15 @@ namespace UE
 		{
 			static FString ProcessorPath = [&]()
 			{
-				FString Path = FPaths::Combine(FPaths::EngineDir(), TEXT("Binaries"));
-
-#if PLATFORM_MAC
-				Path = FPaths::Combine(Path, TEXT("Mac/InterchangeWorker"));
-#elif PLATFORM_LINUX
-				Path = FPaths::Combine(Path, TEXT("Linux/InterchangeWorker"));
-#elif PLATFORM_WINDOWS
-#if UE_BUILD_DEBUG
+				
+				const FString InterchangeWorkerApplicationName = TEXT("InterchangeWorker");
+				FString Path = FPlatformProcess::GenerateApplicationPath(InterchangeWorkerApplicationName, FApp::GetBuildConfiguration());
+				if (!FPaths::FileExists(Path))
 				{
-					FString DebugPath = FPaths::Combine(Path, TEXT("Win64/InterchangeWorker-Win64-Debug.exe"));
-					if (!FPaths::FileExists(DebugPath))
-					{
-						//Try the development build if we cannot find the debug build
-						Path = FPaths::Combine(Path, TEXT("Win64/InterchangeWorker.exe"));
-					}
-					else
-					{
-						Path = MoveTemp(DebugPath);
-					}
+					//Force the development build if the path is not right
+					Path = FPlatformProcess::GenerateApplicationPath(InterchangeWorkerApplicationName, EBuildConfiguration::Development);
 				}
-#else //#if UE_BUILD_DEBUG
-				Path = FPaths::Combine(Path, TEXT("Win64/InterchangeWorker.exe"));
-#endif //#else UE_BUILD_DEBUG
-#endif //#elif PLATFORM_WINDOWS
+
 				if (!FPaths::FileExists(Path))
 				{
 					UE_LOG(LogInterchangeDispatcher, Display, TEXT("InterchangeWorker executable not found. Expected location: %s"), *FPaths::ConvertRelativePathToFull(Path));
diff --git a/Engine/Source/Programs/InterchangeWorker/InterchangeWorker.Target.cs b/Engine/Source/Programs/InterchangeWorker/InterchangeWorker.Target.cs
--- a/Engine/Source/Programs/InterchangeWorker/InterchangeWorker.Target.cs
+++ b/Engine/Source/Programs/InterchangeWorker/InterchangeWorker.Target.cs
@@ -4,7 +4,7 @@ using UnrealBuildTool;
 using System.Collections.Generic;
 
 [SupportedPlatforms(UnrealPlatformClass.Desktop)]
-[SupportedConfigurations(UnrealTargetConfiguration.Debug, UnrealTargetConfiguration.Development)]
+[SupportedConfigurations(UnrealTargetConfiguration.Debug, UnrealTargetConfiguration.Development, UnrealTargetConfiguration.Shipping)]
 public class InterchangeWorkerTarget : TargetRules
 {
 	public InterchangeWorkerTarget(TargetInfo Target) : base(Target)
