Fixed issue of missing meshes when importing into level

#rnx

diff --git a/Engine/Plugins/Interchange/Runtime/Source/Pipelines/Private/InterchangeGenericMeshPipeline.cpp b/Engine/Plugins/Interchange/Runtime/Source/Pipelines/Private/InterchangeGenericMeshPipeline.cpp
--- a/Engine/Plugins/Interchange/Runtime/Source/Pipelines/Private/InterchangeGenericMeshPipeline.cpp
+++ b/Engine/Plugins/Interchange/Runtime/Source/Pipelines/Private/InterchangeGenericMeshPipeline.cpp
@@ -154,7 +154,7 @@ void UInterchangeGenericMeshPipeline::ExecutePipeline(UInterchangeBaseNodeContai
 	DataContext.bConvertSkeletalMeshToStaticMesh = (CommonMeshesProperties->ForceAllMeshAsType == EInterchangeForceMeshType::IFMT_StaticMesh);
 	DataContext.bConvertStaticsWithMorphTargetsToSkeletals = CommonSkeletalMeshesAndAnimationsProperties->bConvertStaticsWithMorphTargetsToSkeletals;
 	DataContext.bImportMeshesInBoneHierarchy = CommonSkeletalMeshesAndAnimationsProperties->bImportMeshesInBoneHierarchy;
-	DataContext.bQueryGeometryOnlyIfNoInstance = true;
+	DataContext.bQueryGeometryOnlyIfNoInstance = CommonMeshesProperties->bBakeMeshes;
 	PipelineMeshesUtilities->SetContext(DataContext);
 
 	//Create skeletalmesh factory nodes
