Fix CIS non unity build
#rnx

diff --git a/Engine/Plugins/Interchange/Runtime/Source/Dispatcher/Private/InterchangeWorkerHandler.cpp b/Engine/Plugins/Interchange/Runtime/Source/Dispatcher/Private/InterchangeWorkerHandler.cpp
--- a/Engine/Plugins/Interchange/Runtime/Source/Dispatcher/Private/InterchangeWorkerHandler.cpp
+++ b/Engine/Plugins/Interchange/Runtime/Source/Dispatcher/Private/InterchangeWorkerHandler.cpp
@@ -9,11 +9,13 @@
 #include "HAL/FileManager.h"
 #include "HAL/PlatformProcess.h"
 #include "HAL/PlatformTime.h"
+#include "Misc/App.h"
 #include "Misc/Paths.h"
 #include "Misc/ScopeLock.h"
 #include "Sockets.h"
 #include "SocketSubsystem.h"
 
+
 namespace UE
 {
 	namespace Interchange
