Fix morph target shape being imported has staticmeshes
#rnx

diff --git a/Engine/Plugins/Interchange/Runtime/Source/Pipelines/Private/InterchangePipelineMeshesUtilities.cpp b/Engine/Plugins/Interchange/Runtime/Source/Pipelines/Private/InterchangePipelineMeshesUtilities.cpp
--- a/Engine/Plugins/Interchange/Runtime/Source/Pipelines/Private/InterchangePipelineMeshesUtilities.cpp
+++ b/Engine/Plugins/Interchange/Runtime/Source/Pipelines/Private/InterchangePipelineMeshesUtilities.cpp
@@ -219,7 +219,10 @@ bool FInterchangePipelineMeshesUtilitiesContext::IsStaticMeshGeometry(const FInt
 	{
 		return false;
 	}
-
+	if (MeshGeometry.MeshNode->IsMorphTarget())
+	{
+		return false;
+	}
 	return !IsSkeletalMeshGeometry(MeshGeometry);
 }
 
@@ -235,6 +238,11 @@ bool FInterchangePipelineMeshesUtilitiesContext::IsSkeletalMeshGeometry(const FI
 		return false;
 	}
 
+	if (MeshGeometry.MeshNode->IsMorphTarget())
+	{
+		return false;
+	}
+
 	if (bConvertStaticMeshToSkeletalMesh
 		|| MeshGeometry.MeshNode->IsSkinnedMesh()
 		|| (bConvertStaticsWithMorphTargetsToSkeletals && (MeshGeometry.MeshNode->GetMorphTargetDependeciesCount() > 0)))
