Fixed a couple of issues reported by a customer

#rnx

diff --git a/Engine/Plugins/Interchange/Runtime/Source/Import/Private/Mesh/InterchangeOBJTranslator.cpp b/Engine/Plugins/Interchange/Runtime/Source/Import/Private/Mesh/InterchangeOBJTranslator.cpp
--- a/Engine/Plugins/Interchange/Runtime/Source/Import/Private/Mesh/InterchangeOBJTranslator.cpp
+++ b/Engine/Plugins/Interchange/Runtime/Source/Import/Private/Mesh/InterchangeOBJTranslator.cpp
@@ -243,7 +243,14 @@ FBox FObjData::GetGroupBoundingBox(const FString& GroupName) const
 
 	for (int32 PositionIndex : GetVertexIndicesUsedByGroup(GroupData))
 	{
-		Box += FVector(PositionToUEBasis(Positions[PositionIndex]));
+		if(Positions.IsValidIndex(PositionIndex))
+		{
+			Box += FVector(PositionToUEBasis(Positions[PositionIndex]));
+		}
+		else
+		{
+			INTERCHANGE_OBJ_TRANSLATOR_LOG_ONCE(LogTemp, Warning, TEXT("FObjData::GetGroupBoundingBox: One or more Position index were invalid, skipping those position to compute the bounding box"));
+		}
 	}
 
 	return Box;
@@ -325,13 +332,22 @@ FMeshDescription FObjData::MakeMeshDescriptionForGroup(const FString& GroupName,
 	// Create vertices and initialize positions
 	// Note that we create a contiguous list of vertex indices from 0...n-1, referencing potentially sparse obj vertices
 
+	TVertexAttributesRef<FVector3f> MeshPositions = Attributes.GetVertexPositions();
 	MeshDescription.ReserveNewVertices(VertexIndexMapping.Num());
 	for (int32 ObjVertexIndex : VertexIndexMapping)
 	{
 		FVertexID VertexIndex = MeshDescription.CreateVertex();
-		FVector3f& Position = Attributes.GetVertexPositions()[VertexIndex];
-		Position = PositionToUEBasis(Positions[ObjVertexIndex]);
-		TransformPosition(TotalMatrix, Position);
+		if (MeshPositions.GetRawArray().IsValidIndex(VertexIndex) && Positions.IsValidIndex(ObjVertexIndex))
+		{
+			FVector3f& Position = Attributes.GetVertexPositions()[VertexIndex];
+			Position = PositionToUEBasis(Positions[ObjVertexIndex]);
+			TransformPosition(TotalMatrix, Position);
+		}
+		else
+		{
+			INTERCHANGE_OBJ_TRANSLATOR_LOG_ONCE(LogTemp, Warning, TEXT("FObjData::MakeMeshDescriptionForGroup: One or more vertex position index are not valid, skipping those vertex position"));
+		}
+		
 	}
 
 	// Create UVs and initialize values
@@ -1315,7 +1331,7 @@ bool UInterchangeOBJTranslator::Translate(UInterchangeBaseNodeContainer& BaseNod
 			const FObjData::FMaterialData& MaterialData = Material.Value;
 			FString NodeUid = UInterchangeShaderGraphNode::MakeNodeUid(MakeShaderGraphNodeName(MaterialName));
 
-			UInterchangeShaderGraphNode* ShaderGraphNode = UInterchangeShaderGraphNode::Create(&BaseNodeContainer, MaterialName);
+			UInterchangeShaderGraphNode * ShaderGraphNode = NewObject<UInterchangeShaderGraphNode>(&BaseNodeContainer);
 			ShaderGraphNode->InitializeNode(NodeUid, MaterialName, EInterchangeNodeContainerType::TranslatedAsset);
 			BaseNodeContainer.AddNode(ShaderGraphNode);
 
