Fixed regression introduced by CL 26740353

#rnx

diff --git a/Engine/Plugins/Interchange/Runtime/Source/Pipelines/Public/InterchangeGenericAnimationPipeline.h b/Engine/Plugins/Interchange/Runtime/Source/Pipelines/Public/InterchangeGenericAnimationPipeline.h
--- a/Engine/Plugins/Interchange/Runtime/Source/Pipelines/Public/InterchangeGenericAnimationPipeline.h
+++ b/Engine/Plugins/Interchange/Runtime/Source/Pipelines/Public/InterchangeGenericAnimationPipeline.h
@@ -131,7 +131,10 @@ private:
 
 	void CreateAnimSequenceFactoryNode(UInterchangeSkeletalAnimationTrackNode& Node);
 
+	// Set as a property to carry this value over during a duplicate
+	UPROPERTY()
 	bool bSceneImport = false;
+
 	UInterchangeBaseNodeContainer* BaseNodeContainer = nullptr;
 	TArray<const UInterchangeSourceData*> SourceDatas;
 
