Fixed conversion of FBX color to UE

#rnx

diff --git a/Engine/Plugins/Interchange/Runtime/Source/Parsers/Fbx/Private/FbxMaterial.cpp b/Engine/Plugins/Interchange/Runtime/Source/Parsers/Fbx/Private/FbxMaterial.cpp
--- a/Engine/Plugins/Interchange/Runtime/Source/Parsers/Fbx/Private/FbxMaterial.cpp
+++ b/Engine/Plugins/Interchange/Runtime/Source/Parsers/Fbx/Private/FbxMaterial.cpp
@@ -110,37 +110,32 @@ namespace UE
 				using namespace Materials::Standard::Nodes;
 
 				const int32 TextureCount = Property.GetSrcObjectCount<FbxFileTexture>();
-				const FbxDataType DataType = Property.GetPropertyDataType();
+				const EFbxType DataType = Property.GetPropertyDataType().GetType();
 				FString InputToConnectTo = InputName.ToString();
 
 				if (TextureCount == 0)
 				{
-					if (DataType.GetType() == eFbxDouble || DataType.GetType() == eFbxFloat || DataType.GetType() == eFbxInt)
-					{
-						float PropertyValue = Property.Get<float>() * Factor;
-						if (bInverse)
-						{
-							PropertyValue = 1.f - PropertyValue;
-						}
+					const FString InputAttributeKey = UInterchangeShaderPortsAPI::MakeInputValueKey(InputName.ToString());
 
-						ShaderGraphNode->AddFloatAttribute(UInterchangeShaderPortsAPI::MakeInputValueKey(InputToConnectTo), PropertyValue);
+					if (DataType == eFbxDouble || DataType == eFbxFloat || DataType == eFbxInt)
+					{
+						const float PropertyValue = Property.Get<float>() * Factor;
+						ShaderGraphNode->AddFloatAttribute(InputAttributeKey, bInverse ? 1.f - PropertyValue : PropertyValue);
 					}
-					else if (DataType.GetType() == eFbxDouble3)
+					else if (DataType == eFbxDouble3 || DataType == eFbxDouble4)
 					{
-						FLinearColor PropertyValue = FFbxConvert::ConvertColor(Property.Get<FbxDouble3>()) * Factor;
-						if (bInverse)
-						{
-							PropertyValue = FLinearColor::White - PropertyValue;
-						}
+						FbxDouble3 Color = DataType == eFbxDouble3 ? Property.Get<FbxDouble3>() : Property.Get<FbxDouble4>();
+						FVector3f FbxValue = FVector3f(Color[0], Color[1], Color[2]) * Factor;
+						FLinearColor PropertyValue = bInverse ? FVector3f::OneVector - FbxValue : FbxValue;
 
 						if (DefaultValue.IsType<FLinearColor>())
 						{
-							ShaderGraphNode->AddLinearColorAttribute(UInterchangeShaderPortsAPI::MakeInputValueKey(InputToConnectTo), PropertyValue);
+							ShaderGraphNode->AddLinearColorAttribute(InputAttributeKey, PropertyValue);
 						}
 						else if (DefaultValue.IsType<float>())
 						{
 							// We're connecting a linear color to a float input. Ideally, we'd go through a desaturate, but for now we'll just take the red channel and ignore the rest.
-							ShaderGraphNode->AddFloatAttribute(UInterchangeShaderPortsAPI::MakeInputValueKey(InputToConnectTo), PropertyValue.R);
+							ShaderGraphNode->AddFloatAttribute(InputAttributeKey, PropertyValue.R);
 						}
 					}
 
