Fixed issue with telemetry in glTF translator

#rnx

diff --git a/Engine/Plugins/Interchange/Runtime/Source/Import/Private/Gltf/InterchangeGltfTranslator.cpp b/Engine/Plugins/Interchange/Runtime/Source/Import/Private/Gltf/InterchangeGltfTranslator.cpp
--- a/Engine/Plugins/Interchange/Runtime/Source/Import/Private/Gltf/InterchangeGltfTranslator.cpp
+++ b/Engine/Plugins/Interchange/Runtime/Source/Import/Private/Gltf/InterchangeGltfTranslator.cpp
@@ -188,16 +188,16 @@ namespace UE::Interchange::Gltf::Private
 			switch (TranslationResult)
 			{
 			case SUCCESSFULL:
-				GLTFAnalytics.Add(FAnalyticsEventAttribute(TEXT("TranslationResult"), "Successfull."));
+				GLTFAnalytics.Add(FAnalyticsEventAttribute(TEXT("TranslationStatus"), "Successful."));
 				break;
 			case INPUT_FILE_NOTFOUND:
-				GLTFAnalytics.Add(FAnalyticsEventAttribute(TEXT("TranslationResult"), "[Failed] Input File Not Found."));
+				GLTFAnalytics.Add(FAnalyticsEventAttribute(TEXT("TranslationStatus"), "[Failed] Input File Not Found."));
 				break;
 			case GLTFREADER_FAILED:
-				GLTFAnalytics.Add(FAnalyticsEventAttribute(TEXT("TranslationResult"), "[Failed] Parsing error: " + GLTFReaderLogMessage));
+				GLTFAnalytics.Add(FAnalyticsEventAttribute(TEXT("TranslationStatus"), "[Failed] Parsing error."));
 				break;
 			case NOTSUPPORTED_EXTENSION_FOUND:
-				GLTFAnalytics.Add(FAnalyticsEventAttribute(TEXT("TranslationResult"), "[Failed] Unsupported Extension Found."));
+				GLTFAnalytics.Add(FAnalyticsEventAttribute(TEXT("TranslationStatus"), "[Failed] Unsupported Extension Found."));
 				break;
 			default:
 				break;
