[Interchange - FBX] Import - Animation problem - fixes.
Finished GeometricTransform support for Interchange.

diff --git a/Engine/Plugins/Interchange/Runtime/Source/FactoryNodes/Private/InterchangeMeshActorFactoryNode.cpp b/Engine/Plugins/Interchange/Runtime/Source/FactoryNodes/Private/InterchangeMeshActorFactoryNode.cpp
--- a/Engine/Plugins/Interchange/Runtime/Source/FactoryNodes/Private/InterchangeMeshActorFactoryNode.cpp
+++ b/Engine/Plugins/Interchange/Runtime/Source/FactoryNodes/Private/InterchangeMeshActorFactoryNode.cpp
@@ -40,4 +40,14 @@ bool UInterchangeMeshActorFactoryNode::SetCustomAnimationAssetUidToPlay(const FS
 bool UInterchangeMeshActorFactoryNode::GetCustomAnimationAssetUidToPlay(FString& AttributeValue) const
 {
 	IMPLEMENT_NODE_ATTRIBUTE_GETTER(AnimationAssetUidToPlay, FString);
+}
+
+bool UInterchangeMeshActorFactoryNode::GetCustomGeometricTransform(FTransform& AttributeValue) const
+{
+	IMPLEMENT_NODE_ATTRIBUTE_GETTER(GeometricTransform, FTransform);
+}
+
+bool UInterchangeMeshActorFactoryNode::SetCustomGeometricTransform(const FTransform& AttributeValue)
+{
+	IMPLEMENT_NODE_ATTRIBUTE_SETTER_NODELEGATE(GeometricTransform, FTransform);
 }
\ No newline at end of file
diff --git a/Engine/Plugins/Interchange/Runtime/Source/FactoryNodes/Public/InterchangeMeshActorFactoryNode.h b/Engine/Plugins/Interchange/Runtime/Source/FactoryNodes/Public/InterchangeMeshActorFactoryNode.h
--- a/Engine/Plugins/Interchange/Runtime/Source/FactoryNodes/Public/InterchangeMeshActorFactoryNode.h
+++ b/Engine/Plugins/Interchange/Runtime/Source/FactoryNodes/Public/InterchangeMeshActorFactoryNode.h
@@ -58,9 +58,20 @@ public:
 	UFUNCTION(BlueprintCallable, Category = "Interchange | Node | SkeletalMesh")
 	bool GetCustomAnimationAssetUidToPlay(FString& AttributeValue) const;
 
+	/** Get the geometric offset. Any mesh attach to this scene node will be offset using this transform. */
+	UFUNCTION(BlueprintCallable, Category = "Interchange | Node | StaticMesh")
+	bool GetCustomGeometricTransform(FTransform& AttributeValue) const;
+
+	/** Set the geometric offset. Any mesh attach to this scene node will be offset using this transform. */
+	UFUNCTION(BlueprintCallable, Category = "Interchange | Node | StaticMesh")
+	bool SetCustomGeometricTransform(const FTransform& AttributeValue);
+
 private:
 	UE::Interchange::TMapAttributeHelper<FString, FString> SlotMaterialDependencies;
 
 	//A scene node can reference an animation asset on top of base asset:
 	const UE::Interchange::FAttributeKey Macro_CustomAnimationAssetUidToPlayKey = UE::Interchange::FAttributeKey(TEXT("AnimationAssetUidToPlay"));
+
+	//A scene node can have a transform apply to the mesh it reference.
+	const UE::Interchange::FAttributeKey Macro_CustomGeometricTransformKey = UE::Interchange::FAttributeKey(TEXT("GeometricTransform"));
 };
\ No newline at end of file
diff --git a/Engine/Plugins/Interchange/Runtime/Source/Import/Private/Mesh/InterchangeSkeletalMeshFactory.cpp b/Engine/Plugins/Interchange/Runtime/Source/Import/Private/Mesh/InterchangeSkeletalMeshFactory.cpp
--- a/Engine/Plugins/Interchange/Runtime/Source/Import/Private/Mesh/InterchangeSkeletalMeshFactory.cpp
+++ b/Engine/Plugins/Interchange/Runtime/Source/Import/Private/Mesh/InterchangeSkeletalMeshFactory.cpp
@@ -1237,7 +1237,7 @@ UInterchangeFactoryBase::FImportAssetResult UInterchangeSkeletalMeshFactory::Imp
 					FTransform SceneNodeGeometricTransform;
 					if(MeshReference.SceneNode->GetCustomGeometricTransform(SceneNodeGeometricTransform))
 					{
-						SceneNodeTransform *= SceneNodeGeometricTransform;
+						SceneNodeTransform = SceneNodeGeometricTransform * SceneNodeTransform;
 					}
 					MeshReference.SceneGlobalTransform = SceneNodeTransform;
 				}
diff --git a/Engine/Plugins/Interchange/Runtime/Source/Import/Private/Scene/InterchangeStaticMeshActorFactory.cpp b/Engine/Plugins/Interchange/Runtime/Source/Import/Private/Scene/InterchangeStaticMeshActorFactory.cpp
--- a/Engine/Plugins/Interchange/Runtime/Source/Import/Private/Scene/InterchangeStaticMeshActorFactory.cpp
+++ b/Engine/Plugins/Interchange/Runtime/Source/Import/Private/Scene/InterchangeStaticMeshActorFactory.cpp
@@ -29,7 +29,72 @@ UObject* UInterchangeStaticMeshActorFactory::ProcessActor(AActor& SpawnedActor,
 
 	if (UStaticMeshComponent* StaticMeshComponent = StaticMeshActor->GetStaticMeshComponent())
 	{
-		StaticMeshComponent->UnregisterComponent();
+		bool bHasGeometricTransform = false;
+		FTransform GeometricTransform;
+		if (FactoryNode.IsA(UInterchangeMeshActorFactoryNode::StaticClass()))
+		{
+			const UInterchangeMeshActorFactoryNode& MeshActorFactoryNode = static_cast<const UInterchangeMeshActorFactoryNode&>(FactoryNode);
+
+			if (MeshActorFactoryNode.GetCustomGeometricTransform(GeometricTransform))
+			{
+				bHasGeometricTransform = true;
+			}
+		}
+
+
+		if (bHasGeometricTransform)
+		{
+			UStaticMeshComponent* GeometricTransformMeshComponent = nullptr;
+			if (StaticMeshComponent->GetNumChildrenComponents() > 0)
+			{
+				USceneComponent* ChildSceneComponent = StaticMeshComponent->GetChildComponent(0);
+				if (ChildSceneComponent->IsA(UStaticMeshComponent::StaticClass()))
+				{
+					AActor* ParentActor = ChildSceneComponent->GetAttachParentActor();
+					if (ParentActor == StaticMeshActor)
+					{
+						GeometricTransformMeshComponent = Cast<UStaticMeshComponent>(ChildSceneComponent);
+					}
+				}
+			}
+
+			if (GeometricTransformMeshComponent == nullptr)
+			{
+				StaticMeshActor->UnregisterAllComponents();
+
+				GeometricTransformMeshComponent = NewObject<UStaticMeshComponent>(StaticMeshActor->GetRootComponent(), TEXT("GeometricTransform"));
+
+#if WITH_EDITORONLY_DATA
+				GeometricTransformMeshComponent->bVisualizeComponent = true;
+#endif
+				StaticMeshActor->AddInstanceComponent(GeometricTransformMeshComponent);
+
+				GeometricTransformMeshComponent->SetMobility(StaticMeshComponent->Mobility);
+				
+				if (const UInterchangeFactoryBaseNode* MeshNode = ActorHelper::FindAssetInstanceFactoryNode(&NodeContainer, &FactoryNode))
+				{
+					FSoftObjectPath ReferenceObject;
+					MeshNode->GetCustomReferenceObject(ReferenceObject);
+					if (UStaticMesh* StaticMesh = Cast<UStaticMesh>(ReferenceObject.TryLoad()))
+					{
+						if (StaticMesh != GeometricTransformMeshComponent->GetStaticMesh())
+						{
+							GeometricTransformMeshComponent->SetStaticMesh(StaticMesh);
+						}
+					}
+				}
+
+				GeometricTransformMeshComponent->SetupAttachment(StaticMeshComponent);
+
+				StaticMeshActor->ReregisterAllComponents();
+
+				GeometricTransformMeshComponent->SetRelativeTransform(GeometricTransform);
+
+				return StaticMeshComponent;
+			}
+		}
+
+		StaticMeshComponent->UnregisterComponent();		
 
 		if (const UInterchangeFactoryBaseNode* MeshNode = ActorHelper::FindAssetInstanceFactoryNode(&NodeContainer, &FactoryNode))
 		{
diff --git a/Engine/Plugins/Interchange/Runtime/Source/Parsers/Fbx/Private/FbxScene.cpp b/Engine/Plugins/Interchange/Runtime/Source/Parsers/Fbx/Private/FbxScene.cpp
--- a/Engine/Plugins/Interchange/Runtime/Source/Parsers/Fbx/Private/FbxScene.cpp
+++ b/Engine/Plugins/Interchange/Runtime/Source/Parsers/Fbx/Private/FbxScene.cpp
@@ -49,7 +49,11 @@ namespace UE
 				if (MeshNode)
 				{
 					UnrealSceneNode->SetCustomAssetInstanceUid(MeshNode->GetUniqueID());
-					UnrealSceneNode->SetCustomGeometricTransform(GeometricTransform);
+
+					if (!GeometricTransform.Equals(FTransform::Identity))
+					{
+						UnrealSceneNode->SetCustomGeometricTransform(GeometricTransform);
+					}
 
 					// @todo: Nothing is using the SceneInstanceUid in the MeshNode. Do we even need to support it?
 					// For the moment an ugly const_cast so we can mutate it (it was fetched from the NodeContainer and is hence const).
diff --git a/Engine/Plugins/Interchange/Runtime/Source/Pipelines/Private/InterchangeGenericScenesPipeline.cpp b/Engine/Plugins/Interchange/Runtime/Source/Pipelines/Private/InterchangeGenericScenesPipeline.cpp
--- a/Engine/Plugins/Interchange/Runtime/Source/Pipelines/Private/InterchangeGenericScenesPipeline.cpp
+++ b/Engine/Plugins/Interchange/Runtime/Source/Pipelines/Private/InterchangeGenericScenesPipeline.cpp
@@ -529,6 +529,12 @@ void UInterchangeGenericLevelPipeline::SetUpFactoryNode(UInterchangeActorFactory
 			}
 
 			MeshActorFactoryNode->AddFactoryDependencyUid(UInterchangeFactoryBaseNode::BuildFactoryNodeUid(MeshNode->GetUniqueID()));
+
+			FTransform GeometricTransform;
+			if (SceneNode->GetCustomGeometricTransform(GeometricTransform))
+			{
+				MeshActorFactoryNode->SetCustomGeometricTransform(GeometricTransform);
+			}
 		}
 	}
 	else if (const UInterchangeBaseLightNode* BaseLightNode = Cast<UInterchangeBaseLightNode>(TranslatedAssetNode))
