Duplicate Materials detection in Interchange
- Interchange can identify duplicate materials and create material instances
- Added parameter keys support for detecting material parameters

diff --git a/Engine/Plugins/Experimental/Enterprise/DatasmithInterchange/Source/DatasmithInterchange/Private/InterchangeDatasmithMaterialNode.cpp b/Engine/Plugins/Experimental/Enterprise/DatasmithInterchange/Source/DatasmithInterchange/Private/InterchangeDatasmithMaterialNode.cpp
--- a/Engine/Plugins/Experimental/Enterprise/DatasmithInterchange/Source/DatasmithInterchange/Private/InterchangeDatasmithMaterialNode.cpp
+++ b/Engine/Plugins/Experimental/Enterprise/DatasmithInterchange/Source/DatasmithInterchange/Private/InterchangeDatasmithMaterialNode.cpp
@@ -17,6 +17,7 @@
 #include "InterchangeShaderGraphNode.h"
 #include "InterchangeTexture2DNode.h"
 #include "InterchangeTexture2DFactoryNode.h"
+#include "Nodes/InterchangeUserDefinedAttribute.h"
 
 #include "Async/TaskGraphInterfaces.h"
 #include "Materials/MaterialExpression.h"
@@ -100,7 +101,7 @@ namespace UE::DatasmithInterchange::MaterialUtils
 
 		static UMaterialExpression* GetDefaultMaterialExpression(const TCHAR* GenericExpressionname);
 
-		UInterchangeShaderNode* BuildGenericExpression(const IDatasmithMaterialExpressionGeneric& DatasmithExpression, UInterchangeShaderNode& ParentNode, TArray<FName>* OutputNames = nullptr);
+		UInterchangeShaderNode* BuildGenericExpression(const IDatasmithMaterialExpressionGeneric& DatasmithExpression, const UInterchangeShaderNode& ParentNode, TArray<FName>* OutputNames = nullptr);
 
 		template<class T = UInterchangeShaderNode>
 		T* CreateExpressionNode(const IDatasmithMaterialExpression& DatasmithExpression, const FString& ParentNodeUid)
@@ -112,9 +113,14 @@ namespace UE::DatasmithInterchange::MaterialUtils
 				return nullptr;
 			}
 
+			if (ProcessedExpressions[*ExpressionIndexPtr] != nullptr)
+			{
+				return Cast<T>(ProcessedExpressions[*ExpressionIndexPtr]);
+			}
+
 			const FString NodeName = FDatasmithUtils::SanitizeObjectName(DatasmithExpression.GetName());
 			const FString NodeUid = NodeUtils::MaterialExpressionPrefix + MaterialID + TEXT("_") + FString::FromInt(*ExpressionIndexPtr) + TEXT("_") + NodeName;
-
+			
 			UInterchangeShaderNode* ExpressionNode = NewObject< T >(&NodeContainer);
 			ExpressionNode->InitializeNode(NodeUid, NodeName, EInterchangeNodeContainerType::TranslatedAsset);
 			NodeContainer.AddNode(ExpressionNode);
@@ -228,19 +234,19 @@ namespace UE::DatasmithInterchange::MaterialUtils
 			return;
 		}
 
+		ExpressionNode->SetCustomShaderType(TextureSample::Name.ToString());
+		const bool bIsAParameter = true;
+		const FString TextureInputParameterName = TextureSample::Inputs::Texture.ToString();
 		if (FPackageName::IsValidObjectPath(DatasmithExpression.GetTexturePathName()))
 		{
-			// TODO: Handle case of existing asset
+			ExpressionNode->AddStringInput(TextureInputParameterName, DatasmithExpression.GetTexturePathName(), bIsAParameter);
 		}
 		else
 		{
-			ExpressionNode->SetCustomShaderType(TextureSample::Name.ToString());
-
 			const FString TextureUid = NodeUtils::TexturePrefix + FDatasmithUtils::SanitizeObjectName(DatasmithExpression.GetTexturePathName());
 			if (NodeContainer.GetNode(TextureUid))
 			{
-				const FString ValueKey = UInterchangeShaderPortsAPI::MakeInputValueKey(TextureSample::Inputs::Texture.ToString());
-				ExpressionNode->AddStringAttribute(ValueKey, TextureUid);
+				ExpressionNode->AddStringInput(TextureInputParameterName, TextureUid, bIsAParameter);
 			}
 			else
 			{
@@ -256,9 +262,10 @@ namespace UE::DatasmithInterchange::MaterialUtils
 										TextureSample::Outputs::A,
 										TextureSample::Outputs::RGBA };
 
+
 		ConnectionData.ConnectOuputToInput(ExpressionNode->GetUniqueID(), OutputNames[OutputIndex]);
 
-		FConnectionData CoordConnectionData{ExpressionNode, TextureSample::Inputs::Coordinates.ToString() };
+		FConnectionData CoordConnectionData{ ExpressionNode, TextureSample::Inputs::Coordinates.ToString() };
 
 		const IDatasmithExpressionInput& ExpressionInput = DatasmithExpression.GetInputCoordinate();
 		ConnectExpression(ExpressionInput.GetExpression(), CoordConnectionData, ExpressionInput.GetOutputIndex());
@@ -367,7 +374,7 @@ namespace UE::DatasmithInterchange::MaterialUtils
 		return ExpressionClass->GetDefaultObject<UMaterialExpression>();
 	}
 
-	UInterchangeShaderNode* FPbrMaterialHelper::BuildGenericExpression(const IDatasmithMaterialExpressionGeneric& DatasmithExpression, UInterchangeShaderNode& ParentNode, TArray<FName>* OutputNames)
+	UInterchangeShaderNode* FPbrMaterialHelper::BuildGenericExpression(const IDatasmithMaterialExpressionGeneric& DatasmithExpression, const UInterchangeShaderNode& ParentNode, TArray<FName>* OutputNames)
 	{
 		UInterchangeShaderNode* ExpressionNode = nullptr;
 
@@ -391,6 +398,55 @@ namespace UE::DatasmithInterchange::MaterialUtils
 
 		// TODO: What about if MaterialExpression->HasAParameterName()?
 
+		
+		for (int32 PropertyIndex = 0; PropertyIndex < DatasmithExpression.GetPropertiesCount(); ++PropertyIndex)
+		{
+			const TSharedPtr<IDatasmithKeyValueProperty> Property = DatasmithExpression.GetProperty(PropertyIndex);
+
+			switch (Property->GetPropertyType())
+			{
+			case EDatasmithKeyValuePropertyType::Float:
+			{
+				const FString PropertyKey = Property->GetName();
+				const float PropertyValue = FCString::Atof(Property->GetValue());
+				const bool bRequiresDelegate = true;
+				UInterchangeUserDefinedAttributesAPI::CreateUserDefinedAttribute_Float(ExpressionNode, PropertyKey, PropertyValue,FString(), bRequiresDelegate);
+				break;
+			}
+			case EDatasmithKeyValuePropertyType::Texture:
+			{
+				using namespace UE::Interchange::Materials::Standard::Nodes;
+				const FString TextureUid = NodeUtils::TexturePrefix + FDatasmithUtils::SanitizeObjectName(Property->GetValue());
+				if (NodeContainer.GetNode(TextureUid))
+				{
+					// Textures are set in the MaterialFactory that require the TextureFactoryUid as opposed to TextureUid, 
+					// hence assuming that the factory will be built for the texture, we are storing the TextureFactoryUid.
+					const FString TextureFactoryUid = UInterchangeFactoryBaseNode::BuildFactoryNodeUid(TextureUid);
+
+					// Passing it through the KeyValueProperty pipeline to ensure that it will be transferred to FactoryNode but no delegate would be added for this attribute.
+					UInterchangeUserDefinedAttributesAPI::CreateUserDefinedAttribute_FString(ExpressionNode, TextureSample::Inputs::Texture.ToString(), TextureFactoryUid, FString());
+				}
+				else
+				{
+					// Existing texture
+					UInterchangeUserDefinedAttributesAPI::CreateUserDefinedAttribute_FString(ExpressionNode, TextureSample::Inputs::Texture.ToString(), TextureUid, FString());
+				}
+				break;
+			}
+			// TODO: Can be added later.
+			case EDatasmithKeyValuePropertyType::Bool:
+			case EDatasmithKeyValuePropertyType::Vector:
+			case EDatasmithKeyValuePropertyType::String:
+			case EDatasmithKeyValuePropertyType::Color:
+			case EDatasmithKeyValuePropertyType::Integer:
+			default:
+			{
+				break;
+			}
+			}
+		}
+		
+
 		for (int32 InputIndex = 0; InputIndex < DatasmithExpression.GetInputCount(); ++InputIndex)
 		{
 			if (const IDatasmithExpressionInput* ExpressionInput = DatasmithExpression.GetInput(InputIndex))
@@ -568,23 +624,76 @@ namespace UE::DatasmithInterchange::MaterialUtils
 			}
 		}
 
+		const FString ParameterName(MaterialExpression->GetName());
+		const bool bIsAParameter = !ParameterName.IsEmpty();
+
 		if (MaterialExpression->IsSubType(EDatasmithMaterialExpressionType::ConstantBool))
 		{
+			using namespace UE::Interchange::Materials::Standard::Nodes;
+
 			const IDatasmithMaterialExpressionBool* BoolExpression = static_cast<const IDatasmithMaterialExpressionBool*>(MaterialExpression);
-			const FString InputName(UInterchangeShaderPortsAPI::MakeInputValueKey(ConnectionData.InputName));
-			ConnectionData.InputNode->AddBooleanAttribute(InputName, BoolExpression->GetBool());
+			
+			ensure(!ParameterName.IsEmpty());
+
+			UInterchangeShaderNode* ExpressionNode = CreateExpressionNode(*BoolExpression, ConnectionData.InputNode->GetUniqueID());
+			if (!ensure(ExpressionNode))
+			{
+				// TODO: Log error
+				return;
+			}
+
+			const FString InputName(UInterchangeShaderPortsAPI::MakeInputParameterKey(StaticBoolParameter::Attributes::DefaultValue.ToString()));
+			ExpressionNode->AddBooleanAttribute(InputName, BoolExpression->GetBool());
+			ExpressionNode->SetCustomShaderType(StaticBoolParameter::Name.ToString());
+			UInterchangeShaderPortsAPI::ConnectDefaultOuputToInput(ConnectionData.InputNode, ConnectionData.InputName, ExpressionNode->GetUniqueID());
 		}
 		else if (MaterialExpression->IsSubType(EDatasmithMaterialExpressionType::ConstantScalar))
 		{
+			using namespace UE::Interchange::Materials::Standard::Nodes;
+
 			const IDatasmithMaterialExpressionScalar* ScalarExpression = static_cast<const IDatasmithMaterialExpressionScalar*>(MaterialExpression);
-			const FString InputName(UInterchangeShaderPortsAPI::MakeInputValueKey(ConnectionData.InputName));
-			ConnectionData.InputNode->AddFloatAttribute(InputName, ScalarExpression->GetScalar());
+
+			if (bIsAParameter)
+			{
+				UInterchangeShaderNode* ExpressionNode = CreateExpressionNode(*ScalarExpression, ConnectionData.InputNode->GetUniqueID());
+				if (!ensure(ExpressionNode))
+				{
+					// TODO: Log error
+					return;
+				}
+
+				ExpressionNode->AddFloatInput(ScalarParameter::Attributes::DefaultValue.ToString(), ScalarExpression->GetScalar(), bIsAParameter);
+				ExpressionNode->SetCustomShaderType(ScalarParameter::Name.ToString());
+				UInterchangeShaderPortsAPI::ConnectDefaultOuputToInput(ConnectionData.InputNode, ConnectionData.InputName, ExpressionNode->GetUniqueID());
+			}
+			else
+			{
+				ConnectionData.InputNode->AddFloatInput(ConnectionData.InputName, ScalarExpression->GetScalar(), bIsAParameter);
+			}
 		}
 		else if (MaterialExpression->IsSubType(EDatasmithMaterialExpressionType::ConstantColor))
 		{
+			using namespace UE::Interchange::Materials::Standard::Nodes;
+
 			const IDatasmithMaterialExpressionColor* ColorExpression = static_cast<const IDatasmithMaterialExpressionColor*>(MaterialExpression);
-			const FString InputName(UInterchangeShaderPortsAPI::MakeInputValueKey(ConnectionData.InputName));
-			ConnectionData.InputNode->AddLinearColorAttribute(InputName, ColorExpression->GetColor());
+
+			if (bIsAParameter)
+			{
+				UInterchangeShaderNode* ExpressionNode = CreateExpressionNode(*ColorExpression, ConnectionData.InputNode->GetUniqueID());
+				if (!ensure(ExpressionNode))
+				{
+					// TODO: Log error
+					return;
+				}
+
+				ExpressionNode->AddLinearColorInput(VectorParameter::Attributes::DefaultValue.ToString(), ColorExpression->GetColor(), bIsAParameter);
+				ExpressionNode->SetCustomShaderType(VectorParameter::Name.ToString());
+				UInterchangeShaderPortsAPI::ConnectDefaultOuputToInput(ConnectionData.InputNode, ConnectionData.InputName, ExpressionNode->GetUniqueID());
+			}
+			else
+			{
+				ConnectionData.InputNode->AddLinearColorInput(ConnectionData.InputName, ColorExpression->GetColor(), bIsAParameter);
+			}
 		}
 		else if (MaterialExpression->IsSubType(EDatasmithMaterialExpressionType::Texture))
 		{
diff --git a/Engine/Plugins/Experimental/Enterprise/DatasmithInterchange/Source/DatasmithInterchange/Private/InterchangeDatasmithMaterialPipeline.cpp b/Engine/Plugins/Experimental/Enterprise/DatasmithInterchange/Source/DatasmithInterchange/Private/InterchangeDatasmithMaterialPipeline.cpp
--- a/Engine/Plugins/Experimental/Enterprise/DatasmithInterchange/Source/DatasmithInterchange/Private/InterchangeDatasmithMaterialPipeline.cpp
+++ b/Engine/Plugins/Experimental/Enterprise/DatasmithInterchange/Source/DatasmithInterchange/Private/InterchangeDatasmithMaterialPipeline.cpp
@@ -15,9 +15,11 @@
 #include "InterchangeManager.h"
 #include "InterchangeMaterialDefinitions.h"
 #include "InterchangeMaterialFactoryNode.h"
+
 #include "InterchangeTexture2DNode.h"
 #include "InterchangeTexture2DFactoryNode.h"
 #include "InterchangeMaterialInstanceNode.h"
+#include "Nodes/InterchangeUserDefinedAttribute.h"
 
 #include "Materials/MaterialExpressionMaterialFunctionCall.h"
 #include "Materials/MaterialFunction.h"
@@ -25,6 +27,31 @@
 #include "Materials/MaterialInstanceDynamic.h"
 #include "Misc/PackageName.h"
 
+DEFINE_LOG_CATEGORY(LogInterchangeMaterialPipeline);
+
+namespace UE::Interchange::MaterialUtils
+{
+	UClass* FindMaterialExpressionClass(const TCHAR* ClassName)
+	{
+		if (!ensure(ClassName))
+		{
+			return nullptr;
+		}
+
+		if (UClass* Result = FindFirstObject<UClass>(ClassName, EFindFirstObjectOptions::None, ELogVerbosity::Warning, TEXT("Datasmith FindClass")))
+		{
+			return Result;
+		}
+
+		if (UObjectRedirector* RenamedClassRedirector = FindFirstObject<UObjectRedirector>(ClassName, EFindFirstObjectOptions::None, ELogVerbosity::Warning, TEXT("Datasmith FindClass")))
+		{
+			return CastChecked<UClass>(RenamedClassRedirector->DestinationObject);
+		}
+
+		return nullptr;
+	}
+}
+
 UInterchangeDatasmithMaterialPipeline::UInterchangeDatasmithMaterialPipeline()
 	: Super()
 {
@@ -294,7 +321,26 @@ void UInterchangeDatasmithMaterialPipeline::UpdateMaterialFactoryNodes(const TAr
 		{
 			continue;
 		}
-
+		
+		
+		if (UInterchangeMaterialExpressionFactoryNode* MaterialExpressionFactoryNode = Cast<UInterchangeMaterialExpressionFactoryNode>(FactoryNode))
+		{
+			FString ExpressionClassName;
+			MaterialExpressionFactoryNode->GetCustomExpressionClassName(ExpressionClassName);
+			
+			UClass* MaterialExpressionClass = UE::Interchange::MaterialUtils::FindMaterialExpressionClass(*ExpressionClassName);
+			if (!MaterialExpressionClass)
+			{
+				UE_LOG(LogInterchangeMaterialPipeline, Warning, TEXT("Invalid Material Expression Class."));
+			}
+			else
+			{
+				// This will transfer both porperties with delegates and transferable properties to the Factory Node.
+				UInterchangeUserDefinedAttributesAPI::DuplicateAllUserDefinedAttribute(ShaderNode, MaterialExpressionFactoryNode, false);
+				UInterchangeUserDefinedAttributesAPI::AddApplyAndFillDelegatesToFactory(MaterialExpressionFactoryNode, MaterialExpressionClass);
+			}
+		}
+		
 		if (UInterchangeDatasmithPbrMaterialNode* PbrMaterialNode = Cast<UInterchangeDatasmithPbrMaterialNode>(ShaderNode))
 		{
 			UInterchangeBaseMaterialFactoryNode* BaseMaterialFactoryNode = Cast<UInterchangeBaseMaterialFactoryNode>(FactoryNode);
@@ -339,5 +385,7 @@ void UInterchangeDatasmithMaterialPipeline::UpdateMaterialFactoryNodes(const TAr
 				//}
 			}
 		}
+
+
 	}
-}
+}
\ No newline at end of file
diff --git a/Engine/Plugins/Experimental/Enterprise/DatasmithInterchange/Source/DatasmithInterchange/Private/InterchangeDatasmithMaterialPipeline.h b/Engine/Plugins/Experimental/Enterprise/DatasmithInterchange/Source/DatasmithInterchange/Private/InterchangeDatasmithMaterialPipeline.h
--- a/Engine/Plugins/Experimental/Enterprise/DatasmithInterchange/Source/DatasmithInterchange/Private/InterchangeDatasmithMaterialPipeline.h
+++ b/Engine/Plugins/Experimental/Enterprise/DatasmithInterchange/Source/DatasmithInterchange/Private/InterchangeDatasmithMaterialPipeline.h
@@ -16,6 +16,8 @@ class UInterchangeMaterialInstanceFactoryNode;
 class UInterchangeShaderGraphNode;
 class UInterchangeMaterialInstanceNode;
 
+DECLARE_LOG_CATEGORY_EXTERN(LogInterchangeMaterialPipeline, Log, All);
+
 UCLASS(BlueprintType, Experimental)
 class DATASMITHINTERCHANGE_API UInterchangeDatasmithMaterialPipeline : public UInterchangeGenericMaterialPipeline
 {
diff --git a/Engine/Plugins/Interchange/Runtime/Source/Import/Private/Material/InterchangeMaterialFactory.cpp b/Engine/Plugins/Interchange/Runtime/Source/Import/Private/Material/InterchangeMaterialFactory.cpp
--- a/Engine/Plugins/Interchange/Runtime/Source/Import/Private/Material/InterchangeMaterialFactory.cpp
+++ b/Engine/Plugins/Interchange/Runtime/Source/Import/Private/Material/InterchangeMaterialFactory.cpp
@@ -27,6 +27,7 @@
 #include "Materials/MaterialInstanceDynamic.h"
 #include "Nodes/InterchangeBaseNode.h"
 #include "Nodes/InterchangeBaseNodeContainer.h"
+#include "Nodes/InterchangeUserDefinedAttribute.h"
 #include "UObject/ObjectRedirector.h"
 
 #include UE_INLINE_GENERATED_CPP_BY_NAME(InterchangeMaterialFactory)
@@ -42,424 +43,615 @@
 
 #define LOCTEXT_NAMESPACE "InterchangeMaterialFactory"
 
-namespace UE
+namespace UE::Interchange::MaterialFactory::Internal
 {
-	namespace Interchange
+
+#if WITH_EDITOR
+	/**
+	 * Finds a UMaterialExpression class by name.
+	 * @param ClassName		The name of the class to look for (ie:UClass*->GetName()).
+	 * @return				A sub class of UMaterialExpression or nullptr.
+	 */
+	TSubclassOf<UMaterialExpression> FindExpressionClass(const TCHAR* ClassName)
 	{
-		namespace MaterialFactory
+		check(ClassName);
+
+		UClass* MaterialExpressionClass = FindFirstObject<UClass>(ClassName, EFindFirstObjectOptions::EnsureIfAmbiguous);
+
+		if (!MaterialExpressionClass)
 		{
-			namespace Internal
+			if (UObjectRedirector* RenamedClassRedirector = FindFirstObject<UObjectRedirector>(ClassName, EFindFirstObjectOptions::EnsureIfAmbiguous))
 			{
-#if WITH_EDITOR
-				/**
-				 * Finds a UMaterialExpression class by name.
-				 * @param ClassName		The name of the class to look for (ie:UClass*->GetName()).
-				 * @return				A sub class of UMaterialExpression or nullptr.
-				 */
-				TSubclassOf<UMaterialExpression> FindExpressionClass(const TCHAR* ClassName)
-				{
-					check(ClassName);
+				MaterialExpressionClass = CastChecked<UClass>(RenamedClassRedirector->DestinationObject);
+			}
+		}
 
-					UClass* MaterialExpressionClass = FindFirstObject<UClass>(ClassName, EFindFirstObjectOptions::EnsureIfAmbiguous);
+		if (MaterialExpressionClass && MaterialExpressionClass->IsChildOf<UMaterialExpression>())
+		{
+			return MaterialExpressionClass;
+		}
+		else
+		{
+			return nullptr;
+		}
+	}
 
-					if (!MaterialExpressionClass)
-					{
-						if (UObjectRedirector* RenamedClassRedirector = FindFirstObject<UObjectRedirector>(ClassName, EFindFirstObjectOptions::EnsureIfAmbiguous))
-						{
-							MaterialExpressionClass = CastChecked<UClass>(RenamedClassRedirector->DestinationObject);
-						}
-					}
+	int32 GetInputIndex(UMaterialExpression& MaterialExpression, const FString& InputName)
+	{
+		int32 ExpressionInputIndex = 0;
 
-					if (MaterialExpressionClass && MaterialExpressionClass->IsChildOf<UMaterialExpression>())
-					{
-						return MaterialExpressionClass;
-					}
-					else
-					{
-						return nullptr;
-					}
-				}
+		for (const FExpressionInput* ExpressionInput : MaterialExpression.GetInputsView())
+		{
+			// MaterialFuncCall appends the type to the input name when calling GetInputName
+			// and the InputName in FExpressionInput is optional so we'll check both here to be safe
+			if (MaterialExpression.GetInputName(ExpressionInputIndex) == *InputName ||
+				(ExpressionInput && ExpressionInput->InputName == *InputName))
+			{
+				return ExpressionInputIndex;
+			}
 
-				int32 GetInputIndex(UMaterialExpression& MaterialExpression, const FString& InputName)
-				{
-					int32 ExpressionInputIndex = 0;
+			++ExpressionInputIndex;
+		}
 
-					for (const FExpressionInput* ExpressionInput : MaterialExpression.GetInputsView())
-					{
-						// MaterialFuncCall appends the type to the input name when calling GetInputName
-						// and the InputName in FExpressionInput is optional so we'll check both here to be safe
-						if (MaterialExpression.GetInputName(ExpressionInputIndex) == *InputName ||
-							(ExpressionInput && ExpressionInput->InputName == *InputName))
-						{
-							return ExpressionInputIndex;
-						}
+		return INDEX_NONE;
+	}
 
-						++ExpressionInputIndex;
-					}
+	int32 GetOutputIndex(UMaterialExpression& MaterialExpression, const FString& OutputName)
+	{
+		// Check whether OutputName stores an index
+		int32 ExpressionOutputIndex = UInterchangeShaderPortsAPI::GetOutputIndexFromName(OutputName);
+		if (ExpressionOutputIndex != INDEX_NONE)
+		{
+			if (MaterialExpression.GetOutputs().IsValidIndex(ExpressionOutputIndex))
+			{
+				return ExpressionOutputIndex;
+			}
+		}
 
-					return INDEX_NONE;
-				}
+		ExpressionOutputIndex = 0;
 
-				int32 GetOutputIndex(UMaterialExpression& MaterialExpression, const FString& OutputName)
-				{
-					// Check whether OutputName stores an index
-					int32 ExpressionOutputIndex = UInterchangeShaderPortsAPI::GetOutputIndexFromName(OutputName);
-					if (ExpressionOutputIndex != INDEX_NONE)
-					{
-						if (MaterialExpression.GetOutputs().IsValidIndex(ExpressionOutputIndex))
-						{
-							return ExpressionOutputIndex;
-						}
-					}
+		for (const FExpressionOutput& ExpressionOutput : MaterialExpression.GetOutputs())
+		{
+			if (ExpressionOutput.OutputName == *OutputName)
+			{
+				return ExpressionOutputIndex;
+			}
 
-					ExpressionOutputIndex = 0;
+			++ExpressionOutputIndex;
+		}
 
-					for (const FExpressionOutput& ExpressionOutput : MaterialExpression.GetOutputs())
-					{
-						if (ExpressionOutput.OutputName == *OutputName)
-						{
-							return ExpressionOutputIndex;
-						}
+		return 0; // Consider 0 as the default output to connect to since most expressions have a single output
+	}
 
-						++ExpressionOutputIndex;
-					}
+	void SetupFunctionCallExpression(const UInterchangeFactoryBase::FImportAssetObjectParams& Arguments, const UInterchangeMaterialExpressionFactoryNode& ExpressionNode, UMaterialExpressionMaterialFunctionCall* FunctionCallExpression)
+	{
+		FInterchangeImportMaterialAsyncHelper& AsyncHelper = FInterchangeImportMaterialAsyncHelper::GetInstance();
+		if (const UInterchangeMaterialFunctionCallExpressionFactoryNode* FunctionCallFactoryNode = Cast<UInterchangeMaterialFunctionCallExpressionFactoryNode>(&ExpressionNode))
+		{
+			FString MaterialFactoryNodeUid;
+			FunctionCallFactoryNode->GetCustomMaterialFunctionDependency(MaterialFactoryNodeUid);
 
-					return 0; // Consider 0 as the default output to connect to since most expressions have a single output
+			if (const UInterchangeMaterialFunctionFactoryNode* MaterialFactoryNode = Cast<UInterchangeMaterialFunctionFactoryNode>(Arguments.NodeContainer->GetNode(MaterialFactoryNodeUid)))
+			{
+				FSoftObjectPath ReferenceObject;
+				MaterialFactoryNode->GetCustomReferenceObject(ReferenceObject);
+				if (UMaterialFunctionInterface* MaterialFunction = Cast<UMaterialFunctionInterface>(ReferenceObject.TryLoad()))
+				{
+					AsyncHelper.UpdateFromFunctionResource(MaterialFunction);
+					FunctionCallExpression->SetMaterialFunction(MaterialFunction);
 				}
+			}
+		}
 
-				void SetupFunctionCallExpression(const UInterchangeFactoryBase::FImportAssetObjectParams& Arguments, const UInterchangeMaterialExpressionFactoryNode& ExpressionNode, UMaterialExpressionMaterialFunctionCall* FunctionCallExpression)
-				{
-					FInterchangeImportMaterialAsyncHelper& AsyncHelper = FInterchangeImportMaterialAsyncHelper::GetInstance();
-					if (const UInterchangeMaterialFunctionCallExpressionFactoryNode* FunctionCallFactoryNode = Cast<UInterchangeMaterialFunctionCallExpressionFactoryNode>(&ExpressionNode))
-					{
-						FString MaterialFactoryNodeUid;
-						FunctionCallFactoryNode->GetCustomMaterialFunctionDependency(MaterialFactoryNodeUid);
+		AsyncHelper.UpdateFromFunctionResource(FunctionCallExpression);
+	}
 
-						if (const UInterchangeMaterialFunctionFactoryNode* MaterialFactoryNode = Cast<UInterchangeMaterialFunctionFactoryNode>(Arguments.NodeContainer->GetNode(MaterialFactoryNodeUid)))
-						{
-							FSoftObjectPath ReferenceObject;
-							MaterialFactoryNode->GetCustomReferenceObject(ReferenceObject);
-							if (UMaterialFunctionInterface* MaterialFunction = Cast<UMaterialFunctionInterface>(ReferenceObject.TryLoad()))
-							{
-								AsyncHelper.UpdateFromFunctionResource(MaterialFunction);
-								FunctionCallExpression->SetMaterialFunction(MaterialFunction);
-							}
-						}
-					}
+	void SetupTextureExpression(const UInterchangeFactoryBase::FImportAssetObjectParams& Arguments, const UInterchangeMaterialExpressionFactoryNode* ExpressionNode, UMaterialExpressionTextureBase* TextureExpression)
+	{
+		using namespace UE::Interchange::Materials::Standard::Nodes::TextureSample;
 
-					AsyncHelper.UpdateFromFunctionResource(FunctionCallExpression);
-				}
+		const bool bIsAParameter = UInterchangeShaderPortsAPI::HasParameter(ExpressionNode, Inputs::Texture);
+		const FString InputKey = bIsAParameter
+			? UInterchangeShaderPortsAPI::MakeInputParameterKey(Inputs::Texture.ToString())
+			: UInterchangeShaderPortsAPI::MakeInputValueKey(Inputs::Texture.ToString());
 
-				void SetupTextureExpression(const UInterchangeFactoryBase::FImportAssetObjectParams& Arguments, const UInterchangeMaterialExpressionFactoryNode* ExpressionNode, UMaterialExpressionTextureBase* TextureExpression)
-				{
-					using namespace UE::Interchange::Materials::Standard::Nodes::TextureSample;
+		FString TextureFactoryNodeUid;
+		ExpressionNode->GetStringAttribute(InputKey, TextureFactoryNodeUid);
 
-					FString TextureFactoryNodeUid;
-					ExpressionNode->GetStringAttribute(UInterchangeShaderPortsAPI::MakeInputValueKey(Inputs::Texture.ToString()), TextureFactoryNodeUid);
 
-					if (const UInterchangeTextureFactoryNode* TextureFactoryNode = Cast<UInterchangeTextureFactoryNode>(Arguments.NodeContainer->GetNode(TextureFactoryNodeUid)))
-					{
-						FSoftObjectPath ReferenceObject;
-						TextureFactoryNode->GetCustomReferenceObject(ReferenceObject);
-						if (UTexture* Texture = Cast<UTexture>(ReferenceObject.TryLoad()))
-						{
-							TextureExpression->Texture = Texture;
-						}
-					}
+		// It is possible that the Texture Object Parameter was used that uses KeyValueProperty to reference the texture instead of extra texture input connection.
+		if (TextureFactoryNodeUid.IsEmpty())
+		{
+			FString PayloadKey;
+			UInterchangeUserDefinedAttributesAPI::GetUserDefinedAttribute_FString(ExpressionNode, Inputs::Texture.ToString(), TextureFactoryNodeUid, PayloadKey);
+		}
 
-					TextureExpression->AutoSetSampleType();
-				}
+		if (const UInterchangeTextureFactoryNode* TextureFactoryNode = Cast<UInterchangeTextureFactoryNode>(Arguments.NodeContainer->GetNode(TextureFactoryNodeUid)))
+		{
+			FSoftObjectPath ReferenceObject;
+			TextureFactoryNode->GetCustomReferenceObject(ReferenceObject);
+			if (UTexture* Texture = Cast<UTexture>(ReferenceObject.TryLoad()))
+			{
+				TextureExpression->Texture = Texture;
+			}
+		}
 
-				UMaterialExpression* CreateMaterialExpression(UMaterial* Material, UMaterialFunction* MaterialFunction, const TSubclassOf<UMaterialExpression>& ExpressionClass)
-				{
-					UObject* const SelectedAsset = nullptr;
-					const int32 NodePosX = 0;
-					const int32 NodePosY = 0;
-					const bool bAllowMarkingPackageDirty = false;
+		TextureExpression->AutoSetSampleType();
+	}
 
-					return UMaterialEditingLibrary::CreateMaterialExpressionEx(Material, MaterialFunction, ExpressionClass, SelectedAsset, NodePosX, NodePosY, bAllowMarkingPackageDirty);
-				}
+	UMaterialExpression* CreateMaterialExpression(UMaterial* Material, UMaterialFunction* MaterialFunction, const TSubclassOf<UMaterialExpression>& ExpressionClass)
+	{
+		UObject* const SelectedAsset = nullptr;
+		const int32 NodePosX = 0;
+		const int32 NodePosY = 0;
+		const bool bAllowMarkingPackageDirty = false;
 
-				template<class T>
-				T* CreateMaterialExpression(UMaterial* Material, UMaterialFunction* MaterialFunction)
-				{
-					return Cast<T>(CreateMaterialExpression(Material, MaterialFunction, T::StaticClass()));
-				}
+		return UMaterialEditingLibrary::CreateMaterialExpressionEx(Material, MaterialFunction, ExpressionClass, SelectedAsset, NodePosX, NodePosY, bAllowMarkingPackageDirty);
+	}
 
+	template<class T>
+	T* CreateMaterialExpression(UMaterial* Material, UMaterialFunction* MaterialFunction)
+	{
+		return Cast<T>(CreateMaterialExpression(Material, MaterialFunction, T::StaticClass()));
+	}
 
-				class FMaterialExpressionBuilder
-				{
-				public:
-					FMaterialExpressionBuilder(UMaterial* InMaterial, UMaterialFunction* InMaterialFunction, const UInterchangeFactoryBase::FImportAssetObjectParams& InArguments)
-						: Material(InMaterial)
-						, MaterialFunction(InMaterialFunction)
-						, Arguments(InArguments)
-					{
-						check((Material || MaterialFunction) && !(Material && MaterialFunction));
 
-						//We need to put in place a better mechanism for the reimport, for the moment let's delete all material expressions to avoid duplicates	
-						DeleteAllMaterialExpressions();
-					}
+	class FMaterialExpressionBuilder
+	{
+	public:
+		FMaterialExpressionBuilder(UMaterial* InMaterial, UMaterialFunction* InMaterialFunction, const UInterchangeFactoryBase::FImportAssetObjectParams& InArguments, TObjectPtr<UInterchangeResultsContainer> InResultsContainer)
+			: Material(InMaterial)
+			, MaterialFunction(InMaterialFunction)
+			, Arguments(InArguments)
+			, ResultsContainer(InResultsContainer)
+		{
+			check((Material || MaterialFunction) && !(Material && MaterialFunction));
 
-					UMaterialExpression* CreateExpressionsForNode(const UInterchangeMaterialExpressionFactoryNode& ExpressionNode)
-					{
-						UMaterialExpression* MaterialExpression = Expressions.FindRef(ExpressionNode.GetUniqueID());
+			//We need to put in place a better mechanism for the reimport, for the moment let's delete all material expressions to avoid duplicates	
+			DeleteAllMaterialExpressions();
+		}
+		UMaterialExpression* CreateExpressionsForNode(const UInterchangeMaterialExpressionFactoryNode& ExpressionNode)
+		{
+			UMaterialExpression* MaterialExpression = Expressions.FindRef(ExpressionNode.GetUniqueID());
 
-						if (MaterialExpression)
-						{
-							return MaterialExpression;
-						}
+			if (MaterialExpression)
+			{
+				return MaterialExpression;
+			}
 
-						MaterialExpression = CreateExpression(ExpressionNode);
-						if (!MaterialExpression)
-						{
-							return nullptr;
-						}
+			MaterialExpression = CreateExpression(ExpressionNode);
+			if (!MaterialExpression)
+			{
+				return nullptr;
+			}
 
-						Expressions.Add(ExpressionNode.GetUniqueID()) = MaterialExpression;
+			Expressions.Add(ExpressionNode.GetUniqueID()) = MaterialExpression;
 
-						TArray<FString> Inputs;
-						UInterchangeShaderPortsAPI::GatherInputs(&ExpressionNode, Inputs);
+			TArray<FString> Inputs;
+			UInterchangeShaderPortsAPI::GatherInputs(&ExpressionNode, Inputs);
 
-						for (const FString& InputName : Inputs)
+			for (const FString& InputName : Inputs)
+			{
+				FString ConnectedExpressionUid;
+				FString OutputName;
+				if (UInterchangeShaderPortsAPI::GetInputConnection(&ExpressionNode, InputName, ConnectedExpressionUid, OutputName))
+				{
+					const UInterchangeMaterialExpressionFactoryNode* ConnectedExpressionNode = Cast<UInterchangeMaterialExpressionFactoryNode>(Arguments.NodeContainer->GetNode(ConnectedExpressionUid));
+					if (ConnectedExpressionNode)
+					{
+						UMaterialExpression* ConnectedExpression = Expressions.FindRef(ConnectedExpressionUid);
+						if (!ConnectedExpression)
 						{
-							FString ConnectedExpressionUid;
-							FString OutputName;
-							if (UInterchangeShaderPortsAPI::GetInputConnection(&ExpressionNode, InputName, ConnectedExpressionUid, OutputName))
+							ConnectedExpression = CreateExpressionsForNode(*ConnectedExpressionNode);
+							// This case is particularly observed when there are broken function references and there are Empty Material Function Names.
+							if (!ConnectedExpression)
 							{
-								const UInterchangeMaterialExpressionFactoryNode* ConnectedExpressionNode = Cast<UInterchangeMaterialExpressionFactoryNode>(Arguments.NodeContainer->GetNode(ConnectedExpressionUid));
-								if (ConnectedExpressionNode)
-								{
-									UMaterialExpression* ConnectedExpression = Expressions.FindRef(ConnectedExpressionUid);
-									if (!ConnectedExpression)
-									{
-										ConnectedExpression = CreateExpressionsForNode(*ConnectedExpressionNode);
-									}
-
-									const int32 InputIndex = GetInputIndex(*MaterialExpression, InputName);
-									const int32 OutputIndex = GetOutputIndex(*ConnectedExpression, OutputName);
-
-									if (InputIndex != INDEX_NONE)
-									{
-										FExpressionInput* ExpressionInput = MaterialExpression->GetInput(InputIndex);
-										if (ExpressionInput)
-										{
-											ConnectedExpression->ConnectExpression(ExpressionInput, OutputIndex);
-										}
-									}
-									else
-									{
-										Messages.Add(FText::Format(LOCTEXT("InputNotFound", "Invalid input {0} for material expression node {1}."),
-											FText::FromString(InputName),
-											FText::FromString(ExpressionNode.GetDisplayLabel())));
-									}
-								}
+								continue;
 							}
 						}
 
-						return MaterialExpression;
+						const int32 InputIndex = GetInputIndex(*MaterialExpression, InputName);
+						const int32 OutputIndex = GetOutputIndex(*ConnectedExpression, OutputName);
+
+						if (InputIndex != INDEX_NONE)
+						{
+							FExpressionInput* ExpressionInput = MaterialExpression->GetInput(InputIndex);
+							if (ExpressionInput)
+							{
+								ConnectedExpression->ConnectExpression(ExpressionInput, OutputIndex);
+							}
+						}
+						else
+						{
+							PostMessage<UInterchangeResultError_Generic>(FText::Format(
+								LOCTEXT("InputNotFound", "Invalid input {0} for material expression node {1}."),
+								FText::FromString(InputName),
+								FText::FromString(ExpressionNode.GetDisplayLabel())
+							));
+						}
 					}
+				}
+			}
 
-					TArray<FText> Messages;
+			return MaterialExpression;
+		}
 
-				private:
-					UMaterialExpression* CreateExpression(const UInterchangeMaterialExpressionFactoryNode& ExpressionNode)
-					{
-						FString ExpressionClassName;
-						ExpressionNode.GetCustomExpressionClassName(ExpressionClassName);
+	private:
+		UMaterialExpression* CreateExpression(const UInterchangeMaterialExpressionFactoryNode& ExpressionNode)
+		{
+			FString ExpressionClassName;
+			ExpressionNode.GetCustomExpressionClassName(ExpressionClassName);
 
-						TSubclassOf<UMaterialExpression> ExpressionClass = FindExpressionClass(*ExpressionClassName);
+			TSubclassOf<UMaterialExpression> ExpressionClass = FindExpressionClass(*ExpressionClassName);
 
-						if (!ExpressionClass.Get())
-						{
-							Messages.Add(FText::Format(LOCTEXT("ExpressionClassNotFound", "Invalid class {0} for material expression node {1}."),
-								FText::FromString(ExpressionClassName),
-								FText::FromString(ExpressionNode.GetDisplayLabel())));
+			if (!ExpressionClass.Get())
+			{
+				PostMessage<UInterchangeResultError_Generic>(FText::Format(
+					LOCTEXT("ExpressionClassNotFound", "Invalid class {0} for material expression node {1}."),
+					FText::FromString(ExpressionClassName),
+					FText::FromString(ExpressionNode.GetDisplayLabel())
+				));
+				return nullptr;
+			}
 
-							return nullptr;
-						}
+			UMaterialExpression* MaterialExpression = CreateMaterialExpression(Material, MaterialFunction, ExpressionClass);
 
-						UMaterialExpression* MaterialExpression = CreateMaterialExpression(Material, MaterialFunction, ExpressionClass);
+			if (!MaterialExpression)
+			{
+				PostMessage<UInterchangeResultError_Generic>(FText::Format(
+					LOCTEXT("MaterialExpressionCreationFailed", "Failed to create {0} object for material expression node {1}."),
+					FText::FromString(ExpressionClassName),
+					FText::FromString(ExpressionNode.GetDisplayLabel())
+				));
+				return nullptr;
+			}
 
-						if (!MaterialExpression)
-						{
-							Messages.Add(FText::Format(LOCTEXT("MaterialExpressionCreationFailed", "Failed to create {0} object for material expression node {1}."),
-								FText::FromString(ExpressionClassName),
-								FText::FromString(ExpressionNode.GetDisplayLabel())));
+			if (Material)
+			{
+				if (FNameProperty* Property = FindFProperty<FNameProperty>(MaterialExpression->GetClass(), GET_MEMBER_NAME_CHECKED(UMaterialExpressionParameter, ParameterName)))
+				{
+					const FString MaterialExpressionName = ExpressionNode.GetDisplayLabel();
+					FName ParameterName;
+					if (MaterialExpressionName == TEXT("Null"))
+					{
+						const int32 RightChopIndex = FString(TEXT("MaterialExpression")).Len();
+						FString ExpressionName = ExpressionClassName.RightChop(RightChopIndex);
+						ParameterName = FName(*(ExpressionName + TEXT("_(") + LexToString(Material->GetExpressions().Num()) + TEXT(")")));
+					}
+					else
+					{
+						ParameterName = FName(*MaterialExpressionName);
+					}
 
-							return nullptr;
-						}
+					// Set the parameter name if the material expression has one (some material expressions don't inherit from UMaterialExpressionParameter, ie: UMaterialExpressionTextureSampleParameter
+					*(Property->ContainerPtrToValuePtr<FName>(MaterialExpression)) = ParameterName;
+				}
+			}
 
-						if (Material)
-						{
-							// Set the parameter name if the material expression has one (some material expressions don't inherit from UMaterialExpressionParameter, ie: UMaterialExpressionTextureSampleParameter
-							if (FNameProperty* Property = FindFProperty<FNameProperty>(MaterialExpression->GetClass(), GET_MEMBER_NAME_CHECKED(UMaterialExpressionParameter, ParameterName)))
-							{
-								*(Property->ContainerPtrToValuePtr<FName>(MaterialExpression)) = FName(*(ExpressionNode.GetDisplayLabel() + LexToString(Material->GetExpressions().Num())));
-							}
-						}
+			ExpressionNode.ApplyAllCustomAttributeToObject(MaterialExpression);
 
-						ExpressionNode.ApplyAllCustomAttributeToObject(MaterialExpression);
+			if (UMaterialExpressionTextureBase* TextureExpression = Cast<UMaterialExpressionTextureBase>(MaterialExpression))
+			{
+				SetupTextureExpression(Arguments, &ExpressionNode, TextureExpression);
+			}
+			else if (UMaterialExpressionMaterialFunctionCall* FunctionCallExpression = Cast<UMaterialExpressionMaterialFunctionCall>(MaterialExpression))
+			{
+				SetupFunctionCallExpression(Arguments, ExpressionNode, FunctionCallExpression);
+			}
 
-						if (UMaterialExpressionTextureBase* TextureExpression = Cast<UMaterialExpressionTextureBase>(MaterialExpression))
-						{
-							SetupTextureExpression(Arguments, &ExpressionNode, TextureExpression);
-						}
-						else if (UMaterialExpressionMaterialFunctionCall* FunctionCallExpression = Cast<UMaterialExpressionMaterialFunctionCall>(MaterialExpression))
-						{
-							SetupFunctionCallExpression(Arguments, ExpressionNode, FunctionCallExpression);
-						}
+			return MaterialExpression;
+		}
 
-						return MaterialExpression;
+		void DeleteAllMaterialExpressions()
+		{
+			if (Material)
+			{
+				// Copy the Material Expressions in a TArray, otherwise working directly on the TArrayView messes up with the expressions
+				// and doesn't make a full clean up, especially when we're not in the Game Thread
+				if (TArray<UMaterialExpression*> MaterialExpressions(Material->GetExpressions()); !MaterialExpressions.IsEmpty())
+				{
+					Material->Modify();
+					for (UMaterialExpression* MaterialExpression : MaterialExpressions)
+					{
+						MaterialExpression->Modify();
+						Material->GetExpressionCollection().RemoveExpression(MaterialExpression);
+						Material->RemoveExpressionParameter(MaterialExpression);
+						// Make sure the deleted expression is caught by gc
+						MaterialExpression->MarkAsGarbage();
 					}
 
-					void DeleteAllMaterialExpressions()
-					{
-						if(Material)
-						{
-							// Copy the Material Expressions in a TArray, otherwise working directly on the TArrayView messes up with the expressions
-							// and doesn't make a full clean up, especially when we're not in the Game Thread
-							if(TArray<UMaterialExpression*> MaterialExpressions(Material->GetExpressions()); !MaterialExpressions.IsEmpty())
-							{
-								Material->Modify();
-								for(UMaterialExpression* MaterialExpression : MaterialExpressions)
-								{
-									MaterialExpression->Modify();
-									Material->GetExpressionCollection().RemoveExpression(MaterialExpression);
-									Material->RemoveExpressionParameter(MaterialExpression);
-									// Make sure the deleted expression is caught by gc
-									MaterialExpression->MarkAsGarbage();
-								}
+					Material->MarkPackageDirty();
+				}
+			}
+		}
 
-								Material->MarkPackageDirty();
-							}
-						}
-					}
+		template<class TInterchangeResultType>
+		void PostMessage(FText&& MessageText)
+		{
+			if (ResultsContainer)
+			{
+				TInterchangeResultType* Result = ResultsContainer->Add<TInterchangeResultType>();
+				Result->Text = MoveTemp(MessageText);
+			}
+		}
 
-					UMaterial* Material = nullptr;
-					UMaterialFunction* MaterialFunction = nullptr;
-					const UInterchangeFactoryBase::FImportAssetObjectParams& Arguments;
-					TMap<FString, UMaterialExpression*> Expressions;
-				};
+	private:
+		UMaterial* Material = nullptr;
+		UMaterialFunction* MaterialFunction = nullptr;
+		const UInterchangeFactoryBase::FImportAssetObjectParams& Arguments;
+		TMap<FString, UMaterialExpression*> Expressions;
+		TObjectPtr<UInterchangeResultsContainer> ResultsContainer;
+	};
 #endif // #if WITH_EDITOR
 
-				void UpdateParameterBool(UMaterialInstance& MaterialInstance, const FString& InputName, const UInterchangeMaterialInstanceFactoryNode& FactoryNode)
-				{
+	void UpdateParameterBool(UMaterialInstance& MaterialInstance, const FString& InputName, const UInterchangeMaterialInstanceFactoryNode& FactoryNode)
+	{
 #if WITH_EDITORONLY_DATA
-					const FName ParameterName = *InputName;
-					bool bInstanceValue;
-					FGuid Uid;
-					if (MaterialInstance.GetStaticSwitchParameterValue(ParameterName, bInstanceValue, Uid))
+		const FName ParameterName = *InputName;
+		bool bInstanceValue;
+		FGuid Uid;
+		if (MaterialInstance.GetStaticSwitchParameterValue(ParameterName, bInstanceValue, Uid))
+		{
+			bool bInputValue = false;
+			FactoryNode.GetBooleanAttribute(UInterchangeShaderPortsAPI::MakeInputValueKey(InputName), bInputValue);
+
+			if (bInputValue != bInstanceValue)
+			{
+				if (UMaterialInstanceConstant* MaterialInstanceConstant = Cast<UMaterialInstanceConstant>(&MaterialInstance))
+				{
+					MaterialInstanceConstant->SetStaticSwitchParameterValueEditorOnly(ParameterName, bInputValue);
+				}
+			}
+		}
+#endif // #if WITH_EDITORONLY_DATA
+	}
+
+	void UpdateParameterFloat(UMaterialInstance& MaterialInstance, const FString& InputName, const UInterchangeMaterialInstanceFactoryNode& FactoryNode)
+	{
+		const FName ParameterName = *InputName;
+		float InstanceValue;
+
+		if (MaterialInstance.GetScalarParameterValue(ParameterName, InstanceValue))
+		{
+			float InputValue = 0.f;
+			FactoryNode.GetFloatAttribute(UInterchangeShaderPortsAPI::MakeInputValueKey(InputName), InputValue);
+
+			if (!FMath::IsNearlyEqual(InputValue, InstanceValue))
+			{
+#if WITH_EDITOR
+				if (UMaterialInstanceConstant* MaterialInstanceConstant = Cast<UMaterialInstanceConstant>(&MaterialInstance))
+				{
+					MaterialInstanceConstant->SetScalarParameterValueEditorOnly(ParameterName, InputValue);
+				}
+				else
+#endif // #if WITH_EDITOR
+					if (UMaterialInstanceDynamic* MaterialInstanceDynamic = Cast<UMaterialInstanceDynamic>(&MaterialInstance))
 					{
-						bool bInputValue = false;
-						FactoryNode.GetBooleanAttribute(UInterchangeShaderPortsAPI::MakeInputValueKey(InputName), bInputValue);
+						MaterialInstanceDynamic->SetScalarParameterValue(ParameterName, InputValue);
+					}
+			}
+		}
+	}
+
+	void UpdateParameterLinearColor(UMaterialInstance& MaterialInstance, const FString& InputName, const UInterchangeMaterialInstanceFactoryNode& FactoryNode)
+	{
+		const FName ParameterName = *InputName;
+		FLinearColor InstanceValue;
 
-						if (bInputValue != bInstanceValue)
+		if (MaterialInstance.GetVectorParameterValue(ParameterName, InstanceValue))
+		{
+			FLinearColor InputValue;
+			if (FactoryNode.GetLinearColorAttribute(UInterchangeShaderPortsAPI::MakeInputValueKey(InputName), InputValue))
+			{
+				if (!InputValue.Equals(InstanceValue))
+				{
+#if WITH_EDITOR
+					if (UMaterialInstanceConstant* MaterialInstanceConstant = Cast<UMaterialInstanceConstant>(&MaterialInstance))
+					{
+						MaterialInstanceConstant->SetVectorParameterValueEditorOnly(ParameterName, InputValue);
+					}
+					else
+#endif // #if WITH_EDITOR
+						if (UMaterialInstanceDynamic* MaterialInstanceDynamic = Cast<UMaterialInstanceDynamic>(&MaterialInstance))
 						{
-							if (UMaterialInstanceConstant* MaterialInstanceConstant = Cast<UMaterialInstanceConstant>(&MaterialInstance))
-							{
-								MaterialInstanceConstant->SetStaticSwitchParameterValueEditorOnly(ParameterName, bInputValue);
-							}
+							MaterialInstanceDynamic->SetVectorParameterValue(ParameterName, InputValue);
 						}
-					}
-#endif // #if WITH_EDITORONLY_DATA
 				}
+			}
+		}
+	}
 
-				void UpdateParameterFloat(UMaterialInstance& MaterialInstance, const FString& InputName, const UInterchangeMaterialInstanceFactoryNode& FactoryNode)
-				{
-					const FName ParameterName = *InputName;
-					float InstanceValue;
+	void UpdateParameterTexture(UMaterialInstance& MaterialInstance, const FString& InputName, const UInterchangeMaterialInstanceFactoryNode& FactoryNode, const UInterchangeBaseNodeContainer& NodeContainer)
+	{
+		const FName ParameterName = *InputName;
+		UTexture* InstanceValue;
 
-					if (MaterialInstance.GetScalarParameterValue(ParameterName, InstanceValue))
+		if (MaterialInstance.GetTextureParameterValue(ParameterName, InstanceValue))
+		{
+			FString InputValue;
+			if (FactoryNode.GetStringAttribute(UInterchangeShaderPortsAPI::MakeInputValueKey(InputName), InputValue))
+			{
+				if (const UInterchangeTextureFactoryNode* TextureFactoryNode = Cast<UInterchangeTextureFactoryNode>(NodeContainer.GetNode(InputValue)))
+				{
+					FSoftObjectPath ReferenceObject;
+					TextureFactoryNode->GetCustomReferenceObject(ReferenceObject);
+					if (UTexture* InputTexture = Cast<UTexture>(ReferenceObject.TryLoad()))
 					{
-						float InputValue = 0.f;
-						FactoryNode.GetFloatAttribute(UInterchangeShaderPortsAPI::MakeInputValueKey(InputName), InputValue);
-
-						if (!FMath::IsNearlyEqual(InputValue, InstanceValue))
+						if (InputTexture != InstanceValue)
 						{
 #if WITH_EDITOR
 							if (UMaterialInstanceConstant* MaterialInstanceConstant = Cast<UMaterialInstanceConstant>(&MaterialInstance))
 							{
-								MaterialInstanceConstant->SetScalarParameterValueEditorOnly(ParameterName, InputValue);
+								MaterialInstanceConstant->SetTextureParameterValueEditorOnly(ParameterName, InputTexture);
 							}
 							else
 #endif // #if WITH_EDITOR
-							if (UMaterialInstanceDynamic* MaterialInstanceDynamic = Cast<UMaterialInstanceDynamic>(&MaterialInstance))
-							{
-								MaterialInstanceDynamic->SetScalarParameterValue(ParameterName, InputValue);
-							}
+								if (UMaterialInstanceDynamic* MaterialInstanceDynamic = Cast<UMaterialInstanceDynamic>(&MaterialInstance))
+								{
+									MaterialInstanceDynamic->SetTextureParameterValue(ParameterName, InputTexture);
+								}
 						}
 					}
 				}
+			}
+		}
+	}
 
-				void UpdateParameterLinearColor(UMaterialInstance& MaterialInstance, const FString& InputName, const UInterchangeMaterialInstanceFactoryNode& FactoryNode)
+	void OverrideBoolParameter(const UInterchangeMaterialInstanceFactoryNode& FactoryNode, const FString& AttributeKey, UMaterialInstance& MaterialInstance, const FName& ParameterName)
+	{
+		bool AttributeValue;
+		if (FactoryNode.GetBooleanAttribute(AttributeKey, AttributeValue))
+		{
+#if WITH_EDITOR
+			if (UMaterialInstanceConstant* MaterialInstanceConstant = Cast<UMaterialInstanceConstant>(&MaterialInstance))
+			{
+				MaterialInstanceConstant->SetStaticSwitchParameterValueEditorOnly(ParameterName, AttributeValue);
+			}
+			else
+#endif // #if WITH_EDITOR
+				if (UMaterialInstanceDynamic* MaterialInstanceDynamic = Cast<UMaterialInstanceDynamic>(&MaterialInstance))
 				{
-					const FName ParameterName = *InputName;
-					FLinearColor InstanceValue;
+					//TODO: Log Error
+					ensure(false);
+				}
+		}
+	}
 
-					if (MaterialInstance.GetVectorParameterValue(ParameterName, InstanceValue))
-					{
-						FLinearColor InputValue;
-						if (FactoryNode.GetLinearColorAttribute(UInterchangeShaderPortsAPI::MakeInputValueKey(InputName), InputValue))
-						{
-							if (!InputValue.Equals(InstanceValue))
-							{
+	void OverrideScalarParameter(const UInterchangeMaterialInstanceFactoryNode& FactoryNode, const FString& AttributeKey, UMaterialInstance& MaterialInstance, const FName& ParameterName)
+	{
+		float AttributeValue;
+		if (FactoryNode.GetFloatAttribute(AttributeKey, AttributeValue))
+		{
 #if WITH_EDITOR
-								if (UMaterialInstanceConstant* MaterialInstanceConstant = Cast<UMaterialInstanceConstant>(&MaterialInstance))
-								{
-									MaterialInstanceConstant->SetVectorParameterValueEditorOnly(ParameterName, InputValue);
-								}
-								else
+			if (UMaterialInstanceConstant* MaterialInstanceConstant = Cast<UMaterialInstanceConstant>(&MaterialInstance))
+			{
+				MaterialInstanceConstant->SetScalarParameterValueEditorOnly(ParameterName, AttributeValue);
+			}
+			else
 #endif // #if WITH_EDITOR
-								if (UMaterialInstanceDynamic* MaterialInstanceDynamic = Cast<UMaterialInstanceDynamic>(&MaterialInstance))
-								{
-									MaterialInstanceDynamic->SetVectorParameterValue(ParameterName, InputValue);
-								}
-							}
-						}
-					}
+				if (UMaterialInstanceDynamic* MaterialInstanceDynamic = Cast<UMaterialInstanceDynamic>(&MaterialInstance))
+				{
+					MaterialInstanceDynamic->SetScalarParameterValue(ParameterName, AttributeValue);
 				}
+		}
+	}
 
-				void UpdateParameterTexture(UMaterialInstance& MaterialInstance, const FString& InputName, const UInterchangeMaterialInstanceFactoryNode& FactoryNode, const UInterchangeBaseNodeContainer& NodeContainer)
+	void OverrideVectorParameter(const UInterchangeMaterialInstanceFactoryNode& FactoryNode, const FString& AttributeKey, UMaterialInstance& MaterialInstance, const FName& ParameterName)
+	{
+		FLinearColor AttributeValue;
+		if (FactoryNode.GetLinearColorAttribute(AttributeKey, AttributeValue))
+		{
+#if WITH_EDITOR
+			if (UMaterialInstanceConstant* MaterialInstanceConstant = Cast<UMaterialInstanceConstant>(&MaterialInstance))
+			{
+				MaterialInstanceConstant->SetVectorParameterValueEditorOnly(ParameterName, AttributeValue);
+			}
+			else
+#endif // #if WITH_EDITOR
+				if (UMaterialInstanceDynamic* MaterialInstanceDynamic = Cast<UMaterialInstanceDynamic>(&MaterialInstance))
 				{
-					const FName ParameterName = *InputName;
-					UTexture* InstanceValue;
+					MaterialInstanceDynamic->SetVectorParameterValue(ParameterName, AttributeValue);
+				}
+		}
+	}
 
-					if (MaterialInstance.GetTextureParameterValue(ParameterName, InstanceValue))
-					{
-						FString InputValue;
-						if (FactoryNode.GetStringAttribute(UInterchangeShaderPortsAPI::MakeInputValueKey(InputName), InputValue))
-						{
-							if (const UInterchangeTextureFactoryNode* TextureFactoryNode = Cast<UInterchangeTextureFactoryNode>(NodeContainer.GetNode(InputValue)))
-							{
-								FSoftObjectPath ReferenceObject;
-								TextureFactoryNode->GetCustomReferenceObject(ReferenceObject);
-								if (UTexture* InputTexture = Cast<UTexture>(ReferenceObject.TryLoad()))
-								{
-									if (InputTexture != InstanceValue)
-									{
+	void OverrideTextureParameter(const UInterchangeMaterialInstanceFactoryNode& FactoryNode, const FString& AttributeKey, const UInterchangeBaseNodeContainer& NodeContainer, UMaterialInstance& MaterialInstance, const FName& ParameterName)
+	{
+		// The String Attribute is required to be a TextureFactoryUid as opposed to a TextureUid to correctly override.
+		FString TextureFactoryUid;
+		if (FactoryNode.GetStringAttribute(AttributeKey, TextureFactoryUid))
+		{
+			if (const UInterchangeTextureFactoryNode* TextureFactoryNode = Cast<UInterchangeTextureFactoryNode>(NodeContainer.GetNode(TextureFactoryUid)))
+			{
+				FSoftObjectPath ReferenceObject;
+				TextureFactoryNode->GetCustomReferenceObject(ReferenceObject);
+				if (UTexture* InputTexture = Cast<UTexture>(ReferenceObject.TryLoad()))
+				{
 #if WITH_EDITOR
-										if (UMaterialInstanceConstant* MaterialInstanceConstant = Cast<UMaterialInstanceConstant>(&MaterialInstance))
-										{
-											MaterialInstanceConstant->SetTextureParameterValueEditorOnly(ParameterName, InputTexture);
-										}
-										else
+					if (UMaterialInstanceConstant* MaterialInstanceConstant = Cast<UMaterialInstanceConstant>(&MaterialInstance))
+					{
+						MaterialInstanceConstant->SetTextureParameterValueEditorOnly(ParameterName, InputTexture);
+					}
+					else
 #endif // #if WITH_EDITOR
-										if (UMaterialInstanceDynamic* MaterialInstanceDynamic = Cast<UMaterialInstanceDynamic>(&MaterialInstance))
-										{
-											MaterialInstanceDynamic->SetTextureParameterValue(ParameterName, InputTexture);
-										}
-									}
-								}
-							}
+						if (UMaterialInstanceDynamic* MaterialInstanceDynamic = Cast<UMaterialInstanceDynamic>(&MaterialInstance))
+						{
+							MaterialInstanceDynamic->SetTextureParameterValue(ParameterName, InputTexture);
 						}
-					}
 				}
 			}
 		}
 	}
 }
 
+namespace UE::Interchange::Materials::HashUtils
+{
+
+	const TCHAR* FInterchangeMaterialInstanceOverridesAPI::ExpressionNameAttributeKey = TEXT("MaterialExpressionNameOverride");
+	const TCHAR* FInterchangeMaterialInstanceOverridesAPI::OverrideParameterPrefix = TEXT("LeafInput");
+	const TCHAR* FInterchangeMaterialInstanceOverridesAPI::OverrideParameterSeparator = TEXT(":");
+	const TCHAR* FInterchangeMaterialInstanceOverridesAPI::OverrideHashSeparator = TEXT("_");
+
+	FString FInterchangeMaterialInstanceOverridesAPI::MakeOverrideParameterName(UE::Interchange::EAttributeTypes AttributeType, int32 Hash, bool Prefix /*= true*/)
+	{
+		TStringBuilder<128> StringBuilder;
+		if (Prefix)
+		{
+			StringBuilder.Append(OverrideParameterPrefix);
+			StringBuilder.Append(OverrideParameterSeparator);
+		}
+		StringBuilder.Append(AttributeTypeToString(AttributeType));
+		StringBuilder.Append(OverrideHashSeparator);
+		StringBuilder.Append(FString::FromInt(Hash));
+		return StringBuilder.ToString();
+	}
+
+	FString FInterchangeMaterialInstanceOverridesAPI::MakeOverrideParameterName(const FString& DisplayLabel)
+	{
+		TStringBuilder<128> StringBuilder;
+		StringBuilder.Append(OverrideParameterPrefix);
+		StringBuilder.Append(OverrideParameterSeparator);
+		StringBuilder.Append(DisplayLabel);
+		return StringBuilder.ToString();
+	}
+
+	FString FInterchangeMaterialInstanceOverridesAPI::MakeExpressionNameString()
+	{
+		return ExpressionNameAttributeKey;
+	}
+
+	bool FInterchangeMaterialInstanceOverridesAPI::GetOverrideParameterName(const UE::Interchange::FAttributeKey& AttributeKey, FString& OverrideParameterName)
+	{
+		FString Prefix;
+		if (AttributeKey.Key.Split(OverrideParameterSeparator, &Prefix, &OverrideParameterName))
+		{
+			return true;
+		}
+		return false;
+	}
+
+	bool FInterchangeMaterialInstanceOverridesAPI::HasMaterialExpressionNameOverride(const UInterchangeBaseNode* BaseNode)
+	{
+		return BaseNode->HasAttribute(UE::Interchange::FAttributeKey(MakeExpressionNameString()));
+	}
+
+	void FInterchangeMaterialInstanceOverridesAPI::GatherLeafInputs(const UInterchangeBaseNode* BaseNode, TArray<UE::Interchange::FAttributeKey>& OutLeafInputAttributeKeys)
+	{
+		TArray<UE::Interchange::FAttributeKey> AttributeKeys;
+		BaseNode->GetAttributeKeys(AttributeKeys);
+
+		OutLeafInputAttributeKeys.Empty();
+		for (const UE::Interchange::FAttributeKey& AttributeKey : AttributeKeys)
+		{
+			if (AttributeKey.Key.StartsWith(OverrideParameterPrefix))
+			{
+				OutLeafInputAttributeKeys.Add(AttributeKey);
+			}
+		}
+	}
+}
+
+
 UClass* UInterchangeMaterialFactory::GetFactoryClass() const
 {
 	return UMaterialInterface::StaticClass();
@@ -604,23 +796,39 @@ UInterchangeFactoryBase::FImportAssetResult UInterchangeMaterialFactory::BeginIm
 				if (UMaterialInstanceConstant* MaterialInstanceConstant = Cast<UMaterialInstanceConstant>(MaterialInstance))
 				{
 					FString ParentPath;
+					FSoftObjectPath ParentMaterialPath;
 					if (MaterialInstanceFactoryNode->GetCustomParent(ParentPath))
 					{
-						UMaterialInterface* ParentMaterial = Cast<UMaterialInterface>(FSoftObjectPath(ParentPath).TryLoad());
-						MaterialInstanceConstant->SetParentEditorOnly(ParentMaterial);
+						if (FPackageName::IsValidObjectPath(ParentPath))
+						{
+							ParentMaterialPath = FSoftObjectPath(ParentPath);
+						}
+						else if (UInterchangeFactoryBaseNode* ParentFactory = Arguments.NodeContainer->GetFactoryNode(ParentPath))
+						{
+							ParentFactory->GetCustomReferenceObject(ParentMaterialPath);
+						}
+
+						if (ParentMaterialPath.IsValid())
+						{
+							if (UMaterialInterface* ParentMaterialAsset = Cast<UMaterialInterface>(ParentMaterialPath.TryLoad()))
+							{
+								MaterialInstanceConstant->SetParentEditorOnly(ParentMaterialAsset);
+								MaterialInstanceConstant->PostEditChange();
+							}
+						}
 					}
 				}
 #endif
-
 				SetupMaterialInstance(*MaterialInstance, *Arguments.NodeContainer, *MaterialInstanceFactoryNode, !Arguments.ReimportObject);
 			}
 		}
 	}
-
+	
 	ImportAssetResult.ImportedObject = Material;
 	return ImportAssetResult;
 }
 
+
 UInterchangeFactoryBase::FImportAssetResult UInterchangeMaterialFactory::ImportAsset_Async(const FImportAssetObjectParams& Arguments)
 {
 	FImportAssetResult ImportAssetResult;
@@ -781,7 +989,7 @@ void UInterchangeMaterialFactory::SetupMaterial(UMaterial* Material, const FImpo
 		}
 	}
 
-	FMaterialExpressionBuilder Builder(Material, nullptr, Arguments);
+	FMaterialExpressionBuilder Builder(Material, nullptr, Arguments, *Results.Get());
 
 	// Substrate
 	if(UInterchangeShaderPortsAPI::HasInput(MaterialFactoryNode, Substrate::Parameters::FrontMaterial))
@@ -1296,11 +1504,11 @@ void UInterchangeMaterialFactory::SetupMaterialInstance(UMaterialInstance& Mater
 	{
 		const FName ParameterName = *InputName;
 
-		switch(UInterchangeShaderPortsAPI::GetInputType(&FactoryNode, InputName))
+		switch (UInterchangeShaderPortsAPI::GetInputType(&FactoryNode, InputName))
 		{
 		case UE::Interchange::EAttributeTypes::Bool:
 			UpdateParameterBool(MaterialInstance, InputName, FactoryNode);
-		    break;
+			break;
 		case UE::Interchange::EAttributeTypes::Float:
 			UpdateParameterFloat(MaterialInstance, InputName, FactoryNode);
 			break;
@@ -1312,6 +1520,42 @@ void UInterchangeMaterialFactory::SetupMaterialInstance(UMaterialInstance& Mater
 			break;
 		}
 	}
+
+	TArray<UE::Interchange::FAttributeKey> LeafInputAttributeKeys;
+	UE::Interchange::Materials::HashUtils::FInterchangeMaterialInstanceOverridesAPI::GatherLeafInputs(&FactoryNode, LeafInputAttributeKeys);
+
+	for (const UE::Interchange::FAttributeKey& AttributeKey : LeafInputAttributeKeys)
+	{
+		FString OverrideParameterName;
+		if (UE::Interchange::Materials::HashUtils::FInterchangeMaterialInstanceOverridesAPI::GetOverrideParameterName(AttributeKey, OverrideParameterName))
+		{
+			const FName ParameterName(OverrideParameterName);
+			switch (FactoryNode.GetAttributeType(AttributeKey))
+			{
+			case UE::Interchange::EAttributeTypes::Bool:
+			{
+				// #todo_Vedang_Javdekar: WIP
+				OverrideBoolParameter(FactoryNode, AttributeKey.Key, MaterialInstance, ParameterName);
+			}
+			break;
+			case UE::Interchange::EAttributeTypes::Float:
+			{
+				OverrideScalarParameter(FactoryNode, AttributeKey.Key, MaterialInstance, ParameterName);
+			}
+			break;
+			case UE::Interchange::EAttributeTypes::LinearColor:
+			{
+				OverrideVectorParameter(FactoryNode, AttributeKey.Key, MaterialInstance, ParameterName);
+			}
+			break;
+			case UE::Interchange::EAttributeTypes::String:
+			{
+				OverrideTextureParameter(FactoryNode, AttributeKey.Key, NodeContainer, MaterialInstance, ParameterName);
+			}
+			break;
+			}
+		}
+	}
 }
 
 void UInterchangeMaterialFactory::SetupReimportedMaterialInstance(UMaterialInstance& MaterialInstance, const UInterchangeBaseNodeContainer& NodeContainer, const UInterchangeMaterialInstanceFactoryNode& FactoryNode, const UInterchangeMaterialInstanceFactoryNode& PreviousFactoryNode)
@@ -1600,7 +1844,7 @@ void UInterchangeMaterialFunctionFactory::SetupMaterial(UMaterialFunction* Mater
 	using namespace UE::Interchange::MaterialFactory::Internal;
 	using namespace UE::Interchange::Materials;
 
-	FMaterialExpressionBuilder Builder(nullptr, MaterialFunction, Arguments);
+	FMaterialExpressionBuilder Builder(nullptr, MaterialFunction, Arguments, *Results.Get());
 
 	UMaterialExpressionFunctionOutput* Output = CreateMaterialExpression<UMaterialExpressionFunctionOutput>(nullptr, MaterialFunction);
 
@@ -1676,5 +1920,4 @@ FInterchangeImportMaterialAsyncHelper& FInterchangeImportMaterialAsyncHelper::Ge
 	return Instance;
 }
 
-#undef LOCTEXT_NAMESPACE
-
+#undef LOCTEXT_NAMESPACE
\ No newline at end of file
diff --git a/Engine/Plugins/Interchange/Runtime/Source/Import/Public/Material/InterchangeMaterialFactory.h b/Engine/Plugins/Interchange/Runtime/Source/Import/Public/Material/InterchangeMaterialFactory.h
--- a/Engine/Plugins/Interchange/Runtime/Source/Import/Public/Material/InterchangeMaterialFactory.h
+++ b/Engine/Plugins/Interchange/Runtime/Source/Import/Public/Material/InterchangeMaterialFactory.h
@@ -18,6 +18,49 @@ class UMaterialExpression;
 class UMaterialExpressionMaterialFunctionCall;
 class UMaterialInstance;
 
+namespace UE::Interchange::Materials::HashUtils
+{
+	class INTERCHANGEIMPORT_API FInterchangeMaterialInstanceOverridesAPI
+	{
+	public:
+		/**
+		 * Generates Attribute Key String from the AttributeType and Hash
+		 */
+		static FString MakeOverrideParameterName(UE::Interchange::EAttributeTypes AttributeType, int32 Hash, bool Prefix = true);
+
+		/**
+		 * Generates Attribute Key String from the DisplayLable
+		 */
+		static FString MakeOverrideParameterName(const FString& DisplayLabel);
+
+		/**
+		 * Generates Attribute Key for Expression Name
+		 */
+		static FString MakeExpressionNameString();
+
+		/**
+		 * Retrieves the Parameter Name from AttributeKey
+		 */
+		static bool GetOverrideParameterName(const UE::Interchange::FAttributeKey& AttributeKey, FString& OverrideParameterName);
+
+		/**
+		 * Checks if the Node has a MaterialExpressionName override
+		 */
+		static bool HasMaterialExpressionNameOverride(const UInterchangeBaseNode* BaseNode);
+
+		/**
+		 * Retrieves the Attribute Keys of the Leaf Inputs.
+		 */
+		static void GatherLeafInputs(const UInterchangeBaseNode* BaseNode, TArray<UE::Interchange::FAttributeKey>& OutLeafInputAttributeKeys);
+
+	private:
+		static const TCHAR* ExpressionNameAttributeKey;
+		static const TCHAR* OverrideParameterPrefix;
+		static const TCHAR* OverrideParameterSeparator;
+		static const TCHAR* OverrideHashSeparator;
+	};
+}
+
 UCLASS(BlueprintType)
 class INTERCHANGEIMPORT_API UInterchangeMaterialFactory : public UInterchangeFactoryBase
 {
diff --git a/Engine/Plugins/Interchange/Runtime/Source/Nodes/Private/InterchangeShaderGraphNode.cpp b/Engine/Plugins/Interchange/Runtime/Source/Nodes/Private/InterchangeShaderGraphNode.cpp
--- a/Engine/Plugins/Interchange/Runtime/Source/Nodes/Private/InterchangeShaderGraphNode.cpp
+++ b/Engine/Plugins/Interchange/Runtime/Source/Nodes/Private/InterchangeShaderGraphNode.cpp
@@ -8,6 +8,7 @@
 const TCHAR* UInterchangeShaderPortsAPI::InputPrefix = TEXT("Inputs");
 const TCHAR* UInterchangeShaderPortsAPI::InputSeparator = TEXT(":");
 const TCHAR* UInterchangeShaderPortsAPI::OutputByIndex = TEXT("__ByIndex__");
+const TCHAR* UInterchangeShaderPortsAPI::ParameterSuffix = TEXT("Parameter");
 
 FString UInterchangeShaderPortsAPI::MakeInputConnectionKey(const FString& InputName)
 {
@@ -33,6 +34,18 @@ FString UInterchangeShaderPortsAPI::MakeInputValueKey(const FString& InputName)
 	return StringBuilder.ToString();
 }
 
+FString UInterchangeShaderPortsAPI::MakeInputParameterKey(const FString& InputName)
+{
+	TStringBuilder<128> StringBuilder;
+	StringBuilder.Append(InputPrefix);
+	StringBuilder.Append(InputSeparator);
+	StringBuilder.Append(InputName);
+	StringBuilder.Append(InputSeparator);
+	StringBuilder.Append(ParameterSuffix);
+
+	return StringBuilder.ToString();
+}
+
 FString UInterchangeShaderPortsAPI::MakeInputName(const FString& InputKey)
 {
 	FString InputName;
@@ -45,12 +58,21 @@ FString UInterchangeShaderPortsAPI::MakeInputName(const FString& InputKey)
 }
 
 bool UInterchangeShaderPortsAPI::IsAnInput(const FString& AttributeKey)
+{
+	TStringBuilder<128> PrefixStringBuilder;
+	PrefixStringBuilder.Append(InputPrefix);
+	PrefixStringBuilder.Append(InputSeparator);
+
+	return AttributeKey.StartsWith(PrefixStringBuilder.ToString());
+}
+
+bool UInterchangeShaderPortsAPI::IsAParameter(const FString& AttributeKey)
 {
 	TStringBuilder<128> StringBuilder;
-	StringBuilder.Append(InputPrefix);
 	StringBuilder.Append(InputSeparator);
+	StringBuilder.Append(ParameterSuffix);
 
-	return AttributeKey.StartsWith(StringBuilder.ToString());
+	return AttributeKey.EndsWith(StringBuilder.ToString());
 }
 
 bool UInterchangeShaderPortsAPI::HasInput(const UInterchangeBaseNode* InterchangeNode, const FName& InInputName)
@@ -69,6 +91,12 @@ bool UInterchangeShaderPortsAPI::HasInput(const UInterchangeBaseNode* Interchang
 	return false;
 }
 
+bool UInterchangeShaderPortsAPI::HasParameter(const UInterchangeBaseNode* InterchangeNode, const FName& InInputName)
+{
+	FString ParameterKey = MakeInputParameterKey(InInputName.ToString());
+	return InterchangeNode->HasAttribute(UE::Interchange::FAttributeKey(ParameterKey));
+}
+
 void UInterchangeShaderPortsAPI::GatherInputs(const UInterchangeBaseNode* InterchangeNode, TArray<FString>& OutInputNames)
 {
 	TArray< UE::Interchange::FAttributeKey > AttributeKeys;
@@ -112,9 +140,16 @@ bool UInterchangeShaderPortsAPI::ConnectOuputToInputByIndex(UInterchangeBaseNode
 	return InterchangeNode->AddStringAttribute(MakeInputConnectionKey(InputName), StringBuilder.ToString());
 }
 
-UE::Interchange::EAttributeTypes UInterchangeShaderPortsAPI::GetInputType(const UInterchangeBaseNode* InterchangeNode, const FString& InputName)
+UE::Interchange::EAttributeTypes UInterchangeShaderPortsAPI::GetInputType(const UInterchangeBaseNode* InterchangeNode, const FString& InputName, bool bIsAParameter /*= false*/)
 {
-	return InterchangeNode->GetAttributeType(UE::Interchange::FAttributeKey(MakeInputValueKey(InputName)));
+	if (bIsAParameter)
+	{
+		return InterchangeNode->GetAttributeType(UE::Interchange::FAttributeKey(MakeInputParameterKey(InputName)));
+	}
+	else
+	{
+		return InterchangeNode->GetAttributeType(UE::Interchange::FAttributeKey(MakeInputValueKey(InputName)));
+	}
 }
 
 bool UInterchangeShaderPortsAPI::GetInputConnection(const UInterchangeBaseNode* InterchangeNode, const FString& InputName, FString& OutExpressionUid, FString& OutputName)
@@ -180,6 +215,51 @@ FString UInterchangeShaderNode::GetTypeName() const
 	return TypeName;
 }
 
+bool UInterchangeShaderNode::AddFloatInput(const FString& InputName, const float& AttributeValue, bool bIsAParameter /*= false*/)
+{
+	bool bResult;
+	if (bIsAParameter)
+	{
+		bResult = AddFloatAttribute(UInterchangeShaderPortsAPI::MakeInputParameterKey(InputName), AttributeValue);
+	}
+	else
+	{
+		bResult = AddFloatAttribute(UInterchangeShaderPortsAPI::MakeInputValueKey(InputName), AttributeValue);
+	}
+
+	return bResult;
+}
+
+bool UInterchangeShaderNode::AddLinearColorInput(const FString& InputName, const FLinearColor& AttributeValue, bool bIsAParameter /*= false*/)
+{
+	bool bResult;
+	if (bIsAParameter)
+	{
+		bResult = AddLinearColorAttribute(UInterchangeShaderPortsAPI::MakeInputParameterKey(InputName), AttributeValue);
+	}
+	else
+	{
+		bResult = AddLinearColorAttribute(UInterchangeShaderPortsAPI::MakeInputValueKey(InputName), AttributeValue);
+	}
+
+	return bResult;
+}
+
+bool UInterchangeShaderNode::AddStringInput(const FString& InputName, const FString& AttributeValue, bool bIsAParameter)
+{
+	bool bResult;
+	if (bIsAParameter)
+	{
+		bResult = AddStringAttribute(UInterchangeShaderPortsAPI::MakeInputParameterKey(InputName), AttributeValue);
+	}
+	else
+	{
+		bResult = AddStringAttribute(UInterchangeShaderPortsAPI::MakeInputValueKey(InputName), AttributeValue);
+	}
+
+	return bResult;
+}
+
 bool UInterchangeShaderNode::GetCustomShaderType(FString& AttributeValue) const
 {
 	IMPLEMENT_NODE_ATTRIBUTE_GETTER(ShaderType, FString);
diff --git a/Engine/Plugins/Interchange/Runtime/Source/Nodes/Public/InterchangeMaterialDefinitions.h b/Engine/Plugins/Interchange/Runtime/Source/Nodes/Public/InterchangeMaterialDefinitions.h
--- a/Engine/Plugins/Interchange/Runtime/Source/Nodes/Public/InterchangeMaterialDefinitions.h
+++ b/Engine/Plugins/Interchange/Runtime/Source/Nodes/Public/InterchangeMaterialDefinitions.h
@@ -279,6 +279,33 @@ namespace Interchange
 				{
 					const FName Name = TEXT("VertexColor");
 				}
+
+				namespace ScalarParameter
+				{
+					const FName Name = TEXT("ScalarParameter");
+					namespace Attributes
+					{
+						const FName DefaultValue = TEXT("DefaultValue");
+					}
+				}
+
+				namespace VectorParameter
+				{
+					const FName Name = TEXT("VectorParameter");
+					namespace Attributes
+					{
+						const FName DefaultValue = TEXT("DefaultValue");
+					}
+				}
+
+				namespace StaticBoolParameter
+				{
+					const FName Name = TEXT("StaticBoolParameter");
+					namespace Attributes
+					{
+						const FName DefaultValue = TEXT("DefaultValue");
+					}
+				}
 			}
 		}
 
diff --git a/Engine/Plugins/Interchange/Runtime/Source/Nodes/Public/InterchangeShaderGraphNode.h b/Engine/Plugins/Interchange/Runtime/Source/Nodes/Public/InterchangeShaderGraphNode.h
--- a/Engine/Plugins/Interchange/Runtime/Source/Nodes/Public/InterchangeShaderGraphNode.h
+++ b/Engine/Plugins/Interchange/Runtime/Source/Nodes/Public/InterchangeShaderGraphNode.h
@@ -30,7 +30,14 @@ public:
 	 */
 	UFUNCTION(BlueprintCallable, Category = "Interchange | Node | Material")
 	static FString MakeInputValueKey(const FString& InputName);
- 
+
+	/**
+	 * Makes an attribute key to represent a parameter being given to an input (ie: Inputs:InputName:Parameter).
+	 * This is more relevant to Materials, but could be used to differentiate between constant values and parameters.
+	 */
+	UFUNCTION(BlueprintCallable, Category = "Interchange | Node | Material")
+	static FString MakeInputParameterKey(const FString& InputName);
+
 	/**
 	 * From an attribute key associated with an input (ie: Inputs:InputName:Value), retrieves the input name from it.
 	 */
@@ -42,6 +49,12 @@ public:
 	 */
 	UFUNCTION(BlueprintCallable, Category = "Interchange | Node | Material")
 	static bool IsAnInput(const FString& AttributeKey);
+
+	/**
+	 * Returns true if the attribute key is an input that represents parameters(ends with ":Parameter").
+	 */
+	UFUNCTION(BlueprintCallable, Category = "Interchange | Node | Material")
+	static bool IsAParameter(const FString& AttributeKey);
  
 	/**
 	 * Checks if a particular input exists on a given node.
@@ -49,6 +62,12 @@ public:
 	UFUNCTION(BlueprintCallable, Category = "Interchange | Node | Material")
 	static bool HasInput(const UInterchangeBaseNode* InterchangeNode, const FName& InInputName);
  
+	/**
+	 * Checks if a particular input exists as a parameter on a given node.
+	 */
+	UFUNCTION(BlueprintCallable, Category = "Interchange | Node | Material")
+	static bool HasParameter(const UInterchangeBaseNode* InterchangeNode, const FName& InInputName);
+
 	/**
 	 * Retrieves the names of all the inputs for a given node.
 	 */
@@ -103,7 +122,7 @@ public:
 	/**
 	 * For an input with a value, returns the type of the stored value.
 	 */
-	static UE::Interchange::EAttributeTypes GetInputType(const UInterchangeBaseNode* InterchangeNode, const FString& InputName);
+	static UE::Interchange::EAttributeTypes GetInputType(const UInterchangeBaseNode* InterchangeNode, const FString& InputName, bool bIsAParameter = false);
 
 	/**
 	 * Returns INDEX_NONE if OutputName is not an index.
@@ -114,6 +133,8 @@ private:
 	static const TCHAR* InputPrefix;
 	static const TCHAR* InputSeparator;
 	static const TCHAR* OutputByIndex;
+
+	static const TCHAR* ParameterSuffix;
 };
 
 /**
@@ -138,6 +159,32 @@ public:
 	virtual FString GetTypeName() const override;
  
 public:
+
+	/**
+	 * Sets the Float Attribute on the Shader Node. If bIsAParameter is set to true, it would be treated as a ScalarParameter
+	 * when the Material Pipeline creates the materials. Otherwise it would be a constant expression in the shader graph.
+	 * Note: It is assumed that the input name would be the parameter name when bIsAParameter is true.
+	 */
+	UFUNCTION(BlueprintCallable, Category = "Interchange | Node | Material")
+	bool AddFloatInput(const FString& InputName, const float& AttributeValue, bool bIsAParameter = false);
+
+	/**
+	 * Sets the Linear Color Attribute on the Shader Node. If bIsAParameter is set to true, it would be treated as a VectorParameter
+	 * when the Material Pipeline creates the materials. Otherwise it would be a constant 3 vector expression in the shader graph.
+	 * Note: It is assumed that the input name would be the parameter name when bIsAParameter is true.
+	 */
+	UFUNCTION(BlueprintCallable, Category = "Interchange | Node | Material")
+	bool AddLinearColorInput(const FString& InputName, const FLinearColor& AttributeValue, bool bIsAParameter = false);
+
+	/**
+	 * Sets the String Attribute on the Shader Node. If bIsAParameter is set to true, it would be treated as a overridable Texture
+	 * or else it should be treated as a LUT Texture.
+	 * Note: It is assumed that the input name would be the parameter name when bIsAParameter is true.
+	 */
+	UFUNCTION(BlueprintCallable, Category = "Interchange | Node | Material")
+	bool AddStringInput(const FString& InputName, const FString& AttributeValue, bool bIsAParameter = false);
+
+
 	UFUNCTION(BlueprintCallable, Category = "Interchange | Node | Material")
 	bool GetCustomShaderType(FString& AttributeValue) const;
  
diff --git a/Engine/Plugins/Interchange/Runtime/Source/Pipelines/Private/InterchangeGenericMaterialPipeline.cpp b/Engine/Plugins/Interchange/Runtime/Source/Pipelines/Private/InterchangeGenericMaterialPipeline.cpp
--- a/Engine/Plugins/Interchange/Runtime/Source/Pipelines/Private/InterchangeGenericMaterialPipeline.cpp
+++ b/Engine/Plugins/Interchange/Runtime/Source/Pipelines/Private/InterchangeGenericMaterialPipeline.cpp
@@ -19,6 +19,7 @@
 #include "Materials/Material.h"
 #include "Misc/PackageName.h"
 #include "Nodes/InterchangeBaseNode.h"
+#include "Nodes/InterchangeUserDefinedAttribute.h"
 #include "InterchangeMaterialInstanceNode.h"
 
 #include "Materials/MaterialExpressionAdd.h"
@@ -31,6 +32,9 @@
 #include "Materials/MaterialExpressionNoise.h"
 #include "Materials/MaterialExpressionOneMinus.h"
 #include "Materials/MaterialExpressionScalarParameter.h"
+#include "Materials/MaterialExpressionStaticBoolParameter.h"
+#include "Materials/MaterialExpressionConstant.h"
+#include "Materials/MaterialExpressionConstant3Vector.h"
 #include "Materials/MaterialExpressionTextureCoordinate.h"
 #include "Materials/MaterialExpressionTextureObject.h"
 #include "Materials/MaterialExpressionTextureSample.h"
@@ -57,6 +61,9 @@
 #include "UObject/Object.h"
 #include "UObject/ObjectMacros.h"
 
+// Material Hash Utils
+#include "Material/InterchangeMaterialFactory.h"
+
 #include UE_INLINE_GENERATED_CPP_BY_NAME(InterchangeGenericMaterialPipeline)
 
 #define LOCTEXT_NAMESPACE "InterchangeGenericMaterialPipeline"
@@ -193,6 +200,260 @@ namespace UE::Interchange::InterchangeGenericMaterialPipeline::Private
 	}
 }
 
+namespace UE::Interchange::Materials::HashUtils
+{
+#if UE_BUILD_DEBUG
+	class FMaterialHashDebugData
+	{
+	public:
+		FMaterialHashDebugData(const FString& InLogDirectoryPath)
+			:LogDirectoryPath(InLogDirectoryPath) {}
+		
+		void Reset();
+
+		void SaveLogsToFile(const FString& FileName);
+
+		template<typename...TArgs>
+		void LogMessage(const FString& Format, TArgs ...Args)
+		{
+			TArray<FStringFormatArg> FormattedArgs;
+			AddFormattedArg(FormattedArgs, Args...);
+			FString Message = FString::Format(*Format, FormattedArgs);
+			UE_LOG(LogInterchangePipeline, Log, TEXT("%s"), *Message);
+			LogMessageContainer.Add(Message);
+		}
+
+		void LogCurrentNodeAddress()
+		{
+			FString NodeAddress = FString::Printf(TEXT("Current Node Address: %s\n"), *GetCurrentNodeAddress());
+			UE_LOG(LogInterchangePipeline, Log, TEXT("%s"), *NodeAddress);
+			LogMessageContainer.Add(NodeAddress);
+		}
+
+		FString GetCurrentNodeAddress();
+		void AddNodeAddress(const FString& NodeAddress, bool bCreatePopCheckPoint = true);
+		void PopNodeAddressesToLastPopIndex();
+
+	private:
+		template<typename TArg>
+		void AddFormattedArg(TArray<FStringFormatArg>& FormattedArgs, TArg Arg)
+		{
+			FormattedArgs.Add(Arg);
+		}
+
+		template<typename TArg, typename ...TArgs>
+		void AddFormattedArg(TArray<FStringFormatArg>& FormattedArgs, TArg Arg, TArgs ...Args)
+		{
+			FormattedArgs.Add(Arg);
+			AddFormattedArg(FormattedArgs, Args...);
+		}
+
+	private:
+		FString LogDirectoryPath;
+		TArray<FString> LogMessageContainer;
+
+		TArray<FString> NodeAddressStack;
+		TArray<int32> NodeAddressPopCheckPoints;
+	};
+
+	FString FMaterialHashDebugData::GetCurrentNodeAddress()
+	{
+		TStringBuilder<512> StringBuilder;
+		for (int32 i = 0; i < NodeAddressStack.Num(); ++i)
+		{
+			StringBuilder.Append(NodeAddressStack[i]);
+			if (i < NodeAddressStack.Num() - 1)
+			{
+				StringBuilder.Append(TEXT("/"));
+			}
+		}
+
+		return StringBuilder.ToString();
+	}
+
+	void FMaterialHashDebugData::SaveLogsToFile(const FString& FileName)
+	{
+		const FString LogFileExtension = TEXT(".txt");
+
+		if (LogMessageContainer.Num())
+		{
+			static FString FileDirectory = FPaths::ProjectSavedDir() + LogDirectoryPath;
+
+			IPlatformFile& PlatformFile = FPlatformFileManager::Get().GetPlatformFile();
+			if (PlatformFile.CreateDirectoryTree(*FileDirectory))
+			{
+				FString AbsolutePath = FileDirectory + FileName + LogFileExtension;
+				FFileHelper::SaveStringArrayToFile(LogMessageContainer, *AbsolutePath);
+			}
+		}
+	}
+
+	void FMaterialHashDebugData::Reset()
+	{
+		NodeAddressStack.Empty();
+		LogMessageContainer.Empty();
+	}
+
+	void FMaterialHashDebugData::AddNodeAddress(const FString& NodeAddress, bool bCreatePopCheckPoint /*= true*/)
+	{
+		if (bCreatePopCheckPoint)
+		{
+			NodeAddressPopCheckPoints.Add(NodeAddressStack.Num());
+		}
+
+		NodeAddressStack.Add(NodeAddress);
+	}
+
+	void FMaterialHashDebugData::PopNodeAddressesToLastPopIndex()
+	{
+		int32 TargetStackSize = 0;
+		if (NodeAddressPopCheckPoints.Num())
+		{
+			TargetStackSize = NodeAddressPopCheckPoints.Last();
+		}
+
+		if (!TargetStackSize)
+		{
+			NodeAddressStack.Empty();
+			NodeAddressPopCheckPoints.Empty();
+		}
+		else
+		{
+			while (NodeAddressStack.Num() && NodeAddressStack.Num() > TargetStackSize)
+			{
+				NodeAddressStack.Pop();
+			}
+
+			NodeAddressPopCheckPoints.Pop();
+		}
+	}
+#endif
+
+	class FDuplicateMaterialHelper
+	{
+	public:
+#if UE_BUILD_DEBUG
+		FDuplicateMaterialHelper(UInterchangeGenericMaterialPipeline& InGenericMaterialPipeline, FMaterialHashDebugData* InHashDebugData)
+			:GenericMaterialPipeline(InGenericMaterialPipeline),
+			HashDebugData(InHashDebugData)
+		{}
+#else
+		FDuplicateMaterialHelper(UInterchangeGenericMaterialPipeline& InGenericMaterialPipeline)
+			:GenericMaterialPipeline(InGenericMaterialPipeline)
+		{}
+#endif
+		void CopyLeafInputsToFactoryNode(UInterchangeBaseMaterialFactoryNode* FactoryNode);
+
+		void ComputMaterialHash(const UInterchangeShaderGraphNode* ShaderGraphNode);
+		
+		void SetupOverridableScalarParameter(const UInterchangeShaderNode* ShaderNode, const FString& ParameterKey, const FString& OverridableParameterNameKey);
+		void SetupOverridableVectorParameter(const UInterchangeShaderNode* ShaderNode, const FString& ParameterKey, const FString& OverridableParameterNameKey);
+		void SetupOverridableStaticBoolParameter(const UInterchangeShaderNode* ShaderNode, const FString& ParameterKey, const FString& OverridableParameterNameKey);
+		void SetupOverridableTextureParameter(const UInterchangeShaderNode* ShaderNode, const FString& InputKey, const FString& OverridableParameterNameKey);
+		
+		/**
+		 * Creates a Base Material Factory Node based on if the material is a duplicate material or if it is found for the first time.
+		 * If the option to create a material instance for the parent is enabled, then additional material instance factory for parent would also be created.
+		 */
+		UInterchangeBaseMaterialFactoryNode* CreateFactoryForDuplicateMaterials(const UInterchangeShaderGraphNode* ShaderGraphNode, bool bImportUnusedMaterial, bool bCreateInstanceForParent);
+
+		void ResetHashData();
+
+		bool IsDuplicate()const { return bIsDuplicate; }
+		const UInterchangeBaseNode* const GetAttributeStorageNode() const { return AttributeStorageNode; }
+
+		template<class TInterchangeResultType>
+		void PostMessage(FText&& MessageText)
+		{
+			if (GenericMaterialPipeline.Results)
+			{
+				TInterchangeResultType* Result = GenericMaterialPipeline.Results->Add<TInterchangeResultType>();
+				Result->Text = MoveTemp(MessageText);
+			}
+		}
+
+	private:
+		UInterchangeBaseMaterialFactoryNode* CreateMaterialFactory(const UInterchangeShaderGraphNode* ShaderGraphNode);
+		UInterchangeMaterialInstanceFactoryNode* CreateMaterialInstanceFactoryFromReference(const UInterchangeShaderGraphNode* ShaderGraphNode);
+		UInterchangeMaterialInstanceFactoryNode* CreateMaterialInstanceFactoryForParent(const UInterchangeShaderGraphNode* ShaderGraphNode);
+
+		TEnumAsByte<EBlendMode> GetShaderGraphNodeBlendMode(const UInterchangeShaderGraphNode* ShaderGraphNode) const;
+		uint8 GetShaderGraphNodeShadingModel(const UInterchangeShaderGraphNode* ShaderGraphNode) const;
+
+		int32 ComputeShaderGraphNodeHash(const UInterchangeShaderGraphNode* ShaderGraphNode);
+		int32 ComputeShaderNodeHash(const  UInterchangeShaderNode* ShaderNode);
+		int32 ComputeShaderInputHash(const UInterchangeShaderNode* ShaderNode, const FString& InputName);
+		int32 HashCombineCustom(int32 Hash, int32 CombineWith);
+
+	private:
+		UInterchangeGenericMaterialPipeline& GenericMaterialPipeline;
+
+		TMap<int32, UInterchangeBaseMaterialFactoryNode*> ParentMaterialFactoryMap;
+
+		UInterchangeBaseNode* AttributeStorageNode = nullptr;
+
+#if UE_BUILD_DEBUG
+		FMaterialHashDebugData* HashDebugData = nullptr;
+#endif
+
+		TArray<UE::Interchange::FAttributeKey> LeafInputAttributeKeys;
+		TSet<const UInterchangeShaderNode*> LeafInputShaderNodes;
+
+		int32 AccumulatedHash = 0;
+		int32 MaterialHash = 0;
+
+		bool bIsDuplicate = false;
+	};
+
+	FString GetDefaultValueStringForShaderType(FString ShaderType)
+	{
+		if (*ShaderType == Standard::Nodes::ScalarParameter::Name)
+		{
+			return Standard::Nodes::ScalarParameter::Attributes::DefaultValue.ToString();
+		}
+		else if (*ShaderType == Standard::Nodes::VectorParameter::Name)
+		{
+			return Standard::Nodes::VectorParameter::Attributes::DefaultValue.ToString();
+		}
+		else if (*ShaderType == Standard::Nodes::StaticBoolParameter::Name)
+		{
+			return Standard::Nodes::StaticBoolParameter::Attributes::DefaultValue.ToString();
+		}
+
+		return FString();
+	}
+}
+
+#if UE_BUILD_DEBUG
+#define ADD_LOG_MESSAGE(...) if(HashDebugData){\
+HashDebugData->LogMessage(__VA_ARGS__);\
+}
+
+#define ADD_NODE_ADDRESS_MESSAGE() if(HashDebugData){\
+HashDebugData->LogCurrentNodeAddress();\
+}
+
+#define PUSH_NODE_ADDRESS(Node) if(HashDebugData){\
+HashDebugData->AddNodeAddress(Node);\
+}
+
+#define PUSH_NODE_ADDRESS_WITHOUT_CHECKPOINT(Node) if(HashDebugData){\
+HashDebugData->AddNodeAddress(Node, false);\
+}
+
+#define POP_NODE_ADDRESSES() if(HashDebugData){\
+HashDebugData->PopNodeAddressesToLastPopIndex();\
+}
+#else
+
+#define ADD_LOG_MESSAGE(...)
+#define ADD_NODE_ADDRESS_MESSAGE() 
+#define PUSH_NODE_ADDRESS(Node) 
+#define PUSH_NODE_ADDRESS_WITHOUT_CHECKPOINT(Node)
+#define POP_NODE_ADDRESSES()
+
+#endif
+
 UInterchangeGenericMaterialPipeline::UInterchangeGenericMaterialPipeline()
 {
 	TexturePipeline = CreateDefaultSubobject<UInterchangeGenericTexturePipeline>("TexturePipeline");
@@ -257,6 +518,13 @@ void UInterchangeGenericMaterialPipeline::AdjustSettingsForContext(EInterchangeP
 
 void UInterchangeGenericMaterialPipeline::ExecutePipeline(UInterchangeBaseNodeContainer* InBaseNodeContainer, const TArray<UInterchangeSourceData*>& InSourceDatas)
 {
+#if UE_BUILD_DEBUG
+	UE::Interchange::Materials::HashUtils::FMaterialHashDebugData HashDebugData(TEXT("InterchangeDebug/MaterialHashLogs/"));
+	UE::Interchange::Materials::HashUtils::FDuplicateMaterialHelper HashHelper(*this, &HashDebugData);
+#else
+	UE::Interchange::Materials::HashUtils::FDuplicateMaterialHelper HashHelper(*this);
+#endif
+
 	if (!InBaseNodeContainer)
 	{
 		UE_LOG(LogInterchangePipeline, Warning, TEXT("UInterchangeGenericMaterialPipeline: Cannot execute pre-import pipeline because InBaseNodeContrainer is null"));
@@ -316,7 +584,7 @@ void UInterchangeGenericMaterialPipeline::ExecutePipeline(UInterchangeBaseNodeCo
 		bImportUnusedMaterial |= bImportMaterials ;
 	}
 
-	// Can't import materials at runtime, fallback to instances
+	// Can't import materials at runtime, fall back to instances
 	if (FApp::IsGame() && MaterialImport == EInterchangeMaterialImportOption::ImportAsMaterials)
 	{
 		MaterialImport = EInterchangeMaterialImportOption::ImportAsMaterialInstances;
@@ -327,16 +595,31 @@ void UInterchangeGenericMaterialPipeline::ExecutePipeline(UInterchangeBaseNodeCo
 		for (const UInterchangeShaderGraphNode* ShaderGraphNode : MaterialNodes)
 		{
 			UInterchangeBaseMaterialFactoryNode* MaterialBaseFactoryNode = nullptr;
-
+			
 			bool bIsAShaderFunction;
 			if (ShaderGraphNode->GetCustomIsAShaderFunction(bIsAShaderFunction) && bIsAShaderFunction)
 			{
 				MaterialBaseFactoryNode = CreateMaterialFunctionFactoryNode(ShaderGraphNode);
 			}
-			else
+			else if (!bIdentifyDuplicateMaterials)
 			{
 				MaterialBaseFactoryNode = CreateMaterialFactoryNode(ShaderGraphNode);
 			}
+			else
+			{
+				HashHelper.ResetHashData();
+				HashHelper.ComputMaterialHash(ShaderGraphNode);
+				AttributeStorageNode = HashHelper.GetAttributeStorageNode();
+
+				/* Creates Material Instance Factory if duplicate material is found. */
+				MaterialBaseFactoryNode = HashHelper.CreateFactoryForDuplicateMaterials(ShaderGraphNode, bImportUnusedMaterial, bCreateMaterialInstanceForParent);
+
+#if UE_BUILD_DEBUG
+				HashDebugData.SaveLogsToFile(ShaderGraphNode->GetUniqueID());
+#endif
+				/* Clearing the AttributeStorageNode as it might affect how the MaterialFunctionsFactories are created. */
+				AttributeStorageNode = nullptr;
+			}
 
 			if (MaterialBaseFactoryNode)
 			{
@@ -364,13 +647,11 @@ void UInterchangeGenericMaterialPipeline::ExecutePipeline(UInterchangeBaseNodeCo
 	for (UInterchangeMaterialInstanceNode* MaterialNode : MaterialInstanceNodes)
 	{
 		FString ParentPath;
-
 		if (!MaterialNode->GetCustomParent(ParentPath) || ParentPath.IsEmpty())
 		{
 			continue;
 		}
 
-
 		UInterchangeMaterialInstanceFactoryNode* MaterialFactoryNode = nullptr;
 		FString DisplayLabel = MaterialNode->GetDisplayLabel();
 		const FString NodeUid = UInterchangeFactoryBaseNode::BuildFactoryNodeUid(MaterialNode->GetUniqueID());
@@ -413,9 +694,10 @@ void UInterchangeGenericMaterialPipeline::ExecutePipeline(UInterchangeBaseNodeCo
 
 		for (const FString& InputName : Inputs)
 		{
-			FString InputValueKey = UInterchangeShaderPortsAPI::MakeInputValueKey(InputName);
+			const bool bIsAParameter = UInterchangeShaderPortsAPI::HasParameter(MaterialNode, FName(InputName));
+			FString InputValueKey = CreateInputKey(InputName,bIsAParameter);
 
-			switch (UInterchangeShaderPortsAPI::GetInputType(MaterialNode, InputName))
+			switch (UInterchangeShaderPortsAPI::GetInputType(MaterialNode, InputName, bIsAParameter))
 			{
 			case UE::Interchange::EAttributeTypes::Bool:
 			{
@@ -500,10 +782,17 @@ void UInterchangeGenericMaterialPipeline::SetReimportSourceIndex(UClass* Reimpor
 	}
 }
 
-UInterchangeBaseMaterialFactoryNode* UInterchangeGenericMaterialPipeline::CreateBaseMaterialFactoryNode(const UInterchangeBaseNode* MaterialNode, TSubclassOf<UInterchangeBaseMaterialFactoryNode> NodeType)
+UInterchangeBaseMaterialFactoryNode* UInterchangeGenericMaterialPipeline::CreateBaseMaterialFactoryNode(const UInterchangeBaseNode* MaterialNode, TSubclassOf<UInterchangeBaseMaterialFactoryNode> NodeType, bool bAddMaterialInstanceSuffix /*= false*/)
 {
+	const FString MaterialInstanceSuffix = TEXT("_MI");
+
 	FString DisplayLabel = MaterialNode->GetDisplayLabel();
-	const FString NodeUid = UInterchangeFactoryBaseNode::BuildFactoryNodeUid(MaterialNode->GetUniqueID());
+	FString NodeUid = UInterchangeFactoryBaseNode::BuildFactoryNodeUid(MaterialNode->GetUniqueID());
+	if (bAddMaterialInstanceSuffix)
+	{
+		NodeUid += MaterialInstanceSuffix;
+	}
+
 	UInterchangeBaseMaterialFactoryNode* MaterialFactoryNode = nullptr;
 	if (BaseNodeContainer->IsNodeUidValid(NodeUid))
 	{
@@ -1394,7 +1683,12 @@ void UInterchangeGenericMaterialPipeline::HandleMakeFloat3Node(const UInterchang
 	}
 }
 
-void UInterchangeGenericMaterialPipeline::HandleTextureNode(const UInterchangeTextureNode* TextureNode, UInterchangeBaseMaterialFactoryNode* MaterialFactoryNode, UInterchangeMaterialExpressionFactoryNode* TextureBaseFactoryNode, const FString & ExpressionClassName)
+void UInterchangeGenericMaterialPipeline::HandleTextureNode(
+	const UInterchangeTextureNode* TextureNode, 
+	UInterchangeBaseMaterialFactoryNode* MaterialFactoryNode, 
+	UInterchangeMaterialExpressionFactoryNode* TextureBaseFactoryNode, 
+	const FString & ExpressionClassName,
+	bool bIsAParameter)
 {
 	using namespace UE::Interchange::Materials::Standard::Nodes::TextureSample;
 
@@ -1408,7 +1702,7 @@ void UInterchangeGenericMaterialPipeline::HandleTextureNode(const UInterchangeTe
 	}
 
 	TextureBaseFactoryNode->SetCustomExpressionClassName(ExpressionClassName);
-	TextureBaseFactoryNode->AddStringAttribute(UInterchangeShaderPortsAPI::MakeInputValueKey(Inputs::Texture.ToString()), TextureFactoryUid);
+	TextureBaseFactoryNode->AddStringAttribute(CreateInputKey(Inputs::Texture.ToString(), bIsAParameter), TextureFactoryUid);
 
 	if(UInterchangeTextureFactoryNode* TextureFactoryNode = Cast<UInterchangeTextureFactoryNode>(BaseNodeContainer->GetFactoryNode(TextureFactoryUid)))
 	{
@@ -1462,19 +1756,19 @@ void UInterchangeGenericMaterialPipeline::HandleTextureObjectNode(const UInterch
 {
 	using namespace UE::Interchange::Materials::Standard::Nodes::TextureObject;
 
-	FString TextureUid;
-	ShaderNode->GetStringAttribute(UInterchangeShaderPortsAPI::MakeInputValueKey(Inputs::Texture.ToString()), TextureUid);
+	bool bIsAParameter;
+	FString TextureUid = GetTextureUidAttributeFromShaderNode(ShaderNode, Inputs::Texture, bIsAParameter);
 	FString ExpressionClassName;
 	FString TextureFactoryUid;
 
 	if(const UInterchangeTextureNode* TextureNode = Cast<const UInterchangeTextureNode>(BaseNodeContainer->GetNode(TextureUid)))
 	{
-		HandleTextureNode(TextureNode, MaterialFactoryNode, TextureObjectFactoryNode, UMaterialExpressionTextureObject::StaticClass()->GetName());
+		HandleTextureNode(TextureNode, MaterialFactoryNode, TextureObjectFactoryNode, UMaterialExpressionTextureObject::StaticClass()->GetName(), bIsAParameter);
 	}
 	else
 	{
 		TextureObjectFactoryNode->SetCustomExpressionClassName(UMaterialExpressionTextureObject::StaticClass()->GetName());
-		TextureObjectFactoryNode->AddStringAttribute(UInterchangeShaderPortsAPI::MakeInputValueKey(Inputs::Texture.ToString()), TextureFactoryUid);
+		TextureObjectFactoryNode->AddStringAttribute(CreateInputKey(Inputs::Texture.ToString(), bIsAParameter), TextureFactoryUid);
 	}
 }
 
@@ -1482,9 +1776,8 @@ void UInterchangeGenericMaterialPipeline::HandleTextureSampleNode(const UInterch
 {
 	using namespace UE::Interchange::Materials::Standard::Nodes::TextureSample;
 
-	FString TextureUid;
-	ShaderNode->GetStringAttribute(UInterchangeShaderPortsAPI::MakeInputValueKey(Inputs::Texture.ToString()), TextureUid);
-
+	bool bIsAParameter;
+	FString TextureUid = GetTextureUidAttributeFromShaderNode(ShaderNode, Inputs::Texture, bIsAParameter);
 	FString ExpressionClassName;
 	FString TextureFactoryUid;
 
@@ -1511,12 +1804,12 @@ void UInterchangeGenericMaterialPipeline::HandleTextureSampleNode(const UInterch
 			ExpressionClassName = UMaterialExpressionTextureSampleParameter2D::StaticClass()->GetName();
 		}
 
-		HandleTextureNode(TextureNode, MaterialFactoryNode, TextureSampleFactoryNode, ExpressionClassName);
+		HandleTextureNode(TextureNode, MaterialFactoryNode, TextureSampleFactoryNode, ExpressionClassName, bIsAParameter);
 	}
 	else
 	{
 		TextureSampleFactoryNode->SetCustomExpressionClassName(UMaterialExpressionTextureSampleParameter2D::StaticClass()->GetName());
-		TextureSampleFactoryNode->AddStringAttribute(UInterchangeShaderPortsAPI::MakeInputValueKey(Inputs::Texture.ToString()), TextureFactoryUid);
+		TextureSampleFactoryNode->AddStringAttribute(CreateInputKey(Inputs::Texture.ToString(), bIsAParameter), TextureFactoryUid);
 	}
 
 	// Coordinates
@@ -2073,6 +2366,53 @@ void UInterchangeGenericMaterialPipeline::HandleSwizzleNode(const UInterchangeSh
 	}
 }
 
+void UInterchangeGenericMaterialPipeline::HandleScalarParameterNode(const UInterchangeShaderNode* ShaderNode, UInterchangeMaterialExpressionFactoryNode* ScalarParameterFactoryNode)
+{
+	using namespace UE::Interchange::Materials::Standard::Nodes;
+	const FString ParameterKey = UInterchangeShaderPortsAPI::MakeInputParameterKey(ScalarParameter::Attributes::DefaultValue.ToString());
+	float InputValue;
+	if (ShaderNode->GetFloatAttribute(ParameterKey, InputValue))
+	{
+		ScalarParameterFactoryNode->SetCustomExpressionClassName(UMaterialExpressionScalarParameter::StaticClass()->GetName());
+		const FName DefaultValueMemberName = GET_MEMBER_NAME_CHECKED(UMaterialExpressionScalarParameter, DefaultValue);
+		ScalarParameterFactoryNode->AddFloatAttribute(DefaultValueMemberName.ToString(), InputValue);
+		ScalarParameterFactoryNode->AddApplyAndFillDelegates<float>(DefaultValueMemberName.ToString(), UMaterialExpressionScalarParameter::StaticClass(), DefaultValueMemberName);
+	}
+
+	ScalarParameterFactoryNode->SetDisplayLabel(ShaderNode->GetDisplayLabel());
+}
+void UInterchangeGenericMaterialPipeline::HandleVectorParameterNode(const UInterchangeShaderNode* ShaderNode, UInterchangeMaterialExpressionFactoryNode* VectorParameterFactoryNode)
+{
+	using namespace UE::Interchange::Materials::Standard::Nodes;
+	const FString ParameterKey = UInterchangeShaderPortsAPI::MakeInputParameterKey(VectorParameter::Attributes::DefaultValue.ToString());
+	FLinearColor InputValue;
+	if (ShaderNode->GetLinearColorAttribute(ParameterKey, InputValue))
+	{
+		VectorParameterFactoryNode->SetCustomExpressionClassName(UMaterialExpressionVectorParameter::StaticClass()->GetName());
+		const FName DefaultValueMemberName = GET_MEMBER_NAME_CHECKED(UMaterialExpressionVectorParameter, DefaultValue);
+		VectorParameterFactoryNode->AddLinearColorAttribute(DefaultValueMemberName.ToString(), InputValue);
+		VectorParameterFactoryNode->AddApplyAndFillDelegates<FLinearColor>(DefaultValueMemberName.ToString(), UMaterialExpressionVectorParameter::StaticClass(), DefaultValueMemberName);
+	}
+
+	VectorParameterFactoryNode->SetDisplayLabel(ShaderNode->GetDisplayLabel());
+}
+void UInterchangeGenericMaterialPipeline::HandleStaticBooleanParameterNode(const UInterchangeShaderNode* ShaderNode, UInterchangeMaterialExpressionFactoryNode* StaticBoolParameterFactoryNode)
+{
+	using namespace UE::Interchange::Materials::Standard::Nodes;
+	const FString ParameterKey = UInterchangeShaderPortsAPI::MakeInputParameterKey(StaticBoolParameter::Attributes::DefaultValue.ToString());
+	bool InputValue;
+	if (ShaderNode->GetBooleanAttribute(ParameterKey, InputValue))
+	{
+		StaticBoolParameterFactoryNode->SetCustomExpressionClassName(UMaterialExpressionStaticBoolParameter::StaticClass()->GetName());
+		const FName DefaultValueMemberName = GET_MEMBER_NAME_CHECKED(UMaterialExpressionStaticBoolParameter, DefaultValue);
+		StaticBoolParameterFactoryNode->AddBooleanAttribute(DefaultValueMemberName.ToString(), InputValue);
+		StaticBoolParameterFactoryNode->AddApplyAndFillDelegates<bool>(DefaultValueMemberName.ToString(), UMaterialExpressionStaticBoolParameter::StaticClass(), DefaultValueMemberName);
+	}
+
+	StaticBoolParameterFactoryNode->SetDisplayLabel(ShaderNode->GetDisplayLabel());
+}
+
+
 UInterchangeMaterialExpressionFactoryNode* UInterchangeGenericMaterialPipeline::CreateMaterialExpressionForShaderNode(UInterchangeBaseMaterialFactoryNode* MaterialFactoryNode,
 	const UInterchangeShaderNode* ShaderNode, const FString& ParentUid)
 {
@@ -2167,6 +2507,18 @@ UInterchangeMaterialExpressionFactoryNode* UInterchangeGenericMaterialPipeline::
 	{
 		HandleVectorNoiseNode(ShaderNode, MaterialFactoryNode, MaterialExpression);
 	}
+	else if (*ShaderType == Nodes::ScalarParameter::Name)
+	{
+		HandleScalarParameterNode(ShaderNode, MaterialExpression);
+	}
+	else if (*ShaderType == Nodes::VectorParameter::Name)
+	{
+		HandleVectorParameterNode(ShaderNode, MaterialExpression);
+	}
+	else if (*ShaderType == Nodes::StaticBoolParameter::Name)
+	{
+		HandleStaticBooleanParameterNode(ShaderNode, MaterialExpression);
+	}
 	else if (ensure(!ShaderType.IsEmpty()))
 	{
 		const FString ExpressionClassName = TEXT("MaterialExpression") + ShaderType;
@@ -2201,11 +2553,13 @@ UInterchangeMaterialExpressionFactoryNode* UInterchangeGenericMaterialPipeline::
 		
 		if (*ShaderType == Nodes::TextureSample::Name)
 		{
-			ShaderNode->GetStringAttribute(UInterchangeShaderPortsAPI::MakeInputValueKey(Nodes::TextureSample::Inputs::Texture.ToString()), TextureUid);
+			const bool bIsAParameter = UInterchangeShaderPortsAPI::HasParameter(ShaderNode, Nodes::TextureSample::Inputs::Texture);
+			ShaderNode->GetStringAttribute(CreateInputKey(Nodes::TextureSample::Inputs::Texture.ToString(), bIsAParameter), TextureUid);
 		}
 		else if (*ShaderType == Nodes::TextureObject::Name)
 		{
-			ShaderNode->GetStringAttribute(UInterchangeShaderPortsAPI::MakeInputValueKey(Nodes::TextureObject::Inputs::Texture.ToString()), TextureUid);
+			const bool bIsAParameter = UInterchangeShaderPortsAPI::HasParameter(ShaderNode, Nodes::TextureObject::Inputs::Texture);
+			ShaderNode->GetStringAttribute(CreateInputKey(Nodes::TextureObject::Inputs::Texture.ToString(),bIsAParameter), TextureUid);
 		}
 
 		// Make the material factory node have a dependency on the texture factory node so that the texture asset gets created first
@@ -2247,18 +2601,76 @@ UInterchangeMaterialExpressionFactoryNode* UInterchangeGenericMaterialPipeline::
 	return MaterialExpressionFactoryNode;
 }
 
+UInterchangeMaterialExpressionFactoryNode* UInterchangeGenericMaterialPipeline::HandleFloatInput(const UInterchangeShaderNode* ShaderNode, const FString& InputName, const FString& ParentUid, bool bIsAParameter)
+{
+	if (bIsAParameter)
+	{
+		return CreateScalarParameterExpression(ShaderNode, InputName, ParentUid);
+	}
+	else
+	{
+		return CreateConstantExpression(ShaderNode, InputName, ParentUid);
+	}
+}
+
+UInterchangeMaterialExpressionFactoryNode* UInterchangeGenericMaterialPipeline::CreateConstantExpression(const UInterchangeShaderNode* ShaderNode, const FString& InputName, const FString& ParentUid)
+{
+	UInterchangeMaterialExpressionFactoryNode* MaterialExpressionFactoryNode = CreateExpressionNode(InputName, ParentUid, UMaterialExpressionConstant::StaticClass());
+
+	float InputValue;
+	if (ShaderNode->GetFloatAttribute(UInterchangeShaderPortsAPI::MakeInputValueKey(InputName), InputValue))
+	{
+		const FName RMemberName = GET_MEMBER_NAME_CHECKED(UMaterialExpressionConstant, R);
+		MaterialExpressionFactoryNode->AddFloatAttribute(RMemberName.ToString(), InputValue);
+		MaterialExpressionFactoryNode->AddApplyAndFillDelegates<float>(RMemberName.ToString(), UMaterialExpressionConstant::StaticClass(), RMemberName);
+	}
+
+	return MaterialExpressionFactoryNode;
+}
+
 UInterchangeMaterialExpressionFactoryNode* UInterchangeGenericMaterialPipeline::CreateScalarParameterExpression(const UInterchangeShaderNode* ShaderNode, const FString& InputName, const FString& ParentUid)
 {
+	using namespace UE::Interchange::Materials::Standard;
+
 	UInterchangeMaterialExpressionFactoryNode* MaterialExpressionFactoryNode = CreateExpressionNode(InputName, ParentUid, UMaterialExpressionScalarParameter::StaticClass());
 
 	float InputValue;
-	if (ShaderNode->GetFloatAttribute(UInterchangeShaderPortsAPI::MakeInputValueKey(InputName), InputValue))
+	if (ShaderNode->GetFloatAttribute(UInterchangeShaderPortsAPI::MakeInputParameterKey(InputName), InputValue))
 	{
 		const FName DefaultValueMemberName = GET_MEMBER_NAME_CHECKED(UMaterialExpressionScalarParameter, DefaultValue);
 		MaterialExpressionFactoryNode->AddFloatAttribute(DefaultValueMemberName.ToString(), InputValue);
 		MaterialExpressionFactoryNode->AddApplyAndFillDelegates<float>(DefaultValueMemberName.ToString(), UMaterialExpressionScalarParameter::StaticClass(), DefaultValueMemberName);
 	}
 
+	MaterialExpressionFactoryNode->SetDisplayLabel(InputName);
+
+	return MaterialExpressionFactoryNode;
+}
+
+UInterchangeMaterialExpressionFactoryNode* UInterchangeGenericMaterialPipeline::HandleLinearColorInput(const UInterchangeShaderNode* ShaderNode, const FString& InputName, const FString& ParentUid, bool bIsAParameter)
+{
+	if (bIsAParameter)
+	{
+		return CreateVectorParameterExpression(ShaderNode, InputName, ParentUid);
+	}
+	else
+	{
+		return CreateConstant3VectorExpression(ShaderNode, InputName, ParentUid);
+	}
+}
+
+UInterchangeMaterialExpressionFactoryNode* UInterchangeGenericMaterialPipeline::CreateConstant3VectorExpression(const UInterchangeShaderNode* ShaderNode, const FString& InputName, const FString& ParentUid)
+{
+	UInterchangeMaterialExpressionFactoryNode* MaterialExpressionFactoryNode = CreateExpressionNode(InputName, ParentUid, UMaterialExpressionConstant3Vector::StaticClass());
+
+	FLinearColor InputValue;
+	if (ShaderNode->GetLinearColorAttribute(UInterchangeShaderPortsAPI::MakeInputValueKey(InputName), InputValue))
+	{
+		const FName ConstantMemberName = GET_MEMBER_NAME_CHECKED(UMaterialExpressionConstant3Vector, Constant);
+		MaterialExpressionFactoryNode->AddLinearColorAttribute(ConstantMemberName.ToString(), InputValue);
+		MaterialExpressionFactoryNode->AddApplyAndFillDelegates<FLinearColor>(ConstantMemberName.ToString(), UMaterialExpressionConstant3Vector::StaticClass(), ConstantMemberName);
+	}
+
 	return MaterialExpressionFactoryNode;
 }
 
@@ -2267,13 +2679,32 @@ UInterchangeMaterialExpressionFactoryNode* UInterchangeGenericMaterialPipeline::
 	UInterchangeMaterialExpressionFactoryNode* MaterialExpressionFactoryNode = CreateExpressionNode(InputName, ParentUid, UMaterialExpressionVectorParameter::StaticClass());
 
 	FLinearColor InputValue;
-	if (ShaderNode->GetLinearColorAttribute(UInterchangeShaderPortsAPI::MakeInputValueKey(InputName), InputValue))
+	if (ShaderNode->GetLinearColorAttribute(UInterchangeShaderPortsAPI::MakeInputParameterKey(InputName), InputValue))
 	{
-		const FName DefaultValueMemberName = GET_MEMBER_NAME_CHECKED(UMaterialExpressionVectorParameter, DefaultValue);
-		MaterialExpressionFactoryNode->AddLinearColorAttribute(DefaultValueMemberName.ToString(), InputValue);
-		MaterialExpressionFactoryNode->AddApplyAndFillDelegates<FLinearColor>(DefaultValueMemberName.ToString(), UMaterialExpressionVectorParameter::StaticClass(), DefaultValueMemberName);
+		const FName DefaultValueName = GET_MEMBER_NAME_CHECKED(UMaterialExpressionVectorParameter, DefaultValue);
+		MaterialExpressionFactoryNode->AddLinearColorAttribute(DefaultValueName.ToString(), InputValue);
+		MaterialExpressionFactoryNode->AddApplyAndFillDelegates<FLinearColor>(DefaultValueName.ToString(), UMaterialExpressionVectorParameter::StaticClass(), DefaultValueName);
+	}
+
+	MaterialExpressionFactoryNode->SetDisplayLabel(InputName);
+
+	return MaterialExpressionFactoryNode;
+}
+
+UInterchangeMaterialExpressionFactoryNode* UInterchangeGenericMaterialPipeline::CreateStaticBooleanParameterExpression(const UInterchangeShaderNode* ShaderNode, const FString& InputName, const FString& ParentUid)
+{
+	UInterchangeMaterialExpressionFactoryNode* MaterialExpressionFactoryNode = CreateExpressionNode(InputName, ParentUid, UMaterialExpressionStaticBoolParameter::StaticClass());
+
+	bool InputValue;
+	if (ShaderNode->GetBooleanAttribute(InputName, InputValue))
+	{
+		const FName DefaultValueMemberName = GET_MEMBER_NAME_CHECKED(UMaterialExpressionStaticBoolParameter, DefaultValue);
+		MaterialExpressionFactoryNode->AddBooleanAttribute(DefaultValueMemberName.ToString(), InputValue);
+		MaterialExpressionFactoryNode->AddApplyAndFillDelegates<bool>(DefaultValueMemberName.ToString(), UMaterialExpressionStaticBoolParameter::StaticClass(), DefaultValueMemberName);
 	}
 
+	MaterialExpressionFactoryNode->SetDisplayLabel(InputName);
+
 	return MaterialExpressionFactoryNode;
 }
 
@@ -2328,20 +2759,34 @@ TTuple<UInterchangeMaterialExpressionFactoryNode*, FString> UInterchangeGenericM
 	}
 	else
 	{
-		switch(UInterchangeShaderPortsAPI::GetInputType(ShaderNode, InputName))
+		const bool bIsAParameter = UInterchangeShaderPortsAPI::HasParameter(ShaderNode, FName(InputName));
+		switch(UInterchangeShaderPortsAPI::GetInputType(ShaderNode, InputName, bIsAParameter))
 		{
 		case UE::Interchange::EAttributeTypes::Float:
-			MaterialExpressionFactoryNode = CreateScalarParameterExpression(ShaderNode, InputName, ParentUid);
+			MaterialExpressionFactoryNode = HandleFloatInput(ShaderNode, InputName, ParentUid, bIsAParameter);
 			break;
 		case UE::Interchange::EAttributeTypes::LinearColor:
-			MaterialExpressionFactoryNode = CreateVectorParameterExpression(ShaderNode, InputName, ParentUid);
+			MaterialExpressionFactoryNode = HandleLinearColorInput(ShaderNode, InputName, ParentUid, bIsAParameter);
 			break;
 		case UE::Interchange::EAttributeTypes::Vector2f:
 			MaterialExpressionFactoryNode = CreateVector2ParameterExpression(ShaderNode, InputName, ParentUid);
 			break;
+		case UE::Interchange::EAttributeTypes::Bool:
+			MaterialExpressionFactoryNode = CreateStaticBooleanParameterExpression(ShaderNode, InputName, ParentUid);
+			break;
+		}
+
+		if (MaterialExpressionFactoryNode)
+		{
+			FString MaterialExpressionName;
+			if (AttributeStorageNode && AttributeStorageNode->GetStringAttribute(ShaderNode->GetUniqueID(), MaterialExpressionName))
+			{
+				MaterialExpressionFactoryNode->SetDisplayLabel(MaterialExpressionName);
+			}
 		}
 	}
 
+
 	TTuple<UInterchangeMaterialExpressionFactoryNode*, FString> Result {MaterialExpressionFactoryNode, MaterialExpressionCreationContextStack[ExpressionContextIndex].OutputName};
 	MaterialExpressionCreationContextStack.Pop();
 
@@ -2475,8 +2920,9 @@ UInterchangeMaterialInstanceFactoryNode* UInterchangeGenericMaterialPipeline::Cr
 
 	for (const FString& InputName : Inputs)
 	{
-		TVariant<FString, FLinearColor, float> InputValue;
+		const bool bIsAParameter = UInterchangeShaderPortsAPI::HasParameter(ShaderGraphNode, FName(InputName));
 
+		TVariant<FString, FLinearColor, float> InputValue;
 		FString ConnectedShaderNodeUid;
 		FString OutputName;
 		if (UInterchangeShaderPortsAPI::GetInputConnection(ShaderGraphNode, InputName, ConnectedShaderNodeUid, OutputName))
@@ -2488,19 +2934,19 @@ UInterchangeMaterialInstanceFactoryNode* UInterchangeGenericMaterialPipeline::Cr
 		}
 		else
 		{
-			switch(UInterchangeShaderPortsAPI::GetInputType(ShaderGraphNode, InputName))
+			switch(UInterchangeShaderPortsAPI::GetInputType(ShaderGraphNode, InputName,bIsAParameter))
 			{
 			case UE::Interchange::EAttributeTypes::Float:
 				{
 					float AttributeValue = 0.f;
-					ShaderGraphNode->GetFloatAttribute(UInterchangeShaderPortsAPI::MakeInputValueKey(InputName), AttributeValue);
+					ShaderGraphNode->GetFloatAttribute(CreateInputKey(InputName, bIsAParameter), AttributeValue);
 					InputValue.Set<float>(AttributeValue);
 				}
 				break;
 			case UE::Interchange::EAttributeTypes::LinearColor:
 				{
 					FLinearColor AttributeValue = FLinearColor::White;
-					ShaderGraphNode->GetLinearColorAttribute(UInterchangeShaderPortsAPI::MakeInputValueKey(InputName), AttributeValue);
+					ShaderGraphNode->GetLinearColorAttribute(CreateInputKey(InputName,bIsAParameter), AttributeValue);
 					InputValue.Set<FLinearColor>(AttributeValue);
 				}
 				break;
@@ -2509,19 +2955,19 @@ UInterchangeMaterialInstanceFactoryNode* UInterchangeGenericMaterialPipeline::Cr
 
 		if (InputValue.IsType<float>())
 		{
-			MaterialInstanceFactoryNode->AddFloatAttribute(UInterchangeShaderPortsAPI::MakeInputValueKey(InputName), InputValue.Get<float>());
+			MaterialInstanceFactoryNode->AddFloatAttribute(CreateInputKey(InputName, bIsAParameter), InputValue.Get<float>());
 		}
 		else if (InputValue.IsType<FLinearColor>())
 		{
-			MaterialInstanceFactoryNode->AddLinearColorAttribute(UInterchangeShaderPortsAPI::MakeInputValueKey(InputName), InputValue.Get<FLinearColor>());
+			MaterialInstanceFactoryNode->AddLinearColorAttribute(CreateInputKey(InputName, bIsAParameter), InputValue.Get<FLinearColor>());
 		}
 		else if (InputValue.IsType<FString>())
 		{
 			const FString MapName(InputName + TEXT("Map"));
-			MaterialInstanceFactoryNode->AddStringAttribute(UInterchangeShaderPortsAPI::MakeInputValueKey(MapName), InputValue.Get<FString>());
+			MaterialInstanceFactoryNode->AddStringAttribute(CreateInputKey(MapName, bIsAParameter), InputValue.Get<FString>());
 
 			const FString MapWeightName(MapName + TEXT("Weight"));
-			MaterialInstanceFactoryNode->AddFloatAttribute(UInterchangeShaderPortsAPI::MakeInputValueKey(MapWeightName), 1.f);
+			MaterialInstanceFactoryNode->AddFloatAttribute(CreateInputKey(MapWeightName, bIsAParameter), 1.f);
 
 			MaterialInstanceFactoryNode->AddFactoryDependencyUid(InputValue.Get<FString>());
 		}
@@ -2574,6 +3020,8 @@ TVariant<FString, FLinearColor, float> UInterchangeGenericMaterialPipeline::Visi
 TVariant<FString, FLinearColor, float> UInterchangeGenericMaterialPipeline::VisitShaderInput(const UInterchangeShaderNode* ShaderNode, const FString& InputName) const
 {
 	TVariant<FString, FLinearColor, float> Result;
+	
+	const bool bIsAParameter = UInterchangeShaderPortsAPI::HasParameter(ShaderNode, FName(InputName));
 
 	FString ConnectedShaderNodeUid;
 	FString OutputName;
@@ -2591,14 +3039,14 @@ TVariant<FString, FLinearColor, float> UInterchangeGenericMaterialPipeline::Visi
 		case UE::Interchange::EAttributeTypes::Float:
 			{
 				float InputValue = 0.f;
-				ShaderNode->GetFloatAttribute(UInterchangeShaderPortsAPI::MakeInputValueKey(InputName), InputValue);
+				ShaderNode->GetFloatAttribute(CreateInputKey(InputName, bIsAParameter), InputValue);
 				Result.Set<float>(InputValue);
 			}
 			break;
 		case UE::Interchange::EAttributeTypes::LinearColor:
 			{
 				FLinearColor InputValue = FLinearColor::White;
-				ShaderNode->GetLinearColorAttribute(UInterchangeShaderPortsAPI::MakeInputValueKey(InputName), InputValue);
+				ShaderNode->GetLinearColorAttribute(CreateInputKey(InputName, bIsAParameter), InputValue);
 				Result.Set<FLinearColor>(InputValue);
 			}
 			break;
@@ -2741,8 +3189,10 @@ TVariant<FString, FLinearColor, float> UInterchangeGenericMaterialPipeline::Visi
 
 	TVariant<FString, FLinearColor, float> Result;
 
+	const bool bIsAParameter = UInterchangeShaderPortsAPI::HasParameter(ShaderNode, TextureSample::Inputs::Texture);
+
 	FString TextureUid;
-	if (ShaderNode->GetStringAttribute(UInterchangeShaderPortsAPI::MakeInputValueKey(TextureSample::Inputs::Texture.ToString()), TextureUid))
+	if (ShaderNode->GetStringAttribute(CreateInputKey(TextureSample::Inputs::Texture.ToString(),bIsAParameter), TextureUid))
 	{
 		if (!TextureUid.IsEmpty())
 		{
@@ -2765,6 +3215,26 @@ TVariant<FString, FLinearColor, float> UInterchangeGenericMaterialPipeline::Visi
 	return Result;
 }
 
+FString UInterchangeGenericMaterialPipeline::GetTextureUidAttributeFromShaderNode(const UInterchangeShaderNode* ShaderNode, FName ParameterName, bool& OutIsAParameter) const
+{
+	OutIsAParameter = UInterchangeShaderPortsAPI::HasParameter(ShaderNode, ParameterName);
+	FString TextureUid;
+	ShaderNode->GetStringAttribute(CreateInputKey(ParameterName.ToString(), OutIsAParameter), TextureUid);
+	return TextureUid;
+}
+
+FString UInterchangeGenericMaterialPipeline::CreateInputKey(const FString& InputName, bool bIsAParameter) const
+{
+	if (bIsAParameter)
+	{
+		return UInterchangeShaderPortsAPI::MakeInputParameterKey(InputName);
+	}
+	else
+	{
+		return UInterchangeShaderPortsAPI::MakeInputValueKey(InputName);
+	}
+}
+
 bool UInterchangeGenericMaterialPipeline::HandleBxDFInput(const UInterchangeShaderGraphNode* ShaderGraphNode, UInterchangeMaterialFactoryNode* MaterialFactoryNode)
 {
 	using namespace UE::Interchange::Materials;
@@ -2961,7 +3431,9 @@ UInterchangeMaterialExpressionFactoryNode* UInterchangeGenericMaterialPipeline::
 	{
 		const FString MaterialFunctionFactoryNodeUid = UInterchangeFactoryBaseNode::BuildFactoryNodeUid(MaterialFunctionAttribute);
 		FunctionCallFactoryNode->SetCustomMaterialFunctionDependency(MaterialFunctionFactoryNodeUid);
-		FunctionCallFactoryNode->AddFactoryDependencyUid(MaterialFunctionFactoryNodeUid);
+
+		UClass* CustomExpressionClass = UMaterialExpressionMaterialFunctionCall::StaticClass();
+		FunctionCallFactoryNode->SetCustomExpressionClassName(CustomExpressionClass->GetName());
 	}
 	else if (FPackageName::IsValidObjectPath(MaterialFunctionAttribute))
 	{
@@ -2986,4 +3458,528 @@ UInterchangeMaterialExpressionFactoryNode* UInterchangeGenericMaterialPipeline::
 	return FunctionCallFactoryNode;
 }
 
-#undef LOCTEXT_NAMESPACE
+namespace UE::Interchange::Materials::HashUtils
+{
+	void FDuplicateMaterialHelper::ResetHashData()
+	{
+		AccumulatedHash = 0;
+		MaterialHash = 0;
+		bIsDuplicate = false;
+
+		if (AttributeStorageNode)
+		{
+			AttributeStorageNode->MarkAsGarbage();
+			AttributeStorageNode = nullptr;
+		}
+		AttributeStorageNode = NewObject<UInterchangeBaseNode>();
+
+		LeafInputAttributeKeys.Empty();
+		LeafInputShaderNodes.Empty();
+
+#if UE_BUILD_DEBUG
+		if (HashDebugData)
+		{
+			HashDebugData->Reset();
+		}
+#endif
+	}
+
+	void FDuplicateMaterialHelper::ComputMaterialHash(const UInterchangeShaderGraphNode* ShaderGraphNode)
+	{
+		MaterialHash = ComputeShaderGraphNodeHash(ShaderGraphNode);
+		if (ParentMaterialFactoryMap.Contains(MaterialHash))
+		{
+			bIsDuplicate = true;
+		}
+	}
+
+	int32 FDuplicateMaterialHelper::ComputeShaderGraphNodeHash(const UInterchangeShaderGraphNode* ShaderGraphNode)
+	{
+		using namespace UE::Interchange::Materials::HashUtils;
+
+		/* Two Sided*/
+		bool bTwoSided;
+		ShaderGraphNode->GetCustomTwoSided(bTwoSided);
+		int32 Hash = GetTypeHash(bTwoSided);
+		ADD_LOG_MESSAGE(TEXT("TwoSided: {0}, Hash: {1}"), bTwoSided, Hash);
+
+		/* Use Material Attributes*/
+		bool bUseMaterialAttributes = UInterchangeShaderPortsAPI::HasInput(ShaderGraphNode, UE::Interchange::Materials::Common::Parameters::BxDF);
+		Hash = HashCombine(Hash, GetTypeHash(bUseMaterialAttributes));
+		ADD_LOG_MESSAGE(TEXT("Use Material Attributes: {0}, Hash: {1}"), bUseMaterialAttributes, Hash);
+
+		/* Blend Mode */
+		TEnumAsByte<EBlendMode> BlendMode = GetShaderGraphNodeBlendMode(ShaderGraphNode);
+		Hash = HashCombine(Hash, GetTypeHash(BlendMode));
+		ADD_LOG_MESSAGE(TEXT("Blend Mode: {0}, Hash: {1}"), (uint8)BlendMode, Hash);
+
+		/* Is Thin Surface */
+		Hash = HashCombine(Hash, GetTypeHash(BlendMode == EBlendMode::BLEND_Translucent));
+		ADD_LOG_MESSAGE(TEXT("Is Thin Surface: {0}, Hash: {1}"), (BlendMode == EBlendMode::BLEND_Translucent), Hash);
+
+		/* EDatasmithShadingModel: uint8 */
+		Hash = HashCombine(Hash, GetTypeHash(GetShaderGraphNodeShadingModel(ShaderGraphNode)));
+		ADD_LOG_MESSAGE(TEXT("Shading Model: {0}, Hash: {1}"), GetShaderGraphNodeShadingModel(ShaderGraphNode), Hash);
+
+		Hash = HashCombine(Hash, ComputeShaderNodeHash(ShaderGraphNode));
+		ADD_LOG_MESSAGE(TEXT("ShaderHash: {0}"), Hash);
+		return Hash;
+	}
+
+	uint8 FDuplicateMaterialHelper::GetShaderGraphNodeShadingModel(const UInterchangeShaderGraphNode* ShaderGraphNode) const
+	{
+		if (GenericMaterialPipeline.HasThinTranslucency(ShaderGraphNode))
+		{
+			return 1;
+		}
+		else if (GenericMaterialPipeline.HasSubsurface(ShaderGraphNode))
+		{
+			return 2;
+		}
+		else if (GenericMaterialPipeline.HasClearCoat(ShaderGraphNode))
+		{
+			return 3;
+		}
+		else if (GenericMaterialPipeline.IsUnlitModel(ShaderGraphNode))
+		{
+			return 4;
+		}
+
+		return 0;
+	}
+
+	TEnumAsByte<EBlendMode> FDuplicateMaterialHelper::GetShaderGraphNodeBlendMode(const UInterchangeShaderGraphNode* ShaderGraphNode) const
+	{
+		using namespace UE::Interchange::Materials;
+
+		TEnumAsByte<EBlendMode> BlendMode = EBlendMode::BLEND_Opaque;
+
+		if (GenericMaterialPipeline.HasThinTranslucency(ShaderGraphNode))
+		{
+			BlendMode = EBlendMode::BLEND_Translucent;
+
+		}
+		else if (GenericMaterialPipeline.HasSubsurface(ShaderGraphNode))
+		{
+			BlendMode = EBlendMode::BLEND_Opaque;
+		}
+		else
+		{
+			const bool bHasOpacityInput = UInterchangeShaderPortsAPI::HasInput(ShaderGraphNode, UE::Interchange::Materials::Common::Parameters::Opacity);
+			if (bHasOpacityInput && GenericMaterialPipeline.IsUnlitModel(ShaderGraphNode))
+			{
+				float OpacityClipValue;
+				if (ShaderGraphNode->GetCustomOpacityMaskClipValue(OpacityClipValue))
+				{
+					BlendMode = EBlendMode::BLEND_Masked;
+				}
+				else
+				{
+					BlendMode = EBlendMode::BLEND_Translucent;
+				}
+			}
+		}
+
+		return BlendMode;
+	}
+
+	int32 FDuplicateMaterialHelper::ComputeShaderNodeHash(const UInterchangeShaderNode* ShaderNode)
+	{
+		int32 Hash = 0;
+
+		FString ShaderTypeName;
+		ShaderNode->GetCustomShaderType(ShaderTypeName);
+
+		TArray<FString> Inputs;
+		UInterchangeShaderPortsAPI::GatherInputs(ShaderNode, Inputs);
+
+		if (!ShaderTypeName.IsEmpty())
+		{
+			using namespace UE::Interchange::Materials;
+
+			PUSH_NODE_ADDRESS_WITHOUT_CHECKPOINT(ShaderTypeName)
+
+			Hash = HashCombineCustom(Hash, GetTypeHash(ShaderTypeName));
+			ADD_LOG_MESSAGE(TEXT("{0}, Accumulated Hash: {1}"), ShaderTypeName, AccumulatedHash)
+			TArray<FInterchangeUserDefinedAttributeInfo> UserDefinedAttributes = UInterchangeUserDefinedAttributesAPI::GetUserDefinedAttributeInfos(ShaderNode);
+			if (UserDefinedAttributes.Num())
+			{
+				for (const auto& UserDefinedAttribute : UserDefinedAttributes)
+				{
+					Hash = HashCombineCustom(Hash, GetTypeHash(UserDefinedAttribute.Type));
+					Hash = HashCombineCustom(Hash, GetTypeHash(UserDefinedAttribute.Name));
+
+					ADD_LOG_MESSAGE(TEXT("UDA[Type: {0}, Name: {1}], Accumulated Hash: {2}"), (int32)UserDefinedAttribute.Type, UserDefinedAttribute.Name, AccumulatedHash)
+
+					const FString UserDefinedAttributeType = AttributeTypeToString(UserDefinedAttribute.Type);
+
+					if (UserDefinedAttribute.Type == UE::Interchange::EAttributeTypes::String)
+					{
+						const FString InputValueKey = (UInterchangeUserDefinedAttributesAPI::MakeUserDefinedPropertyValueKey(UserDefinedAttribute.Name, UserDefinedAttribute.RequiresDelegate)).Key;
+						const FString OverrideParameterNameAttributeKey = FInterchangeMaterialInstanceOverridesAPI::MakeOverrideParameterName(ShaderNode->GetDisplayLabel());
+						SetupOverridableTextureParameter(ShaderNode, InputValueKey, OverrideParameterNameAttributeKey);
+					}
+				}
+			}
+		}
+		else
+		{
+			if (const UInterchangeFunctionCallShaderNode* FunctionCallNode = Cast<UInterchangeFunctionCallShaderNode>(ShaderNode))
+			{
+				FString MaterialFunction;
+				if (FunctionCallNode->GetCustomMaterialFunction(MaterialFunction) && !MaterialFunction.IsEmpty())
+				{
+					Hash = HashCombineCustom(Hash, GetTypeHash(MaterialFunction));
+					ADD_LOG_MESSAGE(TEXT("MF[{0}], Accumulate Hash: {1}"), MaterialFunction, AccumulatedHash);
+#if UE_BUILD_DEBUG
+					FString MaterialFunctionName;
+					FString Discard;
+					if (MaterialFunction.Split(TEXT("."), &Discard, &MaterialFunctionName))
+					{
+						PUSH_NODE_ADDRESS_WITHOUT_CHECKPOINT(FString::Printf(TEXT("MaterialFunction[%s]"), *MaterialFunctionName));
+					}
+					else
+					{
+						PUSH_NODE_ADDRESS_WITHOUT_CHECKPOINT(TEXT("MaterialFunction"));
+					}
+#endif
+				}
+			}
+		}
+
+		if (!Inputs.IsEmpty())
+		{
+			for (const FString& InputName : Inputs)
+			{
+				PUSH_NODE_ADDRESS(FString::Printf(TEXT("[%s]"), *InputName))
+				ADD_NODE_ADDRESS_MESSAGE();
+				Hash = HashCombineCustom(Hash, ComputeShaderInputHash(ShaderNode, InputName));
+				POP_NODE_ADDRESSES()
+			}
+		}
+
+		return Hash;
+	}
+
+	int32 FDuplicateMaterialHelper::ComputeShaderInputHash(const UInterchangeShaderNode* ShaderNode, const FString& InputName)
+	{
+		int32 Hash = 0;
+		FString ConnectedShaderNodeUid;
+		FString OutputName;
+		if (UInterchangeShaderPortsAPI::GetInputConnection(ShaderNode, InputName, ConnectedShaderNodeUid, OutputName))
+		{
+			if (const UInterchangeShaderNode* ConnectedShaderNode = Cast<const UInterchangeShaderNode>(GenericMaterialPipeline.BaseNodeContainer->GetNode(ConnectedShaderNodeUid)))
+			{
+				Hash = HashCombineCustom(Hash, ComputeShaderNodeHash(ConnectedShaderNode));
+			}
+
+			if (!OutputName.IsEmpty())
+			{
+				Hash = HashCombineCustom(Hash, GetTypeHash(OutputName));
+			}
+		}
+		else
+		{
+			const bool bIsAParameter = UInterchangeShaderPortsAPI::HasParameter(ShaderNode, FName(InputName));
+			const UE::Interchange::EAttributeTypes InputType = UInterchangeShaderPortsAPI::GetInputType(ShaderNode, InputName, bIsAParameter);
+			Hash = HashCombineCustom(Hash, GetTypeHash(InputType));
+			ADD_LOG_MESSAGE(TEXT("{0}, Accumulated Hash: {1}"), AttributeTypeToString(InputType), AccumulatedHash);
+
+			// Just setup all the Parameters as overridable parameters. Do not include the values in the Hash
+			if (bIsAParameter)
+			{				
+				using namespace UE::Interchange::Materials::HashUtils;
+				const FString ParameterKey = UInterchangeShaderPortsAPI::MakeInputParameterKey(InputName);
+				const FString OverridableParameterNameKey = FInterchangeMaterialInstanceOverridesAPI::MakeOverrideParameterName(ShaderNode->GetDisplayLabel());
+
+				switch (InputType)
+				{
+				case UE::Interchange::EAttributeTypes::Float:
+				{
+					SetupOverridableScalarParameter(ShaderNode, ParameterKey, OverridableParameterNameKey);
+					break;
+				}
+				case UE::Interchange::EAttributeTypes::LinearColor:
+				{
+					SetupOverridableVectorParameter(ShaderNode, ParameterKey, OverridableParameterNameKey);
+					break;
+				}
+				case UE::Interchange::EAttributeTypes::Bool:
+				{
+					SetupOverridableStaticBoolParameter(ShaderNode, ParameterKey, OverridableParameterNameKey);
+					break;
+				}
+				case UE::Interchange::EAttributeTypes::String:
+				{
+					SetupOverridableTextureParameter(ShaderNode, ParameterKey, OverridableParameterNameKey);
+					break;
+				}
+				}
+			}
+			else
+			{
+				const FString InputValueKey = UInterchangeShaderPortsAPI::MakeInputValueKey(InputName);
+				switch (InputType)
+				{
+				case UE::Interchange::EAttributeTypes::Float:
+				{
+					float InputValue;
+					if (ShaderNode->GetFloatAttribute(InputValueKey, InputValue))
+					{
+						Hash = HashCombineCustom(Hash, GetTypeHash(InputValue));
+						ADD_LOG_MESSAGE(TEXT("Unnamed Float({0}), Accumulated Hash: {1}"), FString::SanitizeFloat(InputValue), AccumulatedHash);
+					}
+					break;
+				}
+				case UE::Interchange::EAttributeTypes::LinearColor:
+				{
+					FLinearColor InputValue;
+					if (ShaderNode->GetLinearColorAttribute(InputValueKey, InputValue))
+					{
+						Hash = HashCombineCustom(Hash, GetTypeHash(InputValue));
+						ADD_LOG_MESSAGE(TEXT("Unnamed LinearColor({0}), Accumulated Hash: {1}"), InputValue.ToString(), AccumulatedHash)
+					}
+					break;
+				}
+				case UE::Interchange::EAttributeTypes::String:
+				{
+					FString InputValue;
+					if (ShaderNode->GetStringAttribute(InputValueKey, InputValue))
+					{
+						Hash = HashCombineCustom(Hash, GetTypeHash(InputValue));
+						ADD_LOG_MESSAGE(TEXT("Unnamed String({0}), Accumulated Hash: {1}"), InputValue, AccumulatedHash)
+					}
+					break;
+				}
+				}
+			}
+		}
+		return Hash;
+	}
+
+	void FDuplicateMaterialHelper::SetupOverridableScalarParameter(const UInterchangeShaderNode* ShaderNode, const FString& ParameterKey, const FString& OverridableParameterNameKey)
+	{
+		float InputValue;
+		if (ShaderNode->GetFloatAttribute(ParameterKey, InputValue))
+		{
+			const UE::Interchange::FAttributeKey AttributeKey(OverridableParameterNameKey);
+			if (!AttributeStorageNode->HasAttribute(AttributeKey))
+			{
+				AttributeStorageNode->AddFloatAttribute(OverridableParameterNameKey, InputValue);
+				LeafInputAttributeKeys.Add(AttributeKey);
+				LeafInputShaderNodes.Emplace(ShaderNode);
+				ADD_LOG_MESSAGE(TEXT("Scalar Parameter: {0}({1})"), ShaderNode->GetDisplayLabel(), FString::SanitizeFloat(InputValue));
+			}
+		}
+	}
+
+	void FDuplicateMaterialHelper::SetupOverridableVectorParameter(const UInterchangeShaderNode* ShaderNode, const FString& ParameterKey, const FString& OverridableParameterNameKey)
+	{		
+		FLinearColor InputValue;
+		if (ShaderNode->GetLinearColorAttribute(ParameterKey, InputValue))
+		{
+			const UE::Interchange::FAttributeKey AttributeKey(OverridableParameterNameKey);
+			if (!AttributeStorageNode->HasAttribute(AttributeKey))
+			{
+				AttributeStorageNode->AddLinearColorAttribute(OverridableParameterNameKey, InputValue);
+				LeafInputAttributeKeys.Add(AttributeKey);
+				LeafInputShaderNodes.Emplace(ShaderNode);
+				ADD_LOG_MESSAGE(TEXT("Vector Parameter: {0}({1})"), ShaderNode->GetDisplayLabel(), InputValue.ToString());
+			}
+		}
+	}
+
+	void FDuplicateMaterialHelper::SetupOverridableStaticBoolParameter(const UInterchangeShaderNode* ShaderNode, const FString& ParameterKey, const FString& OverridableParameterNameKey)
+	{
+		bool InputValue;
+		if (ShaderNode->GetBooleanAttribute(ParameterKey, InputValue))
+		{
+			const UE::Interchange::FAttributeKey AttributeKey(OverridableParameterNameKey);
+			if (!AttributeStorageNode->HasAttribute(AttributeKey))
+			{
+				AttributeStorageNode->AddBooleanAttribute(OverridableParameterNameKey, InputValue);
+				LeafInputAttributeKeys.Add(AttributeKey);
+				LeafInputShaderNodes.Emplace(ShaderNode);
+				ADD_LOG_MESSAGE(TEXT("Bool Parameter: {0}({1})"), ShaderNode->GetDisplayLabel(), InputValue);
+			}
+		}
+	}
+
+	void FDuplicateMaterialHelper::SetupOverridableTextureParameter(const UInterchangeShaderNode* ShaderNode, const FString& InputKey, const FString& OverridableParameterNameKey)
+	{
+		FString InputValue;
+		if (ShaderNode->GetStringAttribute(InputKey, InputValue))
+		{
+			const UE::Interchange::FAttributeKey AttributeKey(OverridableParameterNameKey);
+			if (!AttributeStorageNode->HasAttribute(AttributeKey))
+			{
+				if (!FPackageName::IsValidObjectPath(InputValue))
+				{
+					// Material Factory expects Texture Factory Uid as opposed to Texture Uid
+					const FString TextureFactoryUid = UInterchangeFactoryBaseNode::BuildFactoryNodeUid(InputValue);
+					AttributeStorageNode->AddStringAttribute(OverridableParameterNameKey, TextureFactoryUid);
+				}
+				else
+				{
+					AttributeStorageNode->AddStringAttribute(OverridableParameterNameKey, InputValue);
+				}
+
+				LeafInputAttributeKeys.Add(AttributeKey);
+				LeafInputShaderNodes.Emplace(ShaderNode);
+
+				ADD_LOG_MESSAGE(TEXT("Texture Parameter: {0}({1})"), ShaderNode->GetDisplayLabel(), InputValue);
+			}
+		}
+	}
+
+	int32 FDuplicateMaterialHelper::HashCombineCustom(int32 Hash, int32 CombineWith)
+	{
+		Hash = HashCombine(Hash, CombineWith);
+		AccumulatedHash = HashCombine(AccumulatedHash, CombineWith);
+		return Hash;
+	}
+
+	void FDuplicateMaterialHelper::CopyLeafInputsToFactoryNode(UInterchangeBaseMaterialFactoryNode* FactoryNode)
+	{
+		UInterchangeBaseNode::CopyStorageAttributes(AttributeStorageNode, FactoryNode, LeafInputAttributeKeys);
+	}
+	
+	UInterchangeBaseMaterialFactoryNode* FDuplicateMaterialHelper::CreateFactoryForDuplicateMaterials(const UInterchangeShaderGraphNode* ShaderGraphNode, bool bImportUnusedMaterial, bool bCreateMaterialInstanceForParent)
+	{
+		UInterchangeBaseMaterialFactoryNode* MaterialFactoryNode = nullptr;
+		if (IsDuplicate())
+		{
+			MaterialFactoryNode = CreateMaterialInstanceFactoryFromReference(ShaderGraphNode);
+		}
+		else
+		{
+			MaterialFactoryNode = CreateMaterialFactory(ShaderGraphNode);
+			MaterialFactoryNode->SetEnabled(bImportUnusedMaterial);
+
+			if (bCreateMaterialInstanceForParent)
+			{
+				MaterialFactoryNode = CreateMaterialInstanceFactoryForParent(ShaderGraphNode);
+			}
+		}
+
+		return MaterialFactoryNode;
+	}
+
+	UInterchangeBaseMaterialFactoryNode* FDuplicateMaterialHelper::CreateMaterialFactory(const UInterchangeShaderGraphNode* ShaderGraphNode)
+	{
+		UInterchangeBaseMaterialFactoryNode* MaterialFactoryNode = GenericMaterialPipeline.CreateMaterialFactoryNode(ShaderGraphNode);
+		ParentMaterialFactoryMap.Emplace(MaterialHash, MaterialFactoryNode);
+		CopyLeafInputsToFactoryNode(MaterialFactoryNode);
+		return MaterialFactoryNode;
+	}
+
+	UInterchangeMaterialInstanceFactoryNode* FDuplicateMaterialHelper::CreateMaterialInstanceFactoryFromReference(const UInterchangeShaderGraphNode* ShaderGraphNode)
+	{
+		const UInterchangeBaseMaterialFactoryNode* ParentMaterialFactory = ParentMaterialFactoryMap[MaterialHash];
+
+		UInterchangeMaterialInstanceFactoryNode* MaterialInstanceFactoryNode =
+			Cast<UInterchangeMaterialInstanceFactoryNode>(GenericMaterialPipeline.CreateBaseMaterialFactoryNode(ShaderGraphNode, UInterchangeMaterialInstanceFactoryNode::StaticClass()));
+
+		ensure(MaterialInstanceFactoryNode);
+
+		if (ParentMaterialFactory)
+		{
+			MaterialInstanceFactoryNode->SetCustomParent(ParentMaterialFactory->GetUniqueID());
+			MaterialInstanceFactoryNode->AddFactoryDependencyUid(ParentMaterialFactory->GetUniqueID());
+		}
+
+		for (const auto& LeafInputKey : LeafInputAttributeKeys)
+		{
+			UE::Interchange::EAttributeTypes AttributeType = AttributeStorageNode->GetAttributeType(LeafInputKey);
+			switch (AttributeType)
+			{
+			case UE::Interchange::EAttributeTypes::Float:
+			{
+				float ParentValue;
+				float CurrentValue;
+				if (!AttributeStorageNode->GetFloatAttribute(LeafInputKey.Key, CurrentValue))
+				{
+					continue;
+				}
+
+				if (!ParentMaterialFactory->GetFloatAttribute(LeafInputKey.Key, ParentValue))
+				{
+					continue;
+				}
+
+				if (ParentValue != CurrentValue)
+				{
+					MaterialInstanceFactoryNode->AddFloatAttribute(LeafInputKey.Key, CurrentValue);
+				}
+			}
+			break;
+			case UE::Interchange::EAttributeTypes::LinearColor:
+			{
+				FLinearColor ParentValue;
+				FLinearColor CurrentValue;
+				if (!AttributeStorageNode->GetLinearColorAttribute(LeafInputKey.Key, CurrentValue))
+				{
+					continue;
+				}
+
+				if (!ParentMaterialFactory->GetLinearColorAttribute(LeafInputKey.Key, ParentValue))
+				{
+					continue;
+				}
+
+				if (ParentValue != CurrentValue)
+				{
+					MaterialInstanceFactoryNode->AddLinearColorAttribute(LeafInputKey.Key, CurrentValue);
+				}
+			}
+			break;
+			case UE::Interchange::EAttributeTypes::String:
+			{
+				FString ParentValue;
+				FString CurrentValue;
+				if (!AttributeStorageNode->GetStringAttribute(LeafInputKey.Key, CurrentValue))
+				{
+					continue;
+				}
+
+				if (!ParentMaterialFactory->GetStringAttribute(LeafInputKey.Key, ParentValue))
+				{
+					continue;
+				}
+
+				if (ParentValue != CurrentValue)
+				{
+					MaterialInstanceFactoryNode->AddStringAttribute(LeafInputKey.Key, CurrentValue);
+				}
+			}
+			break;
+			}
+		}
+
+		return MaterialInstanceFactoryNode;
+	}
+
+	UInterchangeMaterialInstanceFactoryNode* FDuplicateMaterialHelper::CreateMaterialInstanceFactoryForParent(const UInterchangeShaderGraphNode* ShaderGraphNode)
+	{
+		const UInterchangeBaseMaterialFactoryNode* ParentMaterialFactory = ParentMaterialFactoryMap[MaterialHash];
+
+		UInterchangeMaterialInstanceFactoryNode* MaterialInstanceFactoryNode =
+			Cast<UInterchangeMaterialInstanceFactoryNode>(GenericMaterialPipeline.CreateBaseMaterialFactoryNode(ShaderGraphNode, UInterchangeMaterialInstanceFactoryNode::StaticClass(), true));
+
+		if (ParentMaterialFactory)
+		{
+			MaterialInstanceFactoryNode->SetCustomParent(ParentMaterialFactory->GetUniqueID());
+			MaterialInstanceFactoryNode->AddFactoryDependencyUid(ParentMaterialFactory->GetUniqueID());
+		}
+
+		return MaterialInstanceFactoryNode;
+	}
+}
+
+#undef ADD_LOG_MESSAGE
+#undef ADD_NODE_ADDRESS_MESSAGE
+#undef PUSH_NODE_ADDRESS
+#undef PUSH_NODE_ADDRESS_WITHOUT_CHECKPOINT
+#undef POP_NODE_ADDRESSES
+
+#undef LOCTEXT_NAMESPACE
\ No newline at end of file
diff --git a/Engine/Plugins/Interchange/Runtime/Source/Pipelines/Public/InterchangeGenericMaterialPipeline.h b/Engine/Plugins/Interchange/Runtime/Source/Pipelines/Public/InterchangeGenericMaterialPipeline.h
--- a/Engine/Plugins/Interchange/Runtime/Source/Pipelines/Public/InterchangeGenericMaterialPipeline.h
+++ b/Engine/Plugins/Interchange/Runtime/Source/Pipelines/Public/InterchangeGenericMaterialPipeline.h
@@ -24,6 +24,11 @@ class UInterchangeResult;
 class UMaterialFunction;
 class UInterchangeTextureNode;
 
+namespace UE::Interchange::Materials::HashUtils
+{
+	class FDuplicateMaterialHelper;
+}
+
 UENUM(BlueprintType)
 enum class EInterchangeMaterialImportOption : uint8
 {
@@ -49,6 +54,14 @@ public:
 
 	UPROPERTY(EditAnywhere, BlueprintReadWrite, Category = "Materials", Meta=(EditCondition="bImportMaterials"))
 	EInterchangeMaterialImportOption MaterialImport = EInterchangeMaterialImportOption::ImportAsMaterials;
+	
+	/** If set, reference materials along with respective material instances would be created*/
+	UPROPERTY(EditAnywhere, BlueprintReadWrite, Category = "Materials")
+	bool bIdentifyDuplicateMaterials = true;
+
+	/** If set, additional material instance would be created for reference/parent material*/
+	UPROPERTY(EditAnywhere, BlueprintReadWrite, Category = "Materials", Meta=(EditCondition="bIdentifyDuplicateMaterials"))
+	bool bCreateMaterialInstanceForParent = true;
 
 	/** Optional material used as the parent when importing materials as instances. If no parent material is specified, one will be automatically selected during the import process. */
 	UPROPERTY(EditAnywhere, BlueprintReadWrite, Category = "Materials", Meta= (EditCondition="bImportMaterials && MaterialImport==EInterchangeMaterialImportOption::ImportAsMaterialInstances", AllowedClasses="/Script/Engine.MaterialInterface"))
@@ -81,11 +94,11 @@ private:
 	/** Material factory assets nodes */
 	TArray<UInterchangeBaseMaterialFactoryNode*> MaterialFactoryNodes;
 
-	UInterchangeBaseMaterialFactoryNode* CreateBaseMaterialFactoryNode(const UInterchangeBaseNode* MaterialNode, TSubclassOf<UInterchangeBaseMaterialFactoryNode> NodeType);
+	UInterchangeBaseMaterialFactoryNode* CreateBaseMaterialFactoryNode(const UInterchangeBaseNode* MaterialNode, TSubclassOf<UInterchangeBaseMaterialFactoryNode> NodeType, bool bAddMaterialInstanceSuffix = false);
 	UInterchangeMaterialFactoryNode* CreateMaterialFactoryNode(const UInterchangeShaderGraphNode* ShaderGraphNode);
 	UInterchangeMaterialFunctionFactoryNode* CreateMaterialFunctionFactoryNode(const UInterchangeShaderGraphNode* FunctionCallShaderNode);
 	UInterchangeMaterialInstanceFactoryNode* CreateMaterialInstanceFactoryNode(const UInterchangeShaderGraphNode* ShaderGraphNode);
-
+	
 	/** True if the shader graph has a clear coat input. */
 	bool HasClearCoat(const UInterchangeShaderGraphNode* ShaderGraphNode) const;
 	UE_DEPRECATED(5.3, "Deprecated. Use HasClearCoat.")
@@ -157,7 +170,7 @@ private:
 	void HandleFlattenNormalNode(const UInterchangeShaderNode* ShaderNode, UInterchangeBaseMaterialFactoryNode* MaterialFactoryNode, UInterchangeMaterialExpressionFactoryNode* FlattenNormalFactoryNode);
 	void HandleNormalFromHeightMapNode(const UInterchangeShaderNode* ShaderNode, UInterchangeBaseMaterialFactoryNode* MaterialFactoryNode, UInterchangeMaterialExpressionFactoryNode* NormalFromHeightMapFactoryNode);
 	void HandleMakeFloat3Node(const UInterchangeShaderNode* ShaderNode, UInterchangeBaseMaterialFactoryNode* MaterialFactoryNode, UInterchangeMaterialExpressionFactoryNode* MakeFloat3FactoryNode);
-	void HandleTextureNode(const UInterchangeTextureNode* TextureNode, UInterchangeBaseMaterialFactoryNode* MaterialFactoryNode, UInterchangeMaterialExpressionFactoryNode* TextureBaseFactoryNode, const FString& ExpressionClassName);
+	void HandleTextureNode(const UInterchangeTextureNode* TextureNode, UInterchangeBaseMaterialFactoryNode* MaterialFactoryNode, UInterchangeMaterialExpressionFactoryNode* TextureBaseFactoryNode, const FString& ExpressionClassName, bool bIsAParameter);
 	void HandleTextureObjectNode(const UInterchangeShaderNode* ShaderNode, UInterchangeBaseMaterialFactoryNode* MaterialFactoryNode, UInterchangeMaterialExpressionFactoryNode* TextureObjectFactoryNode);
 	void HandleTextureSampleNode(const UInterchangeShaderNode* ShaderNode, UInterchangeBaseMaterialFactoryNode* MaterialFactoryNode, UInterchangeMaterialExpressionFactoryNode* TextureSampleFactoryNode);
 	void HandleTextureSampleBlurNode(const UInterchangeShaderNode* ShaderNode, UInterchangeBaseMaterialFactoryNode* MaterialFactoryNode, UInterchangeMaterialExpressionFactoryNode* TextureSampleFactoryNode);
@@ -170,13 +183,22 @@ private:
 	void HandleNoiseNode(const UInterchangeShaderNode* ShaderNode, UInterchangeBaseMaterialFactoryNode* MaterialFactoryNode, UInterchangeMaterialExpressionFactoryNode* NoiseFactoryNode);
 	void HandleVectorNoiseNode(const UInterchangeShaderNode* ShaderNode, UInterchangeBaseMaterialFactoryNode* MaterialFactoryNode, UInterchangeMaterialExpressionFactoryNode* NoiseFactoryNode);
 	void HandleSwizzleNode(const UInterchangeShaderNode* ShaderNode, UInterchangeBaseMaterialFactoryNode* MaterialFactoryNode, UInterchangeMaterialExpressionFactoryNode* SwizzleFactoryNode);
+	void HandleScalarParameterNode(const UInterchangeShaderNode* ShaderNode, UInterchangeMaterialExpressionFactoryNode* ScalarParameterFactoryNode);
+	void HandleVectorParameterNode(const UInterchangeShaderNode* ShaderNode, UInterchangeMaterialExpressionFactoryNode* VectorParameterFactoryNode);
+	void HandleStaticBooleanParameterNode(const UInterchangeShaderNode* ShaderNode, UInterchangeMaterialExpressionFactoryNode* StaticBooleanParameterFactoryNode);
 
 	UInterchangeMaterialExpressionFactoryNode* CreateMaterialExpressionForShaderNode(UInterchangeBaseMaterialFactoryNode* MaterialFactoryNode, const UInterchangeShaderNode* ShaderNode, const FString& ParentUid);
 	TTuple<UInterchangeMaterialExpressionFactoryNode*, FString> CreateMaterialExpressionForInput(UInterchangeBaseMaterialFactoryNode* MaterialFactoryNode, const UInterchangeShaderNode* ShaderNode, const FString& InputName, const FString& ParentUid);
 
 	UInterchangeMaterialExpressionFactoryNode* CreateExpressionNode(const FString& ExpressionName, const FString& ParentUid, UClass* MaterialExpressionClass);
+	
+	UInterchangeMaterialExpressionFactoryNode* HandleFloatInput(const UInterchangeShaderNode* ShaderNode, const FString& InputName, const FString& ParentUid, bool bIsAParameter);
+	UInterchangeMaterialExpressionFactoryNode* CreateConstantExpression(const UInterchangeShaderNode* ShaderNode, const FString& InputName, const FString& ParentUid);
 	UInterchangeMaterialExpressionFactoryNode* CreateScalarParameterExpression(const UInterchangeShaderNode* ShaderNode, const FString& InputName, const FString& ParentUid);
+	UInterchangeMaterialExpressionFactoryNode* HandleLinearColorInput(const UInterchangeShaderNode* ShaderNode, const FString& InputName, const FString& ParentUid, bool bIsAParameter);
+	UInterchangeMaterialExpressionFactoryNode* CreateConstant3VectorExpression(const UInterchangeShaderNode* ShaderNode, const FString& InputName, const FString& ParentUid);
 	UInterchangeMaterialExpressionFactoryNode* CreateVectorParameterExpression(const UInterchangeShaderNode* ShaderNode, const FString& InputName, const FString& ParentUid);
+	UInterchangeMaterialExpressionFactoryNode* CreateStaticBooleanParameterExpression(const UInterchangeShaderNode* ShaderNode, const FString& InputName, const FString& ParentUid);
 	UInterchangeMaterialExpressionFactoryNode* CreateVector2ParameterExpression(const UInterchangeShaderNode* ShaderNode, const FString& InputName, const FString& ParentUid);
 	UInterchangeMaterialExpressionFactoryNode* CreateFunctionCallExpression(const UInterchangeShaderNode* ShaderNode, const FString& MaterialExpressionUid, UInterchangeBaseMaterialFactoryNode* MaterialFactoryNode);
 
@@ -197,8 +219,13 @@ private:
 	TVariant<FString, FLinearColor, float> VisitMultiplyNode(const UInterchangeShaderNode* ShaderNode) const;
 	TVariant<FString, FLinearColor, float> VisitOneMinusNode(const UInterchangeShaderNode* ShaderNode) const;
 	TVariant<FString, FLinearColor, float> VisitTextureSampleNode(const UInterchangeShaderNode* ShaderNode) const;
+	
 
+	FString GetTextureUidAttributeFromShaderNode(const UInterchangeShaderNode* ShaderNode, FName ParameterName, bool& OutIsAParameter) const;
+	FString CreateInputKey(const FString& InputName, bool bIsAParameter) const;
 private:
+	friend class UE::Interchange::Materials::HashUtils::FDuplicateMaterialHelper;
+
 	enum class EMaterialInputType : uint8
 	{
 		Unknown,
@@ -221,4 +248,6 @@ private:
 	TArray<FMaterialExpressionCreationContext> MaterialExpressionCreationContextStack;
 
 	using FParameterMaterialInputType = TTuple<FString, EMaterialInputType>;
+
+	const UInterchangeBaseNode * AttributeStorageNode = nullptr;
 };
