[Interchange - FBX] Degenerate polygon threshold value fine tuning.
[Interchange - FBX] Fixing Rigid Animation imports.
[Interchange - FBX] Ensuring Node Name uniqueness.

diff --git a/Engine/Plugins/Interchange/Runtime/Source/Import/Private/Animation/InterchangeAnimationTrackSetFactory.cpp b/Engine/Plugins/Interchange/Runtime/Source/Import/Private/Animation/InterchangeAnimationTrackSetFactory.cpp
--- a/Engine/Plugins/Interchange/Runtime/Source/Import/Private/Animation/InterchangeAnimationTrackSetFactory.cpp
+++ b/Engine/Plugins/Interchange/Runtime/Source/Import/Private/Animation/InterchangeAnimationTrackSetFactory.cpp
@@ -638,7 +638,14 @@ namespace UE::Interchange::Private
 				SceneValue.Value = CurveKey.Value;
 			}
 
-			Channel->Set(FrameNumbers, MovieSceneDoubleValues);
+			if (!MovieSceneDoubleValues.IsEmpty())
+			{
+				Channel->Set(FrameNumbers, MovieSceneDoubleValues);
+			}
+			else
+			{
+				Channel->RemoveDefault();
+			}
 		};
 
 		CopyToChannel(Channels[IndexOffset + 0], Curves[IndexOffset + 0]);
diff --git a/Engine/Plugins/Interchange/Runtime/Source/Import/Private/Fbx/InterchangeFbxTranslator.cpp b/Engine/Plugins/Interchange/Runtime/Source/Import/Private/Fbx/InterchangeFbxTranslator.cpp
--- a/Engine/Plugins/Interchange/Runtime/Source/Import/Private/Fbx/InterchangeFbxTranslator.cpp
+++ b/Engine/Plugins/Interchange/Runtime/Source/Import/Private/Fbx/InterchangeFbxTranslator.cpp
@@ -242,9 +242,11 @@ TFuture<TOptional<UE::Interchange::FMeshPayloadData>> UInterchangeFbxTranslator:
 		ResultParser.FromJson(JsonResult);
 		FString StaticMeshPayloadFilename = ResultParser.GetResultFilename();
 
-		if (!ensure(FPaths::FileExists(StaticMeshPayloadFilename)))
+		//Mesh payload file generation can fail due to invalid Mesh (for ep No Polygons/Only Degenerate Polygons)
+		if (!FPaths::FileExists(StaticMeshPayloadFilename))
 		{
-			// TODO log an error saying the payload file does not exist even if the get payload command succeeded
+			UE_LOG(LogInterchangeImport, Warning, TEXT("Expected mesh payload file does not exist for PayloadKey: %s"), *PayLoadKey.UniqueId);
+
 			Promise->SetValue(TOptional<UE::Interchange::FMeshPayloadData>());
 			return;
 		}
diff --git a/Engine/Plugins/Interchange/Runtime/Source/Parsers/Fbx/Private/FbxAPI.cpp b/Engine/Plugins/Interchange/Runtime/Source/Parsers/Fbx/Private/FbxAPI.cpp
--- a/Engine/Plugins/Interchange/Runtime/Source/Parsers/Fbx/Private/FbxAPI.cpp
+++ b/Engine/Plugins/Interchange/Runtime/Source/Parsers/Fbx/Private/FbxAPI.cpp
@@ -13,6 +13,7 @@
 #include "FbxScene.h"
 #include "InterchangeTextureNode.h"
 #include "Nodes/InterchangeBaseNodeContainer.h"
+#include "Misc/SecureHash.h"
 
 #define LOCTEXT_NAMESPACE "InterchangeFbxParser"
 
@@ -99,14 +100,11 @@ namespace UE
 
 				bool bStatus = SDKImporter->Import(SDKScene);
 
-				const FbxGlobalSettings& GlobalSettings = SDKScene->GetGlobalSettings();
-				FbxTime::EMode TimeMode = GlobalSettings.GetTimeMode();
-				//Set the original framerate from the current fbx file
-				float FbxFramerate = FbxTime::GetFrameRate(TimeMode);
-
 				//We always convert scene to UE axis and units
 				FFbxConvert::ConvertScene(SDKScene);
 
+				FrameRate = FbxTime::GetFrameRate(SDKScene->GetGlobalSettings().GetTimeMode());
+
 				return true;
 			}
 
@@ -202,7 +200,22 @@ namespace UE
 				}
 
 				//////////////////////////////////////////////////////////////////////////
-				// Scene conversion
+				// Esnure Node Name Validity (uniqueness)
+				TSet<FString> NodeNames;
+				for (int32 NodeIndex = 0; NodeIndex < SDKScene->GetNodeCount(); ++NodeIndex)
+				{
+					FbxNode* Node = SDKScene->GetNode(NodeIndex);
+
+					FString NodeName = Node->GetName();
+					FString NodeUniqueID = GetFbxHelper()->GetFbxNodeHierarchyName(Node);
+
+					if (NodeNames.Contains(NodeName))
+					{
+						NodeName += TEXT("_") + FMD5::HashAnsiString(*NodeUniqueID);
+						Node->SetName(TCHAR_TO_UTF8(*NodeName));
+					}
+					NodeNames.Add(NodeName);
+				}
 			}
 		} //ns Private
 	} //ns Interchange
diff --git a/Engine/Plugins/Interchange/Runtime/Source/Parsers/Fbx/Private/FbxAPI.h b/Engine/Plugins/Interchange/Runtime/Source/Parsers/Fbx/Private/FbxAPI.h
--- a/Engine/Plugins/Interchange/Runtime/Source/Parsers/Fbx/Private/FbxAPI.h
+++ b/Engine/Plugins/Interchange/Runtime/Source/Parsers/Fbx/Private/FbxAPI.h
@@ -78,6 +78,8 @@ namespace UE
 
 				FbxScene* GetSDKScene() { return SDKScene; }
 
+				double GetFrameRate() { return FrameRate; }
+
 				/**
 				 * Critical section to avoid getting multiple payload in same time.
 				 * The FBX evaluator use a cache mechanism for evaluating global transform that is not thread safe.
@@ -96,6 +98,9 @@ namespace UE
 				FString SourceFilename;
 				TMap<FString, TSharedPtr<FPayloadContextBase>> PayloadContexts;
 				TSharedPtr<FFbxHelper> FbxHelper;
+
+				//For PivotReset and Animation Conversion:
+				double FrameRate = 30.0;
 			};
 		}//ns Private
 	}//ns Interchange
diff --git a/Engine/Plugins/Interchange/Runtime/Source/Parsers/Fbx/Private/FbxAnimation.cpp b/Engine/Plugins/Interchange/Runtime/Source/Parsers/Fbx/Private/FbxAnimation.cpp
--- a/Engine/Plugins/Interchange/Runtime/Source/Parsers/Fbx/Private/FbxAnimation.cpp
+++ b/Engine/Plugins/Interchange/Runtime/Source/Parsers/Fbx/Private/FbxAnimation.cpp
@@ -26,7 +26,7 @@
 
 namespace UE::Interchange::Private
 {
-	bool ImportCurve(const FbxAnimCurve* SourceFloatCurves, const float Offset, const float ScaleValue, TArray<FInterchangeCurveKey>& DestinationFloatCurve)
+	bool ImportCurve(const FbxAnimCurve* SourceFloatCurves, const float ScaleValue, TArray<FInterchangeCurveKey>& DestinationFloatCurve)
 	{
 		if (!SourceFloatCurves)
 		{
@@ -41,7 +41,7 @@ namespace UE::Interchange::Private
 			FbxAnimCurveKey Key = SourceFloatCurves->KeyGet(KeyIndex);
 			FbxTime KeyTime = Key.GetTime();
 			const float KeyTimeValue = static_cast<float>(KeyTime.GetSecondDouble());
-			float Value = (Key.GetValue() * ScaleValue) + Offset;
+			float Value = Key.GetValue() * ScaleValue;
 			FInterchangeCurveKey& InterchangeCurveKey = DestinationFloatCurve.AddDefaulted_GetRef();
 			InterchangeCurveKey.Time = KeyTimeValue;
 			InterchangeCurveKey.Value = Value;
@@ -648,7 +648,7 @@ namespace UE::Interchange::Private
 					{
 						if (FbxAnimCurve* CurrentAnimCurve = FetchPayloadData.AnimCurves->GetCurve(ChannelIndex, CurveIndex))
 						{
-							ImportCurve(CurrentAnimCurve, 0.0f, 1.0f, InterchangeCurves.AddDefaulted_GetRef().Keys);
+							ImportCurve(CurrentAnimCurve, 1.0f, InterchangeCurves.AddDefaulted_GetRef().Keys);
 						}
 					}
 				}
@@ -670,7 +670,7 @@ namespace UE::Interchange::Private
 			//Fetch TArray<FInterchangeCurve> which are float curves with interpolation
 			TArray<FInterchangeCurve> InterchangeCurves;
 
-			auto ParseInterchangeCurveCurve = [&InterchangeCurves](FbxAnimCurveNode* AnimCurveNode, const char* pChannelName, float Offset = 0.0f, float Scale = 1.0f)
+			auto ParseInterchangeCurveCurve = [&InterchangeCurves](FbxAnimCurveNode* AnimCurveNode, const char* pChannelName, float Scale = 1.0f)
 			{
 				bool bInterchangeCurveAdded = false;
 				if (AnimCurveNode)
@@ -680,7 +680,7 @@ namespace UE::Interchange::Private
 					{
 						if (FbxAnimCurve* CurrentAnimCurve = AnimCurveNode->GetCurve(ChannelIndex))
 						{
-							ImportCurve(CurrentAnimCurve, Offset, Scale, InterchangeCurves.AddDefaulted_GetRef().Keys);
+							ImportCurve(CurrentAnimCurve, Scale, InterchangeCurves.AddDefaulted_GetRef().Keys);
 							bInterchangeCurveAdded = true;
 						}
 					}
@@ -689,32 +689,55 @@ namespace UE::Interchange::Private
 				if (!bInterchangeCurveAdded) InterchangeCurves.AddDefaulted_GetRef();
 			};
 
-			//As the hierarchical transforms are calculated based on 'EvaluateGlobalTransform'
-			//we calculate the difference between Evaluation and Local values
-			//As the FbxAnimCurve Values do not take into account Pre/Post rotations, offsets, pivots
-			//Note: This is not the utmost ideal scenario as Interchange does not have support for Pre/Post changes nor Pivots
-			//		However, this still prooved to provide the most precise outcome.
-			FbxDouble3 TranslationOffset;
-			FbxDouble3 RotationOffset;
+			auto ResetPivotsPrePostRotationsAndSetRotationOrder = [](FbxNode* Node, double FrameRate)
 			{
-				FbxAMatrix LocalFbxMatrix = FetchPayloadData.Node->EvaluateGlobalTransform();
-				if (fbxsdk::FbxNode* ParentNode = FetchPayloadData.Node->GetParent())
+				//This function now clears out all pivots, post and pre rotations and set's the
+				//RotationOrder to XYZ.
+				//Updated per the latest documentation
+				//https://help.autodesk.com/view/FBX/2017/ENU/?guid=__files_GUID_C35D98CB_5148_4B46_82D1_51077D8970EE_htm
+
+				if (!ensure(Node))
 				{
-					FbxAMatrix GlobalFbxParentMatrix = ParentNode->EvaluateGlobalTransform();
-					LocalFbxMatrix = GlobalFbxParentMatrix.Inverse() * LocalFbxMatrix;
+					return;
 				}
 
-				TranslationOffset = LocalFbxMatrix.GetT() - FetchPayloadData.Node->LclTranslation.Get();
-				RotationOffset = LocalFbxMatrix.GetR() - FetchPayloadData.Node->LclRotation.Get();
+				// Activate pivot converting 
+				Node->SetPivotState(FbxNode::eSourcePivot, FbxNode::ePivotActive);
+				Node->SetPivotState(FbxNode::eDestinationPivot, FbxNode::ePivotActive);
+
+				FbxVector4 Zero(0, 0, 0);
+
+				// We want to set all these to 0 and bake them into the transforms. 
+				Node->SetPostRotation(FbxNode::eDestinationPivot, Zero);
+				Node->SetPreRotation(FbxNode::eDestinationPivot, Zero);
+				Node->SetRotationOffset(FbxNode::eDestinationPivot, Zero);
+				Node->SetScalingOffset(FbxNode::eDestinationPivot, Zero);
+				Node->SetRotationPivot(FbxNode::eDestinationPivot, Zero);
+				Node->SetScalingPivot(FbxNode::eDestinationPivot, Zero);
+
+				Node->SetRotationOrder(FbxNode::eDestinationPivot, eEulerXYZ);
+				//When  we support other orders do this.
+				//EFbxRotationOrder ro;
+				//Node->GetRotationOrder(FbxNode::eSourcePivot, ro);
+				//Node->SetRotationOrder(FbxNode::eDestinationPivot, ro);
+
+				//Most DCC's don't have this but 3ds Max does
+				Node->SetGeometricTranslation(FbxNode::eDestinationPivot, Zero);
+				Node->SetGeometricRotation(FbxNode::eDestinationPivot, Zero);
+				Node->SetGeometricScaling(FbxNode::eDestinationPivot, Zero);
+				//NOTE THAT ConvertPivotAnimationRecursive did not seem to work when getting the local transform values!!!
+				Node->ConvertPivotAnimationRecursive(nullptr, FbxNode::eDestinationPivot, FrameRate);
 			};
 
-			ParseInterchangeCurveCurve(FetchPayloadData.TranlsationCurveNode, FBXSDK_CURVENODE_COMPONENT_X, TranslationOffset[0]);
-			ParseInterchangeCurveCurve(FetchPayloadData.TranlsationCurveNode, FBXSDK_CURVENODE_COMPONENT_Y, TranslationOffset[1], -1.0f); //FBX to UE conversion with scale -1.0
-			ParseInterchangeCurveCurve(FetchPayloadData.TranlsationCurveNode, FBXSDK_CURVENODE_COMPONENT_Z, TranslationOffset[2]);
+			ResetPivotsPrePostRotationsAndSetRotationOrder(FetchPayloadData.Node, FetchPayloadData.FrameRate);
+
+			ParseInterchangeCurveCurve(FetchPayloadData.TranlsationCurveNode, FBXSDK_CURVENODE_COMPONENT_X);
+			ParseInterchangeCurveCurve(FetchPayloadData.TranlsationCurveNode, FBXSDK_CURVENODE_COMPONENT_Y, -1.0f); //FBX to UE conversion with scale -1.0
+			ParseInterchangeCurveCurve(FetchPayloadData.TranlsationCurveNode, FBXSDK_CURVENODE_COMPONENT_Z);
 
-			ParseInterchangeCurveCurve(FetchPayloadData.RotationCurveNode, FBXSDK_CURVENODE_COMPONENT_X, RotationOffset[0]);
-			ParseInterchangeCurveCurve(FetchPayloadData.RotationCurveNode, FBXSDK_CURVENODE_COMPONENT_Y, RotationOffset[1], -1.0f); //FBX to UE conversion with scale -1.0
-			ParseInterchangeCurveCurve(FetchPayloadData.RotationCurveNode, FBXSDK_CURVENODE_COMPONENT_Z, RotationOffset[2], -1.0f); //FBX to UE conversion with scale -1.0
+			ParseInterchangeCurveCurve(FetchPayloadData.RotationCurveNode, FBXSDK_CURVENODE_COMPONENT_X);
+			ParseInterchangeCurveCurve(FetchPayloadData.RotationCurveNode, FBXSDK_CURVENODE_COMPONENT_Y, -1.0f); //FBX to UE conversion with scale -1.0
+			ParseInterchangeCurveCurve(FetchPayloadData.RotationCurveNode, FBXSDK_CURVENODE_COMPONENT_Z, -1.0f); //FBX to UE conversion with scale -1.0
 
 			ParseInterchangeCurveCurve(FetchPayloadData.ScaleCurveNode, FBXSDK_CURVENODE_COMPONENT_X);
 			ParseInterchangeCurveCurve(FetchPayloadData.ScaleCurveNode, FBXSDK_CURVENODE_COMPONENT_Y);
@@ -760,7 +783,7 @@ namespace UE::Interchange::Private
 		//Morph target curve in fbx are between 0 and 100, in Unreal we are between 0 and 1, so we must scale
 		//The curve with 0.01
 		constexpr float ScaleCurve = 0.01f;
-		ImportCurve(AnimCurve, 0.0f, ScaleCurve, InterchangeCurves.AddDefaulted_GetRef().Keys);
+		ImportCurve(AnimCurve, ScaleCurve, InterchangeCurves.AddDefaulted_GetRef().Keys);
 		{
 			FLargeMemoryWriter Ar;
 			Ar << InterchangeCurves;
@@ -843,6 +866,7 @@ namespace UE::Interchange::Private
 		{
 			TSharedPtr<FAnimationPayloadContextTransform> AnimPayload = MakeShared<FAnimationPayloadContextTransform>();
 			FAttributeNodeTransformFetchPayloadData FetchPayloadData;
+			FetchPayloadData.FrameRate = Parser.GetFrameRate();
 			FetchPayloadData.Node = Node;
 			FetchPayloadData.TranlsationCurveNode = TranslationCurveNode;
 			FetchPayloadData.RotationCurveNode = RotationCurveNode;
diff --git a/Engine/Plugins/Interchange/Runtime/Source/Parsers/Fbx/Private/FbxAnimation.h b/Engine/Plugins/Interchange/Runtime/Source/Parsers/Fbx/Private/FbxAnimation.h
--- a/Engine/Plugins/Interchange/Runtime/Source/Parsers/Fbx/Private/FbxAnimation.h
+++ b/Engine/Plugins/Interchange/Runtime/Source/Parsers/Fbx/Private/FbxAnimation.h
@@ -39,6 +39,7 @@ namespace UE::Interchange::Private
 	//Rigid Animation
 	struct FAttributeNodeTransformFetchPayloadData
 	{
+		double FrameRate = 30.0;
 		FbxNode* Node = nullptr;
 		FbxAnimCurveNode* TranlsationCurveNode = nullptr;
 		FbxAnimCurveNode* RotationCurveNode = nullptr;
diff --git a/Engine/Plugins/Interchange/Runtime/Source/Parsers/Fbx/Private/FbxMesh.cpp b/Engine/Plugins/Interchange/Runtime/Source/Parsers/Fbx/Private/FbxMesh.cpp
--- a/Engine/Plugins/Interchange/Runtime/Source/Parsers/Fbx/Private/FbxMesh.cpp
+++ b/Engine/Plugins/Interchange/Runtime/Source/Parsers/Fbx/Private/FbxMesh.cpp
@@ -694,7 +694,7 @@ bool FMeshDescriptionImporter::FillMeshDescriptionFromFbxMesh(FbxMesh* Mesh, con
 				int32 PolygonVertexCount = Mesh->GetPolygonSize(PolygonIndex);
 				//Verify if the polygon is degenerate, in this case do not add them
 				{
-					float ComparisonThreshold = SMALL_NUMBER;
+					float ComparisonThreshold = 1.e-12f; //SMALL_NUMBER (1.e-8f)  was not small enough, it produced missing/partial geometry on Drone from UE-192178
 					P.Reset();
 					P.AddUninitialized(PolygonVertexCount);
 					bool bAllCornerValid = true;
