FIX: Dereferencing NULL pointer 'ParentMaterialFactory'
- added an ensure
- added early return

diff --git a/Engine/Plugins/Interchange/Runtime/Source/Pipelines/Private/InterchangeGenericMaterialPipeline.cpp b/Engine/Plugins/Interchange/Runtime/Source/Pipelines/Private/InterchangeGenericMaterialPipeline.cpp
--- a/Engine/Plugins/Interchange/Runtime/Source/Pipelines/Private/InterchangeGenericMaterialPipeline.cpp
+++ b/Engine/Plugins/Interchange/Runtime/Source/Pipelines/Private/InterchangeGenericMaterialPipeline.cpp
@@ -3875,7 +3875,17 @@ namespace UE::Interchange::Materials::HashUtils
 
 	UInterchangeMaterialInstanceFactoryNode* FDuplicateMaterialHelper::CreateMaterialInstanceFactoryFromReference(const UInterchangeShaderGraphNode* ShaderGraphNode)
 	{
-		const UInterchangeBaseMaterialFactoryNode* ParentMaterialFactory = ParentMaterialFactoryMap[MaterialHash];
+		const UInterchangeBaseMaterialFactoryNode* ParentMaterialFactory = nullptr;
+		if (UInterchangeBaseMaterialFactoryNode** ParentMaterialFactoryEntry = ParentMaterialFactoryMap.Find(MaterialHash))
+		{
+			ensure(*ParentMaterialFactoryEntry);
+			ParentMaterialFactory = *ParentMaterialFactoryEntry;
+		}
+		
+		if (!ParentMaterialFactory)
+		{
+			return nullptr;
+		}
 
 		UInterchangeMaterialInstanceFactoryNode* MaterialInstanceFactoryNode =
 			Cast<UInterchangeMaterialInstanceFactoryNode>(GenericMaterialPipeline.CreateBaseMaterialFactoryNode(ShaderGraphNode, UInterchangeMaterialInstanceFactoryNode::StaticClass()));
