Fixed crash when  loading bogus texture in specific FBX file
Fixed missing assignment of ResultContainer on some texture translators.
Fixed issue with texture object and texture sampler nodes having the same unique ID in FBX

#rnx

diff --git a/Engine/Plugins/Interchange/Runtime/Source/Import/Private/Fbx/InterchangeFbxTranslator.cpp b/Engine/Plugins/Interchange/Runtime/Source/Import/Private/Fbx/InterchangeFbxTranslator.cpp
--- a/Engine/Plugins/Interchange/Runtime/Source/Import/Private/Fbx/InterchangeFbxTranslator.cpp
+++ b/Engine/Plugins/Interchange/Runtime/Source/Import/Private/Fbx/InterchangeFbxTranslator.cpp
@@ -201,6 +201,7 @@ TOptional<UE::Interchange::FImportImage> UInterchangeFbxTranslator::GetTexturePa
 	}
 
 	AlternateTexturePath = PayLoadKey;
+	SourceTranslator->SetResultsContainer(Results);
 
 	return TextureTranslator->GetTexturePayloadData(PayLoadKey, AlternateTexturePath);
 }
diff --git a/Engine/Plugins/Interchange/Runtime/Source/Import/Private/Gltf/InterchangeGltfTranslator.cpp b/Engine/Plugins/Interchange/Runtime/Source/Import/Private/Gltf/InterchangeGltfTranslator.cpp
--- a/Engine/Plugins/Interchange/Runtime/Source/Import/Private/Gltf/InterchangeGltfTranslator.cpp
+++ b/Engine/Plugins/Interchange/Runtime/Source/Import/Private/Gltf/InterchangeGltfTranslator.cpp
@@ -785,6 +785,7 @@ TOptional< UE::Interchange::FImportImage > UInterchangeGLTFTranslator::GetTextur
 		// Embedded texture -- try using ImageWrapper to decode it
 		TArray64<uint8> ImageData(GltfTexture.Source.Data, GltfTexture.Source.DataByteLength);
 		UInterchangeImageWrapperTranslator* ImageWrapperTranslator = NewObject<UInterchangeImageWrapperTranslator>(GetTransientPackage(), NAME_None);
+		ImageWrapperTranslator->SetResultsContainer(Results);
 		TexturePayloadData = ImageWrapperTranslator->GetTexturePayloadDataFromBuffer(ImageData);
 	}
 	else
diff --git a/Engine/Plugins/Interchange/Runtime/Source/Import/Private/Texture/TextureTranslatorUtilities.cpp b/Engine/Plugins/Interchange/Runtime/Source/Import/Private/Texture/TextureTranslatorUtilities.cpp
--- a/Engine/Plugins/Interchange/Runtime/Source/Import/Private/Texture/TextureTranslatorUtilities.cpp
+++ b/Engine/Plugins/Interchange/Runtime/Source/Import/Private/Texture/TextureTranslatorUtilities.cpp
@@ -87,20 +87,23 @@ bool UE::Interchange::FTextureTranslatorUtilities::GenericTextureLightProfileTra
 void UE::Interchange::FTextureTranslatorUtilities::LogError(const UInterchangeTranslatorBase& TextureTranslator, FText&& ErrorText)
 {
 	UInterchangeResultError_Generic* ErrorMessage = TextureTranslator.AddMessage<UInterchangeResultError_Generic>();
-	ErrorMessage->AssetType = TextureTranslator.GetClass();
-
-	if (TextureTranslator.GetSourceData())
-	{
-		const FString Filename = TextureTranslator.GetSourceData()->GetFilename();
-		ErrorMessage->SourceAssetName = Filename;
-		ErrorMessage->InterchangeKey = FPaths::GetBaseFilename(Filename);
-	}
-	else
+	if (ensure(ErrorMessage))
 	{
-		ErrorMessage->InterchangeKey = TEXT("Undefined");
-	}
+		ErrorMessage->AssetType = TextureTranslator.GetClass();
+
+		if (TextureTranslator.GetSourceData())
+		{
+			const FString Filename = TextureTranslator.GetSourceData()->GetFilename();
+			ErrorMessage->SourceAssetName = Filename;
+			ErrorMessage->InterchangeKey = FPaths::GetBaseFilename(Filename);
+		}
+		else
+		{
+			ErrorMessage->InterchangeKey = TEXT("Undefined");
+		}
 
-	ErrorMessage->Text = MoveTemp(ErrorText);
+		ErrorMessage->Text = MoveTemp(ErrorText);
+	}
 }
 
 bool UE::Interchange::FTextureTranslatorUtilities::IsTranslatorValid(const UInterchangeTranslatorBase& TextureTranslator, const TCHAR* Format)
diff --git a/Engine/Plugins/Interchange/Runtime/Source/Parsers/Fbx/Private/FbxMaterial.cpp b/Engine/Plugins/Interchange/Runtime/Source/Parsers/Fbx/Private/FbxMaterial.cpp
--- a/Engine/Plugins/Interchange/Runtime/Source/Parsers/Fbx/Private/FbxMaterial.cpp
+++ b/Engine/Plugins/Interchange/Runtime/Source/Parsers/Fbx/Private/FbxMaterial.cpp
@@ -55,7 +55,7 @@ namespace UE
 				return NewTextureNode;
 			}
 
-			UInterchangeShaderNode* FFbxMaterial::CreateTextureSampler(FbxFileTexture* FbxTexture, UInterchangeBaseNodeContainer& NodeContainer, const FString& ShaderUniqueID)
+			const UInterchangeShaderNode* FFbxMaterial::CreateTextureSampler(FbxFileTexture* FbxTexture, UInterchangeBaseNodeContainer& NodeContainer, const FString& ShaderUniqueID)
 			{
 				using namespace Materials::Standard::Nodes;
 
@@ -66,8 +66,17 @@ namespace UE
 
 				const FString TextureFilename = FbxTexture ? UTF8_TO_TCHAR(FbxTexture->GetFileName()) : TEXT("");
 				const FString TextureName = FPaths::GetBaseFilename(TextureFilename);
+				const FString NodeName = TEXT("Sampler_") + TextureName;
 
-				UInterchangeShaderNode* TextureSampleShader = UInterchangeShaderNode::Create(&NodeContainer, TextureName, ShaderUniqueID);
+				// Return already created node if applicable.
+				const FString SamplerNodeUid = UInterchangeShaderNode::MakeNodeUid(NodeName, ShaderUniqueID);
+				if (const UInterchangeShaderNode* SamplerNode = Cast<UInterchangeShaderNode>(NodeContainer.GetNode(SamplerNodeUid)))
+				{
+					return SamplerNode;
+				}
+
+				UInterchangeShaderNode* TextureSampleShader = UInterchangeShaderNode::Create(&NodeContainer, NodeName, ShaderUniqueID);
+				TextureSampleShader->SetDisplayLabel(TextureName);
 				TextureSampleShader->SetCustomShaderType(TextureSample::Name.ToString());
 
 				// Return incomplete texture sampler if texture file does not exist
@@ -81,12 +90,7 @@ namespace UE
 					return TextureSampleShader;
 				}
 
-				const UInterchangeTexture2DNode* TextureNode = Cast<const UInterchangeTexture2DNode>(NodeContainer.GetNode(UInterchangeTextureNode::MakeNodeUid(TextureName)));
-
-				if (!TextureNode)
-				{
-					TextureNode = CreateTexture2DNode(NodeContainer, TextureFilename);
-				}
+				const UInterchangeTexture2DNode* TextureNode = CreateTexture2DNode(NodeContainer, TextureFilename);
 
 				TextureSampleShader->AddStringAttribute(UInterchangeShaderPortsAPI::MakeInputValueKey(TextureSample::Inputs::Texture.ToString()), TextureNode->GetUniqueID());
 
@@ -182,7 +186,7 @@ namespace UE
 
 				// Handles max one texture per property.
 				FbxFileTexture* FbxTexture = Property.GetSrcObject<FbxFileTexture>(0);
-				if (UInterchangeShaderNode* TextureSampleShader = CreateTextureSampler(FbxTexture, NodeContainer, ShaderGraphNode->GetUniqueID()))
+				if (const UInterchangeShaderNode* TextureSampleShader = CreateTextureSampler(FbxTexture, NodeContainer, ShaderGraphNode->GetUniqueID()))
 				{
 					UInterchangeShaderPortsAPI::ConnectDefaultOuputToInput(NodeToConnectTo, InputToConnectTo, TextureSampleShader->GetUniqueID());
 				}
@@ -215,7 +219,7 @@ namespace UE
 				if (MaterialProperty.GetSrcObjectCount<FBXSDK_NAMESPACE::FbxTexture>() > 0)
 				{
 					FbxFileTexture* FbxTexture = MaterialProperty.GetSrcObject<FbxFileTexture>(0);
-					if (UInterchangeShaderNode* TextureSampleShader = CreateTextureSampler(FbxTexture, NodeContainer, ShaderGraphNode->GetUniqueID()))
+					if (const UInterchangeShaderNode* TextureSampleShader = CreateTextureSampler(FbxTexture, NodeContainer, ShaderGraphNode->GetUniqueID()))
 					{
 						FString MultiplyNodeName = Phong::Parameters::Shininess.ToString() + TEXT("_Multiply");
 						UInterchangeShaderNode* MultiplyNode = UInterchangeShaderNode::Create(&NodeContainer, MultiplyNodeName, ShaderGraphNode->GetUniqueID());
diff --git a/Engine/Plugins/Interchange/Runtime/Source/Parsers/Fbx/Private/FbxMaterial.h b/Engine/Plugins/Interchange/Runtime/Source/Parsers/Fbx/Private/FbxMaterial.h
--- a/Engine/Plugins/Interchange/Runtime/Source/Parsers/Fbx/Private/FbxMaterial.h
+++ b/Engine/Plugins/Interchange/Runtime/Source/Parsers/Fbx/Private/FbxMaterial.h
@@ -60,7 +60,7 @@ namespace UE
 			private:
 				const UInterchangeShaderGraphNode* AddShaderGraphNode(FbxSurfaceMaterial* SurfaceMaterial, UInterchangeBaseNodeContainer& NodeContainer);
 				void ConvertShininessToShaderNode(FbxSurfaceMaterial& SurfaceMaterial, UInterchangeBaseNodeContainer& NodeContainer, UInterchangeShaderGraphNode* ShaderGraphNode);
-				UInterchangeShaderNode* CreateTextureSampler(FbxFileTexture* FbxTexture, UInterchangeBaseNodeContainer& NodeContainer, const FString& ShaderUniqueID);
+				const UInterchangeShaderNode* CreateTextureSampler(FbxFileTexture* FbxTexture, UInterchangeBaseNodeContainer& NodeContainer, const FString& ShaderUniqueID);
 
 				FFbxParser& Parser;
 			};
