[Interchange - glTF] Export/Import: Import GLTF exported with UE seems to be far from origin.

diff --git a/Engine/Plugins/Interchange/Runtime/Source/Parsers/GLTFCore/Private/GLTFMeshFactory.cpp b/Engine/Plugins/Interchange/Runtime/Source/Parsers/GLTFCore/Private/GLTFMeshFactory.cpp
--- a/Engine/Plugins/Interchange/Runtime/Source/Parsers/GLTFCore/Private/GLTFMeshFactory.cpp
+++ b/Engine/Plugins/Interchange/Runtime/Source/Parsers/GLTFCore/Private/GLTFMeshFactory.cpp
@@ -346,24 +346,32 @@ namespace GLTF
 
 			TArray<FVector3f>& Positions = GetVectorBuffer(PositionBufferIndex);
 			Primitive.GetPositions(Positions);
+			
+			UpdateVector3fFromMorphTargets(Positions, Primitive, MorphTargetWeights, Vector3fMorphTargetAttributes::POSITIONS);
+
+			FMeshFactory::FIndexVertexIdMap& PositionIndexToVertexId = PositionIndexToVertexIdPerPrim[Index];
+			PositionIndexToVertexId.Empty(Positions.Num());
+
 			if (TotalMatrix != FMatrix::Identity)
 			{
-				for (FVector3f& Position : Positions)
+				for (int32 PositionIndex = 0; PositionIndex < Positions.Num(); ++PositionIndex)
 				{
-					const FVector TransformedPosition = TotalMatrix.TransformPosition(FVector(Position));
-					Position = static_cast<FVector3f>(TransformedPosition);
+					const FVertexID& VertexID = MeshDescription->CreateVertex();
+					FVector3f FinalPosition = Positions[PositionIndex] * ImportUniformScale;
+					const FVector TransformedPosition = TotalMatrix.TransformPosition(FVector(FinalPosition));
+					FinalPosition = static_cast<FVector3f>(TransformedPosition);
+					VertexPositions[VertexID] = FinalPosition;
+					PositionIndexToVertexId.Add(PositionIndex, VertexID);
 				}
 			}
-
-			UpdateVector3fFromMorphTargets(Positions, Primitive, MorphTargetWeights, Vector3fMorphTargetAttributes::POSITIONS);
-
-			FMeshFactory::FIndexVertexIdMap& PositionIndexToVertexId = PositionIndexToVertexIdPerPrim[Index];
-			PositionIndexToVertexId.Empty(Positions.Num());
-			for (int32 PositionIndex = 0; PositionIndex < Positions.Num(); ++PositionIndex)
+			else
 			{
-				const FVertexID& VertexID = MeshDescription->CreateVertex();
-				VertexPositions[VertexID] = Positions[PositionIndex] * ImportUniformScale;
-				PositionIndexToVertexId.Add(PositionIndex, VertexID);
+				for (int32 PositionIndex = 0; PositionIndex < Positions.Num(); ++PositionIndex)
+				{
+					const FVertexID& VertexID = MeshDescription->CreateVertex();
+					VertexPositions[VertexID] = Positions[PositionIndex] * ImportUniformScale;
+					PositionIndexToVertexId.Add(PositionIndex, VertexID);
+				}
 			}
 		}
 
@@ -472,6 +480,8 @@ namespace GLTF
 			TArray<FVector3f>& ReindexBuffer = GetVectorBuffer(ReindexBufferIndex);
 			Primitive.GetNormals(Normals);
 
+			UpdateVector3fFromMorphTargets(Normals, Primitive, MorphTargetWeights, Vector3fMorphTargetAttributes::NORMALS);
+
 			if (TotalMatrixForNormal != FMatrix::Identity)
 			{
 				for (FVector3f& Normal : Normals)
@@ -480,8 +490,6 @@ namespace GLTF
 				}
 			}
 
-			UpdateVector3fFromMorphTargets(Normals, Primitive, MorphTargetWeights, Vector3fMorphTargetAttributes::NORMALS);
-
 			ReIndex(Normals, Indices, ReindexBuffer);
 			Swap(Normals, ReindexBuffer);
 		}
@@ -490,6 +498,9 @@ namespace GLTF
 			TArray<FVector3f>& Positions = GetVectorBuffer(PositionBufferIndex);
 			Primitive.GetPositions(Positions);
 
+			//update positions with morph targets before generating flat normals:
+			UpdateVector3fFromMorphTargets(Positions, Primitive, MorphTargetWeights, Vector3fMorphTargetAttributes::POSITIONS);
+
 			if (TotalMatrix != FMatrix::Identity)
 			{
 				for (FVector3f& Position : Positions)
@@ -498,9 +509,6 @@ namespace GLTF
 				}
 			}
 
-			//update positions with morph targets before generating flat normals:
-			UpdateVector3fFromMorphTargets(Positions, Primitive, MorphTargetWeights, Vector3fMorphTargetAttributes::POSITIONS);
-
 			GenerateFlatNormals(Positions, Indices, Normals);
 			bIgnoreTangents = true;
 		}
@@ -511,6 +519,8 @@ namespace GLTF
 			TArray<FVector4f>& ReindexBuffer = GetVector4dBuffer(Reindex4dBufferIndex);
 			Primitive.GetTangents(Tangents);
 
+			UpdateVector4dFromMorphTargets(Tangents, Primitive, MorphTargetWeights, Vector4dMorphTargetAttributes::TANGENTS);
+
 			if (TotalMatrixForNormal != FMatrix::Identity)
 			{
 				for (FVector4f& Tangent : Tangents)
@@ -519,8 +529,6 @@ namespace GLTF
 				}
 			}
 
-			UpdateVector4dFromMorphTargets(Tangents, Primitive, MorphTargetWeights, Vector4dMorphTargetAttributes::TANGENTS);
-
 			ReIndex(Tangents, Indices, ReindexBuffer);
 			Swap(Tangents, ReindexBuffer);
 		}
