FIX: Dereferencing NULL pointer 'ParentMaterialFactory'
- ensure at correct position

diff --git a/Engine/Plugins/Interchange/Runtime/Source/Pipelines/Private/InterchangeGenericMaterialPipeline.cpp b/Engine/Plugins/Interchange/Runtime/Source/Pipelines/Private/InterchangeGenericMaterialPipeline.cpp
--- a/Engine/Plugins/Interchange/Runtime/Source/Pipelines/Private/InterchangeGenericMaterialPipeline.cpp
+++ b/Engine/Plugins/Interchange/Runtime/Source/Pipelines/Private/InterchangeGenericMaterialPipeline.cpp
@@ -3878,10 +3878,11 @@ namespace UE::Interchange::Materials::HashUtils
 		const UInterchangeBaseMaterialFactoryNode* ParentMaterialFactory = nullptr;
 		if (UInterchangeBaseMaterialFactoryNode** ParentMaterialFactoryEntry = ParentMaterialFactoryMap.Find(MaterialHash))
 		{
-			ensure(*ParentMaterialFactoryEntry);
 			ParentMaterialFactory = *ParentMaterialFactoryEntry;
 		}
 		
+		ensure(ParentMaterialFactory);
+
 		if (!ParentMaterialFactory)
 		{
 			return nullptr;
