Adapted the OBJ translator to work with the new Reference Material Mechanism in Interchange

diff --git a/Engine/Plugins/Interchange/Runtime/Source/Import/Private/Material/InterchangeMaterialNodesHelper.cpp b/Engine/Plugins/Interchange/Runtime/Source/Import/Private/Material/InterchangeMaterialNodesHelper.cpp
new file mode 100644
--- /dev/null
+++ b/Engine/Plugins/Interchange/Runtime/Source/Import/Private/Material/InterchangeMaterialNodesHelper.cpp
@@ -0,0 +1,46 @@
+// Copyright Epic Games, Inc. All Rights Reserved. 
+#include "Material/InterchangeMaterialNodesHelper.h"
+
+#include "Nodes/InterchangeBaseNode.h"
+#include "Nodes/InterchangeBaseNodeContainer.h"
+#include "Nodes/InterchangeUserDefinedAttribute.h"
+#include "InterchangeShaderGraphNode.h"
+#include "InterchangeMaterialDefinitions.h"
+
+namespace UE::Interchange::Materials::Private
+{
+	bool FMaterialNodesHelper::SetupScalarParameter(UInterchangeBaseNodeContainer& BaseNodeContainer, UInterchangeShaderGraphNode* ShaderGraphNode, const FString& InputName, const float& AttributeValue)
+	{
+		using namespace UE::Interchange::Materials::Standard::Nodes;
+
+		UInterchangeShaderNode* ExpressionNode = UInterchangeShaderNode::Create(&BaseNodeContainer, InputName, ShaderGraphNode->GetUniqueID());
+		if (!ensure(ExpressionNode))
+		{
+			// TODO: Log error
+			return false;
+		}
+
+		ExpressionNode->AddFloatInput(ScalarParameter::Attributes::DefaultValue.ToString(), AttributeValue,/*bIsAParameter = */ true);
+		ExpressionNode->SetCustomShaderType(ScalarParameter::Name.ToString());
+
+		UInterchangeShaderPortsAPI::ConnectDefaultOuputToInput(ShaderGraphNode, InputName, ExpressionNode->GetUniqueID());
+		return true;
+	}
+
+	bool FMaterialNodesHelper::SetupVectorParameter(UInterchangeBaseNodeContainer& BaseNodeContainer, UInterchangeShaderGraphNode* ShaderGraphNode, const FString& InputName, const FLinearColor& AttributeValue)
+	{
+		using namespace UE::Interchange::Materials::Standard::Nodes;
+
+		UInterchangeShaderNode* ExpressionNode = UInterchangeShaderNode::Create(&BaseNodeContainer, InputName, ShaderGraphNode->GetUniqueID());
+		if (!ensure(ExpressionNode))
+		{
+			// TODO: Log error
+			return false;
+		}
+
+		ExpressionNode->AddLinearColorInput(VectorParameter::Attributes::DefaultValue.ToString(), AttributeValue,/*bIsAParameter = */ true);
+		ExpressionNode->SetCustomShaderType(VectorParameter::Name.ToString());
+		UInterchangeShaderPortsAPI::ConnectDefaultOuputToInput(ShaderGraphNode, InputName, ExpressionNode->GetUniqueID());
+		return true;
+	}
+}
\ No newline at end of file
diff --git a/Engine/Plugins/Interchange/Runtime/Source/Import/Private/Mesh/InterchangeOBJTranslator.cpp b/Engine/Plugins/Interchange/Runtime/Source/Import/Private/Mesh/InterchangeOBJTranslator.cpp
--- a/Engine/Plugins/Interchange/Runtime/Source/Import/Private/Mesh/InterchangeOBJTranslator.cpp
+++ b/Engine/Plugins/Interchange/Runtime/Source/Import/Private/Mesh/InterchangeOBJTranslator.cpp
@@ -25,7 +25,7 @@
 #include "StaticMeshOperations.h"
 #include "UObject/GCObjectScopeGuard.h"
 #include "UVMapSettings.h"
-
+#include "Material/InterchangeMaterialNodesHelper.h"
 
 #include UE_INLINE_GENERATED_CPP_BY_NAME(InterchangeOBJTranslator)
 
@@ -1090,27 +1090,28 @@ namespace ObjTranslatorUtils
 
 			if (const UInterchangeTexture2DNode* TextureNode = CreateTexture2DNode(BaseNodeContainer, TextureParameter.Path, TextureName))
 			{
-				UInterchangeShaderNode* TextureSampleShader = UInterchangeShaderNode::Create(&BaseNodeContainer, TextureName, ShaderGraphNode->GetUniqueID());
-				TextureSampleShader->SetCustomShaderType(TextureSample::Name.ToString());
-				TextureSampleShader->AddStringAttribute(UInterchangeShaderPortsAPI::MakeInputValueKey(TextureSample::Inputs::Texture.ToString()), TextureNode->GetUniqueID());
-				return TextureSampleShader;
+				UInterchangeShaderNode* TextureSampleShaderNode = UInterchangeShaderNode::Create(&BaseNodeContainer, TextureName, ShaderGraphNode->GetUniqueID());
+				TextureSampleShaderNode->SetCustomShaderType(TextureSample::Name.ToString());
+				TextureSampleShaderNode->AddStringInput(TextureSample::Inputs::Texture.ToString(), TextureNode->GetUniqueID(), /*bIsAParameter =*/ true);
+				return TextureSampleShaderNode;
 			}
 		}
 		return nullptr;
 	}
 
 	// multiply texture by color, handling cases when any of the arguments is missing
-	static bool AddTexturedColoredInput(UInterchangeBaseNodeContainer& BaseNodeContainer, UInterchangeShaderGraphNode* ShaderGraphNode, const FString& InputType, const FVector3f& Color, const FObjData::FTextureParameter& TexturePath) 
+	static bool AddTexturedColoredInput(UInterchangeBaseNodeContainer& BaseNodeContainer, UInterchangeShaderGraphNode* ShaderGraphNode, const FString& InputName, const FVector3f& Color, const FObjData::FTextureParameter& TexturePath)
 	{
-		UInterchangeShaderNode* TextureSampleShader = CreateMaterialTextureSampleNode(BaseNodeContainer, ShaderGraphNode, TexturePath);
+		UInterchangeShaderNode* TextureSampleShaderNode = CreateMaterialTextureSampleNode(BaseNodeContainer, ShaderGraphNode, TexturePath);
 
-		if (TextureSampleShader)
+		if (TextureSampleShaderNode)
 		{
-			UInterchangeShaderPortsAPI::ConnectDefaultOuputToInput(ShaderGraphNode, InputType, TextureSampleShader->GetUniqueID());
+			TextureSampleShaderNode->SetDisplayLabel(InputName);
+			UInterchangeShaderPortsAPI::ConnectDefaultOuputToInput(ShaderGraphNode, InputName, TextureSampleShaderNode->GetUniqueID());
 		}
 		else if (IsColorInitialized(Color))
 		{
-			ShaderGraphNode->AddLinearColorAttribute(UInterchangeShaderPortsAPI::MakeInputValueKey(InputType), FLinearColor(Color));
+			UE::Interchange::Materials::Private::FMaterialNodesHelper::SetupVectorParameter(BaseNodeContainer, ShaderGraphNode, InputName, FLinearColor(Color));
 		}
 		else
 		{
@@ -1119,9 +1120,9 @@ namespace ObjTranslatorUtils
 		return true;
 	}
 
-	static UInterchangeShaderNode* MakeTexturedWeighted(UInterchangeBaseNodeContainer& BaseNodeContainer, UInterchangeShaderGraphNode* ShaderGraphNode, FString MapName, const float& Weight, UInterchangeShaderNode* TextureSampleShader) 
+	static UInterchangeShaderNode* MakeTexturedWeighted(UInterchangeBaseNodeContainer& BaseNodeContainer, UInterchangeShaderGraphNode* ShaderGraphNode, FString MapName, const float& Weight, UInterchangeShaderNode* TextureSampleShaderNode) 
 	{
-		if (!TextureSampleShader)
+		if (!TextureSampleShaderNode)
 		{
 			return nullptr;
 		}
@@ -1134,22 +1135,23 @@ namespace ObjTranslatorUtils
 			UInterchangeShaderNode* MultiplierNode = UInterchangeShaderNode::Create(&BaseNodeContainer, MultiplierNodeName, ShaderGraphNode->GetUniqueID());
 			MultiplierNode->SetCustomShaderType(Multiply::Name.ToString());
 
-			UInterchangeShaderPortsAPI::ConnectDefaultOuputToInput(MultiplierNode, Multiply::Inputs::A.ToString(), TextureSampleShader->GetUniqueID());
+			UInterchangeShaderPortsAPI::ConnectDefaultOuputToInput(MultiplierNode, Multiply::Inputs::A.ToString(), TextureSampleShaderNode->GetUniqueID());
 			MultiplierNode->AddFloatAttribute(UInterchangeShaderPortsAPI::MakeInputValueKey( Multiply::Inputs::B.ToString() ), Weight);
 			return MultiplierNode;
 		}
-		return TextureSampleShader;
+		return TextureSampleShaderNode;
 	}
 
-	static bool AddTexturedWeightedInput(UInterchangeBaseNodeContainer& BaseNodeContainer, UInterchangeShaderGraphNode* ShaderGraphNode, const FString& InputType, const float& Weight, UInterchangeShaderNode* TextureSampleShader) 
+	static bool AddTexturedWeightedInput(UInterchangeBaseNodeContainer& BaseNodeContainer, UInterchangeShaderGraphNode* ShaderGraphNode, const FString& InputName, const float& Weight, UInterchangeShaderNode* TextureSampleShader) 
 	{
-		if (UInterchangeShaderNode* TextureSampleWeighted = MakeTexturedWeighted(BaseNodeContainer, ShaderGraphNode, InputType, Weight, TextureSampleShader))
+		if (UInterchangeShaderNode* TextureSampleWeighted = MakeTexturedWeighted(BaseNodeContainer, ShaderGraphNode, InputName, Weight, TextureSampleShader))
 		{
-			UInterchangeShaderPortsAPI::ConnectDefaultOuputToInput(ShaderGraphNode, InputType, TextureSampleWeighted->GetUniqueID());
+			TextureSampleWeighted->SetDisplayLabel(InputName);
+			UInterchangeShaderPortsAPI::ConnectDefaultOuputToInput(ShaderGraphNode, InputName, TextureSampleWeighted->GetUniqueID());
 		}
 		else if (IsScalarInitialized(Weight))
 		{
-			ShaderGraphNode->AddFloatAttribute(UInterchangeShaderPortsAPI::MakeInputValueKey(InputType), Weight);
+			UE::Interchange::Materials::Private::FMaterialNodesHelper::SetupScalarParameter(BaseNodeContainer, ShaderGraphNode, InputName, Weight);
 		}
 		else
 		{
@@ -1161,9 +1163,9 @@ namespace ObjTranslatorUtils
 	// multiply texture by scalar, handling cases when any of the arguments is missing
 	static bool AddTexturedWeightedInput(UInterchangeBaseNodeContainer& BaseNodeContainer, UInterchangeShaderGraphNode* ShaderGraphNode, const FString& InputType, const float& Weight, const FObjData::FTextureParameter& TexturePath) 
 	{
-		UInterchangeShaderNode* TextureSampleShader = CreateMaterialTextureSampleNode(BaseNodeContainer, ShaderGraphNode, TexturePath);
-
-		return AddTexturedWeightedInput(BaseNodeContainer, ShaderGraphNode, InputType, Weight, TextureSampleShader);
+		UInterchangeShaderNode* TextureSampleShaderNode = CreateMaterialTextureSampleNode(BaseNodeContainer, ShaderGraphNode, TexturePath);
+		
+		return AddTexturedWeightedInput(BaseNodeContainer, ShaderGraphNode, InputType, Weight, TextureSampleShaderNode);
 	}
 
 	static bool HandlePbrAttributes(UInterchangeBaseNodeContainer& BaseNodeContainer, const FObjData::FMaterialData& MaterialData, UInterchangeShaderGraphNode* ShaderGraphNode)
@@ -1187,17 +1189,17 @@ namespace ObjTranslatorUtils
 
 		if (IsScalarInitialized(MaterialData.ClearCoatThickness) && !FMath::IsNearlyZero(MaterialData.ClearCoatThickness))
 		{
-			ShaderGraphNode->AddFloatAttribute(UInterchangeShaderPortsAPI::MakeInputValueKey(Materials::ClearCoat::Parameters::ClearCoat.ToString()), MaterialData.ClearCoatThickness);
+			UE::Interchange::Materials::Private::FMaterialNodesHelper::SetupScalarParameter(BaseNodeContainer, ShaderGraphNode, Materials::ClearCoat::Parameters::ClearCoat.ToString(), MaterialData.ClearCoatThickness);
 
 			if (IsScalarInitialized(MaterialData.ClearCoatRoughness))
 			{
-				ShaderGraphNode->AddFloatAttribute(UInterchangeShaderPortsAPI::MakeInputValueKey(Materials::ClearCoat::Parameters::ClearCoatRoughness.ToString()), MaterialData.ClearCoatRoughness);
+				UE::Interchange::Materials::Private::FMaterialNodesHelper::SetupScalarParameter(BaseNodeContainer, ShaderGraphNode, Materials::ClearCoat::Parameters::ClearCoatRoughness.ToString(), MaterialData.ClearCoatRoughness);
 			}
 		}
 
 		if (IsScalarInitialized(MaterialData.Anisotropy)  && !FMath::IsNearlyZero(MaterialData.Anisotropy))
 		{
-			ShaderGraphNode->AddFloatAttribute(UInterchangeShaderPortsAPI::MakeInputValueKey(Materials::PBRMR::Parameters::Anisotropy.ToString()), MaterialData.Anisotropy);
+			UE::Interchange::Materials::Private::FMaterialNodesHelper::SetupScalarParameter(BaseNodeContainer, ShaderGraphNode, Materials::PBRMR::Parameters::Anisotropy.ToString(), MaterialData.Anisotropy);
 		}
 
 		return true;
@@ -1427,25 +1429,25 @@ bool UInterchangeOBJTranslator::Translate(UInterchangeBaseNodeContainer& BaseNod
 
 						// OBJ Specular Exponent('Ns') is in range 0 to 1000(for sharpest highlight) and Interchange Shininess is up to 100
 						const float SpecularExponentToShininessFactor = 0.1f;
+						const FString ShininessParameterName = Materials::Phong::Parameters::Shininess.ToString();
 						if (UInterchangeShaderNode* SpecularExponentTextureSampleShader = CreateMaterialTextureSampleNode(BaseNodeContainer, ShaderGraphNode, MaterialData.SpecularExponentTexture))
 						{
 							using namespace UE::Interchange::Materials::Standard::Nodes;
 
 							float SpecularExponentScale = (IsScalarInitialized(MaterialData.SpecularExponent) ? MaterialData.SpecularExponent : 1.0f) * SpecularExponentToShininessFactor;
 
-							FString MapName = Materials::Phong::Parameters::Shininess.ToString();
-							const FString MultiplierNodeName = MapName + TEXT("_Multiply");
+							const FString MultiplierNodeName = ShininessParameterName + TEXT("_Multiply");
 							UInterchangeShaderNode* MultiplierNode = UInterchangeShaderNode::Create(&BaseNodeContainer, MultiplierNodeName, ShaderGraphNode->GetUniqueID());
 							MultiplierNode->SetCustomShaderType(Multiply::Name.ToString());
 
 							UInterchangeShaderPortsAPI::ConnectDefaultOuputToInput(MultiplierNode, Multiply::Inputs::A.ToString(), SpecularExponentTextureSampleShader->GetUniqueID());
 							MultiplierNode->AddFloatAttribute(UInterchangeShaderPortsAPI::MakeInputValueKey( Multiply::Inputs::B.ToString() ), SpecularExponentScale);
 
-							UInterchangeShaderPortsAPI::ConnectDefaultOuputToInput(ShaderGraphNode, Materials::Phong::Parameters::Shininess.ToString(), MultiplierNode->GetUniqueID());
+							UInterchangeShaderPortsAPI::ConnectDefaultOuputToInput(ShaderGraphNode, ShininessParameterName, MultiplierNode->GetUniqueID());
 						}
 						else if (IsScalarInitialized(MaterialData.SpecularExponent))
 						{
-							ShaderGraphNode->AddFloatAttribute(UInterchangeShaderPortsAPI::MakeInputValueKey(Materials::Phong::Parameters::Shininess.ToString()), MaterialData.SpecularExponent);
+							UE::Interchange::Materials::Private::FMaterialNodesHelper::SetupScalarParameter(BaseNodeContainer, ShaderGraphNode, ShininessParameterName, MaterialData.SpecularExponent);
 						}
 					}
 
@@ -1453,7 +1455,7 @@ bool UInterchangeOBJTranslator::Translate(UInterchangeBaseNodeContainer& BaseNod
 					if (bReflectionMetal)
 					{
 						// Set ambient to white as it drives the metallic parameter of UE Pbr material.
-						ShaderGraphNode->AddLinearColorAttribute(UInterchangeShaderPortsAPI::MakeInputValueKey(AmbientColorPropertyName), FLinearColor::White);
+						UE::Interchange::Materials::Private::FMaterialNodesHelper::SetupVectorParameter(BaseNodeContainer, ShaderGraphNode, AmbientColorPropertyName, FLinearColor::White);
 					}
 					else
 					{
@@ -1467,28 +1469,27 @@ bool UInterchangeOBJTranslator::Translate(UInterchangeBaseNodeContainer& BaseNod
 
 			// Opacity
 			{
-				FString OpacityInputType = Materials::Common::Parameters::Opacity.ToString();
+				const FString OpacityInputName = Materials::Common::Parameters::Opacity.ToString();
 
 				// Transparency texture replaces Transparency scalar
 				if (UInterchangeShaderNode* TransparencySample = CreateMaterialTextureSampleNode(BaseNodeContainer, ShaderGraphNode, MaterialData.TransparencyTexture))
 				{
 					// Invert transparency to get Opacity
-					FString MapName = OpacityInputType;
-					const FString OneMinusNodeName = MapName + TEXT("OneMinus");
+					const FString OneMinusNodeName = OpacityInputName + TEXT("OneMinus");
 					UInterchangeShaderNode* OneMinusNode = UInterchangeShaderNode::Create( &BaseNodeContainer, OneMinusNodeName, ShaderGraphNode->GetUniqueID() );
 					OneMinusNode->SetCustomShaderType(Materials::Standard::Nodes::OneMinus::Name.ToString());
 					UInterchangeShaderPortsAPI::ConnectDefaultOuputToInput(OneMinusNode, Materials::Standard::Nodes::OneMinus::Inputs::Input.ToString(), TransparencySample->GetUniqueID());
-					UInterchangeShaderPortsAPI::ConnectDefaultOuputToInput(ShaderGraphNode, OpacityInputType, OneMinusNode->GetUniqueID());
+					UInterchangeShaderPortsAPI::ConnectDefaultOuputToInput(ShaderGraphNode, OpacityInputName, OneMinusNode->GetUniqueID());
 				}
 				// Opacity texture is weighted by scalar opacity("d")
-				else if (!AddTexturedWeightedInput(BaseNodeContainer, ShaderGraphNode, OpacityInputType, MaterialData.Opacity, MaterialData.OpacityTexture))
+				else if (!AddTexturedWeightedInput(BaseNodeContainer, ShaderGraphNode, OpacityInputName, MaterialData.Opacity, MaterialData.OpacityTexture))
 				{
 					// When no textures are present compute opacity value depending on which is defined - opacity or transparency
 					// OBJ has two ways to define Opacity/Transparency - "d"(for opacity) and "Tr"(for "inverse opacity" - transparency)
 					float OpacityScalar = IsScalarInitialized(MaterialData.Opacity) ? MaterialData.Opacity : (IsScalarInitialized(MaterialData.Transparency) ? 1.f - MaterialData.Transparency : -1.f);
 					if (IsScalarInitialized(OpacityScalar))
 					{
-						ShaderGraphNode->AddFloatAttribute(UInterchangeShaderPortsAPI::MakeInputValueKey(OpacityInputType), OpacityScalar);
+						UE::Interchange::Materials::Private::FMaterialNodesHelper::SetupScalarParameter(BaseNodeContainer, ShaderGraphNode, OpacityInputName, OpacityScalar);
 					}
 				}
 			}
@@ -1505,14 +1506,16 @@ bool UInterchangeOBJTranslator::Translate(UInterchangeBaseNodeContainer& BaseNod
 			if (IsScalarInitialized(MaterialData.IndexOfRefraction) || bRefraction)
 			{
 				const float DefaultIndexOfRefraction = 1.52f;  // Use glass index of refraction by default
-				float IndexOfRefraction = IsScalarInitialized(MaterialData.IndexOfRefraction) ? MaterialData.IndexOfRefraction : DefaultIndexOfRefraction;
-				ShaderGraphNode->AddFloatAttribute(UInterchangeShaderPortsAPI::MakeInputValueKey(Materials::Common::Parameters::IndexOfRefraction.ToString()), IndexOfRefraction);
+				const float IndexOfRefraction = IsScalarInitialized(MaterialData.IndexOfRefraction) ? MaterialData.IndexOfRefraction : DefaultIndexOfRefraction;
+				const FString ParameterName = Materials::Common::Parameters::IndexOfRefraction.ToString();
+				UE::Interchange::Materials::Private::FMaterialNodesHelper::SetupScalarParameter(BaseNodeContainer, ShaderGraphNode, ParameterName, IndexOfRefraction);
 			}
 
 			if ((IsColorInitialized(MaterialData.TransmissionFilter) && !(MaterialData.TransmissionFilter - FVector3f::One()).IsNearlyZero()) || bTransmission)
 			{
 				FLinearColor TransmissionColor = IsColorInitialized(MaterialData.TransmissionFilter) ? MaterialData.TransmissionFilter : FLinearColor::White;
-				ShaderGraphNode->AddLinearColorAttribute(UInterchangeShaderPortsAPI::MakeInputValueKey(Materials::ThinTranslucent::Parameters::TransmissionColor.ToString()), TransmissionColor);
+				const FString ParameterName = Materials::ThinTranslucent::Parameters::TransmissionColor.ToString();
+				UE::Interchange::Materials::Private::FMaterialNodesHelper::SetupVectorParameter(BaseNodeContainer, ShaderGraphNode, ParameterName, TransmissionColor);
 			}
 
 			// Normal/bump
diff --git a/Engine/Plugins/Interchange/Runtime/Source/Import/Public/Material/InterchangeMaterialNodesHelper.h b/Engine/Plugins/Interchange/Runtime/Source/Import/Public/Material/InterchangeMaterialNodesHelper.h
new file mode 100644
--- /dev/null
+++ b/Engine/Plugins/Interchange/Runtime/Source/Import/Public/Material/InterchangeMaterialNodesHelper.h
@@ -0,0 +1,17 @@
+// Copyright Epic Games, Inc. All Rights Reserved. 
+#pragma once
+#include "CoreMinimal.h"
+
+class UInterchangeBaseNode;
+class UInterchangeBaseNodeContainer;
+class UInterchangeShaderGraphNode;
+
+namespace UE::Interchange::Materials::Private
+{
+	struct FMaterialNodesHelper
+	{
+	public:
+		static bool SetupScalarParameter(UInterchangeBaseNodeContainer& BaseNodeContainer, UInterchangeShaderGraphNode* ShaderGraphNode, const FString& InputName, const float& AttributeValue);
+		static bool SetupVectorParameter(UInterchangeBaseNodeContainer& BaseNodeContainer, UInterchangeShaderGraphNode* ShaderGraphNode, const FString& InputName, const FLinearColor& AttributeValue);
+	};
+} 
