Implement support of IDatasmithDecalActorElement through Interchange
- Added support for Decal Materials as well

diff --git a/Engine/Plugins/Experimental/Enterprise/DatasmithInterchange/Source/DatasmithInterchange/Private/InterchangeDatasmithLevelPipeline.cpp b/Engine/Plugins/Experimental/Enterprise/DatasmithInterchange/Source/DatasmithInterchange/Private/InterchangeDatasmithLevelPipeline.cpp
--- a/Engine/Plugins/Experimental/Enterprise/DatasmithInterchange/Source/DatasmithInterchange/Private/InterchangeDatasmithLevelPipeline.cpp
+++ b/Engine/Plugins/Experimental/Enterprise/DatasmithInterchange/Source/DatasmithInterchange/Private/InterchangeDatasmithLevelPipeline.cpp
@@ -6,6 +6,8 @@
 #include "InterchangeDatasmithAreaLightNode.h"
 #include "InterchangeDatasmithUtils.h"
 
+#include "InterchangeDecalActorFactoryNode.h"
+#include "InterchangeDecalNode.h"
 
 #include "InterchangeAnimationTrackSetFactoryNode.h"
 #include "InterchangeAnimationTrackSetNode.h"
@@ -18,6 +20,7 @@
 #include "Engine/StaticMesh.h"
 #include "Engine/StaticMeshActor.h"
 #include "Materials/MaterialInterface.h"
+#include "Misc/PackageName.h"
 
 void UInterchangeDatasmithLevelPipeline::ExecutePipeline(UInterchangeBaseNodeContainer* NodeContainer, const TArray<UInterchangeSourceData*>& InSourceDatas)
 {
@@ -150,6 +153,11 @@ UInterchangeActorFactoryNode* UInterchangeDatasmithLevelPipeline::CreateActorFac
 	{
 		return NewObject<UInterchangeDatasmithAreaLightFactoryNode>(BaseNodeContainer, NAME_None);
 	}
+	// Experimental - Move to Super class once finalized
+	else if (TranslatedAssetNode && TranslatedAssetNode->IsA<UInterchangeDecalNode>())
+	{
+		return NewObject<UInterchangeDecalActorFactoryNode>(BaseNodeContainer, NAME_None);
+	}
 	else
 	{
 		return Super::CreateActorFactoryNode(SceneNode, TranslatedAssetNode);
@@ -164,6 +172,13 @@ void UInterchangeDatasmithLevelPipeline::SetUpFactoryNode(UInterchangeActorFacto
 		ensure(AreaLightFactory);
 		SetupAreaLight(AreaLightFactory, AreaLightNode);
 	}
+	// Experimental - Move to Super class once finalized
+	else if (const UInterchangeDecalNode* DecalNode = Cast<UInterchangeDecalNode>(TranslatedAssetNode))
+	{
+		UInterchangeDecalActorFactoryNode* DecalActorFactory = Cast<UInterchangeDecalActorFactoryNode>(ActorFactoryNode);
+		ensure(DecalActorFactory);
+		SetupDecalActor(DecalActorFactory, DecalNode);
+	}
 	else
 	{
 		Super::SetUpFactoryNode(ActorFactoryNode, SceneNode, TranslatedAssetNode);
@@ -178,9 +193,9 @@ void UInterchangeDatasmithLevelPipeline::SetUpFactoryNode(UInterchangeActorFacto
 	}\
 
 /**
-* Only modifies the attribute if a condition is satisfied
-* Condition should an expression evaluated to a boolean.
-*/
+ * Only modifies the attribute if a condition is satisfied
+ * Condition should an expression evaluated to a boolean.
+ */
 #define APPLY_FACTORY_ATTRIBUTE_WITH_VALIDATION(AttributeName, AttributeType, Condition) \
 	AttributeType AttributeName;\
 	if (TranslatedNode->GetCustom##AttributeName(AttributeName))\
@@ -263,6 +278,21 @@ void UInterchangeDatasmithLevelPipeline::SetupAreaLight(UInterchangeDatasmithAre
 	APPLY_FACTORY_ATTRIBUTE(SpotlightOuterAngle, float);
 }
 
+void UInterchangeDatasmithLevelPipeline::SetupDecalActor(UInterchangeDecalActorFactoryNode* FactoryNode, const UInterchangeDecalNode* TranslatedNode) const
+{
+	APPLY_FACTORY_ATTRIBUTE(DecalSize, FVector);
+	APPLY_FACTORY_ATTRIBUTE(SortOrder, int32);
+	APPLY_FACTORY_ATTRIBUTE(DecalMaterialPathName, FString);
+
+	// If the path is not a valid object path then it is an Interchange Node UID (Decal Material Node to be specific).
+	if (!FPackageName::IsValidObjectPath(DecalMaterialPathName))
+	{
+		const FString MaterialFactoryUid = UInterchangeFactoryBaseNode::BuildFactoryNodeUid(DecalMaterialPathName);
+		FactoryNode->SetCustomDecalMaterialPathName(MaterialFactoryUid);
+		FactoryNode->AddFactoryDependencyUid(MaterialFactoryUid);
+	}
+}
+
 #undef APPLY_FACTORY_ATTRIBUTE
 #undef APPLY_FACTORY_ATTRIBUTE_WITH_VALIDATION
 #undef APPLY_FACTORY_ATTRIBUTE_WITH_MODIFIER
\ No newline at end of file
diff --git a/Engine/Plugins/Experimental/Enterprise/DatasmithInterchange/Source/DatasmithInterchange/Private/InterchangeDatasmithLevelPipeline.h b/Engine/Plugins/Experimental/Enterprise/DatasmithInterchange/Source/DatasmithInterchange/Private/InterchangeDatasmithLevelPipeline.h
--- a/Engine/Plugins/Experimental/Enterprise/DatasmithInterchange/Source/DatasmithInterchange/Private/InterchangeDatasmithLevelPipeline.h
+++ b/Engine/Plugins/Experimental/Enterprise/DatasmithInterchange/Source/DatasmithInterchange/Private/InterchangeDatasmithLevelPipeline.h
@@ -13,6 +13,8 @@ class UInterchangeBaseNodeContainer;
 class UInterchangeSceneNode;
 class UInterchangeDatasmithAreaLightFactoryNode;
 class UInterchangeDatasmithAreaLightNode;
+class UInterchangeDecalActorFactoryNode;
+class UInterchangeDecalNode;
 
 UCLASS(BlueprintType, Experimental)
 class UInterchangeDatasmithLevelPipeline : public UInterchangeGenericLevelPipeline
@@ -30,4 +32,5 @@ protected:
 private:
 
 	void SetupAreaLight(UInterchangeDatasmithAreaLightFactoryNode* AreaLightFactoryNode, const UInterchangeDatasmithAreaLightNode* AreaLightNode) const;
+	void SetupDecalActor(UInterchangeDecalActorFactoryNode* DecalActorFactoryNode, const UInterchangeDecalNode* DecalNode) const;
 };
\ No newline at end of file
diff --git a/Engine/Plugins/Experimental/Enterprise/DatasmithInterchange/Source/DatasmithInterchange/Private/InterchangeDatasmithMaterialNode.cpp b/Engine/Plugins/Experimental/Enterprise/DatasmithInterchange/Source/DatasmithInterchange/Private/InterchangeDatasmithMaterialNode.cpp
--- a/Engine/Plugins/Experimental/Enterprise/DatasmithInterchange/Source/DatasmithInterchange/Private/InterchangeDatasmithMaterialNode.cpp
+++ b/Engine/Plugins/Experimental/Enterprise/DatasmithInterchange/Source/DatasmithInterchange/Private/InterchangeDatasmithMaterialNode.cpp
@@ -17,6 +17,7 @@
 #include "InterchangeShaderGraphNode.h"
 #include "InterchangeTexture2DNode.h"
 #include "InterchangeTexture2DFactoryNode.h"
+#include "InterchangeDecalMaterialNode.h"
 #include "Nodes/InterchangeUserDefinedAttribute.h"
 
 #include "Async/TaskGraphInterfaces.h"
@@ -871,6 +872,17 @@ namespace UE::DatasmithInterchange::MaterialUtils
 		return true;
 	}
 
+	bool BuildMaterialNode(IDatasmithDecalMaterialElement& MaterialElement, UInterchangeDecalMaterialNode& MaterialNode)
+	{
+		// #todo_Vedang_Javdekar Check if adding this attribute would help later in the Import Process.
+		//MaterialNode.AddInt32Attribute(MaterialTypeAttrName, int32(EDatasmithReferenceMaterialType::Decal));
+
+		MaterialNode.SetCustomDiffuseTexturePath(MaterialElement.GetDiffuseTexturePathName());
+		MaterialNode.SetCustomNormalTexturePath(MaterialElement.GetNormalTexturePathName());
+
+		return true;
+	}
+
 	void ProcessDependencies(IDatasmithUEPbrMaterialElement& MaterialElement, TArray<TSharedPtr<IDatasmithBaseMaterialElement>>& OutMaterialElements, TMap<FString, TSharedPtr<IDatasmithBaseMaterialElement>>& MaterialsToSort)
 	{
 		for (int32 ExpressionIndex = 0; ExpressionIndex < MaterialElement.GetExpressionsCount(); ExpressionIndex++)
@@ -1037,6 +1049,21 @@ namespace UE::DatasmithInterchange::MaterialUtils
 
 				BaseNode = MaterialNode;
 			}
+			else if (MaterialElement->IsA(EDatasmithElementType::DecalMaterial))
+			{
+				IDatasmithDecalMaterialElement& DecalMaterialElement = static_cast<IDatasmithDecalMaterialElement&>(*MaterialElement);
+
+				UInterchangeDecalMaterialNode* MaterialNode = NewObject<UInterchangeDecalMaterialNode>(&NodeContainer);
+				const FString MaterialNodeUid = NodeUtils::DecalMaterialPrefix + FDatasmithUtils::SanitizeObjectName(MaterialElement->GetName());
+				MaterialNode->InitializeNode(MaterialNodeUid, *FDatasmithUtils::SanitizeObjectName(MaterialElement->GetLabel()), EInterchangeNodeContainerType::TranslatedAsset);
+
+				if (!BuildMaterialNode(DecalMaterialElement, *MaterialNode))
+				{
+					MaterialNode = nullptr;
+				}
+
+				BaseNode = MaterialNode;
+			}
 
 			if (BaseNode)
 			{
diff --git a/Engine/Plugins/Experimental/Enterprise/DatasmithInterchange/Source/DatasmithInterchange/Private/InterchangeDatasmithMaterialPipeline.cpp b/Engine/Plugins/Experimental/Enterprise/DatasmithInterchange/Source/DatasmithInterchange/Private/InterchangeDatasmithMaterialPipeline.cpp
--- a/Engine/Plugins/Experimental/Enterprise/DatasmithInterchange/Source/DatasmithInterchange/Private/InterchangeDatasmithMaterialPipeline.cpp
+++ b/Engine/Plugins/Experimental/Enterprise/DatasmithInterchange/Source/DatasmithInterchange/Private/InterchangeDatasmithMaterialPipeline.cpp
@@ -5,6 +5,7 @@
 #include "InterchangeDatasmithMaterialNode.h"
 #include "InterchangeDatasmithMaterialPipeline.h"
 #include "InterchangeDatasmithUtils.h"
+#include "InterchangeDecalMaterialNode.h"
 
 #include "DatasmithMaterialElements.h"
 #include "IDatasmithSceneElements.h"
@@ -15,7 +16,7 @@
 #include "InterchangeManager.h"
 #include "InterchangeMaterialDefinitions.h"
 #include "InterchangeMaterialFactoryNode.h"
-
+#include "InterchangeDecalMaterialFactoryNode.h"
 #include "InterchangeTexture2DNode.h"
 #include "InterchangeTexture2DFactoryNode.h"
 #include "InterchangeMaterialInstanceNode.h"
@@ -64,6 +65,7 @@ void UInterchangeDatasmithMaterialPipeline::ExecutePipeline(UInterchangeBaseNode
 	using namespace UE::DatasmithInterchange;
 
 	TArray<UInterchangeMaterialInstanceNode*> InstancedMaterials;
+	TArray<UInterchangeDecalMaterialNode*> DecalMaterials;
 	TArray<UInterchangeShaderNode*> ShaderNodes;
 
 	//Find all translated node we need for this pipeline
@@ -73,6 +75,10 @@ void UInterchangeDatasmithMaterialPipeline::ExecutePipeline(UInterchangeBaseNode
 			{
 				InstancedMaterials.Add(MaterialNode);
 			}
+			else if (UInterchangeDecalMaterialNode* DecalNode = Cast<UInterchangeDecalMaterialNode>(Node))
+			{
+				DecalMaterials.Add(DecalNode);
+			}
 			else if (UInterchangeShaderNode* ShaderNode = Cast<UInterchangeShaderNode>(Node))
 			{
 				ShaderNodes.Add(ShaderNode);
@@ -87,6 +93,11 @@ void UInterchangeDatasmithMaterialPipeline::ExecutePipeline(UInterchangeBaseNode
 	{
 		PreImportMaterialNode(NodeContainer, MaterialNode);
 	}
+
+	for (UInterchangeDecalMaterialNode* MaterialNode : DecalMaterials)
+	{
+		PreImportDecalMaterialNode(NodeContainer, MaterialNode);
+	}
 }
 
 void UInterchangeDatasmithMaterialPipeline::ExecutePostImportPipeline(const UInterchangeBaseNodeContainer* NodeContainer, const FString& NodeKey, UObject* CreatedAsset, bool bIsAReimport)
@@ -389,4 +400,44 @@ void UInterchangeDatasmithMaterialPipeline::UpdateMaterialFactoryNodes(const TAr
 
 
 	}
-}
\ No newline at end of file
+}
+
+void UInterchangeDatasmithMaterialPipeline::PreImportDecalMaterialNode(UInterchangeBaseNodeContainer* NodeContainer, UInterchangeDecalMaterialNode* MaterialNode)
+{
+	using namespace UE::DatasmithInterchange;
+
+	const FString MaterialNodeUid = UInterchangeDecalMaterialFactoryNode::GetMaterialFactoryNodeUidFromMaterialNodeUid(MaterialNode->GetUniqueID());
+	UInterchangeDecalMaterialFactoryNode* MaterialFactoryNode = NodeUtils::FindOrAddFactoryNode<UInterchangeDecalMaterialFactoryNode>(MaterialNode, NodeContainer, MaterialNodeUid);
+
+	// Handle Diffuse Texture
+	FString DiffuseTexturePath;
+	if (MaterialNode->GetCustomDiffuseTexturePath(DiffuseTexturePath))
+	{
+		if (!FPackageName::IsValidObjectPath(DiffuseTexturePath))
+		{
+			const FString TextureUid = UInterchangeTexture2DFactoryNode::GetTextureFactoryNodeUidFromTextureNodeUid(NodeUtils::TexturePrefix + DiffuseTexturePath);
+			MaterialFactoryNode->SetCustomDiffuseTexturePath(TextureUid);
+			MaterialFactoryNode->AddFactoryDependencyUid(TextureUid);
+		}
+	}
+	else
+	{
+		UE_LOG(LogInterchangeMaterialPipeline, Warning, TEXT("Decal Material Node doesn't have Diffuse Texture Path Attribute set."));
+	}
+
+	// Handle Normal Texture
+	FString NormalTexturePath;
+	if (MaterialNode->GetCustomNormalTexturePath(NormalTexturePath))
+	{
+		if (!FPackageName::IsValidObjectPath(DiffuseTexturePath))
+		{
+			const FString TextureUid = UInterchangeTexture2DFactoryNode::GetTextureFactoryNodeUidFromTextureNodeUid(NodeUtils::TexturePrefix + DiffuseTexturePath);
+			MaterialFactoryNode->SetCustomNormalTexturePath(TextureUid);
+			MaterialFactoryNode->AddFactoryDependencyUid(TextureUid);
+		}
+	}
+	else
+	{
+		UE_LOG(LogInterchangeMaterialPipeline, Warning, TEXT("Decal Material Node doesn't have Normal Texture Path Attribute set."));
+	}	
+}
diff --git a/Engine/Plugins/Experimental/Enterprise/DatasmithInterchange/Source/DatasmithInterchange/Private/InterchangeDatasmithMaterialPipeline.h b/Engine/Plugins/Experimental/Enterprise/DatasmithInterchange/Source/DatasmithInterchange/Private/InterchangeDatasmithMaterialPipeline.h
--- a/Engine/Plugins/Experimental/Enterprise/DatasmithInterchange/Source/DatasmithInterchange/Private/InterchangeDatasmithMaterialPipeline.h
+++ b/Engine/Plugins/Experimental/Enterprise/DatasmithInterchange/Source/DatasmithInterchange/Private/InterchangeDatasmithMaterialPipeline.h
@@ -15,6 +15,7 @@ class UInterchangeMaterialFactoryNode;
 class UInterchangeMaterialInstanceFactoryNode;
 class UInterchangeShaderGraphNode;
 class UInterchangeMaterialInstanceNode;
+class UInterchangeDecalMaterialNode;
 
 DECLARE_LOG_CATEGORY_EXTERN(LogInterchangeMaterialPipeline, Log, All);
 
@@ -39,4 +40,5 @@ private:
 	void PreImportMaterialNode(UInterchangeBaseNodeContainer* NodeContainer, UInterchangeMaterialInstanceNode* MaterialNode);
 	void PostImportMaterialInstanceFactoryNode(const UInterchangeBaseNodeContainer* NodeContainer, UInterchangeMaterialInstanceFactoryNode* FactoryNode, UMaterialInterface* CreatedMaterial);
 	void UpdateMaterialFactoryNodes(const TArray<UInterchangeShaderNode*>& ShaderNodes);
+	void PreImportDecalMaterialNode(UInterchangeBaseNodeContainer* NodeContainer, UInterchangeDecalMaterialNode* MaterialNode);
 };
\ No newline at end of file
diff --git a/Engine/Plugins/Experimental/Enterprise/DatasmithInterchange/Source/DatasmithInterchange/Private/InterchangeDatasmithTranslator.cpp b/Engine/Plugins/Experimental/Enterprise/DatasmithInterchange/Source/DatasmithInterchange/Private/InterchangeDatasmithTranslator.cpp
--- a/Engine/Plugins/Experimental/Enterprise/DatasmithInterchange/Source/DatasmithInterchange/Private/InterchangeDatasmithTranslator.cpp
+++ b/Engine/Plugins/Experimental/Enterprise/DatasmithInterchange/Source/DatasmithInterchange/Private/InterchangeDatasmithTranslator.cpp
@@ -24,6 +24,7 @@
 #include "InterchangeCameraNode.h"
 #include "InterchangeAnimationTrackSetNode.h"
 #include "InterchangeLightNode.h"
+#include "InterchangeDecalNode.h"
 #include "InterchangeManager.h"
 #include "InterchangeMaterialDefinitions.h"
 #include "InterchangeMaterialInstanceNode.h"
@@ -36,6 +37,8 @@
 #include "StaticMeshOperations.h"
 
 #include "Misc/App.h"
+#include "Misc/PackageName.h"
+
 
 #define LOCTEXT_NAMESPACE "DatasmithInterchange"
 
@@ -394,6 +397,13 @@ void UInterchangeDatasmithTranslator::HandleDatasmithActor(UInterchangeBaseNodeC
 		UInterchangeBaseLightNode* LightNode = AddLightNode(BaseNodeContainer, LightActor);
 		InterchangeSceneNode->SetCustomAssetInstanceUid(LightNode->GetUniqueID());
 	}
+	else if (ActorElement->IsA(EDatasmithElementType::Decal))
+	{
+		TSharedRef<IDatasmithDecalActorElement> DecalActor = StaticCastSharedRef<IDatasmithDecalActorElement>(ActorElement);
+
+		UInterchangeDecalNode* DecalNode= AddDecalNode(BaseNodeContainer, DecalActor);
+		InterchangeSceneNode->SetCustomAssetInstanceUid(DecalNode->GetUniqueID());
+	}
 
 	for (int32 ChildIndex = 0, ChildrenCount = ActorElement->GetChildrenCount(); ChildIndex < ChildrenCount; ++ChildIndex)
 	{
@@ -510,6 +520,32 @@ UInterchangeBaseLightNode* UInterchangeDatasmithTranslator::AddLightNode(UInterc
 	return LightNode;
 }
 
+UInterchangeDecalNode* UInterchangeDatasmithTranslator::AddDecalNode(UInterchangeBaseNodeContainer& BaseNodeContainer, const TSharedRef<IDatasmithDecalActorElement>& DecalActor) const
+{
+	using namespace UE::DatasmithInterchange;
+
+	UInterchangeDecalNode* DecalNode = NewObject<UInterchangeDecalNode>(&BaseNodeContainer);
+	const FString DecalUid = NodeUtils::DecalPrefix + DecalActor->GetName();
+	DecalNode->InitializeNode(DecalUid, DecalActor->GetLabel(), EInterchangeNodeContainerType::TranslatedAsset);
+	BaseNodeContainer.AddNode(DecalNode);
+
+	DecalNode->SetCustomSortOrder(DecalActor->GetSortOrder());
+	DecalNode->SetCustomDecalSize(DecalActor->GetDimensions());
+
+	FString DecalMaterialPathName = DecalActor->GetDecalMaterialPathName();
+	if(!FPackageName::IsValidObjectPath(DecalActor->GetDecalMaterialPathName()))
+	{
+		const FString DecalMaterialUid = NodeUtils::DecalMaterialPrefix + DecalActor->GetDecalMaterialPathName();
+		if (BaseNodeContainer.IsNodeUidValid(DecalMaterialUid))
+		{
+			DecalMaterialPathName = DecalMaterialUid;
+		}
+	}
+	DecalNode->SetCustomDecalMaterialPathName(DecalMaterialPathName);
+
+	return DecalNode;
+}
+
 TOptional<UE::Interchange::FImportImage> UInterchangeDatasmithTranslator::GetTexturePayloadData(const FString& PayloadKey, TOptional<FString>& AlternateTexturePath) const
 {
 	if (!LoadedExternalSource || !LoadedExternalSource->GetDatasmithScene())
diff --git a/Engine/Plugins/Experimental/Enterprise/DatasmithInterchange/Source/DatasmithInterchange/Private/InterchangeDatasmithUtils.cpp b/Engine/Plugins/Experimental/Enterprise/DatasmithInterchange/Source/DatasmithInterchange/Private/InterchangeDatasmithUtils.cpp
--- a/Engine/Plugins/Experimental/Enterprise/DatasmithInterchange/Source/DatasmithInterchange/Private/InterchangeDatasmithUtils.cpp
+++ b/Engine/Plugins/Experimental/Enterprise/DatasmithInterchange/Source/DatasmithInterchange/Private/InterchangeDatasmithUtils.cpp
@@ -18,6 +18,8 @@ namespace UE::DatasmithInterchange::NodeUtils
 	const FString LevelSequencePrefix(FString(TEXT("\\LevelSequence\\")));
 	const FString LevelVariantSetPrefix(FString(TEXT("\\LevelVariantSet\\")));
 	const FString LightPrefix(FString(TEXT("\\Light\\")));
+	const FString DecalPrefix(FString(TEXT("\\Decal\\")));
+	const FString DecalMaterialPrefix(FString(TEXT("\\DecalMaterial\\")));
 	const FString MaterialPrefix(FString(TEXT("\\Material\\")));
 	const FString MaterialExpressionPrefix(FString(TEXT("\\MaterialExpression\\")));
 	const FString MaterialFunctionPrefix(FString(TEXT("\\MaterialFunction\\")));
diff --git a/Engine/Plugins/Experimental/Enterprise/DatasmithInterchange/Source/DatasmithInterchange/Private/InterchangeDatasmithUtils.h b/Engine/Plugins/Experimental/Enterprise/DatasmithInterchange/Source/DatasmithInterchange/Private/InterchangeDatasmithUtils.h
--- a/Engine/Plugins/Experimental/Enterprise/DatasmithInterchange/Source/DatasmithInterchange/Private/InterchangeDatasmithUtils.h
+++ b/Engine/Plugins/Experimental/Enterprise/DatasmithInterchange/Source/DatasmithInterchange/Private/InterchangeDatasmithUtils.h
@@ -30,6 +30,8 @@ namespace UE::DatasmithInterchange::NodeUtils
 	extern const FString LevelSequencePrefix;
 	extern const FString LevelVariantSetPrefix;
 	extern const FString LightPrefix;
+	extern const FString DecalPrefix;
+	extern const FString DecalMaterialPrefix;
 	extern const FString MaterialPrefix;
 	extern const FString MaterialExpressionPrefix;
 	extern const FString MaterialFunctionPrefix;
diff --git a/Engine/Plugins/Experimental/Enterprise/DatasmithInterchange/Source/DatasmithInterchange/Public/InterchangeDatasmithTranslator.h b/Engine/Plugins/Experimental/Enterprise/DatasmithInterchange/Source/DatasmithInterchange/Public/InterchangeDatasmithTranslator.h
--- a/Engine/Plugins/Experimental/Enterprise/DatasmithInterchange/Source/DatasmithInterchange/Public/InterchangeDatasmithTranslator.h
+++ b/Engine/Plugins/Experimental/Enterprise/DatasmithInterchange/Source/DatasmithInterchange/Public/InterchangeDatasmithTranslator.h
@@ -23,10 +23,12 @@ class IDatasmithActorElement;
 class IDatasmithBaseAnimationElement;
 class IDatasmithCameraActorElement;
 class IDatasmithLightActorElement;
+class IDatasmithDecalActorElement;
 class IDatasmithScene;
 class IDatasmithTransformAnimationElement;
 class UInterchangePhysicalCameraNode;
 class UInterchangeBaseLightNode;
+class UInterchangeDecalNode;
 class UInterchangeSceneNode;
 
 namespace UE::Interchange
@@ -101,6 +103,8 @@ private:
 
 	UInterchangeBaseLightNode* AddLightNode(UInterchangeBaseNodeContainer& BaseNodeContainer, const TSharedRef<IDatasmithLightActorElement>& LightActor) const;
 
+	UInterchangeDecalNode* AddDecalNode(UInterchangeBaseNodeContainer& BaseNodeContainer, const TSharedRef<IDatasmithDecalActorElement>& DecalActor) const;
+
 	TSharedPtr<UE::DatasmithImporter::FExternalSource> LoadedExternalSource;
 
 	mutable uint64 StartTime = 0;
diff --git a/Engine/Plugins/Interchange/Runtime/Source/FactoryNodes/Private/InterchangeDecalActorFactoryNode.cpp b/Engine/Plugins/Interchange/Runtime/Source/FactoryNodes/Private/InterchangeDecalActorFactoryNode.cpp
new file mode 100644
--- /dev/null
+++ b/Engine/Plugins/Interchange/Runtime/Source/FactoryNodes/Private/InterchangeDecalActorFactoryNode.cpp
@@ -0,0 +1,52 @@
+// Copyright Epic Games, Inc. All Rights Reserved.
+
+#include "InterchangeDecalActorFactoryNode.h"
+#include "Components/DecalComponent.h"
+
+#if WITH_ENGINE
+#include "Engine/DecalActor.h"
+#endif
+
+UClass* UInterchangeDecalActorFactoryNode::GetObjectClass() const
+{
+#if WITH_ENGINE
+	return ADecalActor::StaticClass();
+#else
+	return nullptr;
+#endif
+}
+
+bool UInterchangeDecalActorFactoryNode::GetCustomSortOrder(int32& AttributeValue) const
+{
+	IMPLEMENT_NODE_ATTRIBUTE_GETTER(SortOrder, int32);
+}
+
+UFUNCTION(BlueprintCallable, Category = "Interchange | Node | Decal")
+bool UInterchangeDecalActorFactoryNode::SetCustomSortOrder(const int32& AttributeValue, bool bAddApplyDelegate /*= true*/)
+{
+	IMPLEMENT_NODE_ATTRIBUTE_SETTER_WITH_CUSTOM_DELEGATE_WITH_CUSTOM_CLASS(UInterchangeDecalActorFactoryNode, SortOrder, int32, UDecalComponent);
+}
+
+UFUNCTION(BlueprintCallable, Category = "Interchange | Node | Decal")
+bool UInterchangeDecalActorFactoryNode::GetCustomDecalSize(FVector& AttributeValue) const
+{
+	IMPLEMENT_NODE_ATTRIBUTE_GETTER(DecalSize, FVector);
+}
+
+UFUNCTION(BlueprintCallable, Category = "Interchange | Node | Decal")
+bool UInterchangeDecalActorFactoryNode::SetCustomDecalSize(const FVector& AttributeValue, bool bAddApplyDelegate /*= true*/)
+{
+	IMPLEMENT_NODE_ATTRIBUTE_SETTER_WITH_CUSTOM_DELEGATE_WITH_CUSTOM_CLASS(UInterchangeDecalActorFactoryNode, DecalSize, FVector, UDecalComponent);
+}
+
+UFUNCTION(BlueprintCallable, Category = "Interchange | Node | Decal")
+bool UInterchangeDecalActorFactoryNode::GetCustomDecalMaterialPathName(FString& AttributeValue) const
+{
+	IMPLEMENT_NODE_ATTRIBUTE_GETTER(DecalMaterialPathName, FString);
+}
+
+UFUNCTION(BlueprintCallable, Category = "Interchange | Node | Decal")
+bool UInterchangeDecalActorFactoryNode::SetCustomDecalMaterialPathName(const FString& AttributeValue)
+{
+	IMPLEMENT_NODE_ATTRIBUTE_SETTER_NODELEGATE(DecalMaterialPathName, FString);
+}
\ No newline at end of file
diff --git a/Engine/Plugins/Interchange/Runtime/Source/FactoryNodes/Private/InterchangeDecalMaterialFactoryNode.cpp b/Engine/Plugins/Interchange/Runtime/Source/FactoryNodes/Private/InterchangeDecalMaterialFactoryNode.cpp
new file mode 100644
--- /dev/null
+++ b/Engine/Plugins/Interchange/Runtime/Source/FactoryNodes/Private/InterchangeDecalMaterialFactoryNode.cpp
@@ -0,0 +1,33 @@
+// Copyright Epic Games, Inc. All Rights Reserved.
+
+#include "InterchangeDecalMaterialFactoryNode.h"
+#include "Materials/MaterialInstanceConstant.h"
+
+UClass* UInterchangeDecalMaterialFactoryNode::GetObjectClass() const
+{
+#if WITH_ENGINE
+	return UMaterialInstanceConstant::StaticClass();
+#else
+	return nullptr;
+#endif
+}
+
+bool UInterchangeDecalMaterialFactoryNode::GetCustomDiffuseTexturePath(FString& AttributeValue) const
+{
+	IMPLEMENT_NODE_ATTRIBUTE_GETTER(DiffuseTexturePath, FString);
+}
+
+bool UInterchangeDecalMaterialFactoryNode::SetCustomDiffuseTexturePath(const FString& AttributeValue)
+{
+	IMPLEMENT_NODE_ATTRIBUTE_SETTER_NODELEGATE(DiffuseTexturePath, FString);
+}
+
+bool UInterchangeDecalMaterialFactoryNode::GetCustomNormalTexturePath(FString& AttributeValue) const
+{
+	IMPLEMENT_NODE_ATTRIBUTE_GETTER(NormalTexturePath, FString);
+}
+
+bool UInterchangeDecalMaterialFactoryNode::SetCustomNormalTexturePath(const FString& AttributeValue)
+{
+	IMPLEMENT_NODE_ATTRIBUTE_SETTER_NODELEGATE(NormalTexturePath, FString);
+}
diff --git a/Engine/Plugins/Interchange/Runtime/Source/FactoryNodes/Public/InterchangeDecalActorFactoryNode.h b/Engine/Plugins/Interchange/Runtime/Source/FactoryNodes/Public/InterchangeDecalActorFactoryNode.h
new file mode 100644
--- /dev/null
+++ b/Engine/Plugins/Interchange/Runtime/Source/FactoryNodes/Public/InterchangeDecalActorFactoryNode.h
@@ -0,0 +1,45 @@
+// Copyright Epic Games, Inc. All Rights Reserved.
+
+#pragma once
+
+#include "InterchangeActorFactoryNode.h"
+
+#include "InterchangeDecalActorFactoryNode.generated.h"
+
+UCLASS(BlueprintType, Experimental)
+class INTERCHANGEFACTORYNODES_API UInterchangeDecalActorFactoryNode : public UInterchangeActorFactoryNode
+{
+	GENERATED_BODY()
+
+public:
+	UFUNCTION(BlueprintCallable, Category = "Interchange | Node | Decal")
+	UClass* GetObjectClass()const override;
+
+	UFUNCTION(BlueprintCallable, Category = "Interchange | Node | Decal")
+	bool GetCustomSortOrder(int32& AttributeValue) const;
+	
+	UFUNCTION(BlueprintCallable, Category = "Interchange | Node | Decal")
+	bool SetCustomSortOrder(const int32& AttributeValue, bool bAddApplyDelegate = true);
+	
+	UFUNCTION(BlueprintCallable, Category = "Interchange | Node | Decal")
+	bool GetCustomDecalSize(FVector& AttributeValue) const;
+	
+	UFUNCTION(BlueprintCallable, Category = "Interchange | Node | Decal")
+	bool SetCustomDecalSize(const FVector& AttributeValue, bool bAddApplyDelegate = true);
+	
+	UFUNCTION(BlueprintCallable, Category = "Interchange | Node | Decal")
+	bool GetCustomDecalMaterialPathName(FString& AttributeValue) const;
+	
+
+	UFUNCTION(BlueprintCallable, Category = "Interchange | Node | Decal")
+	bool SetCustomDecalMaterialPathName(const FString& AttributeValue);
+	
+private:
+	IMPLEMENT_NODE_ATTRIBUTE_DELEGATE_BY_PROPERTYNAME(SortOrder, int32, UDecalComponent, TEXT("SortOrder"));
+	IMPLEMENT_NODE_ATTRIBUTE_DELEGATE_BY_PROPERTYNAME(DecalSize, FVector, UDecalComponent, TEXT("DecalSize"));
+
+private:
+	const UE::Interchange::FAttributeKey Macro_CustomSortOrderKey = UE::Interchange::FAttributeKey(TEXT("SortOrder"));
+	const UE::Interchange::FAttributeKey Macro_CustomDecalSizeKey = UE::Interchange::FAttributeKey(TEXT("DecalSize"));
+	const UE::Interchange::FAttributeKey Macro_CustomDecalMaterialPathNameKey = UE::Interchange::FAttributeKey(TEXT("DecalMaterialPathName"));
+};
\ No newline at end of file
diff --git a/Engine/Plugins/Interchange/Runtime/Source/FactoryNodes/Public/InterchangeDecalMaterialFactoryNode.h b/Engine/Plugins/Interchange/Runtime/Source/FactoryNodes/Public/InterchangeDecalMaterialFactoryNode.h
new file mode 100644
--- /dev/null
+++ b/Engine/Plugins/Interchange/Runtime/Source/FactoryNodes/Public/InterchangeDecalMaterialFactoryNode.h
@@ -0,0 +1,34 @@
+// Copyright Epic Games, Inc. All Rights Reserved.
+
+#pragma once
+
+#include "CoreMinimal.h"
+#include "InterchangeMaterialFactoryNode.h"
+
+#include "InterchangeDecalMaterialFactoryNode.generated.h"
+
+UCLASS(BlueprintType, Experimental)
+class INTERCHANGEFACTORYNODES_API UInterchangeDecalMaterialFactoryNode : public UInterchangeBaseMaterialFactoryNode
+{
+	GENERATED_BODY()
+
+public:
+	UFUNCTION(BlueprintCallable, Category = "Interchange | Node | Decal")
+	UClass* GetObjectClass()const override;
+
+	UFUNCTION(BlueprintCallable, Category = "Interchange | Node | Decal")
+	bool GetCustomDiffuseTexturePath(FString& AttributeValue) const;
+
+	UFUNCTION(BlueprintCallable, Category = "Interchange | Node | Decal")
+	bool SetCustomDiffuseTexturePath(const FString& AttributeValue);
+
+	UFUNCTION(BlueprintCallable, Category = "Interchange | Node | Decal")
+	bool GetCustomNormalTexturePath(FString& AttributeValue) const;
+
+	UFUNCTION(BlueprintCallable, Category = "Interchange | Node | Decal")
+	bool SetCustomNormalTexturePath(const FString& AttributeValue);
+
+private:
+	const UE::Interchange::FAttributeKey Macro_CustomDiffuseTexturePathKey = UE::Interchange::FAttributeKey(TEXT("DiffuseTexturePath"));
+	const UE::Interchange::FAttributeKey Macro_CustomNormalTexturePathKey = UE::Interchange::FAttributeKey(TEXT("NormalTexturePath"));
+};
\ No newline at end of file
diff --git a/Engine/Plugins/Interchange/Runtime/Source/Import/Private/InterchangeImportModule.cpp b/Engine/Plugins/Interchange/Runtime/Source/Import/Private/InterchangeImportModule.cpp
--- a/Engine/Plugins/Interchange/Runtime/Source/Import/Private/InterchangeImportModule.cpp
+++ b/Engine/Plugins/Interchange/Runtime/Source/Import/Private/InterchangeImportModule.cpp
@@ -26,6 +26,7 @@
 #include "Scene/InterchangeSceneVariantSetsFactory.h"
 #include "Scene/InterchangeStaticMeshActorFactory.h"
 #include "Scene/InterchangeSkeletalMeshActorFactory.h"
+#include "Scene/InterchangeDecalActorFactory.h"
 #include "Texture/InterchangeImageWrapperTranslator.h"
 #include "Texture/InterchangeDDSTranslator.h"
 #include "Texture/InterchangeIESTranslator.h"
@@ -86,6 +87,7 @@ void FInterchangeImportModule::StartupModule()
 		InterchangeManager.RegisterFactory(UInterchangeCineCameraActorFactory::StaticClass());
 		InterchangeManager.RegisterFactory(UInterchangeCameraActorFactory::StaticClass());
 		InterchangeManager.RegisterFactory(UInterchangeStaticMeshActorFactory::StaticClass());
+		InterchangeManager.RegisterFactory(UInterchangeDecalActorFactory::StaticClass());
 		InterchangeManager.RegisterFactory(UInterchangeSkeletalMeshActorFactory::StaticClass());
 		InterchangeManager.RegisterFactory(UInterchangeSceneVariantSetsFactory::StaticClass());
 		InterchangeManager.RegisterFactory(UInterchangeLightActorFactory::StaticClass());
diff --git a/Engine/Plugins/Interchange/Runtime/Source/Import/Private/Material/InterchangeMaterialFactory.cpp b/Engine/Plugins/Interchange/Runtime/Source/Import/Private/Material/InterchangeMaterialFactory.cpp
--- a/Engine/Plugins/Interchange/Runtime/Source/Import/Private/Material/InterchangeMaterialFactory.cpp
+++ b/Engine/Plugins/Interchange/Runtime/Source/Import/Private/Material/InterchangeMaterialFactory.cpp
@@ -5,6 +5,7 @@
 #include "InterchangeImportCommon.h"
 #include "InterchangeImportLog.h"
 #include "InterchangeMaterialFactoryNode.h"
+#include "InterchangeDecalMaterialFactoryNode.h"
 #include "InterchangeShaderGraphNode.h"
 #include "InterchangeSourceData.h"
 #include "InterchangeTextureNode.h"
@@ -157,11 +158,8 @@ namespace UE::Interchange::MaterialFactory::Internal
 			: UInterchangeShaderPortsAPI::MakeInputValueKey(Inputs::Texture.ToString());
 
 		FString TextureFactoryNodeUid;
-		ExpressionNode->GetStringAttribute(InputKey, TextureFactoryNodeUid);
-
-
 		// It is possible that the Texture Object Parameter was used that uses KeyValueProperty to reference the texture instead of extra texture input connection.
-		if (TextureFactoryNodeUid.IsEmpty())
+		if (!ExpressionNode->GetStringAttribute(InputKey, TextureFactoryNodeUid))
 		{
 			FString PayloadKey;
 			UInterchangeUserDefinedAttributesAPI::GetUserDefinedAttribute_FString(ExpressionNode, Inputs::Texture.ToString(), TextureFactoryNodeUid, PayloadKey);
@@ -323,7 +321,6 @@ namespace UE::Interchange::MaterialFactory::Internal
 					{
 						ParameterName = FName(*MaterialExpressionName);
 					}
-
 					// Set the parameter name if the material expression has one (some material expressions don't inherit from UMaterialExpressionParameter, ie: UMaterialExpressionTextureSampleParameter
 					*(Property->ContainerPtrToValuePtr<FName>(MaterialExpression)) = ParameterName;
 				}
@@ -385,6 +382,35 @@ namespace UE::Interchange::MaterialFactory::Internal
 	};
 #endif // #if WITH_EDITOR
 
+	/**
+	 * Try loading the asset with treating the Path as a package path, and fallback to NodeContainer if it fails.
+	 * @type AssetType the type to load
+	 * @type AssetFactoryType asset factory to use to get the referenced object
+	 * @param NodeContainer container that stores all the nodes
+	 * @param PathString path or UID of a node
+	 * 
+	 * @returns the loaded asset if there is any
+	 */
+	template<typename AssetType, typename AssetFactoryType>
+	AssetType* TryLoadAsset(const UInterchangeBaseNodeContainer& NodeContainer, FString PathString)
+	{
+		AssetType* OutAsset = nullptr;
+		if (FPackageName::IsValidObjectPath(PathString))
+		{
+			FSoftObjectPath AssetPath(PathString);
+			OutAsset = Cast<AssetType>(AssetPath.TryLoad());
+		}
+		else if (const AssetFactoryType* TextureFactoryNode = Cast<AssetFactoryType>(NodeContainer.GetNode(PathString)))
+		{
+			FSoftObjectPath ReferenceObject;
+			TextureFactoryNode->GetCustomReferenceObject(ReferenceObject);
+			OutAsset = Cast<AssetType>(ReferenceObject.TryLoad());
+		}
+
+		return OutAsset;
+	}
+
+
 	void UpdateParameterBool(UMaterialInstance& MaterialInstance, const FString& InputName, const UInterchangeMaterialInstanceFactoryNode& FactoryNode)
 	{
 #if WITH_EDITORONLY_DATA
@@ -426,10 +452,10 @@ namespace UE::Interchange::MaterialFactory::Internal
 				}
 				else
 #endif // #if WITH_EDITOR
-					if (UMaterialInstanceDynamic* MaterialInstanceDynamic = Cast<UMaterialInstanceDynamic>(&MaterialInstance))
-					{
-						MaterialInstanceDynamic->SetScalarParameterValue(ParameterName, InputValue);
-					}
+				if (UMaterialInstanceDynamic* MaterialInstanceDynamic = Cast<UMaterialInstanceDynamic>(&MaterialInstance))
+				{
+					MaterialInstanceDynamic->SetScalarParameterValue(ParameterName, InputValue);
+				}
 			}
 		}
 	}
@@ -453,10 +479,10 @@ namespace UE::Interchange::MaterialFactory::Internal
 					}
 					else
 #endif // #if WITH_EDITOR
-						if (UMaterialInstanceDynamic* MaterialInstanceDynamic = Cast<UMaterialInstanceDynamic>(&MaterialInstance))
-						{
-							MaterialInstanceDynamic->SetVectorParameterValue(ParameterName, InputValue);
-						}
+					if (UMaterialInstanceDynamic* MaterialInstanceDynamic = Cast<UMaterialInstanceDynamic>(&MaterialInstance))
+					{
+						MaterialInstanceDynamic->SetVectorParameterValue(ParameterName, InputValue);
+					}
 				}
 			}
 		}
@@ -487,10 +513,10 @@ namespace UE::Interchange::MaterialFactory::Internal
 							}
 							else
 #endif // #if WITH_EDITOR
-								if (UMaterialInstanceDynamic* MaterialInstanceDynamic = Cast<UMaterialInstanceDynamic>(&MaterialInstance))
-								{
-									MaterialInstanceDynamic->SetTextureParameterValue(ParameterName, InputTexture);
-								}
+							if (UMaterialInstanceDynamic* MaterialInstanceDynamic = Cast<UMaterialInstanceDynamic>(&MaterialInstance))
+							{
+								MaterialInstanceDynamic->SetTextureParameterValue(ParameterName, InputTexture);
+							}
 						}
 					}
 				}
@@ -510,11 +536,11 @@ namespace UE::Interchange::MaterialFactory::Internal
 			}
 			else
 #endif // #if WITH_EDITOR
-				if (UMaterialInstanceDynamic* MaterialInstanceDynamic = Cast<UMaterialInstanceDynamic>(&MaterialInstance))
-				{
-					//TODO: Log Error
-					ensure(false);
-				}
+			if (UMaterialInstanceDynamic* MaterialInstanceDynamic = Cast<UMaterialInstanceDynamic>(&MaterialInstance))
+			{
+				//TODO: Log Error
+				ensure(false);
+			}
 		}
 	}
 
@@ -530,10 +556,10 @@ namespace UE::Interchange::MaterialFactory::Internal
 			}
 			else
 #endif // #if WITH_EDITOR
-				if (UMaterialInstanceDynamic* MaterialInstanceDynamic = Cast<UMaterialInstanceDynamic>(&MaterialInstance))
-				{
-					MaterialInstanceDynamic->SetScalarParameterValue(ParameterName, AttributeValue);
-				}
+			if (UMaterialInstanceDynamic* MaterialInstanceDynamic = Cast<UMaterialInstanceDynamic>(&MaterialInstance))
+			{
+				MaterialInstanceDynamic->SetScalarParameterValue(ParameterName, AttributeValue);
+			}
 		}
 	}
 
@@ -549,10 +575,10 @@ namespace UE::Interchange::MaterialFactory::Internal
 			}
 			else
 #endif // #if WITH_EDITOR
-				if (UMaterialInstanceDynamic* MaterialInstanceDynamic = Cast<UMaterialInstanceDynamic>(&MaterialInstance))
-				{
-					MaterialInstanceDynamic->SetVectorParameterValue(ParameterName, AttributeValue);
-				}
+			if (UMaterialInstanceDynamic* MaterialInstanceDynamic = Cast<UMaterialInstanceDynamic>(&MaterialInstance))
+			{
+				MaterialInstanceDynamic->SetVectorParameterValue(ParameterName, AttributeValue);
+			}
 		}
 	}
 
@@ -575,16 +601,16 @@ namespace UE::Interchange::MaterialFactory::Internal
 					}
 					else
 #endif // #if WITH_EDITOR
-						if (UMaterialInstanceDynamic* MaterialInstanceDynamic = Cast<UMaterialInstanceDynamic>(&MaterialInstance))
-						{
-							MaterialInstanceDynamic->SetTextureParameterValue(ParameterName, InputTexture);
-						}
+					if (UMaterialInstanceDynamic* MaterialInstanceDynamic = Cast<UMaterialInstanceDynamic>(&MaterialInstance))
+					{
+						MaterialInstanceDynamic->SetTextureParameterValue(ParameterName, InputTexture);
+					}
 				}
 			}
 		}
 	}
 }
-
+		
 namespace UE::Interchange::Materials::HashUtils
 {
 
@@ -796,26 +822,20 @@ UInterchangeFactoryBase::FImportAssetResult UInterchangeMaterialFactory::BeginIm
 #if WITH_EDITOR
 				if (UMaterialInstanceConstant* MaterialInstanceConstant = Cast<UMaterialInstanceConstant>(MaterialInstance))
 				{
+					using namespace UE::Interchange::MaterialFactory::Internal;
+
 					FString ParentPath;
 					FSoftObjectPath ParentMaterialPath;
 					if (MaterialInstanceFactoryNode->GetCustomParent(ParentPath))
 					{
-						if (FPackageName::IsValidObjectPath(ParentPath))
-						{
-							ParentMaterialPath = FSoftObjectPath(ParentPath);
-						}
-						else if (UInterchangeFactoryBaseNode* ParentFactory = Arguments.NodeContainer->GetFactoryNode(ParentPath))
+						if (UMaterialInterface* ParentMaterialAsset = TryLoadAsset<UMaterialInterface, UInterchangeFactoryBaseNode>(*Arguments.NodeContainer, ParentPath))
 						{
-							ParentFactory->GetCustomReferenceObject(ParentMaterialPath);
+							MaterialInstanceConstant->SetParentEditorOnly(ParentMaterialAsset);
+							MaterialInstanceConstant->PostEditChange();
 						}
-
-						if (ParentMaterialPath.IsValid())
+						else
 						{
-							if (UMaterialInterface* ParentMaterialAsset = Cast<UMaterialInterface>(ParentMaterialPath.TryLoad()))
-							{
-								MaterialInstanceConstant->SetParentEditorOnly(ParentMaterialAsset);
-								MaterialInstanceConstant->PostEditChange();
-							}
+							UE_LOG(LogInterchangeImport, Error, TEXT("No parent material was found."))
 						}
 					}
 				}
@@ -824,6 +844,44 @@ UInterchangeFactoryBase::FImportAssetResult UInterchangeMaterialFactory::BeginIm
 			}
 		}
 	}
+	else if (const UInterchangeDecalMaterialFactoryNode* DecalMaterialFactoryNode = Cast<UInterchangeDecalMaterialFactoryNode>(MaterialFactoryNode))
+	{
+#if WITH_EDITOR
+		const FSoftObjectPath DecalMaterialParent(TEXT("/Interchange/Materials/DecalMaterial.DecalMaterial"));
+		const FName DiffuseTextureParameterName(TEXT("DecalTexture"));
+		const FName NormalTextureParameterName(TEXT("NormalTexture"));
+		if (UMaterialInstanceConstant* MaterialInstanceConstant = Cast<UMaterialInstanceConstant>(Material))
+		{
+			if (UMaterial* ParentMaterial = Cast<UMaterial>(DecalMaterialParent.TryLoad()))
+			{
+				using namespace UE::Interchange::MaterialFactory::Internal;
+				MaterialInstanceConstant->SetParentEditorOnly(ParentMaterial);
+
+				FString DiffuseTexturePath;
+				if (DecalMaterialFactoryNode->GetCustomDiffuseTexturePath(DiffuseTexturePath))
+				{
+					if (UTexture* InputTexture = TryLoadAsset<UTexture, UInterchangeTextureFactoryNode>(*Arguments.NodeContainer, DiffuseTexturePath))
+					{
+						MaterialInstanceConstant->SetTextureParameterValueEditorOnly(DiffuseTextureParameterName, InputTexture);
+					}
+				}
+
+				FString NormalTexturePath;
+				if (DecalMaterialFactoryNode->GetCustomNormalTexturePath(NormalTexturePath))
+				{
+					if (UTexture* InputTexture = TryLoadAsset<UTexture, UInterchangeTextureFactoryNode>(*Arguments.NodeContainer, NormalTexturePath))
+					{
+						MaterialInstanceConstant->SetTextureParameterValueEditorOnly(NormalTextureParameterName, InputTexture);
+					}
+				}
+			}
+			else
+			{
+				UE_LOG(LogInterchangeImport, Error, TEXT("Invalid Decal Material Parent Path. Can't create a Decal Material Instance."));
+			}
+		}
+#endif // #if WITH_EDITOR
+	}
 	
 	ImportAssetResult.ImportedObject = Material;
 	return ImportAssetResult;
@@ -1535,7 +1593,6 @@ void UInterchangeMaterialFactory::SetupMaterialInstance(UMaterialInstance& Mater
 			{
 			case UE::Interchange::EAttributeTypes::Bool:
 			{
-				// #todo_Vedang_Javdekar: WIP
 				OverrideBoolParameter(FactoryNode, AttributeKey.Key, MaterialInstance, ParameterName);
 			}
 			break;
@@ -1921,4 +1978,4 @@ FInterchangeImportMaterialAsyncHelper& FInterchangeImportMaterialAsyncHelper::Ge
 	return Instance;
 }
 
-#undef LOCTEXT_NAMESPACE
\ No newline at end of file
+#undef LOCTEXT_NAMESPACE
diff --git a/Engine/Plugins/Interchange/Runtime/Source/Import/Private/Scene/InterchangeDecalActorFactory.cpp b/Engine/Plugins/Interchange/Runtime/Source/Import/Private/Scene/InterchangeDecalActorFactory.cpp
new file mode 100644
--- /dev/null
+++ b/Engine/Plugins/Interchange/Runtime/Source/Import/Private/Scene/InterchangeDecalActorFactory.cpp
@@ -0,0 +1,61 @@
+// Copyright Epic Games, Inc. All Rights Reserved.
+
+#include "Scene/InterchangeDecalActorFactory.h"
+
+#include "InterchangeDecalActorFactoryNode.h"
+#include "InterchangeDecalMaterialFactoryNode.h"
+#include "Nodes/InterchangeBaseNodeContainer.h"
+
+#include "Engine/DecalActor.h"
+#include "Components/DecalComponent.h"
+
+#include "Misc/PackageName.h"
+#include "InterchangeImportLog.h"
+
+UClass* UInterchangeDecalActorFactory::GetFactoryClass() const
+{
+	return ADecalActor::StaticClass();
+}
+
+UObject* UInterchangeDecalActorFactory::ProcessActor(AActor& SpawnedActor, const UInterchangeActorFactoryNode& FactoryNode, const UInterchangeBaseNodeContainer& NodeContainer)
+{
+	ADecalActor* DecalActor = Cast<ADecalActor>(&SpawnedActor);
+	if (DecalActor)
+	{
+		if (const UInterchangeDecalActorFactoryNode* DecalActorFactoryNode = Cast<UInterchangeDecalActorFactoryNode>(&FactoryNode))
+		{
+			FString DecalMaterialPath;
+			if (DecalActorFactoryNode->GetCustomDecalMaterialPathName(DecalMaterialPath))
+			{
+				UMaterialInterface* DecalMaterial = nullptr;
+				if (FPackageName::IsValidObjectPath(DecalMaterialPath))
+				{
+					const FSoftObjectPath MaterialPath(DecalMaterialPath);
+					DecalMaterial = Cast<UMaterialInterface>(MaterialPath.TryLoad());
+				}
+				else
+				{
+					if (UInterchangeDecalMaterialFactoryNode* MaterialFactoryNode = Cast<UInterchangeDecalMaterialFactoryNode>(NodeContainer.GetFactoryNode(DecalMaterialPath)))
+					{
+						FSoftObjectPath ReferenceObject;
+						if (MaterialFactoryNode->GetCustomReferenceObject(ReferenceObject))
+						{
+							DecalMaterial = Cast<UMaterialInterface>(ReferenceObject.TryLoad());
+						}
+					}
+				}
+
+				if (DecalMaterial)
+				{
+					DecalActor->SetDecalMaterial(DecalMaterial);
+				}
+				else
+				{
+					UE_LOG(LogInterchangeImport, Error, TEXT("No valid decal material found. Make sure that the DecalMaterialPath is valid: %s"), *DecalMaterialPath);
+				}
+			}
+		}
+	}
+
+	return DecalActor ? DecalActor->GetDecal() : nullptr;
+}
\ No newline at end of file
diff --git a/Engine/Plugins/Interchange/Runtime/Source/Import/Public/Scene/InterchangeDecalActorFactory.h b/Engine/Plugins/Interchange/Runtime/Source/Import/Public/Scene/InterchangeDecalActorFactory.h
new file mode 100644
--- /dev/null
+++ b/Engine/Plugins/Interchange/Runtime/Source/Import/Public/Scene/InterchangeDecalActorFactory.h
@@ -0,0 +1,33 @@
+// Copyright Epic Games, Inc. All Rights Reserved.
+
+#pragma once
+
+#include "CoreMinimal.h"
+#include "Scene/InterchangeActorFactory.h"
+
+#include "InterchangeDecalActorFactory.generated.h"
+
+
+UCLASS(BlueprintType)
+class INTERCHANGEIMPORT_API UInterchangeDecalActorFactory : public UInterchangeActorFactory
+{
+	GENERATED_BODY()
+public:
+
+	//////////////////////////////////////////////////////////////////////////
+	// Interchange factory base interface begin
+
+	virtual UClass* GetFactoryClass() const override;
+
+	// Interchange factory base interface end
+	//////////////////////////////////////////////////////////////////////////
+
+	//////////////////////////////////////////////////////////////////////////
+	// Interchange actor factory interface begin
+
+	virtual UObject* ProcessActor(class AActor& SpawnedActor, const UInterchangeActorFactoryNode& FactoryNode, const UInterchangeBaseNodeContainer& NodeContainer) override;
+
+	// Interchange actor factory interface end
+	//////////////////////////////////////////////////////////////////////////
+};
+
diff --git a/Engine/Plugins/Interchange/Runtime/Source/Nodes/Public/InterchangeDecalMaterialNode.h b/Engine/Plugins/Interchange/Runtime/Source/Nodes/Public/InterchangeDecalMaterialNode.h
new file mode 100644
--- /dev/null
+++ b/Engine/Plugins/Interchange/Runtime/Source/Nodes/Public/InterchangeDecalMaterialNode.h
@@ -0,0 +1,43 @@
+// Copyright Epic Games, Inc. All Rights Reserved.
+
+#pragma once
+
+#include "CoreMinimal.h"
+#include "InterchangeShaderGraphNode.h"
+
+#include "InterchangeDecalMaterialNode.generated.h"
+
+class IDatasmithBaseMaterialElement;
+class UInterchangeBaseNodeContainer;
+
+UCLASS(BlueprintType, Experimental)
+class INTERCHANGENODES_API UInterchangeDecalMaterialNode : public UInterchangeShaderNode
+{
+	GENERATED_BODY()
+
+public:
+
+	bool GetCustomDiffuseTexturePath(FString& AttributeValue) const
+	{
+		IMPLEMENT_NODE_ATTRIBUTE_GETTER(DiffuseTexturePath, FString);
+	}
+
+	bool SetCustomDiffuseTexturePath(FString AttributeValue)
+	{
+		IMPLEMENT_NODE_ATTRIBUTE_SETTER_NODELEGATE(DiffuseTexturePath, FString);
+	}
+
+	bool GetCustomNormalTexturePath(FString& AttributeValue) const
+	{
+		IMPLEMENT_NODE_ATTRIBUTE_GETTER(NormalTexturePath, FString);
+	}
+
+	bool SetCustomNormalTexturePath(FString AttributeValue)
+	{
+		IMPLEMENT_NODE_ATTRIBUTE_SETTER_NODELEGATE(NormalTexturePath, FString);
+	}
+
+private:
+	const UE::Interchange::FAttributeKey Macro_CustomDiffuseTexturePathKey = UE::Interchange::FAttributeKey(TEXT("DiffuseTexturePath"));
+	const UE::Interchange::FAttributeKey Macro_CustomNormalTexturePathKey= UE::Interchange::FAttributeKey(TEXT("NormalTexturePath"));
+};
\ No newline at end of file
diff --git a/Engine/Plugins/Interchange/Runtime/Source/Nodes/Public/InterchangeDecalNode.h b/Engine/Plugins/Interchange/Runtime/Source/Nodes/Public/InterchangeDecalNode.h
new file mode 100644
--- /dev/null
+++ b/Engine/Plugins/Interchange/Runtime/Source/Nodes/Public/InterchangeDecalNode.h
@@ -0,0 +1,56 @@
+// Copyright Epic Games, Inc. All Rights Reserved.
+
+#pragma once
+
+#include "CoreMinimal.h"
+#include "Nodes/InterchangeBaseNode.h"
+
+#include "InterchangeDecalNode.generated.h"
+
+UCLASS(BlueprintType, Experimental)
+class INTERCHANGENODES_API UInterchangeDecalNode : public UInterchangeBaseNode
+{
+	GENERATED_BODY()
+
+public:
+	UFUNCTION(BlueprintCallable, Category = "Interchange | Node | Decal")
+	bool GetCustomSortOrder(int32& AttributeValue) const
+	{
+		IMPLEMENT_NODE_ATTRIBUTE_GETTER(SortOrder, int32);
+	}
+
+	UFUNCTION(BlueprintCallable, Category = "Interchange | Node | Decal")
+	bool SetCustomSortOrder(const int32& AttributeValue)
+	{
+		IMPLEMENT_NODE_ATTRIBUTE_SETTER_NODELEGATE(SortOrder, int32);
+	}
+
+	UFUNCTION(BlueprintCallable, Category = "Interchange | Node | Decal")
+	bool GetCustomDecalSize(FVector& AttributeValue) const
+	{
+		IMPLEMENT_NODE_ATTRIBUTE_GETTER(DecalSize, FVector);
+	}
+
+	UFUNCTION(BlueprintCallable, Category = "Interchange | Node | Decal")
+	bool SetCustomDecalSize(const FVector& AttributeValue)
+	{
+		IMPLEMENT_NODE_ATTRIBUTE_SETTER_NODELEGATE(DecalSize, FVector);
+	}
+
+	UFUNCTION(BlueprintCallable, Category = "Interchange | Node | Decal")
+	bool GetCustomDecalMaterialPathName(FString& AttributeValue) const
+	{
+		IMPLEMENT_NODE_ATTRIBUTE_GETTER(DecalMaterialPathName, FString);
+	}
+
+	UFUNCTION(BlueprintCallable, Category = "Interchange | Node | Decal")
+	bool SetCustomDecalMaterialPathName(const FString& AttributeValue)
+	{
+		IMPLEMENT_NODE_ATTRIBUTE_SETTER_NODELEGATE(DecalMaterialPathName, FString);
+	}
+
+private:
+	const UE::Interchange::FAttributeKey Macro_CustomSortOrderKey = UE::Interchange::FAttributeKey(TEXT("SortOrder"));
+	const UE::Interchange::FAttributeKey Macro_CustomDecalSizeKey = UE::Interchange::FAttributeKey(TEXT("DecalSize"));
+	const UE::Interchange::FAttributeKey Macro_CustomDecalMaterialPathNameKey = UE::Interchange::FAttributeKey(TEXT("DecalMaterialPathName"));
+};
