Interchange: Fix component material override when importing a scene

#rnx

diff --git a/Engine/Plugins/Interchange/Runtime/Source/Import/Private/Scene/InterchangeStaticMeshActorFactory.cpp b/Engine/Plugins/Interchange/Runtime/Source/Import/Private/Scene/InterchangeStaticMeshActorFactory.cpp
--- a/Engine/Plugins/Interchange/Runtime/Source/Import/Private/Scene/InterchangeStaticMeshActorFactory.cpp
+++ b/Engine/Plugins/Interchange/Runtime/Source/Import/Private/Scene/InterchangeStaticMeshActorFactory.cpp
@@ -31,11 +31,10 @@ UObject* UInterchangeStaticMeshActorFactory::ProcessActor(AActor& SpawnedActor,
 	{
 		bool bHasGeometricTransform = false;
 		FTransform GeometricTransform;
-		if (FactoryNode.IsA(UInterchangeMeshActorFactoryNode::StaticClass()))
+		const UInterchangeMeshActorFactoryNode* MeshActorFactoryNode = Cast<UInterchangeMeshActorFactoryNode>(&FactoryNode);
+		if (MeshActorFactoryNode)
 		{
-			const UInterchangeMeshActorFactoryNode& MeshActorFactoryNode = static_cast<const UInterchangeMeshActorFactoryNode&>(FactoryNode);
-
-			if (MeshActorFactoryNode.GetCustomGeometricTransform(GeometricTransform))
+			if (MeshActorFactoryNode->GetCustomGeometricTransform(GeometricTransform))
 			{
 				bHasGeometricTransform = true;
 			}
@@ -81,6 +80,10 @@ UObject* UInterchangeStaticMeshActorFactory::ProcessActor(AActor& SpawnedActor,
 						{
 							GeometricTransformMeshComponent->SetStaticMesh(StaticMesh);
 						}
+						if (MeshActorFactoryNode)
+						{
+							UE::Interchange::ActorHelper::ApplySlotMaterialDependencies(NodeContainer, *MeshActorFactoryNode, *GeometricTransformMeshComponent);
+						}
 					}
 				}
 
@@ -106,6 +109,10 @@ UObject* UInterchangeStaticMeshActorFactory::ProcessActor(AActor& SpawnedActor,
 				{
 					StaticMeshComponent->SetStaticMesh(StaticMesh);
 				}
+				if (MeshActorFactoryNode)
+				{
+					UE::Interchange::ActorHelper::ApplySlotMaterialDependencies(NodeContainer, *MeshActorFactoryNode, *StaticMeshComponent);
+				}
 			}
 		}
 		else
