[Interchange] Collision prefix identification fixes

diff --git a/Engine/Plugins/Interchange/Runtime/Source/Pipelines/Private/InterchangeGenericStaticMeshPipeline.cpp b/Engine/Plugins/Interchange/Runtime/Source/Pipelines/Private/InterchangeGenericStaticMeshPipeline.cpp
--- a/Engine/Plugins/Interchange/Runtime/Source/Pipelines/Private/InterchangeGenericStaticMeshPipeline.cpp
+++ b/Engine/Plugins/Interchange/Runtime/Source/Pipelines/Private/InterchangeGenericStaticMeshPipeline.cpp
@@ -30,23 +30,22 @@ enum class EMeshCollisionType
 	Convex
 };
 
-
-static FStringView GetMeshNameFromUid(const FString& NodeUid)
+static TTuple<EMeshCollisionType, FString> GetCollisionMeshType(const UInterchangeBaseNodeContainer& NodeContainer, const FString& NodeUid, const TArray<FString>& AllNodeUids)
 {
-	FStringView MeshName = NodeUid;
-	int32 FinalDot = INDEX_NONE;
-	if (MeshName.FindLastChar(TEXT('.'), FinalDot))
+	auto GetNodeName = [&NodeContainer](const FString& NodeUid)
 	{
-		MeshName.RightChopInline(FinalDot + 1);
-	}
-
-	return MeshName;
-}
-
+		const UInterchangeBaseNode* BaseNode = NodeContainer.GetNode(NodeUid);
+		if (BaseNode)
+		{
+			return BaseNode->GetDisplayLabel();
+		}
+		else
+		{
+			return FString();
+		}
+	};
 
-static TTuple<EMeshCollisionType, FStringView> GetCollisionMeshType(const FString& NodeUid, const TArray<FString>& AllNodeUids)
-{
-	FStringView MeshName = GetMeshNameFromUid(NodeUid);
+	FString MeshName = GetNodeName(NodeUid);
 	EMeshCollisionType CollisionType = EMeshCollisionType::None;
 
 	// Determine if the mesh name is a potential collision mesh
@@ -69,7 +68,7 @@ static TTuple<EMeshCollisionType, FStringView> GetCollisionMeshType(const FStrin
 	}
 	else
 	{
-		return { EMeshCollisionType::None, FStringView() };
+		return { EMeshCollisionType::None, FString() };
 	}
 
 	// We have a mesh name with a collision type suffix.
@@ -82,10 +81,10 @@ static TTuple<EMeshCollisionType, FStringView> GetCollisionMeshType(const FStrin
 	int32 LastUnderscore = INDEX_NONE;
 	verify(MeshName.FindLastChar(TEXT('_'), LastUnderscore));
 
-	auto MatchPredicate = [](FStringView Body)
+	auto MatchPredicate = [&GetNodeName](FString Body)
 	{
 		// Generate a predicate to be used by the below Finds.
-		return [Body](const FString& ToCompare) { return Body == GetMeshNameFromUid(ToCompare); };
+		return [Body, &GetNodeName](const FString& ToCompare) { return Body == GetNodeName(ToCompare); };
 	};
 
 	// If we find a mesh named the same as the collision mesh (following the collision prefix), we have a match
@@ -111,21 +110,21 @@ static TTuple<EMeshCollisionType, FStringView> GetCollisionMeshType(const FStrin
 
 	// Mesh had a collision type prefix, but no corresponding mesh, so don't treat it as a collision mesh
 
-	return { EMeshCollisionType::None, FStringView() };
+	return { EMeshCollisionType::None, FString() };
 }
 
 
-static bool IsCollisionMeshUid(const FString& MeshUid, const TArray<FString>& MeshUids)
+static bool IsCollisionMeshUid(const UInterchangeBaseNodeContainer& NodeContainer, const FString& MeshUid, const TArray<FString>& MeshUids)
 {
-	return GetCollisionMeshType(MeshUid, MeshUids).Get<0>() != EMeshCollisionType::None;
+	return GetCollisionMeshType(NodeContainer, MeshUid, MeshUids).Get<0>() != EMeshCollisionType::None;
 }
 
 
-static void BuildMeshToCollisionMeshMap(const TArray<FString>& MeshUids, TMap<FString, TArray<FString>>& MeshToCollisionMeshMap)
+static void BuildMeshToCollisionMeshMap(const UInterchangeBaseNodeContainer& NodeContainer, const TArray<FString>& MeshUids, TMap<FString, TArray<FString>>& MeshToCollisionMeshMap)
 {
 	for (const FString& MeshUid : MeshUids)
 	{
-		TTuple<EMeshCollisionType, FStringView> CollisionType = GetCollisionMeshType(MeshUid, MeshUids);
+		TTuple<EMeshCollisionType, FString> CollisionType = GetCollisionMeshType(NodeContainer, MeshUid, MeshUids);
 		if (CollisionType.Get<0>() != EMeshCollisionType::None)
 		{
 			MeshToCollisionMeshMap.FindOrAdd(FString(CollisionType.Get<1>())).Emplace(MeshUid);
@@ -235,13 +234,13 @@ void UInterchangeGenericMeshPipeline::ExecutePreImportPipelineStaticMesh()
 				TMap<FString, TArray<FString>> MeshToCollisionMeshMap;
 				if (bImportCollisionAccordingToMeshName)
 				{
-					BuildMeshToCollisionMeshMap(MeshUids, MeshToCollisionMeshMap);
+					BuildMeshToCollisionMeshMap(*BaseNodeContainer, MeshUids, MeshToCollisionMeshMap);
 				}
 
 				// Now iterate through each mesh UID, creating a new factory for each one
 				for (const FString& MeshUid : MeshUids)
 				{
-					if (bImportCollisionAccordingToMeshName && IsCollisionMeshUid(MeshUid, MeshUids))
+					if (bImportCollisionAccordingToMeshName && IsCollisionMeshUid(*BaseNodeContainer, MeshUid, MeshUids))
 					{
 						// If this is a collision mesh, don't add a factory; it will be added as part of another factory
 						continue;
@@ -288,12 +287,12 @@ void UInterchangeGenericMeshPipeline::ExecutePreImportPipelineStaticMesh()
 				TMap<FString, TArray<FString>> MeshToCollisionMeshMap;
 				if (bImportCollisionAccordingToMeshName)
 				{
-					BuildMeshToCollisionMeshMap(MeshUids, MeshToCollisionMeshMap);
+					BuildMeshToCollisionMeshMap(*BaseNodeContainer, MeshUids, MeshToCollisionMeshMap);
 				}
 
 				for (const FString& MeshUid : MeshUids)
 				{
-					if (bImportCollisionAccordingToMeshName && IsCollisionMeshUid(MeshUid, MeshUids))
+					if (bImportCollisionAccordingToMeshName && IsCollisionMeshUid(*BaseNodeContainer, MeshUid, MeshUids))
 					{
 						// If this is a collision mesh, don't add a factory; it will be added as part of another factory
 						continue;
@@ -546,7 +545,7 @@ void UInterchangeGenericMeshPipeline::AddLodDataToStaticMesh(UInterchangeStaticM
 
 			if (bImportCollisionAccordingToMeshName)
 			{
-				TTuple<EMeshCollisionType, FStringView> MeshType = GetCollisionMeshType(NodeUid, NodeUids);
+				TTuple<EMeshCollisionType, FStringView> MeshType = GetCollisionMeshType(*BaseNodeContainer, NodeUid, NodeUids);
 				switch (MeshType.Get<0>())
 				{
 				case EMeshCollisionType::None:
