[Interchange - glTF] import has UV issue when importing into level

diff --git a/Engine/Plugins/Interchange/Runtime/Source/Import/Private/Gltf/InterchangeGLTFMaterial.cpp b/Engine/Plugins/Interchange/Runtime/Source/Import/Private/Gltf/InterchangeGLTFMaterial.cpp
--- a/Engine/Plugins/Interchange/Runtime/Source/Import/Private/Gltf/InterchangeGLTFMaterial.cpp
+++ b/Engine/Plugins/Interchange/Runtime/Source/Import/Private/Gltf/InterchangeGLTFMaterial.cpp
@@ -372,8 +372,6 @@ namespace UE::Interchange::GLTFMaterials
 						AngleRadians = TWO_PI - AngleRadians;
 					}
 
-					AngleRadians = 1.0f - (AngleRadians / TWO_PI);
-
 					SetScalar(Name + TEXT("_Rotation"), AngleRadians, 0.f);
 				}
 			}
