Fixed default values in generic material pipeline for check on duplicated materials
  - Default values should be false. Overlooked during the review :-(

diff --git a/Engine/Plugins/Interchange/Runtime/Source/Pipelines/Public/InterchangeGenericMaterialPipeline.h b/Engine/Plugins/Interchange/Runtime/Source/Pipelines/Public/InterchangeGenericMaterialPipeline.h
--- a/Engine/Plugins/Interchange/Runtime/Source/Pipelines/Public/InterchangeGenericMaterialPipeline.h
+++ b/Engine/Plugins/Interchange/Runtime/Source/Pipelines/Public/InterchangeGenericMaterialPipeline.h
@@ -57,11 +57,11 @@ public:
 	
 	/** If set, reference materials along with respective material instances would be created*/
 	UPROPERTY(EditAnywhere, BlueprintReadWrite, Category = "Materials")
-	bool bIdentifyDuplicateMaterials = true;
+	bool bIdentifyDuplicateMaterials = false;
 
 	/** If set, additional material instance would be created for reference/parent material*/
 	UPROPERTY(EditAnywhere, BlueprintReadWrite, Category = "Materials", Meta=(EditCondition="bIdentifyDuplicateMaterials"))
-	bool bCreateMaterialInstanceForParent = true;
+	bool bCreateMaterialInstanceForParent = false;
 
 	/** Optional material used as the parent when importing materials as instances. If no parent material is specified, one will be automatically selected during the import process. */
 	UPROPERTY(EditAnywhere, BlueprintReadWrite, Category = "Materials", Meta= (EditCondition="bImportMaterials && MaterialImport==EInterchangeMaterialImportOption::ImportAsMaterialInstances", AllowedClasses="/Script/Engine.MaterialInterface"))
