[Interchange] GeometryTransform related component mobility fixes

diff --git a/Engine/Plugins/Interchange/Runtime/Source/Import/Private/Scene/InterchangeStaticMeshActorFactory.cpp b/Engine/Plugins/Interchange/Runtime/Source/Import/Private/Scene/InterchangeStaticMeshActorFactory.cpp
--- a/Engine/Plugins/Interchange/Runtime/Source/Import/Private/Scene/InterchangeStaticMeshActorFactory.cpp
+++ b/Engine/Plugins/Interchange/Runtime/Source/Import/Private/Scene/InterchangeStaticMeshActorFactory.cpp
@@ -68,8 +68,6 @@ UObject* UInterchangeStaticMeshActorFactory::ProcessActor(AActor& SpawnedActor,
 #endif
 				StaticMeshActor->AddInstanceComponent(GeometricTransformMeshComponent);
 
-				GeometricTransformMeshComponent->SetMobility(StaticMeshComponent->Mobility);
-				
 				if (const UInterchangeFactoryBaseNode* MeshNode = ActorHelper::FindAssetInstanceFactoryNode(&NodeContainer, &FactoryNode))
 				{
 					FSoftObjectPath ReferenceObject;
@@ -91,7 +89,10 @@ UObject* UInterchangeStaticMeshActorFactory::ProcessActor(AActor& SpawnedActor,
 
 				StaticMeshActor->ReregisterAllComponents();
 
+				EComponentMobility::Type MobilityToSet = StaticMeshComponent->Mobility;
+				GeometricTransformMeshComponent->SetMobility(EComponentMobility::Type::Movable); //so that RelativeTransform can be set
 				GeometricTransformMeshComponent->SetRelativeTransform(GeometricTransform);
+				GeometricTransformMeshComponent->SetMobility(MobilityToSet);
 
 				return StaticMeshComponent;
 			}
