Interchange: Fix a crash when translating texture during mesh import. We must avoid using FGCObjectScopeGuard outside the game thread.

#rnx

diff --git a/Engine/Plugins/Interchange/Runtime/Source/Import/Private/Fbx/InterchangeFbxTranslator.cpp b/Engine/Plugins/Interchange/Runtime/Source/Import/Private/Fbx/InterchangeFbxTranslator.cpp
--- a/Engine/Plugins/Interchange/Runtime/Source/Import/Private/Fbx/InterchangeFbxTranslator.cpp
+++ b/Engine/Plugins/Interchange/Runtime/Source/Import/Private/Fbx/InterchangeFbxTranslator.cpp
@@ -10,6 +10,7 @@
 #include "InterchangeDispatcherTask.h"
 #include "InterchangeManager.h"
 #include "InterchangeImportLog.h"
+#include "InterchangeTranslatorHelper.h"
 #include "Mesh/InterchangeMeshPayload.h"
 #include "Misc/ConfigCacheIni.h"
 #include "Misc/FileHelper.h"
@@ -185,24 +186,13 @@ void UInterchangeFbxTranslator::ImportFinish()
 
 TOptional<UE::Interchange::FImportImage> UInterchangeFbxTranslator::GetTexturePayloadData(const FString& PayLoadKey, TOptional<FString>& AlternateTexturePath) const
 {
-	UInterchangeSourceData* PayloadSourceData = UInterchangeManager::GetInterchangeManager().CreateSourceData(PayLoadKey);
-	FGCObjectScopeGuard ScopedSourceData(PayloadSourceData);
-	
-	if (!PayloadSourceData)
-	{
-		return TOptional<UE::Interchange::FImportImage>();
-	}
-	UInterchangeTranslatorBase* SourceTranslator = UInterchangeManager::GetInterchangeManager().GetTranslatorForSourceData(PayloadSourceData);
-	FGCObjectScopeGuard ScopedSourceTranslator(SourceTranslator);
-	const IInterchangeTexturePayloadInterface* TextureTranslator = Cast<IInterchangeTexturePayloadInterface>(SourceTranslator);
+	UE::Interchange::Private::FScopedTranslator ScopedTranslator(PayLoadKey, Results);
+	const IInterchangeTexturePayloadInterface* TextureTranslator = ScopedTranslator.GetPayLoadInterface<IInterchangeTexturePayloadInterface>();
 	if (!ensure(TextureTranslator))
 	{
 		return TOptional<UE::Interchange::FImportImage>();
 	}
-
 	AlternateTexturePath = PayLoadKey;
-	SourceTranslator->SetResultsContainer(Results);
-
 	return TextureTranslator->GetTexturePayloadData(PayLoadKey, AlternateTexturePath);
 }
 
diff --git a/Engine/Plugins/Interchange/Runtime/Source/Import/Private/Gltf/InterchangeGltfTranslator.cpp b/Engine/Plugins/Interchange/Runtime/Source/Import/Private/Gltf/InterchangeGltfTranslator.cpp
--- a/Engine/Plugins/Interchange/Runtime/Source/Import/Private/Gltf/InterchangeGltfTranslator.cpp
+++ b/Engine/Plugins/Interchange/Runtime/Source/Import/Private/Gltf/InterchangeGltfTranslator.cpp
@@ -21,6 +21,7 @@
 #include "InterchangeSceneNode.h"
 #include "InterchangeShaderGraphNode.h"
 #include "InterchangeTexture2DNode.h"
+#include "InterchangeTranslatorHelper.h"
 #include "InterchangeVariantSetNode.h"
 #include "Nodes/InterchangeSourceNode.h"
 
@@ -791,26 +792,13 @@ TOptional< UE::Interchange::FImportImage > UInterchangeGLTFTranslator::GetTextur
 	else
 	{
 		const FString TextureFilePath = FPaths::ConvertRelativePathToFull(GltfTexture.Source.FilePath);
-
-		UInterchangeSourceData* PayloadSourceData = UInterchangeManager::GetInterchangeManager().CreateSourceData(TextureFilePath);
-		FGCObjectScopeGuard ScopedSourceData(PayloadSourceData);
-
-		if (!PayloadSourceData)
-		{
-			return TexturePayloadData;
-		}
-
-		UInterchangeTranslatorBase* SourceTranslator = UInterchangeManager::GetInterchangeManager().GetTranslatorForSourceData(PayloadSourceData);
-		FGCObjectScopeGuard ScopedSourceTranslator(SourceTranslator);
-		const IInterchangeTexturePayloadInterface* TextureTranslator = Cast< IInterchangeTexturePayloadInterface >(SourceTranslator);
+		UE::Interchange::Private::FScopedTranslator ScopedTranslator(TextureFilePath, Results);
+		const IInterchangeTexturePayloadInterface* TextureTranslator = ScopedTranslator.GetPayLoadInterface<IInterchangeTexturePayloadInterface>();
 		if (!ensure(TextureTranslator))
 		{
 			return TexturePayloadData;
 		}
-		SourceTranslator->SetResultsContainer(Results);
-
 		AlternateTexturePath = TextureFilePath;
-
 		TexturePayloadData = TextureTranslator->GetTexturePayloadData(PayloadKey, AlternateTexturePath);
 	}
 
diff --git a/Engine/Plugins/Interchange/Runtime/Source/Import/Private/InterchangeTranslatorHelper.cpp b/Engine/Plugins/Interchange/Runtime/Source/Import/Private/InterchangeTranslatorHelper.cpp
new file mode 100644
--- /dev/null
+++ b/Engine/Plugins/Interchange/Runtime/Source/Import/Private/InterchangeTranslatorHelper.cpp
@@ -0,0 +1,53 @@
+// Copyright Epic Games, Inc. All Rights Reserved. 
+
+#include "InterchangeTranslatorHelper.h"
+
+#include "CoreGlobals.h"
+#include "InterchangeManager.h"
+#include "InterchangeResultsContainer.h"
+#include "InterchangeSourceData.h"
+#include "InterchangeTranslatorBase.h"
+#include "UObject/Object.h"
+
+namespace UE::Interchange::Private
+{
+	FScopedTranslator::FScopedTranslator(const FString& PayLoadKey, UInterchangeResultsContainer* Results)
+	{
+		PayloadSourceData = UInterchangeManager::GetInterchangeManager().CreateSourceData(PayLoadKey);
+		if (PayloadSourceData)
+		{
+			PayloadSourceData->SetInternalFlags(EInternalObjectFlags::Async);
+
+			SourceTranslator = UInterchangeManager::GetInterchangeManager().GetTranslatorForSourceData(PayloadSourceData);
+			if (SourceTranslator)
+			{
+				SourceTranslator->SetInternalFlags(EInternalObjectFlags::Async);
+				if (Results)
+				{
+					SourceTranslator->SetResultsContainer(Results);
+				}
+			}
+		}
+	}
+
+	FScopedTranslator::~FScopedTranslator()
+	{
+		if (PayloadSourceData)
+		{
+			PayloadSourceData->ClearInternalFlags(EInternalObjectFlags::Async);
+		}
+
+		if (SourceTranslator)
+		{
+			SourceTranslator->ClearInternalFlags(EInternalObjectFlags::Async);
+		}
+
+		SourceTranslator = nullptr;
+		PayloadSourceData = nullptr;
+	}
+
+	bool FScopedTranslator::IsValid() const
+	{
+		return SourceTranslator != nullptr;
+	}
+} //namespace UE::Interchange::Private
diff --git a/Engine/Plugins/Interchange/Runtime/Source/Import/Private/Mesh/InterchangeOBJTranslator.cpp b/Engine/Plugins/Interchange/Runtime/Source/Import/Private/Mesh/InterchangeOBJTranslator.cpp
--- a/Engine/Plugins/Interchange/Runtime/Source/Import/Private/Mesh/InterchangeOBJTranslator.cpp
+++ b/Engine/Plugins/Interchange/Runtime/Source/Import/Private/Mesh/InterchangeOBJTranslator.cpp
@@ -13,6 +13,7 @@
 #include "InterchangeShaderGraphNode.h"
 #include "InterchangeTexture2DNode.h"
 #include "InterchangeTextureNode.h"
+#include "InterchangeTranslatorHelper.h"
 #include "MaterialDomain.h"
 #include "Materials/Material.h"
 #include "Misc/ConfigCacheIni.h"
@@ -1576,25 +1577,13 @@ TOptional<UE::Interchange::FImportImage> UInterchangeOBJTranslator::GetTexturePa
 	// Here, we have a texture source referenced by the .obj which is ideally handled by the existing texture translators.
 	// In this case, we can just implement the texture payload interface here and use a temporary texture translator instance
 	// to generate the payload, but what if the translator needed to do non-trivial work during the Translate() phase?
-	UInterchangeSourceData* PayloadSourceData = UInterchangeManager::GetInterchangeManager().CreateSourceData(PayLoadKey);
-	FGCObjectScopeGuard ScopedSourceData(PayloadSourceData);
-	if (!PayloadSourceData)
-	{
-		return TOptional<UE::Interchange::FImportImage>();
-	}
-
-	UInterchangeTranslatorBase* SourceTranslator = UInterchangeManager::GetInterchangeManager().GetTranslatorForSourceData(PayloadSourceData);
-	FGCObjectScopeGuard ScopedSourceTranslator(SourceTranslator);
-	const IInterchangeTexturePayloadInterface* TextureTranslator = Cast<IInterchangeTexturePayloadInterface>(SourceTranslator);
+	UE::Interchange::Private::FScopedTranslator ScopedTranslator(PayLoadKey, Results);
+	const IInterchangeTexturePayloadInterface* TextureTranslator = ScopedTranslator.GetPayLoadInterface<IInterchangeTexturePayloadInterface>();
 	if (!ensure(TextureTranslator))
 	{
 		return TOptional<UE::Interchange::FImportImage>();
 	}
-
-	SourceTranslator->SetResultsContainer(Results);
-
 	AlternateTexturePath = PayLoadKey;
-
 	return TextureTranslator->GetTexturePayloadData(PayLoadKey, AlternateTexturePath);
 }
 
diff --git a/Engine/Plugins/Interchange/Runtime/Source/Import/Public/InterchangeTranslatorHelper.h b/Engine/Plugins/Interchange/Runtime/Source/Import/Public/InterchangeTranslatorHelper.h
new file mode 100644
--- /dev/null
+++ b/Engine/Plugins/Interchange/Runtime/Source/Import/Public/InterchangeTranslatorHelper.h
@@ -0,0 +1,35 @@
+// Copyright Epic Games, Inc. All Rights Reserved. 
+#pragma once
+
+#include "InterchangeSourceData.h"
+#include "InterchangeTranslatorBase.h"
+#include "Containers/UnrealString.h"
+
+class UInterchangeResultsContainer;
+
+namespace UE::Interchange::Private
+{
+	struct INTERCHANGEIMPORT_API FScopedTranslator
+	{
+	public:
+		FScopedTranslator(const FString& PayLoadKey, UInterchangeResultsContainer* Results);
+		~FScopedTranslator();
+
+		template <class T>
+		const T* GetPayLoadInterface() const
+		{
+			if(!IsValid())
+			{ 
+				return nullptr;
+			}
+			const T* PayloadInterface = Cast<T>(SourceTranslator);
+			return PayloadInterface;
+		}
+	
+	private:
+		bool IsValid() const;
+
+		TObjectPtr<UInterchangeTranslatorBase> SourceTranslator = nullptr;
+		TObjectPtr<UInterchangeSourceData> PayloadSourceData = nullptr;
+	};
+} //namespace UE::Interchange::Private
