[Interchange - glTF] Animation handling refactor.
[Interchange - glTF] Animation fixes for LocalBindPoses without animations.

diff --git a/Engine/Plugins/Interchange/Runtime/Source/Import/Private/Gltf/InterchangeGltfAnimation.cpp b/Engine/Plugins/Interchange/Runtime/Source/Import/Private/Gltf/InterchangeGltfAnimation.cpp
--- a/Engine/Plugins/Interchange/Runtime/Source/Import/Private/Gltf/InterchangeGltfAnimation.cpp
+++ b/Engine/Plugins/Interchange/Runtime/Source/Import/Private/Gltf/InterchangeGltfAnimation.cpp
@@ -11,8 +11,15 @@
 #include "InterchangeImportLog.h"
 #include "Animation/AnimTypes.h"
 
+#include "Nodes/InterchangeBaseNodeContainer.h"
+#include "InterchangeSceneNode.h"
+#include "InterchangeMeshNode.h"
+#include "InterchangeAnimationTrackSetNode.h"
+
 namespace UE::Interchange::Gltf::Private
 {
+	const FString BIND_POSE_FIX = TEXT("BIND_POSE_FIX<->");
+
 	/*
 	* According to gltf specification Seconds acquired from Samplers.input:
 	* "The values represent time in seconds with time[0] >= 0.0, and strictly increasing values, Index.e., time[n + 1] > time[n]."
@@ -451,19 +458,31 @@ namespace UE::Interchange::Gltf::Private
 
 	bool GetBakedAnimationTransformPayloadData(const FString& PayLoadKey, const GLTF::FAsset& GltfAsset, FAnimationPayloadData& PayloadData)
 	{
+		const double BakeInterval = 1.0 / PayloadData.BakeFrequency;
+		const double SequenceLength = FMath::Max<double>(PayloadData.RangeEndTime - PayloadData.RangeStartTime, MINIMUM_ANIMATION_LENGTH);
+		int32 FrameCount = FMath::RoundToInt32(SequenceLength * PayloadData.BakeFrequency);
+		int32 BakeKeyCount = FrameCount + 1;
+
 		TArray<int32> ChannelIndices;
 		int32 AnimationIndex;
 
 		if (!ParsePayLoadKey(GltfAsset, PayLoadKey, AnimationIndex, ChannelIndices))
 		{
+			if (PayLoadKey.Contains(BIND_POSE_FIX))
+			{
+				FString CutPayloadKey = PayLoadKey.Replace(*BIND_POSE_FIX, TEXT(""));
+				int32 NodeToSetIndex = INDEX_NONE;
+				LexFromString(NodeToSetIndex, *CutPayloadKey);
+
+				if (GltfAsset.Nodes.IsValidIndex(NodeToSetIndex))
+				{
+					PayloadData.Transforms.Init(GltfAsset.Nodes[NodeToSetIndex].Transform, BakeKeyCount);
+					return true;
+				}
+			}
 			return false;
 		}
 		const GLTF::FAnimation& GltfAnimation = GltfAsset.Animations[AnimationIndex];
-		
-		const double BakeInterval = 1.0 / PayloadData.BakeFrequency;
-		const double SequenceLength = FMath::Max<double>(PayloadData.RangeEndTime - PayloadData.RangeStartTime, MINIMUM_ANIMATION_LENGTH);
-		int32 FrameCount = FMath::RoundToInt32(SequenceLength * PayloadData.BakeFrequency);
-		int32 BakeKeyCount = FrameCount + 1;
 
 		//buffers to use for final Transform:
 		TArray<FVector3f> TranslationData;
@@ -577,4 +596,452 @@ namespace UE::Interchange::Gltf::Private
 
 		return true;
 	}
+
+	namespace AnimationHelpers
+	{
+		struct FAnimationHandler
+		{
+			UInterchangeBaseNodeContainer& NodeContainer;
+			int32 AnimationIndex;
+			const GLTF::FAnimation& GLTFAnimation;
+			const TArray<GLTF::FNode>& GLTFNodes;
+			const TMap<const GLTF::FNode*, FString>& GLTFNodeToInterchangeUidMap; //UInterchangeGLTFTranslator::NodeUidMap
+
+			TMap<FString, UInterchangeSkeletalAnimationTrackNode*> RootJointIndexToTrackNodeMap;
+
+			FAnimationHandler(UInterchangeBaseNodeContainer& InNodeContainer, 
+				int32 InAnimationIndex, const GLTF::FAnimation& InGLTFAnimation, 
+				const TArray<GLTF::FNode>& InGLTFNodes,
+				const TMap<const GLTF::FNode*, FString>& InGLTFNodeToInterchangeUidMap)
+				: NodeContainer(InNodeContainer)
+				, AnimationIndex(InAnimationIndex)
+				, GLTFAnimation(InGLTFAnimation)
+				, GLTFNodes(InGLTFNodes)
+				, GLTFNodeToInterchangeUidMap(InGLTFNodeToInterchangeUidMap)
+			{}
+
+			struct FAnimationTimes
+			{
+				//@ StartTime = 0;
+				//from gltf documentation: https://registry.khronos.org/glTF/specs/2.0/glTF-2.0.html#animations
+				//Implementation Note
+				//	For example, if the earliest sampler input for an animation is t = 10, a client implementation must begin playback of that animation channel at t = 0 with output clamped to the first available output value.
+				double StartTime = 0.0;
+
+				double FrameRate = 30.0;
+				double SingleFrameDuration = 1.0 / FrameRate;
+				double StopTime = SingleFrameDuration;
+				int32 FrameNumber = 0;
+			};
+
+
+			UInterchangeSkeletalAnimationTrackNode* AcquireTrackNode(const FString& SkeletonNodeUid, FAnimationTimes& AnimationTimes)
+			{
+				UInterchangeSkeletalAnimationTrackNode* TrackNode = nullptr;
+				if (RootJointIndexToTrackNodeMap.Contains(SkeletonNodeUid))
+				{
+					TrackNode = RootJointIndexToTrackNodeMap[SkeletonNodeUid];
+
+					double SampleRate;
+					if (TrackNode->GetCustomAnimationSampleRate(SampleRate))
+					{
+						AnimationTimes.FrameRate = SampleRate;
+						AnimationTimes.SingleFrameDuration = 1.0 / SampleRate;
+					}
+
+					double StopTime;
+					if (TrackNode->GetCustomAnimationStopTime(StopTime))
+					{
+						AnimationTimes.StopTime = StopTime;
+						AnimationTimes.FrameNumber = FMath::RoundToInt32(StopTime / AnimationTimes.SingleFrameDuration);
+					}
+				}
+				else
+				{
+					TrackNode = NewObject< UInterchangeSkeletalAnimationTrackNode >(&NodeContainer);
+					FString TrackNodeUid = "\\SkeletalAnimation\\" + SkeletonNodeUid + "_" + LexToString(AnimationIndex);
+					TrackNode->InitializeNode(TrackNodeUid, GLTFAnimation.Name, EInterchangeNodeContainerType::TranslatedAsset);
+					TrackNode->SetCustomSkeletonNodeUid(SkeletonNodeUid);
+
+					NodeContainer.AddNode(TrackNode);
+
+					RootJointIndexToTrackNodeMap.Add(SkeletonNodeUid, TrackNode);
+				}
+
+				return TrackNode;
+			};
+
+			bool ValidateChannelSampler(const GLTF::FAnimation::FChannel& Channel, const GLTF::FAnimation::FSampler& Sampler)
+			{
+				if (Sampler.Interpolation == GLTF::FAnimation::EInterpolation::CubicSpline)
+				{
+					if (Sampler.Input.Count != 3 * Sampler.Output.Count)
+					{
+						return false;
+
+					}
+				}
+				else
+				{
+					if (Channel.Target.Path != GLTF::FAnimation::EPath::Weights && (Sampler.Input.Count != Sampler.Output.Count))
+					{
+						return false;
+					}
+				}
+				return true;
+			};
+
+			void UpdateAnimationTimes(const GLTF::FAnimation::FSampler& Sampler, FAnimationTimes& AnimationTimes)
+			{
+				TArray<float> Seconds;
+				Sampler.Input.GetFloatArray(Seconds);
+
+				int32 CurrentFrameNumber = 0;
+
+				if (Seconds.Num() > 0)
+				{
+					//calculate FrameNumber and currentStopTime:
+					float CurrentStopTime = Seconds[Seconds.Num() - 1];
+
+					float CurrentFrameNumberCandidate = CurrentStopTime / AnimationTimes.SingleFrameDuration;
+					CurrentFrameNumber = int(CurrentFrameNumberCandidate);
+					if (int(CurrentFrameNumberCandidate) < CurrentFrameNumberCandidate)
+					{
+						CurrentFrameNumber++;
+					}
+					CurrentStopTime = CurrentFrameNumber * AnimationTimes.SingleFrameDuration;
+				}
+
+				if (AnimationTimes.FrameNumber < CurrentFrameNumber)
+				{
+					AnimationTimes.FrameNumber = CurrentFrameNumber;
+					AnimationTimes.StopTime = AnimationTimes.FrameNumber * AnimationTimes.SingleFrameDuration;
+				}
+			};
+
+			void AcquireJointsWithBindPose(int32 CurrentIndex, TSet<int32>& Joints)
+			{
+				if (GLTFNodes[CurrentIndex].Type == GLTF::FNode::EType::Joint)
+				{
+					if (GLTFNodes[CurrentIndex].bHasLocalBindPose)
+					{
+						Joints.Add(CurrentIndex);
+					}
+
+					for (int32 ChildIndex : GLTFNodes[CurrentIndex].Children)
+					{
+						AcquireJointsWithBindPose(ChildIndex, Joints);
+					}
+				}
+			}
+
+			void ProcessRiggedAnimations(const TMap<int32, TMap<int32, TSet<int32>>>& RiggedAnimations)
+			{
+				for (const TPair<int32, TMap<int32, TSet<int32>>>& RiggedAnimation : RiggedAnimations)
+				{
+					int32 NodeIndex = RiggedAnimation.Key;
+					const GLTF::FNode* GLTFNode = &GLTFNodes[NodeIndex];
+
+					//Acquire Node Interchange UniqueID:
+					const FString* NodeUidPtr = GLTFNodeToInterchangeUidMap.Find(GLTFNode);
+					if (!ensure(NodeUidPtr))
+					{
+						continue;
+					}
+
+					FAnimationTimes AnimationTimes;
+
+					//Acquire/Create UInterchangeSkeletalAnimationTrackNode:
+					UInterchangeSkeletalAnimationTrackNode* TrackNode = AcquireTrackNode(*NodeUidPtr, AnimationTimes);
+
+					//iterate (AnimatedNodeIndex, [Channels])
+					for (const TTuple<int32, TSet<int32>>& AnimatedNodeIndexToChannelIndices : RiggedAnimation.Value)
+					{
+						int32 AnimatedNodeIndex = AnimatedNodeIndexToChannelIndices.Key;
+						const FString* AnimatedNodeUidPtr = GLTFNodeToInterchangeUidMap.Find(&GLTFNodes[AnimatedNodeIndex]);
+						if (!ensure(AnimatedNodeUidPtr))
+						{
+							continue;
+						}
+
+						//double PreviousStopTime = StopTime;
+						//check channel length and build payload
+						FString Payload = TEXT("");
+						for (int32 ChannelIndex : AnimatedNodeIndexToChannelIndices.Value)
+						{
+							const GLTF::FAnimation::FChannel& Channel = GLTFAnimation.Channels[ChannelIndex];
+							const GLTF::FAnimation::FSampler& Sampler = GLTFAnimation.Samplers[Channel.Sampler];
+
+							if (!ValidateChannelSampler(Channel, Sampler))
+							{
+								// if any of the channels are corrupt the joint will not receive any of the  animation data
+								UE_LOG(LogInterchangeImport, Warning, TEXT("GLTF Sampler Corrupt. Input and Output not meeting expectations."));
+								break;
+							}
+
+							UpdateAnimationTimes(Sampler, AnimationTimes);
+
+							if (Channel.Target.Path != GLTF::FAnimation::EPath::Weights)
+							{
+								Payload += ":" + LexToString(ChannelIndex);
+							}
+						}
+
+						if (Payload.Len() > 0)
+						{
+							Payload = LexToString(AnimationIndex) + Payload;
+							TrackNode->SetAnimationPayloadKeyForSceneNodeUid(*AnimatedNodeUidPtr, Payload, EInterchangeAnimationPayLoadType::BAKED);
+						}
+					}
+
+					//set animation length:
+					TrackNode->SetCustomAnimationSampleRate(AnimationTimes.FrameRate);
+					TrackNode->SetCustomAnimationStartTime(AnimationTimes.StartTime);
+					TrackNode->SetCustomAnimationStopTime(AnimationTimes.StopTime);
+				}
+			}
+
+			void ProcessMorphTargetAnimations(const TMap<int32, TMap<int32, TSet<int32>>>& MorphTargetAnimations)
+			{
+				ProcessRiggedAnimations(MorphTargetAnimations);
+
+				for (const TPair<int32, TMap<int32, TSet<int32>>>& RiggedAnimation : MorphTargetAnimations)
+				{
+					int32 NodeIndex = RiggedAnimation.Key;
+					const GLTF::FNode* GLTFNode = &GLTFNodes[NodeIndex];
+
+					//Acquire Interchange UniqueID:
+					const FString* NodeUidPtr = GLTFNodeToInterchangeUidMap.Find(GLTFNode);
+					if (!ensure(NodeUidPtr))
+					{
+						continue;
+					}
+
+					FAnimationTimes AnimationTimes;
+
+					//Acquire/Create UInterchangeSkeletalAnimationTrackNode:
+					UInterchangeSkeletalAnimationTrackNode* TrackNode = AcquireTrackNode(*NodeUidPtr, AnimationTimes);
+
+					//iterate (AnimatedNodeIndex, [Channels])
+					for (const TTuple<int32, TSet<int32>>& AnimatedNodeIndexToChannelIndices : RiggedAnimation.Value)
+					{
+						int32 AnimatedNodeIndex = AnimatedNodeIndexToChannelIndices.Key;
+						const FString* AnimatedNodeUidPtr = GLTFNodeToInterchangeUidMap.Find(&GLTFNodes[AnimatedNodeIndex]);
+						if (!ensure(AnimatedNodeUidPtr))
+						{
+							continue;
+						}
+
+						const TSet<int32>& Channels = AnimatedNodeIndexToChannelIndices.Value;
+						ensure(Channels.Num() == 1);
+
+						//Find SceneNode that references the MeshNode:
+						if (const UInterchangeSceneNode* ConstSceneMeshActorNode = Cast< UInterchangeSceneNode >(NodeContainer.GetNode(*AnimatedNodeUidPtr)))
+						{
+							FString SkeletalMeshUid;
+							if (ConstSceneMeshActorNode->GetCustomAssetInstanceUid(SkeletalMeshUid))
+							{
+								if (const UInterchangeMeshNode* MeshNode = Cast< UInterchangeMeshNode >(NodeContainer.GetNode(SkeletalMeshUid)))
+								{
+									TArray<FString> MorphTargetDependencies;
+									MeshNode->GetMorphTargetDependencies(MorphTargetDependencies);
+									for (const FString& MorphTargetDependencyUid : MorphTargetDependencies)
+									{
+										if (const UInterchangeMeshNode* MorphTargetNodeConst = Cast< UInterchangeMeshNode >(NodeContainer.GetNode(MorphTargetDependencyUid)))
+										{
+											if (MorphTargetNodeConst->GetPayLoadKey().IsSet())
+											{
+												FInterchangeMeshPayLoadKey PayLoadKey = MorphTargetNodeConst->GetPayLoadKey().GetValue();
+												FString PayLoadKeyUniqueId = LexToString(AnimationIndex) + TEXT(":") + LexToString(*Channels.begin()) + TEXT(":") + PayLoadKey.UniqueId;
+
+												TrackNode->SetAnimationPayloadKeyForMorphTargetNodeUid(MorphTargetDependencyUid, PayLoadKeyUniqueId, EInterchangeAnimationPayLoadType::MORPHTARGETCURVE);
+											}
+										}
+									}
+								}
+							}
+						}
+					}
+				}
+			}
+
+			void ProcessRigidAnimation(const TMap<int32, TSet<int32>>& RigidAnimation)
+			{
+				if (RigidAnimation.Num() == 0)
+				{
+					return;
+				}
+
+				UInterchangeAnimationTrackSetNode* TrackSetNode = NewObject< UInterchangeAnimationTrackSetNode >(&NodeContainer);
+
+				const FString AnimTrackSetNodeUid = TEXT("\\Animation\\") + GLTFAnimation.UniqueId;
+				TrackSetNode->InitializeNode(AnimTrackSetNodeUid, GLTFAnimation.Name, EInterchangeNodeContainerType::TranslatedAsset);
+
+				for (const TTuple<int32, TSet<int32>>& NodeChannelsEntry : RigidAnimation)
+				{
+					if (NodeChannelsEntry.Value.Num() == 0)
+					{
+						continue;
+					}
+
+					const GLTF::FNode& GltfNode = GLTFNodes[NodeChannelsEntry.Key];
+					const FString* NodeUid = GLTFNodeToInterchangeUidMap.Find(&GltfNode);
+					if (!ensure(NodeUid))
+					{
+						continue;
+					}
+
+					UInterchangeTransformAnimationTrackNode* TransformAnimTrackNode = NewObject< UInterchangeTransformAnimationTrackNode >(&NodeContainer);
+
+					const FString TransformAnimTrackNodeName = FString::Printf(TEXT("%s_%s"), *GltfNode.Name, *GLTFAnimation.Name);
+					const FString TransformAnimTrackNodeUid = TEXT("\\AnimationTrack\\") + TransformAnimTrackNodeName;
+
+					TransformAnimTrackNode->InitializeNode(TransformAnimTrackNodeUid, TransformAnimTrackNodeName, EInterchangeNodeContainerType::TranslatedAsset);
+
+					TransformAnimTrackNode->SetCustomActorDependencyUid(*NodeUid);
+
+					FString PayloadKey = FString::FromInt(AnimationIndex);
+
+					constexpr int32 TranslationChannel = 0x0001 | 0x0002 | 0x0004;
+					constexpr int32 RotationChannel = 0x0008 | 0x0010 | 0x0020;
+					constexpr int32 ScaleChannel = 0x0040 | 0x0080 | 0x0100;
+
+					int32 UsedChannels = 0;
+
+					for (int32 ChannelIndex : NodeChannelsEntry.Value)
+					{
+						PayloadKey += TEXT(":") + FString::FromInt(ChannelIndex);
+
+						const GLTF::FAnimation::FChannel& Channel = GLTFAnimation.Channels[ChannelIndex];
+
+						switch (Channel.Target.Path)
+						{
+						case GLTF::FAnimation::EPath::Translation:
+						{
+							UsedChannels |= TranslationChannel;
+						} break;
+
+						case GLTF::FAnimation::EPath::Rotation:
+						{
+							UsedChannels |= RotationChannel;
+						} break;
+
+						case GLTF::FAnimation::EPath::Scale:
+						{
+							UsedChannels |= ScaleChannel;
+						} break;
+						default: break;
+						}
+					}
+
+					TransformAnimTrackNode->SetCustomAnimationPayloadKey(PayloadKey, EInterchangeAnimationPayLoadType::CURVE);
+					TransformAnimTrackNode->SetCustomUsedChannels(UsedChannels);
+
+					NodeContainer.AddNode(TransformAnimTrackNode);
+
+					TrackSetNode->AddCustomAnimationTrackUid(TransformAnimTrackNodeUid);
+				}
+
+				NodeContainer.AddNode(TrackSetNode);
+			}
+
+			void FixSkeletalAnimations(TMap<int32, TSet<int32>>& SkeletonRootToAnimatedJointNodeIndicesMap)
+			{
+				for (const TTuple<int32, TSet<int32>>& AnimatedJointNodeIndices : SkeletonRootToAnimatedJointNodeIndicesMap)
+				{
+					int32 SkeletonRootIndex = AnimatedJointNodeIndices.Key;
+
+					const FString* SkeletonUid = GLTFNodeToInterchangeUidMap.Find(&GLTFNodes[SkeletonRootIndex]);
+					if (!ensure(SkeletonUid))
+					{
+						continue;
+					}
+
+					FAnimationTimes AnimationTimes;
+					UInterchangeSkeletalAnimationTrackNode* TrackNode = AcquireTrackNode(*SkeletonUid, AnimationTimes);
+
+					TSet<int32> JointsWithBindPose;
+					AcquireJointsWithBindPose(SkeletonRootIndex, JointsWithBindPose);
+
+					TSet<int32> NodesToAddToAnimation = JointsWithBindPose.Difference(AnimatedJointNodeIndices.Value);
+					for (int32 NodeToAddIndex : NodesToAddToAnimation)
+					{
+						const FString* NodeToAddUidPtr = GLTFNodeToInterchangeUidMap.Find(&GLTFNodes[NodeToAddIndex]);
+						if (!ensure(NodeToAddUidPtr))
+						{
+							continue;
+						}
+						TrackNode->SetAnimationPayloadKeyForSceneNodeUid(*NodeToAddUidPtr, BIND_POSE_FIX + LexToString(NodeToAddIndex), EInterchangeAnimationPayLoadType::BAKED);
+					}
+				}
+			}
+
+			void Process()
+			{
+				TMap<int32, TSet<int32>> RigidAnimation;							// (AnimatedNodeIndex, [Channels])
+				TMap<int32, TMap<int32, TSet<int32>>> RiggedAnimations;				// (SkeletonRootIndex, (AnimatedNodeIndex, [Channels])
+				TMap<int32, TMap<int32, TSet<int32>>> MorphTargetAnimations;		// (AnimatedNodeIndex, (AnimatedNodeIndex, [Channels]) //Channel is a Target.Path==Weight channel.
+																					// glTF allows MorphTargets on StaticMeshes as well.
+																					// However UE only does MorphTargets on Skeletals, for this reason we presume Skeletals and Skeletal Animations for Morph Target Animations.
+
+				// Important: AnimatedNodes where bHasLocalBindPose==true
+				TMap<int32, TSet<int32>> SkeletonRootToAnimatedJointNodeIndicesMap;	// (SkeletonRootIndex, [AnimatedNodeIndices])
+
+				for (int32 ChannelIndex = 0; ChannelIndex < GLTFAnimation.Channels.Num(); ++ChannelIndex)
+				{
+					const GLTF::FAnimation::FChannel& Channel = GLTFAnimation.Channels[ChannelIndex];
+					const GLTF::FNode& AnimatedNode = Channel.Target.Node;
+					int32 AnimatedNodeIndex = AnimatedNode.Index;
+
+					bool bMorphTargetAnimation = Channel.Target.Path == GLTF::FAnimation::EPath::Weights;
+					bool bSkeletalAnimation = AnimatedNode.Type == GLTF::FNode::EType::Joint && GLTFNodes.IsValidIndex(AnimatedNode.RootJointIndex);
+
+					if (bMorphTargetAnimation)
+					{
+						TMap<int32, TSet<int32>>& AnimatedNodeIndicesToChannelIndices = MorphTargetAnimations.FindOrAdd(AnimatedNodeIndex);
+						TSet<int32>& ChannelIndices = AnimatedNodeIndicesToChannelIndices.FindOrAdd(AnimatedNodeIndex);
+						ChannelIndices.Add(ChannelIndex);
+					}
+					else if (bSkeletalAnimation)
+					{
+						TMap<int32, TSet<int32>>& AnimatedNodeIndicesToChannelIndices = RiggedAnimations.FindOrAdd(AnimatedNode.RootJointIndex);
+						TSet<int32>& ChannelIndices = AnimatedNodeIndicesToChannelIndices.FindOrAdd(AnimatedNodeIndex);
+						ChannelIndices.Add(ChannelIndex);
+					}
+					else
+					{
+						//Rigid animation:
+						TSet<int32>& ChannelIndices = RigidAnimation.FindOrAdd(AnimatedNodeIndex);
+						ChannelIndices.Add(ChannelIndex);
+					}
+
+					if (bSkeletalAnimation)
+					{
+						TSet<int32>& AnimatedJointNodeIndices = SkeletonRootToAnimatedJointNodeIndicesMap.FindOrAdd(AnimatedNode.RootJointIndex);
+						AnimatedJointNodeIndices.Add(AnimatedNode.Index);
+					}
+				}
+
+				ProcessRiggedAnimations(RiggedAnimations);
+				ProcessMorphTargetAnimations(MorphTargetAnimations);
+				ProcessRigidAnimation(RigidAnimation);
+
+				FixSkeletalAnimations(SkeletonRootToAnimatedJointNodeIndicesMap);
+			}
+		};
+	}
+
+	void HandleGLTFAnimations(UInterchangeBaseNodeContainer& NodeContainer,
+		TArray<GLTF::FAnimation> Animations,
+		const TArray<GLTF::FNode>& GLTFNodes,
+		const TMap<const GLTF::FNode*, FString>& GLTFNodeToInterchangeUidMap)
+	{
+		using namespace AnimationHelpers;
+
+		for (int32 AnimationIndex = 0; AnimationIndex < Animations.Num(); AnimationIndex++)
+		{
+			FAnimationHandler AnimationHadler(NodeContainer, AnimationIndex, Animations[AnimationIndex], GLTFNodes, GLTFNodeToInterchangeUidMap);
+			AnimationHadler.Process();
+		}
+	}
 }
\ No newline at end of file
diff --git a/Engine/Plugins/Interchange/Runtime/Source/Import/Private/Gltf/InterchangeGltfPrivate.h b/Engine/Plugins/Interchange/Runtime/Source/Import/Private/Gltf/InterchangeGltfPrivate.h
--- a/Engine/Plugins/Interchange/Runtime/Source/Import/Private/Gltf/InterchangeGltfPrivate.h
+++ b/Engine/Plugins/Interchange/Runtime/Source/Import/Private/Gltf/InterchangeGltfPrivate.h
@@ -3,11 +3,14 @@
 
 #include "CoreMinimal.h"
 
+class UInterchangeBaseNodeContainer;
 struct FMeshDescription;
 
 namespace GLTF
 {
 	struct FAsset;
+	struct FAnimation;
+	struct FNode;
 }
 
 namespace UE::Interchange
@@ -34,5 +37,11 @@ namespace UE::Interchange
 		bool GetStaticMeshPayloadDataForPayLoadKey(const GLTF::FAsset& GltfAsset, const FString& PayLoadKey, const FTransform& MeshGlobalTransform
 		, FMeshDescription& MeshDescription);
 		//
+
+		//Process/Handle GLTF Animations:
+		void HandleGLTFAnimations(UInterchangeBaseNodeContainer& NodeContainer,
+			TArray<GLTF::FAnimation> Animations,
+			const TArray<GLTF::FNode>& GLTFNodes,
+			const TMap<const GLTF::FNode*, FString>& GLTFNodeToInterchangeUidMap);
 	}
 }
\ No newline at end of file
diff --git a/Engine/Plugins/Interchange/Runtime/Source/Import/Private/Gltf/InterchangeGltfTranslator.cpp b/Engine/Plugins/Interchange/Runtime/Source/Import/Private/Gltf/InterchangeGltfTranslator.cpp
--- a/Engine/Plugins/Interchange/Runtime/Source/Import/Private/Gltf/InterchangeGltfTranslator.cpp
+++ b/Engine/Plugins/Interchange/Runtime/Source/Import/Private/Gltf/InterchangeGltfTranslator.cpp
@@ -692,10 +692,7 @@ bool UInterchangeGLTFTranslator::Translate( UInterchangeBaseNodeContainer& NodeC
 
 	// Animations
 	{
-		for ( int32 AnimationIndex = 0; AnimationIndex < GltfAsset.Animations.Num(); ++AnimationIndex )
-		{
-			HandleGltfAnimation( NodeContainer, AnimationIndex );
-		}
+		UE::Interchange::Gltf::Private::HandleGLTFAnimations(NodeContainer, GltfAsset.Animations, GltfAsset.Nodes, NodeUidMap);
 	}
 
 	// Variants
@@ -855,307 +852,6 @@ TFuture<TOptional<UE::Interchange::FAnimationPayloadData>> UInterchangeGLTFTrans
 	);
 }
 
-void UInterchangeGLTFTranslator::HandleGltfAnimation(UInterchangeBaseNodeContainer& NodeContainer, int32 AnimationIndex) const
-{
-	const GLTF::FAnimation& GltfAnimation = GltfAsset.Animations[AnimationIndex];
-
-	TMap< const GLTF::FNode*, TArray<int32> > NodeChannelsMap;
-
-	TMap<FString, UInterchangeSkeletalAnimationTrackNode*> RootJointIndexToTrackNodeMap;
-	TMap<UInterchangeSkeletalAnimationTrackNode*, TMap<FString, TArray<int32>>> TrackNodeToJointUidWithChannelsUsedMap;
-
-	for (int32 ChannelIndex = 0; ChannelIndex < GltfAnimation.Channels.Num(); ++ChannelIndex)
-	{
-		const GLTF::FAnimation::FChannel& Channel = GltfAnimation.Channels[ChannelIndex];
-		const GLTF::FNode* AnimatedNode = &Channel.Target.Node;
-
-		const FString* AnimatedNodeUidPtr = NodeUidMap.Find(AnimatedNode);
-		if (!ensure(AnimatedNodeUidPtr))
-		{
-			continue;
-		}
-
-		FString AnimatedNodeUid = *AnimatedNodeUidPtr;
-
-		auto CreateSkeletalAnimationTrackNode = [&NodeContainer, &RootJointIndexToTrackNodeMap, &GltfAnimation, &AnimationIndex, &TrackNodeToJointUidWithChannelsUsedMap, &ChannelIndex, &AnimatedNodeUid](const FString& SkeletonNodeUid, const TMap<FString, FString>& AnimationPayloadKeyForMorphTargetNodeUids)
-		{
-			if (SkeletonNodeUid.Len() > 0)
-			{
-				UInterchangeSkeletalAnimationTrackNode* TrackNode = nullptr;
-				if (RootJointIndexToTrackNodeMap.Contains(SkeletonNodeUid))
-				{
-					TrackNode = RootJointIndexToTrackNodeMap[SkeletonNodeUid];
-				}
-				else
-				{
-					TrackNode = NewObject< UInterchangeSkeletalAnimationTrackNode >(&NodeContainer);
-					FString TrackNodeUid = "\\SkeletalAnimation\\" + SkeletonNodeUid + "_" + LexToString(AnimationIndex);
-					TrackNode->InitializeNode(TrackNodeUid, GltfAnimation.Name, EInterchangeNodeContainerType::TranslatedAsset);
-					TrackNode->SetCustomSkeletonNodeUid(SkeletonNodeUid);
-
-					NodeContainer.AddNode(TrackNode);
-
-					RootJointIndexToTrackNodeMap.Add(SkeletonNodeUid, TrackNode);
-				}
-
-				for (const TPair<FString, FString>& AnimationPayloadKeyForMorphTargetNodeUid : AnimationPayloadKeyForMorphTargetNodeUids)
-				{
-					TrackNode->SetAnimationPayloadKeyForMorphTargetNodeUid(AnimationPayloadKeyForMorphTargetNodeUid.Key, AnimationPayloadKeyForMorphTargetNodeUid.Value, EInterchangeAnimationPayLoadType::MORPHTARGETCURVE);
-				}
-
-				TMap<FString, TArray<int32>>& JointUidWithChannelsUsedMap = TrackNodeToJointUidWithChannelsUsedMap.FindOrAdd(TrackNode);
-				TArray<int32>& ChannelsUsed = JointUidWithChannelsUsedMap.FindOrAdd(AnimatedNodeUid);
-				ChannelsUsed.Add(ChannelIndex);
-			}
-		};
-
-		
-		bool bAnimationChannelProcessed = false;
-
-		bool bSkeletalAnimation = AnimatedNode->Type == GLTF::FNode::EType::Joint && GltfAsset.Nodes.IsValidIndex(AnimatedNode->RootJointIndex);
-		if (bSkeletalAnimation)
-		{
-			const FString* SkeletonUidPtr = NodeUidMap.Find(&GltfAsset.Nodes[AnimatedNode->RootJointIndex]);
-			FString SkeletonNodeUid = *SkeletonUidPtr;
-
-			CreateSkeletalAnimationTrackNode(SkeletonNodeUid, TMap<FString, FString>());
-
-			bAnimationChannelProcessed = true;
-		}
-
-		bool bMorphTargetAnimation = Channel.Target.Path == GLTF::FAnimation::EPath::Weights;
-		if (bMorphTargetAnimation)
-		{
-			TMap<FString, FString> AnimationPayloadKeyForMorphTargetNodeUids;
-			//Find SceneNode that references the MeshNode:
-			if (const UInterchangeSceneNode* ConstSceneMeshActorNode = Cast< UInterchangeSceneNode >(NodeContainer.GetNode(*AnimatedNodeUid)))
-			{
-				FString SkeletalMeshUid;
-				if (ConstSceneMeshActorNode->GetCustomAssetInstanceUid(SkeletalMeshUid))
-				{
-					if (const UInterchangeMeshNode* MeshNode = Cast< UInterchangeMeshNode >(NodeContainer.GetNode(SkeletalMeshUid)))
-					{
-						TArray<FString> MorphTargetDependencies;
-						MeshNode->GetMorphTargetDependencies(MorphTargetDependencies);
-						for (const FString& MorphTargetDependencyUid : MorphTargetDependencies)
-						{
-							if (const UInterchangeMeshNode* MorphTargetNodeConst = Cast< UInterchangeMeshNode >(NodeContainer.GetNode(MorphTargetDependencyUid)))
-							{
-								if (MorphTargetNodeConst->GetPayLoadKey().IsSet())
-								{
-									FInterchangeMeshPayLoadKey PayLoadKey = MorphTargetNodeConst->GetPayLoadKey().GetValue();
-									FString PayLoadKeyUniqueId = LexToString(AnimationIndex) + TEXT(":") + LexToString(ChannelIndex) + TEXT(":") + PayLoadKey.UniqueId;
-
-									AnimationPayloadKeyForMorphTargetNodeUids.Add(MorphTargetDependencyUid, PayLoadKeyUniqueId);
-								}
-							}
-						}
-					}
-				}
-			}
-
-			if (AnimationPayloadKeyForMorphTargetNodeUids.Num() > 0)
-			{
-				FString SkeletonNodeUid = AnimatedNodeUid;
-				CreateSkeletalAnimationTrackNode(AnimatedNodeUid, AnimationPayloadKeyForMorphTargetNodeUids);
-			}
-			
-			bAnimationChannelProcessed = true;
-		}
-
-		if (bAnimationChannelProcessed)
-		{
-			continue;
-		}
-
-		TArray<int32>& NodeChannels = NodeChannelsMap.FindOrAdd(AnimatedNode);
-		NodeChannels.Add(ChannelIndex);
-	}
-
-	//setup rigged animations:
-	{
-		for (const TTuple<UInterchangeSkeletalAnimationTrackNode*, TMap<FString, TArray<int32>>>& TrackNodeAndJointNodeChannels : TrackNodeToJointUidWithChannelsUsedMap)
-		{
-			//@ StartTime = 0;
-			//from gltf documentation: https://registry.khronos.org/glTF/specs/2.0/glTF-2.0.html#animations
-			//Implementation Note
-			//	For example, if the earliest sampler input for an animation is t = 10, a client implementation must begin playback of that animation channel at t = 0 with output clamped to the first available output value.
-			double StartTime = 0.0;
-			double StopTime = 0.0;
-			double FrameRate = 30.0;
-			double SingleFrameDuration = 1.0 / FrameRate;
-			int32 FrameNumber = 0;
-			bool bHasAnimationPayloadSet = false;
-
-			for (const TTuple<FString, TArray<int32>>& JointNodeUidAndChannelsUsedPair : TrackNodeAndJointNodeChannels.Value)
-			{
-				bool bHasNonWeightAnimationChannel = false;
-				double PreviousStopTime = StopTime;
-				//check channel length and build payload
-				FString Payload = LexToString(AnimationIndex);
-				for (int32 ChannelIndex : JointNodeUidAndChannelsUsedPair.Value)
-				{
-					const GLTF::FAnimation::FChannel& Channel = GltfAnimation.Channels[ChannelIndex];
-					const GLTF::FAnimation::FSampler& Sampler = GltfAnimation.Samplers[Channel.Sampler];
-					TArray<float> Seconds;
-					Sampler.Input.GetFloatArray(Seconds);
-
-					if (Sampler.Interpolation == GLTF::FAnimation::EInterpolation::CubicSpline)
-					{
-						if (Sampler.Input.Count != 3 * Sampler.Output.Count)
-						{
-							// if any of the channels are corrupt the joint will not receive any of the  animation data
-							Payload = "";
-							UE_LOG(LogInterchangeImport, Warning, TEXT("GLTF Sampler Corrupt. Input and Output not meeting expectations."));
-							break;
-						}
-					}
-					else
-					{
-						if (Channel.Target.Path != GLTF::FAnimation::EPath::Weights && (Sampler.Input.Count != Sampler.Output.Count))
-						{
-							// if any of the channels are corrupt the joint will not receive any of the  animation data
-							Payload = "";
-							UE_LOG(LogInterchangeImport, Warning, TEXT("GLTF Sampler Corrupt. Input and Output not meeting expectations."));
-							break;
-						}
-					}
-
-					float CurrentStopTime = 0;
-					int32 CurrentFrameNumber = 0;
-
-					if (Seconds.Num() > 0)
-					{
-						//calculate FrameNumber and currentStopTime:
-						CurrentStopTime = Seconds[Seconds.Num() - 1];
-
-						float CurrentFrameNumberCandidate = CurrentStopTime / SingleFrameDuration;
-						CurrentFrameNumber = int(CurrentFrameNumberCandidate);
-						if (int(CurrentFrameNumberCandidate) < CurrentFrameNumberCandidate)
-						{
-							CurrentFrameNumber++;
-						}
-						CurrentStopTime = CurrentFrameNumber * SingleFrameDuration;
-					}
-
-					if (FrameNumber < CurrentFrameNumber)
-					{
-						FrameNumber = CurrentFrameNumber;
-						StopTime = FrameNumber * SingleFrameDuration;
-					}
-
-					if (Channel.Target.Path != GLTF::FAnimation::EPath::Weights)
-					{
-						bHasNonWeightAnimationChannel = true;
-						Payload += ":" + LexToString(ChannelIndex);
-					}
-				}
-
-				//set payload:
-				if (Payload.Len() > 0)
-				{
-					if (bHasNonWeightAnimationChannel)
-					{
-						TrackNodeAndJointNodeChannels.Key->SetAnimationPayloadKeyForSceneNodeUid(JointNodeUidAndChannelsUsedPair.Key, Payload, EInterchangeAnimationPayLoadType::BAKED);
-					}
-					
-					bHasAnimationPayloadSet = true;
-				}
-				else
-				{
-					StopTime = PreviousStopTime;
-				}
-			}
-
-			//set animation length:
-			if (bHasAnimationPayloadSet)
-			{
-				TrackNodeAndJointNodeChannels.Key->SetCustomAnimationSampleRate(FrameRate);
-				TrackNodeAndJointNodeChannels.Key->SetCustomAnimationStartTime(StartTime);
-				TrackNodeAndJointNodeChannels.Key->SetCustomAnimationStopTime(StopTime);
-			}
-		}
-	}
-
-	if (NodeChannelsMap.IsEmpty())
-	{
-		return;
-	}
-
-	UInterchangeAnimationTrackSetNode* TrackSetNode = NewObject< UInterchangeAnimationTrackSetNode >(&NodeContainer);
-
-	const FString AnimTrackSetNodeUid = TEXT("\\Animation\\") + GltfAnimation.UniqueId;
-	TrackSetNode->InitializeNode(AnimTrackSetNodeUid, GltfAnimation.Name, EInterchangeNodeContainerType::TranslatedAsset);
-
-	for (const TTuple<const GLTF::FNode*, TArray<int32>>& NodeChannelsEntry : NodeChannelsMap)
-	{
-		const TArray<int32>& ChannelIndices = NodeChannelsEntry.Value;
-		if (ChannelIndices.Num() == 0)
-		{
-			continue;
-		}
-
-		const GLTF::FNode* GltfNode = NodeChannelsEntry.Key;
-		const FString* NodeUid = NodeUidMap.Find(GltfNode);
-		if (!ensure(NodeUid))
-		{
-			continue;
-		}
-
-		UInterchangeTransformAnimationTrackNode* TransformAnimTrackNode = NewObject< UInterchangeTransformAnimationTrackNode >(&NodeContainer);
-
-		const FString TransformAnimTrackNodeName = FString::Printf(TEXT("%s_%s"), *GltfNode->Name, *GltfAnimation.Name);
-		const FString TransformAnimTrackNodeUid = TEXT("\\AnimationTrack\\") + TransformAnimTrackNodeName;
-
-		TransformAnimTrackNode->InitializeNode(TransformAnimTrackNodeUid, TransformAnimTrackNodeName, EInterchangeNodeContainerType::TranslatedAsset);
-
-		TransformAnimTrackNode->SetCustomActorDependencyUid(*NodeUid);
-
-		FString PayloadKey = FString::FromInt(AnimationIndex);
-
-		constexpr int32 TranslationChannel = 0x0001 | 0x0002 | 0x0004;
-		constexpr int32 RotationChannel = 0x0008 | 0x0010 | 0x0020;
-		constexpr int32 ScaleChannel = 0x0040 | 0x0080 | 0x0100;
-
-		int32 UsedChannels = 0;
-
-		for (int32 ChannelIndex : ChannelIndices)
-		{
-			PayloadKey += TEXT(":") + FString::FromInt(ChannelIndex);
-
-			const GLTF::FAnimation::FChannel& Channel = GltfAnimation.Channels[ChannelIndex];
-
-			switch (Channel.Target.Path)
-			{
-				case GLTF::FAnimation::EPath::Translation:
-				{
-					UsedChannels |= TranslationChannel;
-				} break;
-
-				case GLTF::FAnimation::EPath::Rotation:
-				{
-					UsedChannels |= RotationChannel;
-				} break;
-
-				case GLTF::FAnimation::EPath::Scale:
-				{
-					UsedChannels |= ScaleChannel;
-				} break;
-				default: break;
-			}
-		}
-
-		TransformAnimTrackNode->SetCustomAnimationPayloadKey(PayloadKey, EInterchangeAnimationPayLoadType::CURVE);
-		TransformAnimTrackNode->SetCustomUsedChannels(UsedChannels);
-
-		NodeContainer.AddNode(TransformAnimTrackNode);
-
-		TrackSetNode->AddCustomAnimationTrackUid(TransformAnimTrackNodeUid);
-	}
-
-	NodeContainer.AddNode(TrackSetNode);
-}
-
 void UInterchangeGLTFTranslator::SetTextureSRGB(UInterchangeBaseNodeContainer& NodeContainer, const GLTF::FTextureMap& TextureMap, bool bSRGB) const
 {
 	if (GltfAsset.Textures.IsValidIndex(TextureMap.TextureIndex))
diff --git a/Engine/Plugins/Interchange/Runtime/Source/Import/Public/Gltf/InterchangeGltfTranslator.h b/Engine/Plugins/Interchange/Runtime/Source/Import/Public/Gltf/InterchangeGltfTranslator.h
--- a/Engine/Plugins/Interchange/Runtime/Source/Import/Public/Gltf/InterchangeGltfTranslator.h
+++ b/Engine/Plugins/Interchange/Runtime/Source/Import/Public/Gltf/InterchangeGltfTranslator.h
@@ -66,7 +66,6 @@ protected:
 	void HandleGltfSkeletons( UInterchangeBaseNodeContainer& NodeContainer, const FString& SceneNodeUid, const TArray<int32>& SkinnedMeshNodes, TSet<int>& UnusedMeshIndices ) const;
 	void HandleGltfNode( UInterchangeBaseNodeContainer& NodeContainer, const GLTF::FNode& GltfNode, const FString& ParentNodeUid, const int32 NodeIndex, 
 		bool &bHasVariants, TArray<int32>& SkinnedMeshNodes, TSet<int>& UnusedMeshIndices ) const;
-	void HandleGltfAnimation(UInterchangeBaseNodeContainer& NodeContainer, int32 AnimationIndex) const;
 	void HandleGltfVariants(UInterchangeBaseNodeContainer& NodeContainer, const FString& FileName) const;
 	UInterchangeMeshNode* HandleGltfMesh(UInterchangeBaseNodeContainer& NodeContainer, const GLTF::FMesh& GltfMesh, int MeshIndex,
 		TSet<int>& UnusedMeshIndices, const FString& SkeletalName = "" /*If set it creates the mesh even if it was already created (for Skeletals)*/, const FString& SkeletalId = "") const;
diff --git a/Engine/Plugins/Interchange/Runtime/Source/Parsers/GLTFCore/Private/GLTF/GLTFReader.cpp b/Engine/Plugins/Interchange/Runtime/Source/Parsers/GLTFCore/Private/GLTF/GLTFReader.cpp
--- a/Engine/Plugins/Interchange/Runtime/Source/Parsers/GLTFCore/Private/GLTF/GLTFReader.cpp
+++ b/Engine/Plugins/Interchange/Runtime/Source/Parsers/GLTFCore/Private/GLTF/GLTFReader.cpp
@@ -443,6 +443,8 @@ namespace GLTF
 	{
 		FNode& Node = Asset->Nodes.Emplace_GetRef();
 
+		Node.Index = Asset->Nodes.Num() - 1;
+
 		Node.Name = GetString(Object, TEXT("name"));
 
 		if (Object.HasField(TEXT("matrix")))
diff --git a/Engine/Plugins/Interchange/Runtime/Source/Parsers/GLTFCore/Public/GLTFNode.h b/Engine/Plugins/Interchange/Runtime/Source/Parsers/GLTFCore/Public/GLTFNode.h
--- a/Engine/Plugins/Interchange/Runtime/Source/Parsers/GLTFCore/Public/GLTFNode.h
+++ b/Engine/Plugins/Interchange/Runtime/Source/Parsers/GLTFCore/Public/GLTFNode.h
@@ -37,6 +37,7 @@ namespace GLTF
 		int32 CameraIndex;  // index into FAsset::Cameras
 		int32 LightIndex;   // index into FAsset::Lights
 
+		int32 Index;		// index of FNode in GLTFAsset.Nodes array.
 		int32 ParentIndex;
 		int32 RootJointIndex; //only valid if node is of Joint Type
 
@@ -57,6 +58,7 @@ namespace GLTF
 		    , Skindex(INDEX_NONE)
 		    , CameraIndex(INDEX_NONE)
 		    , LightIndex(INDEX_NONE)
+			, Index(INDEX_NONE)
 			, ParentIndex(INDEX_NONE)
 			, RootJointIndex(INDEX_NONE)
 			, bHasLocalBindPose(false)
diff --git a/Engine/Plugins/Interchange/Runtime/Source/Pipelines/Private/InterchangeGenericAnimationPipeline.cpp b/Engine/Plugins/Interchange/Runtime/Source/Pipelines/Private/InterchangeGenericAnimationPipeline.cpp
--- a/Engine/Plugins/Interchange/Runtime/Source/Pipelines/Private/InterchangeGenericAnimationPipeline.cpp
+++ b/Engine/Plugins/Interchange/Runtime/Source/Pipelines/Private/InterchangeGenericAnimationPipeline.cpp
@@ -94,7 +94,17 @@ void UInterchangeGenericAnimationPipeline::ExecutePipeline(UInterchangeBaseNodeC
 				TMap<FString, float> MorphTargetCurveWeights;
 				SceneNode->GetMorphTargetCurveWeights(MorphTargetCurveWeights);
 
-				if (MorphTargetCurveWeights.Num() > 0)
+				bool CreateAnimSequence = false;
+				for (const TTuple<FString, float>& MorphTargetCurveWeight : MorphTargetCurveWeights)
+				{
+					if (MorphTargetCurveWeight.Value != 0)
+					{
+						CreateAnimSequence = true;
+						break;
+					}
+				}
+
+				if (CreateAnimSequence)
 				{
 					SceneNodesWithMorphTargetCurveWeights.Add(SceneNode);
 				}
@@ -203,7 +213,11 @@ void UInterchangeGenericAnimationPipeline::ExecutePipeline(UInterchangeBaseNodeC
 				}
 			}
 
-			BaseNodeContainer->AddNode(SkeletalAnimationNode);
+			if (bCustomSkeletonNodeUidSet)
+			{
+				//if no SkeletonNodeUid can be set, then the conversion failed and should not be added to the BaseNodeContainer.
+				BaseNodeContainer->AddNode(SkeletalAnimationNode);
+			}
 		}
 	}
 	else
