[Interchange] Fixing Mesh transform baking for SkeletalMeshes.

diff --git a/Engine/Plugins/Interchange/Runtime/Source/Import/Private/Mesh/InterchangeSkeletalMeshFactory.cpp b/Engine/Plugins/Interchange/Runtime/Source/Import/Private/Mesh/InterchangeSkeletalMeshFactory.cpp
--- a/Engine/Plugins/Interchange/Runtime/Source/Import/Private/Mesh/InterchangeSkeletalMeshFactory.cpp
+++ b/Engine/Plugins/Interchange/Runtime/Source/Import/Private/Mesh/InterchangeSkeletalMeshFactory.cpp
@@ -332,15 +332,11 @@ namespace UE
 					}
 					
 					FTransform JointBindPoseGlobalTransform;
-					if (!JointNode->GetCustomBindPoseGlobalTransform(NodeContainer, GlobalOffsetTransform, JointBindPoseGlobalTransform))
+					if (!ensure(JointNode->GetCustomBindPoseGlobalTransform(NodeContainer, GlobalOffsetTransform, JointBindPoseGlobalTransform)))
 					{
-						//If there is no bind pose we will fall back on the CustomGlobalTransform of the link.
-						//We ensure here because any scenenode should have a valid CustomGlobalTransform.
-						if (!ensure(JointNode->GetCustomGlobalTransform(NodeContainer, GlobalOffsetTransform, JointBindPoseGlobalTransform)))
-						{
-							//No value to convert from, skip this joint.
-							continue;
-						}
+						//BindPose will fall back on LocalTransform in case its not present.
+						//If neither is present: No Value to convert from, skip this joint.
+						continue;
 					}
 
 					FTransform JointTimeZeroGlobalTransform;
@@ -1228,7 +1224,7 @@ UInterchangeFactoryBase::FImportAssetResult UInterchangeSkeletalMeshFactory::Imp
 					FTransform SceneNodeTransform;
 					if (!bUseTimeZeroAsBindPose || !MeshReference.SceneNode->GetCustomTimeZeroGlobalTransform(Arguments.NodeContainer, GlobalOffsetTransform, SceneNodeTransform))
 					{
-						ensure(MeshReference.SceneNode->GetCustomGlobalTransform(Arguments.NodeContainer, GlobalOffsetTransform, SceneNodeTransform));
+						ensure(MeshReference.SceneNode->GetCustomBindPoseGlobalTransform(Arguments.NodeContainer, GlobalOffsetTransform, SceneNodeTransform));
 						if (!bBakeMeshes)
 						{
 							SceneNodeTransform *= BakeToRootJointTransfromModifier;
diff --git a/Engine/Plugins/Interchange/Runtime/Source/Nodes/Private/InterchangeSceneNode.cpp b/Engine/Plugins/Interchange/Runtime/Source/Nodes/Private/InterchangeSceneNode.cpp
--- a/Engine/Plugins/Interchange/Runtime/Source/Nodes/Private/InterchangeSceneNode.cpp
+++ b/Engine/Plugins/Interchange/Runtime/Source/Nodes/Private/InterchangeSceneNode.cpp
@@ -303,9 +303,18 @@ bool UInterchangeSceneNode::GetGlobalTransformInternal(const UE::Interchange::FA
 	, FTransform& AttributeValue
 	, bool bForceRecache) const
 {
-	if (!Attributes->ContainAttribute(LocalTransformKey))
+	UE::Interchange::FAttributeKey TransformKey = LocalTransformKey;
+	if (!Attributes->ContainAttribute(TransformKey))
 	{
-		return false;
+		//Fallback to LocalTransform:
+		if (Attributes->ContainAttribute(Macro_CustomLocalTransformKey))
+		{
+			TransformKey = Macro_CustomLocalTransformKey;
+		}
+		else
+		{
+			return false;
+		}
 	}
 	if (bForceRecache)
 	{
@@ -314,7 +323,7 @@ bool UInterchangeSceneNode::GetGlobalTransformInternal(const UE::Interchange::FA
 	if (!CacheTransform.IsSet())
 	{
 		FTransform LocalTransform;
-		UE::Interchange::FAttributeStorage::TAttributeHandle<FTransform> AttributeHandle = GetAttributeHandle<FTransform>(LocalTransformKey);
+		UE::Interchange::FAttributeStorage::TAttributeHandle<FTransform> AttributeHandle = GetAttributeHandle<FTransform>(TransformKey);
 		if (AttributeHandle.IsValid() && AttributeHandle.Get(LocalTransform) == UE::Interchange::EAttributeStorageResult::Operation_Success)
 		{
 			//Compute the Global
@@ -329,19 +338,11 @@ bool UInterchangeSceneNode::GetGlobalTransformInternal(const UE::Interchange::FA
 					}
 					else if (LocalTransformKey == Macro_CustomBindPoseLocalTransformKey)
 					{
-						//Its possible the root skeleton id have a parent that is not a joint, in that case we need to fall back on the normal transform
-						if (!ParentSceneNode->GetCustomBindPoseGlobalTransform(BaseNodeContainer, GlobalOffsetTransform, GlobalParent, bForceRecache))
-						{
-							ParentSceneNode->GetCustomGlobalTransform(BaseNodeContainer, GlobalOffsetTransform, GlobalParent, bForceRecache);
-						}
+						ParentSceneNode->GetCustomBindPoseGlobalTransform(BaseNodeContainer, GlobalOffsetTransform, GlobalParent, bForceRecache);
 					}
 					else if (LocalTransformKey == Macro_CustomTimeZeroLocalTransformKey)
 					{
-						//Its possible the root skeleton id have a parent that is not a joint, in that case we need to fall back on the normal transform
-						if(!ParentSceneNode->GetCustomTimeZeroGlobalTransform(BaseNodeContainer, GlobalOffsetTransform, GlobalParent, bForceRecache))
-						{
-							ParentSceneNode->GetCustomGlobalTransform(BaseNodeContainer, GlobalOffsetTransform, GlobalParent, bForceRecache);
-						}
+						ParentSceneNode->GetCustomTimeZeroGlobalTransform(BaseNodeContainer, GlobalOffsetTransform, GlobalParent, bForceRecache);
 					}
 				}
 				CacheTransform = LocalTransform * GlobalParent;
