Interchange: support the skeletal mesh VertexAttributes import feature.
#rnx

diff --git a/Engine/Plugins/Interchange/Runtime/Source/FactoryNodes/Private/InterchangeSkeletalMeshFactoryNode.cpp b/Engine/Plugins/Interchange/Runtime/Source/FactoryNodes/Private/InterchangeSkeletalMeshFactoryNode.cpp
--- a/Engine/Plugins/Interchange/Runtime/Source/FactoryNodes/Private/InterchangeSkeletalMeshFactoryNode.cpp
+++ b/Engine/Plugins/Interchange/Runtime/Source/FactoryNodes/Private/InterchangeSkeletalMeshFactoryNode.cpp
@@ -91,6 +91,16 @@ bool UInterchangeSkeletalMeshFactoryNode::SetCustomImportMorphTarget(const bool&
 	IMPLEMENT_NODE_ATTRIBUTE_SETTER_NODELEGATE(ImportMorphTarget, bool)
 }
 
+bool UInterchangeSkeletalMeshFactoryNode::GetCustomImportVertexAttributes(bool& AttributeValue) const
+{
+	IMPLEMENT_NODE_ATTRIBUTE_GETTER(ImportVertexAttributes, bool)
+}
+
+bool UInterchangeSkeletalMeshFactoryNode::SetCustomImportVertexAttributes(const bool& AttributeValue)
+{
+	IMPLEMENT_NODE_ATTRIBUTE_SETTER_NODELEGATE(ImportVertexAttributes, bool)
+}
+
 bool UInterchangeSkeletalMeshFactoryNode::GetCustomCreatePhysicsAsset(bool& AttributeValue) const
 {
 	IMPLEMENT_NODE_ATTRIBUTE_GETTER(CreatePhysicsAsset, bool)
diff --git a/Engine/Plugins/Interchange/Runtime/Source/FactoryNodes/Public/InterchangeSkeletalMeshFactoryNode.h b/Engine/Plugins/Interchange/Runtime/Source/FactoryNodes/Public/InterchangeSkeletalMeshFactoryNode.h
--- a/Engine/Plugins/Interchange/Runtime/Source/FactoryNodes/Public/InterchangeSkeletalMeshFactoryNode.h
+++ b/Engine/Plugins/Interchange/Runtime/Source/FactoryNodes/Public/InterchangeSkeletalMeshFactoryNode.h
@@ -64,6 +64,14 @@ public:
 	UFUNCTION(BlueprintCallable, Category = "Interchange | Node | SkeletalMesh")
 	bool SetCustomImportMorphTarget(const bool& AttributeValue);
 
+	/** Query weather the skeletal mesh factory should import the vertex attributes. Return false if the attribute was not set.*/
+	UFUNCTION(BlueprintCallable, Category = "Interchange | Node | SkeletalMesh")
+	bool GetCustomImportVertexAttributes(bool& AttributeValue) const;
+
+	/** Set weather the skeletal mesh factory should import the vertex attributes. Return false if the attribute cannot be set.*/
+	UFUNCTION(BlueprintCallable, Category = "Interchange | Node | SkeletalMesh")
+	bool SetCustomImportVertexAttributes(const bool& AttributeValue);
+
 	/** Query weather the skeletal mesh factory should create a physics asset. Return false if the attribute was not set.*/
 	UFUNCTION(BlueprintCallable, Category = "Interchange | Node | SkeletalMesh")
 	bool GetCustomCreatePhysicsAsset(bool& AttributeValue) const;
@@ -157,6 +165,7 @@ private:
 
 	const UE::Interchange::FAttributeKey ClassNameAttributeKey = UE::Interchange::FBaseNodeStaticData::ClassTypeAttributeKey();
 	const UE::Interchange::FAttributeKey Macro_CustomImportMorphTargetKey = UE::Interchange::FAttributeKey(TEXT("ImportMorphTarget"));
+	const UE::Interchange::FAttributeKey Macro_CustomImportVertexAttributesKey = UE::Interchange::FAttributeKey(TEXT("ImportVertexAttributes"));
 	const UE::Interchange::FAttributeKey Macro_CustomSkeletonSoftObjectPathKey = UE::Interchange::FAttributeKey(TEXT("SkeletonSoftObjectPath"));
 	const UE::Interchange::FAttributeKey Macro_CustomCreatePhysicsAssetKey = UE::Interchange::FAttributeKey(TEXT("CreatePhysicsAsset"));
 	const UE::Interchange::FAttributeKey Macro_CustomPhysicAssetSoftObjectPathKey = UE::Interchange::FAttributeKey(TEXT("PhysicAssetSoftObjectPath"));
diff --git a/Engine/Plugins/Interchange/Runtime/Source/Import/Private/Mesh/InterchangeSkeletalMeshFactory.cpp b/Engine/Plugins/Interchange/Runtime/Source/Import/Private/Mesh/InterchangeSkeletalMeshFactory.cpp
--- a/Engine/Plugins/Interchange/Runtime/Source/Import/Private/Mesh/InterchangeSkeletalMeshFactory.cpp
+++ b/Engine/Plugins/Interchange/Runtime/Source/Import/Private/Mesh/InterchangeSkeletalMeshFactory.cpp
@@ -437,6 +437,9 @@ namespace UE
 				TMap<FString, TOptional<UE::Interchange::FMeshPayloadData>> MorphTargetMeshDescriptionsPerMorphTargetName;
 				int32 MorphTargetCount = 0;
 
+				bool bImportVertexAttributes = false;
+				SkeletalMeshFactoryNode->GetCustomImportVertexAttributes(bImportVertexAttributes);
+
 				for (const FMeshNodeContext& MeshNodeContext : MeshReferences)
 				{
 					//Add the payload entry key, the payload data will be fill later in bulk by the translator
@@ -459,6 +462,7 @@ namespace UE
 					const int32 VertexOffset = LodMeshDescription.Vertices().Num();
 
 					FSkeletalMeshOperations::FSkeletalMeshAppendSettings SkeletalMeshAppendSettings;
+					SkeletalMeshAppendSettings.bAppendVertexAttributes = bImportVertexAttributes;
 					SkeletalMeshAppendSettings.SourceVertexIDOffset = VertexOffset;
 					{
 						TRACE_CPUPROFILER_EVENT_SCOPE("RetrieveAllSkeletalMeshPayloadsAndFillImportData::CompactPayload")
diff --git a/Engine/Plugins/Interchange/Runtime/Source/Import/Public/Mesh/InterchangeMeshPayloadInterface.h b/Engine/Plugins/Interchange/Runtime/Source/Import/Public/Mesh/InterchangeMeshPayloadInterface.h
--- a/Engine/Plugins/Interchange/Runtime/Source/Import/Public/Mesh/InterchangeMeshPayloadInterface.h
+++ b/Engine/Plugins/Interchange/Runtime/Source/Import/Public/Mesh/InterchangeMeshPayloadInterface.h
@@ -31,6 +31,7 @@ public:
 	 * This payload will be use by the factories to create the asset.
 	 *
 	 * @param PayloadKey - The key to retrieve the a particular payload contain into the specified source data.
+	 * @param MeshGlobalTransform - The transform to apply to the raw data when creating the payload data.
 	 * @return a PayloadData containing the the data ask with the key.
 	 */
 	virtual TFuture<TOptional<UE::Interchange::FMeshPayloadData>> GetMeshPayloadData(const FInterchangeMeshPayLoadKey& PayLoadKey, const FTransform& MeshGlobalTransform) const = 0;
diff --git a/Engine/Plugins/Interchange/Runtime/Source/Parsers/Fbx/Private/FbxMesh.cpp b/Engine/Plugins/Interchange/Runtime/Source/Parsers/Fbx/Private/FbxMesh.cpp
--- a/Engine/Plugins/Interchange/Runtime/Source/Parsers/Fbx/Private/FbxMesh.cpp
+++ b/Engine/Plugins/Interchange/Runtime/Source/Parsers/Fbx/Private/FbxMesh.cpp
@@ -25,6 +25,231 @@
 
 namespace UE::Interchange::Private {
 
+	//Import vertex Attribute from other vertex color layer
+	void GetVertexAttributeFromMeshVertexColor(FMeshDescription& MeshDescription, FbxMesh* Mesh)
+	{
+		//
+		// Get the vertex attribute layers from all layers, even the first layer which may be used as a vertex color layer.
+		// Currently we're only interested in alpha-only layers, since those are the only layer types the engine
+		// currently exposes for vertex attributes. Internally we can store 1-4 components, but there's no tooling for that
+		// 2-4 channels as of yet.
+		//
+		struct FNamedVertexAttribute
+		{
+			FNamedVertexAttribute(FString&& InAttributeName, TArray<float>&& InAttributeValues, const int32 InComponentCount)
+				: AttributeName(InAttributeName)
+				, AttributeValues(InAttributeValues)
+				, ComponentCount(InComponentCount)
+			{}
+
+			FString AttributeName;
+			TArray<float> AttributeValues;
+			int32 ComponentCount;
+		};
+		int32 ControlPointsCount = Mesh->GetControlPointsCount();
+		int32 TriangleCount = Mesh->GetPolygonCount();
+		int32 MeshLayerCount = Mesh->GetLayerCount();
+		TArray<FNamedVertexAttribute> NamedVertexAttributes;
+
+		for (int32 LayerIndex = 0; LayerIndex < MeshLayerCount; LayerIndex++)
+		{
+			FbxLayerElementVertexColor* LayerElementVertexAttribute = Mesh->GetLayer(LayerIndex)->GetVertexColors();
+			if (!LayerElementVertexAttribute)
+			{
+				continue;
+			}
+
+			// Check if this is an alpha-only attribute, by ensuring the RGB values are all zero, otherwise skip.
+			bool bIsValidAttribute = true;
+			const FbxLayerElementArrayTemplate<FbxColor>& AttributeValues = LayerElementVertexAttribute->GetDirectArray();
+			for (int32 Index = 0; Index < AttributeValues.GetCount(); Index++)
+			{
+				// We do an exact comparison, since that's how empty channels would be represented in the FBX file.
+				const FbxColor& Value = AttributeValues.GetAt(Index);
+				if (Value.mRed != 0.0 || Value.mGreen != 0.0 || Value.mBlue != 0.0)
+				{
+					bIsValidAttribute = false;
+					break;
+				}
+			}
+
+			// We can only do attributes that are mapped per-vertex.
+			if (!bIsValidAttribute)
+			{
+				continue;
+			}
+
+			const int32 AttributeComponentCount = 1;	// Number of component values per vertex. See comment above. 
+			TArray<float> AttributeComponentValues;
+
+			switch (LayerElementVertexAttribute->GetMappingMode())
+			{
+			case FbxLayerElement::eByControlPoint:
+			{
+				AttributeComponentValues.AddZeroed(ControlPointsCount);
+
+				if (LayerElementVertexAttribute->GetReferenceMode() == FbxLayerElement::eDirect)
+				{
+					for (int32 Index = 0; Index < AttributeValues.GetCount(); Index++)
+					{
+						AttributeComponentValues[Index] = AttributeValues.GetAt(Index).mAlpha;
+					}
+				}
+				else // LayerElementVertexAttribute->GetReferenceMode() == FbxLayerElement::eIndexToDirect
+				{
+					const FbxLayerElementArrayTemplate<int>& IndexArray = LayerElementVertexAttribute->GetIndexArray();
+					for (int32 Index = 0; Index < IndexArray.GetCount(); Index++)
+					{
+						AttributeComponentValues[Index] = AttributeValues.GetAt(IndexArray[Index]).mAlpha;
+					}
+				}
+			}
+			break;
+			case FbxLayerElement::eByPolygonVertex:
+			{
+				// Vertex attributes are stored per-vertex, not per-vertex instance. To work around this we average
+				// together values that share a vertex.
+				TArray<int32> SharedVertexCount;
+				SharedVertexCount.AddZeroed(ControlPointsCount);
+				AttributeComponentValues.AddZeroed(ControlPointsCount);
+
+				const FbxLayerElementArrayTemplate<int>* IndexArray = nullptr;
+				if (LayerElementVertexAttribute->GetReferenceMode() == FbxLayerElement::eIndexToDirect)
+				{
+					IndexArray = &LayerElementVertexAttribute->GetIndexArray();
+				}
+
+				const int* PolygonControlPointIndexes = Mesh->GetPolygonVertices();
+
+				for (int32 TriangleIndex = 0; TriangleIndex < TriangleCount; TriangleIndex++)
+				{
+					for (int32 InnerIndex = 0; InnerIndex < 3; InnerIndex++)
+					{
+						const int32 PolygonVertexIndex = TriangleIndex * 3 + InnerIndex;;
+						const int32 PointIndex = PolygonControlPointIndexes[PolygonVertexIndex];
+
+						AttributeComponentValues[PointIndex] +=
+							AttributeValues.GetAt(IndexArray ? IndexArray->GetAt(PolygonVertexIndex) : PolygonVertexIndex).mAlpha;
+						SharedVertexCount[PointIndex]++;
+					}
+				}
+
+				for (int32 PointIndex = 0; PointIndex < ControlPointsCount; PointIndex++)
+				{
+					if (SharedVertexCount[PointIndex] > 1)
+					{
+						AttributeComponentValues[PointIndex] /= static_cast<float>(SharedVertexCount[PointIndex]);
+					}
+				}
+			}
+			break;
+			default:
+				break;
+			}
+
+			if (!AttributeComponentValues.IsEmpty())
+			{
+				FString AttributeName(UTF8_TO_TCHAR(LayerElementVertexAttribute->GetName()));
+				NamedVertexAttributes.Emplace(MoveTemp(AttributeName), MoveTemp(AttributeComponentValues), AttributeComponentCount);
+			}
+		}
+		if(NamedVertexAttributes.Num() > 0)
+		{
+			FSkeletalMeshAttributes MeshAttributes(MeshDescription);
+			MeshAttributes.Register();
+			TVertexAttributesRef<FVector3f> VertexPositions = MeshAttributes.GetVertexPositions();
+
+			TMap<FString, FName> ValidAttributes;
+			for (int32 AttributeIndex = 0; AttributeIndex < NamedVertexAttributes.Num(); AttributeIndex++)
+			{
+				const FNamedVertexAttribute& NamedVertexAttribute = NamedVertexAttributes[AttributeIndex];
+				const FString& VertexAttributeName = NamedVertexAttribute.AttributeName;
+				if (!ensure(NamedVertexAttribute.AttributeValues.Num() == (VertexPositions.GetNumElements() * NamedVertexAttribute.ComponentCount)))
+				{
+					continue;
+				}
+
+				EMeshAttributeFlags DefaultAttributeFlags = EMeshAttributeFlags::Mergeable | EMeshAttributeFlags::Lerpable;
+
+				FName RegisteredName(VertexAttributeName);
+
+				// Ignore attributes with reserved names. This should have been handled at import time or when the attribute
+				// was created/renamed.
+				if (!ensure(!FSkeletalMeshAttributes::IsReservedAttributeName(RegisteredName)))
+				{
+					continue;
+				}
+
+				switch (NamedVertexAttribute.ComponentCount)
+				{
+				case 1:
+					MeshDescription.VertexAttributes().RegisterAttribute<float>(RegisteredName, 1, 0.0f, DefaultAttributeFlags);
+					break;
+				case 2:
+					MeshDescription.VertexAttributes().RegisterAttribute<FVector2f>(RegisteredName, 1, FVector2f::Zero(), DefaultAttributeFlags);
+					break;
+				case 3:
+					MeshDescription.VertexAttributes().RegisterAttribute<FVector3f>(RegisteredName, 1, FVector3f::Zero(), DefaultAttributeFlags);
+					break;
+				case 4:
+					MeshDescription.VertexAttributes().RegisterAttribute<FVector4f>(RegisteredName, 1, FVector4f::Zero(), DefaultAttributeFlags);
+					break;
+				default:
+					continue;
+				}
+
+				ValidAttributes.Add(VertexAttributeName, RegisteredName);
+
+				switch (NamedVertexAttribute.ComponentCount)
+				{
+				case 1:
+				{
+					TVertexAttributesRef<float> AttributeRef = MeshDescription.VertexAttributes().GetAttributesRef<float>(RegisteredName);
+					for (int32 Index = 0; Index < NamedVertexAttribute.AttributeValues.Num(); Index++)
+					{
+						AttributeRef.Set(FVertexID(Index), NamedVertexAttribute.AttributeValues[Index]);
+					}
+					break;
+				}
+				case 2:
+				{
+					TVertexAttributesRef<FVector2f> AttributeRef = MeshDescription.VertexAttributes().GetAttributesRef<FVector2f>(RegisteredName);
+					for (int32 Index = 0; Index < NamedVertexAttribute.AttributeValues.Num(); Index += 2)
+					{
+						AttributeRef.Set(FVertexID(Index / 2),
+							FVector2f(NamedVertexAttribute.AttributeValues[Index], NamedVertexAttribute.AttributeValues[Index + 1]));
+					}
+					break;
+				}
+				case 3:
+				{
+					TVertexAttributesRef<FVector3f> AttributeRef = MeshDescription.VertexAttributes().GetAttributesRef<FVector3f>(RegisteredName);
+					for (int32 Index = 0; Index < NamedVertexAttribute.AttributeValues.Num(); Index += 3)
+					{
+						AttributeRef.Set(FVertexID(Index / 3),
+							FVector3f(NamedVertexAttribute.AttributeValues[Index], NamedVertexAttribute.AttributeValues[Index + 1], NamedVertexAttribute.AttributeValues[Index + 2]));
+					}
+					break;
+				}
+				case 4:
+				{
+					TVertexAttributesRef<FVector4f> AttributeRef = MeshDescription.VertexAttributes().GetAttributesRef<FVector4f>(RegisteredName);
+					for (int32 Index = 0; Index < NamedVertexAttribute.AttributeValues.Num(); Index += 4)
+					{
+						AttributeRef.Set(FVertexID(Index / 4),
+							FVector4f(
+								NamedVertexAttribute.AttributeValues[Index], NamedVertexAttribute.AttributeValues[Index + 1],
+								NamedVertexAttribute.AttributeValues[Index + 2], NamedVertexAttribute.AttributeValues[Index + 3]));
+					}
+					break;
+				}
+				default:
+					checkNoEntry();
+				}
+			}
+		}
+	}
+
 // Wraps some common code useful for multiple fbx import code path
 struct FFBXUVs
 {
@@ -1059,6 +1284,9 @@ bool FMeshDescriptionImporter::FillMeshDescriptionFromFbxMesh(FbxMesh* Mesh, con
 			FSkeletalMeshAttributes SkeletalMeshAttributes(*MeshDescription);
 			SkeletalMeshAttributes.Register(true);
 
+			//Import vertex Attribute from all mesh layer vertex color
+			GetVertexAttributeFromMeshVertexColor(*MeshDescription, Mesh);
+
 			using namespace UE::AnimationCore;
 			TMap<FVertexID, TArray<FBoneWeight>> RawBoneWeights;
 						
diff --git a/Engine/Plugins/Interchange/Runtime/Source/Pipelines/Private/InterchangeGenericSkeletalMeshPipeline.cpp b/Engine/Plugins/Interchange/Runtime/Source/Pipelines/Private/InterchangeGenericSkeletalMeshPipeline.cpp
--- a/Engine/Plugins/Interchange/Runtime/Source/Pipelines/Private/InterchangeGenericSkeletalMeshPipeline.cpp
+++ b/Engine/Plugins/Interchange/Runtime/Source/Pipelines/Private/InterchangeGenericSkeletalMeshPipeline.cpp
@@ -352,6 +352,7 @@ UInterchangeSkeletalMeshFactoryNode* UInterchangeGenericMeshPipeline::CreateSkel
 
 	AddLodDataToSkeletalMesh(SkeletonFactoryNode, SkeletalMeshFactoryNode, MeshUidsPerLodIndex);
 	SkeletalMeshFactoryNode->SetCustomImportMorphTarget(bImportMorphTargets);
+	SkeletalMeshFactoryNode->SetCustomImportVertexAttributes(bImportVertexAttributes);
 
 	SkeletalMeshFactoryNode->SetCustomImportContentType(SkeletalMeshImportContentType);
 
diff --git a/Engine/Plugins/Interchange/Runtime/Source/Pipelines/Public/InterchangeGenericMeshPipeline.h b/Engine/Plugins/Interchange/Runtime/Source/Pipelines/Public/InterchangeGenericMeshPipeline.h
--- a/Engine/Plugins/Interchange/Runtime/Source/Pipelines/Public/InterchangeGenericMeshPipeline.h
+++ b/Engine/Plugins/Interchange/Runtime/Source/Pipelines/Public/InterchangeGenericMeshPipeline.h
@@ -152,6 +152,10 @@ public:
 	UPROPERTY(EditAnywhere, BlueprintReadWrite, Category = "Skeletal Meshes")
 	bool bImportMorphTargets = true;
 
+	/** True to import per-vertex attributes from the FBX file */
+	UPROPERTY(EditAnywhere, BlueprintReadWrite, Category = "Skeletal Meshes", meta = (ToolTip = "If enabled, creates named vertex attributes for secondary vertex color data."))
+	bool bImportVertexAttributes = false;
+
 	/** Enable this option to update Skeleton (of the mesh)'s reference pose. Mesh's reference pose is always updated.  */
 	UPROPERTY(EditAnywhere, BlueprintReadWrite, Category = "Skeletal Meshes")
 	bool bUpdateSkeletonReferencePose = false;
diff --git a/Engine/Source/Runtime/SkeletalMeshDescription/Private/SkeletalMeshOperations.cpp b/Engine/Source/Runtime/SkeletalMeshDescription/Private/SkeletalMeshOperations.cpp
--- a/Engine/Source/Runtime/SkeletalMeshDescription/Private/SkeletalMeshOperations.cpp
+++ b/Engine/Source/Runtime/SkeletalMeshDescription/Private/SkeletalMeshOperations.cpp
@@ -10,6 +10,58 @@ DEFINE_LOG_CATEGORY(LogSkeletalMeshOperations);
 
 #define LOCTEXT_NAMESPACE "SkeletalMeshOperations"
 
+namespace UE::Private
+{
+	template<typename T>
+	struct FCreateAndCopyAttributeValues
+	{
+		FCreateAndCopyAttributeValues(
+			const FMeshDescription& InSourceMesh,
+			FMeshDescription& InTargetMesh,
+			TArray<FName>& InTargetCustomAttributeNames,
+			int32 InTargetVertexIndexOffset)
+			: SourceMesh(InSourceMesh)
+			, TargetMesh(InTargetMesh)
+			, TargetCustomAttributeNames(InTargetCustomAttributeNames)
+			, TargetVertexIndexOffset(InTargetVertexIndexOffset)
+			{}
+
+		void operator()(const FName InAttributeName, TVertexAttributesConstRef<T> InSrcAttribute)
+		{
+			// Ignore attributes with reserved names.
+			if (FSkeletalMeshAttributes::IsReservedAttributeName(InAttributeName))
+			{
+				return;
+			}
+
+			const bool bAppend = TargetCustomAttributeNames.Contains(InAttributeName);
+			if (!bAppend)
+			{
+				TargetMesh.VertexAttributes().RegisterAttribute<T>(InAttributeName, InSrcAttribute.GetNumChannels(), InSrcAttribute.GetDefaultValue(), InSrcAttribute.GetFlags());
+				TargetCustomAttributeNames.Add(InAttributeName);
+			}
+			//Copy the data
+			TVertexAttributesRef<T> TargetVertexAttributes = TargetMesh.VertexAttributes().GetAttributesRef<T>(InAttributeName);
+			for (const FVertexID SourceVertexID : SourceMesh.Vertices().GetElementIDs())
+			{
+				const FVertexID TargetVertexID = FVertexID(TargetVertexIndexOffset + SourceVertexID.GetValue());
+				TargetVertexAttributes.Set(TargetVertexID, InSrcAttribute.Get(SourceVertexID));
+			}
+		}
+
+		// Unhandled sub-types.
+		void operator()(const FName, TVertexAttributesConstRef<TArrayAttribute<T>>) { }
+		void operator()(const FName, TVertexAttributesConstRef<TArrayView<T>>) { }
+
+	private:
+		const FMeshDescription& SourceMesh;
+		FMeshDescription& TargetMesh;
+		TArray<FName>& TargetCustomAttributeNames;
+		int32 TargetVertexIndexOffset = 0;
+	};
+
+} // ns UE::Private
+
 //Add specific skeletal mesh descriptions implementation here
 void FSkeletalMeshOperations::AppendSkinWeight(const FMeshDescription& SourceMesh, FMeshDescription& TargetMesh, FSkeletalMeshAppendSettings& AppendSettings)
 {
@@ -25,6 +77,19 @@ void FSkeletalMeshOperations::AppendSkinWeight(const FMeshDescription& SourceMes
 
 	TargetMesh.SuspendVertexIndexing();
 	
+	//Append Custom VertexAttribute
+	if(AppendSettings.bAppendVertexAttributes)
+	{
+		TArray<FName> TargetCustomAttributeNames;
+		TargetMesh.VertexAttributes().GetAttributeNames(TargetCustomAttributeNames);
+		int32 TargetVertexIndexOffset = FMath::Max(TargetMesh.Vertices().Num() - SourceMesh.Vertices().Num(), 0);
+
+		SourceMesh.VertexAttributes().ForEachByType<float>(UE::Private::FCreateAndCopyAttributeValues<float>(SourceMesh, TargetMesh, TargetCustomAttributeNames, AppendSettings.SourceVertexIDOffset));
+		SourceMesh.VertexAttributes().ForEachByType<FVector2f>(UE::Private::FCreateAndCopyAttributeValues<FVector2f>(SourceMesh, TargetMesh, TargetCustomAttributeNames, AppendSettings.SourceVertexIDOffset));
+		SourceMesh.VertexAttributes().ForEachByType<FVector3f>(UE::Private::FCreateAndCopyAttributeValues<FVector3f>(SourceMesh, TargetMesh, TargetCustomAttributeNames, AppendSettings.SourceVertexIDOffset));
+		SourceMesh.VertexAttributes().ForEachByType<FVector4f>(UE::Private::FCreateAndCopyAttributeValues<FVector4f>(SourceMesh, TargetMesh, TargetCustomAttributeNames, AppendSettings.SourceVertexIDOffset));
+	}
+
 	for (const FVertexID SourceVertexID : SourceMesh.Vertices().GetElementIDs())
 	{
 		const FVertexID TargetVertexID = FVertexID(AppendSettings.SourceVertexIDOffset + SourceVertexID.GetValue());
diff --git a/Engine/Source/Runtime/SkeletalMeshDescription/Public/SkeletalMeshOperations.h b/Engine/Source/Runtime/SkeletalMeshDescription/Public/SkeletalMeshOperations.h
--- a/Engine/Source/Runtime/SkeletalMeshDescription/Public/SkeletalMeshOperations.h
+++ b/Engine/Source/Runtime/SkeletalMeshDescription/Public/SkeletalMeshOperations.h
@@ -28,6 +28,7 @@ public:
 
 		int32 SourceVertexIDOffset;
 		TArray<FBoneIndexType> SourceRemapBoneIndex;
+		bool bAppendVertexAttributes = false;
 	};
 	
 	static SKELETALMESHDESCRIPTION_API void AppendSkinWeight(const FMeshDescription& SourceMesh, FMeshDescription& TargetMesh, FSkeletalMeshAppendSettings& AppendSettings);
