Interchange: Make sure fbx node with no attribute but participating to skinning is view has a joint. When importing a scene, do not create staticmesh component if the mesh node is skinned.

#rnx

diff --git a/Engine/Plugins/Interchange/Runtime/Source/Parsers/Fbx/Private/FbxScene.cpp b/Engine/Plugins/Interchange/Runtime/Source/Parsers/Fbx/Private/FbxScene.cpp
--- a/Engine/Plugins/Interchange/Runtime/Source/Parsers/Fbx/Private/FbxScene.cpp
+++ b/Engine/Plugins/Interchange/Runtime/Source/Parsers/Fbx/Private/FbxScene.cpp
@@ -107,7 +107,8 @@ namespace UE
 				, FbxNode* Node
 				, FbxScene* SDKScene
 				, UInterchangeBaseNodeContainer& NodeContainer
-				, TMap<FString, TSharedPtr<FPayloadContextBase, ESPMode::ThreadSafe>>& PayloadContexts)
+				, TMap<FString, TSharedPtr<FPayloadContextBase, ESPMode::ThreadSafe>>& PayloadContexts
+				, TArray<FbxNode*>& ForceJointNodes)
 			{
 				constexpr bool bResetCache = false;
 				FString NodeName = Parser.GetFbxHelper()->GetFbxObjectName(Node);
@@ -158,7 +159,56 @@ namespace UE
 					}
 				}
 
-				int32 AttributeCount = Node->GetNodeAttributeCount();
+				auto ApplySkeletonAttribute = [this, &SDKScene, &UnrealNode, &Node, &NodeContainer, &bResetCache, &GetConvertedTransform]()
+				{
+					//Add the joint specialized type
+					UnrealNode->AddSpecializedType(FSceneNodeStaticData::GetJointSpecializeTypeString());
+					//Get the bind pose transform for this joint
+					FbxAMatrix GlobalBindPoseJointMatrix;
+					if (FFbxMesh::GetGlobalJointBindPoseTransform(SDKScene, Node, GlobalBindPoseJointMatrix))
+					{
+						FTransform GlobalBindPoseJointTransform = GetConvertedTransform(GlobalBindPoseJointMatrix);
+						//We grab the fbx parent node to compute the local transform
+						if (FbxNode* ParentNode = Node->GetParent())
+						{
+							FbxAMatrix GlobalFbxParentMatrix = ParentNode->EvaluateGlobalTransform();
+							FFbxMesh::GetGlobalJointBindPoseTransform(SDKScene, ParentNode, GlobalFbxParentMatrix);
+							FbxAMatrix	LocalFbxMatrix = GlobalFbxParentMatrix.Inverse() * GlobalBindPoseJointMatrix;
+							FTransform LocalBindPoseJointTransform = GetConvertedTransform(LocalFbxMatrix);
+							UnrealNode->SetCustomBindPoseLocalTransform(&NodeContainer, LocalBindPoseJointTransform, bResetCache);
+						}
+						else
+						{
+							//No parent, set the same matrix has the global
+							UnrealNode->SetCustomBindPoseLocalTransform(&NodeContainer, GlobalBindPoseJointTransform, bResetCache);
+						}
+					}
+
+					//Get time Zero transform for this joint
+					{
+						//Set the global node transform
+						FbxAMatrix GlobalFbxMatrix = Node->EvaluateGlobalTransform(FBXSDK_TIME_ZERO);
+						FTransform GlobalTransform = GetConvertedTransform(GlobalFbxMatrix);
+						if (FbxNode* ParentNode = Node->GetParent())
+						{
+							FbxAMatrix GlobalFbxParentMatrix = ParentNode->EvaluateGlobalTransform(FBXSDK_TIME_ZERO);
+							FbxAMatrix	LocalFbxMatrix = GlobalFbxParentMatrix.Inverse() * GlobalFbxMatrix;
+							FTransform LocalTransform = GetConvertedTransform(LocalFbxMatrix);
+							UnrealNode->SetCustomTimeZeroLocalTransform(&NodeContainer, LocalTransform, bResetCache);
+						}
+						else
+						{
+							//No parent, set the same matrix has the global
+							UnrealNode->SetCustomTimeZeroLocalTransform(&NodeContainer, GlobalTransform, bResetCache);
+						}
+					}
+
+					FString JointNodeName = Parser.GetFbxHelper()->GetFbxObjectName(Node, true);
+					UnrealNode->SetDisplayLabel(JointNodeName);
+				};
+
+				bool bIsNodeContainJointAttribute = false;
+				const int32 AttributeCount = Node->GetNodeAttributeCount();
 				for (int32 AttributeIndex = 0; AttributeIndex < AttributeCount; ++AttributeIndex)
 				{
 					FbxNodeAttribute* NodeAttribute = Node->GetNodeAttributeByIndex(AttributeIndex);
@@ -196,50 +246,8 @@ namespace UE
 						//No break since the eNull act has a skeleton if possible
 						case FbxNodeAttribute::eSkeleton:
 						{
-							//Add the joint specialized type
-							UnrealNode->AddSpecializedType(FSceneNodeStaticData::GetJointSpecializeTypeString());
-							//Get the bind pose transform for this joint
-							FbxAMatrix GlobalBindPoseJointMatrix;
-							if (FFbxMesh::GetGlobalJointBindPoseTransform(SDKScene, Node, GlobalBindPoseJointMatrix))
-							{
-								FTransform GlobalBindPoseJointTransform = GetConvertedTransform(GlobalBindPoseJointMatrix);
-								//We grab the fbx parent node to compute the local transform
-								if (FbxNode* ParentNode = Node->GetParent())
-								{
-									FbxAMatrix GlobalFbxParentMatrix = ParentNode->EvaluateGlobalTransform();
-									FFbxMesh::GetGlobalJointBindPoseTransform(SDKScene, ParentNode, GlobalFbxParentMatrix);
-									FbxAMatrix	LocalFbxMatrix = GlobalFbxParentMatrix.Inverse() * GlobalBindPoseJointMatrix;
-									FTransform LocalBindPoseJointTransform = GetConvertedTransform(LocalFbxMatrix);
-									UnrealNode->SetCustomBindPoseLocalTransform(&NodeContainer, LocalBindPoseJointTransform, bResetCache);
-								}
-								else
-								{
-									//No parent, set the same matrix has the global
-									UnrealNode->SetCustomBindPoseLocalTransform(&NodeContainer, GlobalBindPoseJointTransform, bResetCache);
-								}
-							}
-
-							//Get time Zero transform for this joint
-							{
-								//Set the global node transform
-								FbxAMatrix GlobalFbxMatrix = Node->EvaluateGlobalTransform(FBXSDK_TIME_ZERO);
-								FTransform GlobalTransform = GetConvertedTransform(GlobalFbxMatrix);
-								if (FbxNode* ParentNode = Node->GetParent())
-								{
-									FbxAMatrix GlobalFbxParentMatrix = ParentNode->EvaluateGlobalTransform(FBXSDK_TIME_ZERO);
-									FbxAMatrix	LocalFbxMatrix = GlobalFbxParentMatrix.Inverse() * GlobalFbxMatrix;
-									FTransform LocalTransform = GetConvertedTransform(LocalFbxMatrix);
-									UnrealNode->SetCustomTimeZeroLocalTransform(&NodeContainer, LocalTransform, bResetCache);
-								}
-								else
-								{
-									//No parent, set the same matrix has the global
-									UnrealNode->SetCustomTimeZeroLocalTransform(&NodeContainer, GlobalTransform, bResetCache);
-								}
-							}
-
-							FString JointNodeName = Parser.GetFbxHelper()->GetFbxObjectName(Node, true);
-							UnrealNode->SetDisplayLabel(JointNodeName);
+							ApplySkeletonAttribute();
+							bIsNodeContainJointAttribute = true;
 							break;
 						}
 
@@ -282,6 +290,13 @@ namespace UE
 					}
 				}
 
+				//If there is no attribute, make sure to threat the node like a joint if it's in the ForcejointNodes array
+				if (!bIsNodeContainJointAttribute && ForceJointNodes.Contains(Node))
+				{
+					UnrealNode->AddSpecializedType(FSceneNodeStaticData::GetTransformSpecializeTypeString());
+					ApplySkeletonAttribute();
+				}
+
 				//Add all custom Attributes for the node
 				FbxProperty Property = Node->GetFirstProperty();
 				while (Property.IsValid())
@@ -419,7 +434,7 @@ namespace UE
 				for (int32 ChildIndex = 0; ChildIndex < ChildCount; ++ChildIndex)
 				{
 					FbxNode* ChildNode = Node->GetChild(ChildIndex);
-					AddHierarchyRecursively(UnrealNode, ChildNode, SDKScene, NodeContainer, PayloadContexts);
+					AddHierarchyRecursively(UnrealNode, ChildNode, SDKScene, NodeContainer, PayloadContexts, ForceJointNodes);
 				}
 			}
 
@@ -444,7 +459,12 @@ namespace UE
 			{
 				FbxNode* RootNode = SDKScene->GetRootNode();
 
-				AddHierarchyRecursively(nullptr, RootNode, SDKScene, NodeContainer, PayloadContexts);
+				//Some fbx file have node without attribute that are link in cluster,
+				//We must consider those node has joint
+				TArray<FbxNode*> ForceJointNodes;
+				FindForceJointNode(SDKScene, ForceJointNodes);
+
+				AddHierarchyRecursively(nullptr, RootNode, SDKScene, NodeContainer, PayloadContexts, ForceJointNodes);
 
 				int32 NodeCount = SDKScene->GetNodeCount();
 				for (int32 NodeIndex = 0; NodeIndex < NodeCount; ++NodeIndex)
@@ -455,7 +475,7 @@ namespace UE
 						{
 							if (Node->GetParent() == nullptr)
 							{
-								AddHierarchyRecursively(nullptr, Node, SDKScene, NodeContainer, PayloadContexts);
+								AddHierarchyRecursively(nullptr, Node, SDKScene, NodeContainer, PayloadContexts, ForceJointNodes);
 							}
 						}
 					}
@@ -548,18 +568,85 @@ namespace UE
 				}
 			}
 
+			void FFbxScene::FindForceJointNode(FbxScene* SDKScene, TArray<FbxNode*>& ForceJointNodes)
+			{
+				int32 GeometryCount = SDKScene->GetGeometryCount();
+				for (int32 GeometryIndex = 0; GeometryIndex < GeometryCount; ++GeometryIndex)
+				{
+					FbxGeometry* Geometry = SDKScene->GetGeometry(GeometryIndex);
+					if (Geometry->GetAttributeType() != FbxNodeAttribute::eMesh)
+					{
+						continue;
+					}
+					FbxMesh* Mesh = static_cast<FbxMesh*>(Geometry);
+					if (!Mesh)
+					{
+						continue;
+					}
+					const int32 SkinDeformerCount = Mesh->GetDeformerCount(FbxDeformer::eSkin);
+					for (int32 DeformerIndex = 0; DeformerIndex < SkinDeformerCount; DeformerIndex++)
+					{
+						FbxSkin* Skin = (FbxSkin*)Mesh->GetDeformer(DeformerIndex, FbxDeformer::eSkin);
+						if (!ensure(Skin))
+						{
+							continue;
+						}
+						const int32 ClusterCount = Skin->GetClusterCount();
+						for (int32 ClusterIndex = 0; ClusterIndex < ClusterCount; ClusterIndex++)
+						{
+							FbxCluster* Cluster = Skin->GetCluster(ClusterIndex);
+							// When Maya plug-in exports rigid binding, it will generate "CompensationCluster" for each ancestor links.
+							// FBX writes these "CompensationCluster" out. The CompensationCluster also has weight 1 for vertices.
+							// Unreal importer should skip these clusters.
+							if (!Cluster || (FCStringAnsi::Strcmp(Cluster->GetUserDataID(), "Maya_ClusterHint") == 0 && FCStringAnsi::Strcmp(Cluster->GetUserData(), "CompensationCluster") == 0))
+							{
+								continue;
+							}
+							ForceJointNodes.AddUnique(Cluster->GetLink());
+						}
+					}
+				}
+			}
+
 			void FFbxScene::AddAnimationRecursively(FbxNode* Node
 				, FbxScene* SDKScene
 				, UInterchangeBaseNodeContainer& NodeContainer
 				, TMap<FString, TSharedPtr<FPayloadContextBase, ESPMode::ThreadSafe>>& PayloadContexts
 				, UInterchangeSkeletalAnimationTrackNode* SkeletalAnimationTrackNode, bool SkeletalAnimationAddedToContainer
 				, const FString& RootSceneNodeUid, const TSet<FString>& SkeletonRootNodeUids
-				, const int32& AnimationIndex)
+				, const int32& AnimationIndex
+				, TArray<FbxNode*>& ForceJointNodes)
 			{
 				FString NodeUniqueID = Parser.GetFbxHelper()->GetFbxNodeHierarchyName(Node);
 				if (UInterchangeSceneNode* UnrealNode = const_cast<UInterchangeSceneNode*>(Cast< UInterchangeSceneNode >(NodeContainer.GetNode(NodeUniqueID))))
 				{
 					bool HasSkeletonAttribute = false;
+					auto ApplySkeletonAttribute = [&SDKScene, &HasSkeletonAttribute, &SkeletonRootNodeUids, &NodeUniqueID, &SkeletalAnimationTrackNode, &UnrealNode, &AnimationIndex, &NodeContainer]()
+					{
+						HasSkeletonAttribute = true;
+						if (SkeletonRootNodeUids.Contains(NodeUniqueID))
+						{
+							SkeletalAnimationTrackNode = NewObject< UInterchangeSkeletalAnimationTrackNode >(&NodeContainer);
+							FString TrackNodeUid = "\\SkeletalAnimation\\" + UnrealNode->GetUniqueID() + "_" + FString::FromInt(AnimationIndex);
+							FString DisplayString = "Anim_" + FString::FromInt(AnimationIndex) + "_" + UnrealNode->GetDisplayLabel();
+							SkeletalAnimationTrackNode->InitializeNode(TrackNodeUid, DisplayString, EInterchangeNodeContainerType::TranslatedAsset);
+
+							double FrameRate = FbxTime::GetFrameRate(SDKScene->GetGlobalSettings().GetTimeMode());
+							SkeletalAnimationTrackNode->SetCustomAnimationSampleRate(FrameRate);
+
+							SkeletalAnimationTrackNode->SetCustomSkeletonNodeUid(UnrealNode->GetUniqueID());
+
+							//Calculate AnimationTime:
+							FbxAnimStack* CurrentAnimationStack = (FbxAnimStack*)SDKScene->GetSrcObject<FbxAnimStack>(AnimationIndex);
+							FbxTimeSpan TimeSpan = CurrentAnimationStack->GetLocalTimeSpan();
+
+							SkeletalAnimationTrackNode->SetCustomAnimationStartTime(TimeSpan.GetStart().GetSecondDouble());
+							SkeletalAnimationTrackNode->SetCustomAnimationStopTime(TimeSpan.GetStop().GetSecondDouble());
+						}
+					};
+
+					bool bIsNodeContainJointAttribute = false;
+
 					int32 AttributeCount = Node->GetNodeAttributeCount();
 
 					for (int32 AttributeIndex = 0; AttributeIndex < AttributeCount && !HasSkeletonAttribute; ++AttributeIndex)
@@ -569,31 +656,18 @@ namespace UE
 						{
 						case FbxNodeAttribute::eNull:
 						case FbxNodeAttribute::eSkeleton:
-							HasSkeletonAttribute = true;
-							if (SkeletonRootNodeUids.Contains(NodeUniqueID))
-							{
-								SkeletalAnimationTrackNode = NewObject< UInterchangeSkeletalAnimationTrackNode >(&NodeContainer);
-								FString TrackNodeUid = "\\SkeletalAnimation\\" + UnrealNode->GetUniqueID() + "_" + FString::FromInt(AnimationIndex);
-								FString DisplayString = "Anim_" + FString::FromInt(AnimationIndex) + "_" + UnrealNode->GetDisplayLabel();
-								SkeletalAnimationTrackNode->InitializeNode(TrackNodeUid, DisplayString, EInterchangeNodeContainerType::TranslatedAsset);
-
-								double FrameRate = FbxTime::GetFrameRate(SDKScene->GetGlobalSettings().GetTimeMode());
-								SkeletalAnimationTrackNode->SetCustomAnimationSampleRate(FrameRate);
-
-								SkeletalAnimationTrackNode->SetCustomSkeletonNodeUid(UnrealNode->GetUniqueID());
-
-								//Calculate AnimationTime:
-								FbxAnimStack* CurrentAnimationStack = (FbxAnimStack*)SDKScene->GetSrcObject<FbxAnimStack>(AnimationIndex);
-								FbxTimeSpan TimeSpan = CurrentAnimationStack->GetLocalTimeSpan();
-
-								SkeletalAnimationTrackNode->SetCustomAnimationStartTime(TimeSpan.GetStart().GetSecondDouble());
-								SkeletalAnimationTrackNode->SetCustomAnimationStopTime(TimeSpan.GetStop().GetSecondDouble());
-							}
+							bIsNodeContainJointAttribute = true;
+							ApplySkeletonAttribute();
 							break;
 						default:
 							break;
 						}
 					}
+					//If there is no attribute, make sure to threat the node like a joint if it's in the ForcejointNodes array
+					if (!bIsNodeContainJointAttribute && ForceJointNodes.Contains(Node))
+					{
+						ApplySkeletonAttribute();
+					}
 
 					if (!HasSkeletonAttribute)
 					{
@@ -627,7 +701,7 @@ namespace UE
 				for (int32 ChildIndex = 0; ChildIndex < ChildCount; ++ChildIndex)
 				{
 					FbxNode* ChildNode = Node->GetChild(ChildIndex);
-					AddAnimationRecursively(ChildNode, SDKScene, NodeContainer, PayloadContexts, SkeletalAnimationTrackNode, SkeletalAnimationAddedToContainer, RootSceneNodeUid, SkeletonRootNodeUids, AnimationIndex);
+					AddAnimationRecursively(ChildNode, SDKScene, NodeContainer, PayloadContexts, SkeletalAnimationTrackNode, SkeletalAnimationAddedToContainer, RootSceneNodeUid, SkeletonRootNodeUids, AnimationIndex, ForceJointNodes);
 				}
 			}
 
@@ -636,6 +710,11 @@ namespace UE
 				FbxNode* RootNode = SDKScene->GetRootNode();
 				FString RootSceneNodeUniqueID = Parser.GetFbxHelper()->GetFbxNodeHierarchyName(RootNode);
 
+				//Some fbx file have node without attribute that are link in cluster,
+				//We must consider those node has joint
+				TArray<FbxNode*> ForceJointNodes;
+				FindForceJointNode(SDKScene, ForceJointNodes);
+
 				//acquire Skeletal Node Uids from Meshes (via the skeletondependencies: )
 				TSet<FString> SkeletonRootNodeUids;
 				NodeContainer.IterateNodesOfType<UInterchangeMeshNode>([&](const FString& NodeUid, UInterchangeMeshNode* MeshNode)
@@ -716,7 +795,7 @@ namespace UE
 
 				for (int32 AnimationIndex = 0; AnimationIndex < NumAnimations; AnimationIndex++)
 				{
-					AddAnimationRecursively(RootNode, SDKScene, NodeContainer, PayloadContexts, nullptr, false, RootSceneNodeUniqueID, SkeletonRootNodeUids, AnimationIndex);
+					AddAnimationRecursively(RootNode, SDKScene, NodeContainer, PayloadContexts, nullptr, false, RootSceneNodeUniqueID, SkeletonRootNodeUids, AnimationIndex, ForceJointNodes);
 					int32 NodeCount = SDKScene->GetNodeCount();
 					for (int32 NodeIndex = 0; NodeIndex < NodeCount; ++NodeIndex)
 					{
@@ -726,7 +805,7 @@ namespace UE
 							{
 								if (Node->GetParent() == nullptr)
 								{
-									AddAnimationRecursively(Node, SDKScene, NodeContainer, PayloadContexts, nullptr, false, RootSceneNodeUniqueID, SkeletonRootNodeUids, AnimationIndex);
+									AddAnimationRecursively(Node, SDKScene, NodeContainer, PayloadContexts, nullptr, false, RootSceneNodeUniqueID, SkeletonRootNodeUids, AnimationIndex, ForceJointNodes);
 								}
 							}
 						}
diff --git a/Engine/Plugins/Interchange/Runtime/Source/Parsers/Fbx/Private/FbxScene.h b/Engine/Plugins/Interchange/Runtime/Source/Parsers/Fbx/Private/FbxScene.h
--- a/Engine/Plugins/Interchange/Runtime/Source/Parsers/Fbx/Private/FbxScene.h
+++ b/Engine/Plugins/Interchange/Runtime/Source/Parsers/Fbx/Private/FbxScene.h
@@ -43,7 +43,8 @@ namespace UE
 					, FbxNode* Node
 					, FbxScene* SDKScene
 					, UInterchangeBaseNodeContainer& NodeContainer
-					, TMap<FString, TSharedPtr<FPayloadContextBase, ESPMode::ThreadSafe>>& PayloadContexts);
+					, TMap<FString, TSharedPtr<FPayloadContextBase, ESPMode::ThreadSafe>>& PayloadContexts
+					, TArray<FbxNode*>& ForceJointNodes);
 
 				void AddAnimationRecursively(FbxNode* Node
 					, FbxScene* SDKScene
@@ -51,7 +52,8 @@ namespace UE
 					, TMap<FString, TSharedPtr<FPayloadContextBase, ESPMode::ThreadSafe>>& PayloadContexts
 					, UInterchangeSkeletalAnimationTrackNode* SkeletalAnimationTrackNode, bool SkeletalAnimationAddedToContainer
 					, const FString& RootSceneNodeUid, const TSet<FString>& SkeletonRootNodeUids
-					, const int32& AnimationIndex);
+					, const int32& AnimationIndex
+					, TArray<FbxNode*>& ForceJointNodes);
 
 			private:
 				void AddRigidAnimation(FbxNode* Node
@@ -59,6 +61,8 @@ namespace UE
 					, UInterchangeBaseNodeContainer& NodeContainer
 					, TMap<FString, TSharedPtr<FPayloadContextBase, ESPMode::ThreadSafe>>& PayloadContexts);
 
+				void FindForceJointNode(FbxScene* SDKScene, TArray<FbxNode*>& ForceJointNodes);
+
 				FFbxParser& Parser;
 			};
 		}//ns Private
diff --git a/Engine/Plugins/Interchange/Runtime/Source/Pipelines/Private/InterchangeGenericScenesPipeline.cpp b/Engine/Plugins/Interchange/Runtime/Source/Pipelines/Private/InterchangeGenericScenesPipeline.cpp
--- a/Engine/Plugins/Interchange/Runtime/Source/Pipelines/Private/InterchangeGenericScenesPipeline.cpp
+++ b/Engine/Plugins/Interchange/Runtime/Source/Pipelines/Private/InterchangeGenericScenesPipeline.cpp
@@ -9,6 +9,7 @@
 #include "InterchangeCommonPipelineDataFactoryNode.h"
 #include "InterchangeSceneImportAsset.h"
 #include "InterchangeSceneImportAssetFactoryNode.h"
+#include "InterchangeSkeletalMeshLodDataNode.h"
 #include "InterchangeLightNode.h"
 #include "InterchangeLightFactoryNode.h"
 #include "InterchangeMeshActorFactoryNode.h"
@@ -352,6 +353,42 @@ void UInterchangeGenericLevelPipeline::ExecuteSceneNodePreImport(const FTransfor
 		{
 			TranslatedAssetNode = BaseNodeContainer->GetNode(AssetInstanceUid);
 		}
+		if (const UInterchangeMeshNode* MeshNode = Cast<UInterchangeMeshNode>(TranslatedAssetNode))
+		{
+			bool bIsSkinnedMesh = MeshNode->IsSkinnedMesh();
+			if(!bIsSkinnedMesh)
+			{
+				//In case we have a rigid mesh (static mesh having morph target...), we should find a skeletalmesh factory node holding this mesh node
+				BaseNodeContainer->BreakableIterateNodesOfType<UInterchangeSkeletalMeshFactoryNode>([this, &SceneNode, &bIsSkinnedMesh](const FString& NodeUid, UInterchangeSkeletalMeshFactoryNode* SkeletalMeshNode)
+				{
+					int32 LodCount = SkeletalMeshNode->GetLodDataCount();
+					TArray<FString> LodDataUniqueIds;
+					SkeletalMeshNode->GetLodDataUniqueIds(LodDataUniqueIds);
+					for (int32 LodIndex = 0; LodIndex < LodCount; ++LodIndex)
+					{
+						FString LodUniqueId = LodDataUniqueIds[LodIndex];
+						const UInterchangeSkeletalMeshLodDataNode* LodDataNode = Cast<UInterchangeSkeletalMeshLodDataNode>(BaseNodeContainer->GetNode(LodUniqueId));
+						if (LodDataNode)
+						{
+							TArray<FString> MeshUids;
+							LodDataNode->GetMeshUids(MeshUids);
+							if (MeshUids.Contains(SceneNode->GetUniqueID()))
+							{
+								bIsSkinnedMesh = true;
+								return bIsSkinnedMesh;
+							}
+						}
+					}
+					return false;
+				});	
+			}
+			//Skinned mesh are added when the bRootJointNode is true.
+			//In this case we dont want to add an empty staticmesh actor.
+			if (bIsSkinnedMesh)
+			{
+				return;
+			}
+		}
 	}
 
 	UInterchangeActorFactoryNode* ActorFactoryNode = CreateActorFactoryNode(SceneNode, TranslatedAssetNode);
