Interchange: Fix fbx importer to make sure joint without cluster use the bindpose to get the global bind transform.
Also fix a check issue if we have an empty UV channel, it was crashing the InterchangeWorker application.

#rnx

diff --git a/Engine/Plugins/Interchange/Runtime/Source/Dispatcher/Public/InterchangeCommands.h b/Engine/Plugins/Interchange/Runtime/Source/Dispatcher/Public/InterchangeCommands.h
--- a/Engine/Plugins/Interchange/Runtime/Source/Dispatcher/Public/InterchangeCommands.h
+++ b/Engine/Plugins/Interchange/Runtime/Source/Dispatcher/Public/InterchangeCommands.h
@@ -17,7 +17,7 @@ namespace UE
 		{
 		public:
 			//Major version should be updated when there is an existing API that has been change.
-			static int32 GetMajor() { return 14; }
+			static int32 GetMajor() { return 15; }
 			//Minor version must be updated when there is an addition to the API.
 			static int32 GetMinor() { return 0; }
 			//Patch version should be update if there is some bug fixes in the private code.
diff --git a/Engine/Plugins/Interchange/Runtime/Source/Parsers/Fbx/Private/FbxMesh.cpp b/Engine/Plugins/Interchange/Runtime/Source/Parsers/Fbx/Private/FbxMesh.cpp
--- a/Engine/Plugins/Interchange/Runtime/Source/Parsers/Fbx/Private/FbxMesh.cpp
+++ b/Engine/Plugins/Interchange/Runtime/Source/Parsers/Fbx/Private/FbxMesh.cpp
@@ -331,13 +331,12 @@ struct FFBXUVs
 		UniqueUVCount = UVSets.Num();
 		if (UniqueUVCount > 0)
 		{
-			LayerElementUV.AddZeroed(UniqueUVCount);
-			UVReferenceMode.AddZeroed(UniqueUVCount);
-			UVMappingMode.AddZeroed(UniqueUVCount);
+			LayerElementUV.Reserve(UniqueUVCount);
+			UVReferenceMode.Reserve(UniqueUVCount);
+			UVMappingMode.Reserve(UniqueUVCount);
 		}
 		for (int32 UVIndex = 0; UVIndex < UniqueUVCount; UVIndex++)
 		{
-			LayerElementUV[UVIndex] = NULL;
 			for (int32 UVLayerIndex = 0, LayerCount = Mesh->GetLayerCount(); UVLayerIndex < LayerCount; UVLayerIndex++)
 			{
 				FbxLayer* lLayer = Mesh->GetLayer(UVLayerIndex);
@@ -358,9 +357,9 @@ struct FFBXUVs
 							}
 							if (LocalUVSetName == UVSets[UVIndex])
 							{
-								LayerElementUV[UVIndex] = ElementUV;
-								UVReferenceMode[UVIndex] = ElementUV->GetReferenceMode();
-								UVMappingMode[UVIndex] = ElementUV->GetMappingMode();
+								LayerElementUV.Add(ElementUV);
+								UVReferenceMode.Add(ElementUV->GetReferenceMode());
+								UVMappingMode.Add(ElementUV->GetMappingMode());
 								break;
 							}
 						}
@@ -368,6 +367,7 @@ struct FFBXUVs
 				}
 			}
 		}
+		UniqueUVCount = LayerElementUV.Num();
 
 		if (UniqueUVCount > MAX_MESH_TEXTURE_COORDS_MD)
 		{
@@ -1785,6 +1785,8 @@ void FFbxMesh::AddAllMeshes(FbxScene* SDKScene, FbxGeometryConverter* SDKGeometr
 
 bool FFbxMesh::GetGlobalJointBindPoseTransform(FbxScene* SDKScene, FbxNode* Joint, FbxAMatrix& GlobalBindPoseJointMatrix)
 {
+	//First look for Cluster and then look in the bind pose if no cluster was found
+
 	//Search all skeletalmesh(FbxGeometry with valid deformer) using this joint and see if there is a valid FbxCluster
 	//containing a TransformLinkMatrix for this joint
 	const int32 GeometryCount = SDKScene->GetGeometryCount();
@@ -1822,6 +1824,32 @@ bool FFbxMesh::GetGlobalJointBindPoseTransform(FbxScene* SDKScene, FbxNode* Join
 			}
 		}
 	}
+
+	const int32 PoseCount = SDKScene->GetPoseCount();
+	for (int32 PoseIndex = 0; PoseIndex < PoseCount; PoseIndex++)
+	{
+		FbxPose* CurrentPose = SDKScene->GetPose(PoseIndex);
+
+		// current pose is bind pose, 
+		if (CurrentPose && CurrentPose->IsBindPose())
+		{
+			FString PoseName = CurrentPose->GetName();
+			// all error report status
+			FbxStatus Status;
+
+			int32 PoseLinkIndex = CurrentPose->Find(Joint);
+			if (PoseLinkIndex >= 0)
+			{
+				if (!CurrentPose->IsLocalMatrix(PoseLinkIndex))
+				{
+					FbxMatrix NoneAffineMatrix = CurrentPose->GetMatrix(PoseLinkIndex);
+					GlobalBindPoseJointMatrix = *(FbxAMatrix*)(double*)&NoneAffineMatrix;
+					return true;
+				}
+			}
+		}
+	}
+
 	return false;
 }
 
