[Interchange] MaterialInstancing fixes + test fixes.
[Interchange] GLTF test fixes.
[Interchange-glTF] turned on MaterialInstance creation on default.

diff --git a/Engine/Plugins/Interchange/Runtime/Source/Pipelines/Private/InterchangeGenericMaterialPipeline.cpp b/Engine/Plugins/Interchange/Runtime/Source/Pipelines/Private/InterchangeGenericMaterialPipeline.cpp
--- a/Engine/Plugins/Interchange/Runtime/Source/Pipelines/Private/InterchangeGenericMaterialPipeline.cpp
+++ b/Engine/Plugins/Interchange/Runtime/Source/Pipelines/Private/InterchangeGenericMaterialPipeline.cpp
@@ -2924,7 +2924,7 @@ UInterchangeMaterialInstanceFactoryNode* UInterchangeGenericMaterialPipeline::Cr
 			ParentRootName = TEXT("PBRSurfaceMaterial_");
 		}
 
-		const FString ParentAssetPath = TEXT("Material'/Interchange/Materials/") + ParentRootName + Model + TEXT(".") + ParentRootName + Model + TEXT("'");
+		const FString ParentAssetPath = TEXT("/Interchange/Materials/") + ParentRootName + Model + TEXT(".") + ParentRootName + Model;
 		MaterialInstanceFactoryNode->SetCustomParent(ParentAssetPath);
 	};
 
@@ -2942,20 +2942,20 @@ UInterchangeMaterialInstanceFactoryNode* UInterchangeGenericMaterialPipeline::Cr
 	}
 	else if (IsPhongModel(ShaderGraphNode))
 	{
-		MaterialInstanceFactoryNode->SetCustomParent(TEXT("Material'/Interchange/Materials/PhongSurfaceMaterial.PhongSurfaceMaterial'"));
+		MaterialInstanceFactoryNode->SetCustomParent(TEXT("/Interchange/Materials/PhongSurfaceMaterial.PhongSurfaceMaterial"));
 	}
 	else if (IsLambertModel(ShaderGraphNode))
 	{
-		MaterialInstanceFactoryNode->SetCustomParent(TEXT("Material'/Interchange/Materials/LambertSurfaceMaterial.LambertSurfaceMaterial'"));
+		MaterialInstanceFactoryNode->SetCustomParent(TEXT("/Interchange/Materials/LambertSurfaceMaterial.LambertSurfaceMaterial"));
 	}
 	else if (IsUnlitModel(ShaderGraphNode))
 	{
-		MaterialInstanceFactoryNode->SetCustomParent(TEXT("Material'/Interchange/Materials/UnlitMaterial.UnlitMaterial'"));
+		MaterialInstanceFactoryNode->SetCustomParent(TEXT("/Interchange/Materials/UnlitMaterial.UnlitMaterial"));
 	}
 	else
 	{
 		// Default to PBR
-		MaterialInstanceFactoryNode->SetCustomParent(TEXT("Material'/Interchange/Materials/PBRSurfaceMaterial.PBRSurfaceMaterial'"));
+		MaterialInstanceFactoryNode->SetCustomParent(TEXT("/Interchange/Materials/PBRSurfaceMaterial.PBRSurfaceMaterial"));
 	}
 
 #if WITH_EDITOR
diff --git a/Engine/Plugins/Interchange/Runtime/Source/Pipelines/Public/InterchangeglTFPipeline.h b/Engine/Plugins/Interchange/Runtime/Source/Pipelines/Public/InterchangeglTFPipeline.h
--- a/Engine/Plugins/Interchange/Runtime/Source/Pipelines/Public/InterchangeglTFPipeline.h
+++ b/Engine/Plugins/Interchange/Runtime/Source/Pipelines/Public/InterchangeglTFPipeline.h
@@ -54,7 +54,7 @@ public:
 	TObjectPtr<UGLTFPipelineSettings> GLTFPipelineSettings;
 
 	UPROPERTY(EditAnywhere, BlueprintReadWrite, Category = "glTF", meta = (DisplayName = "Use glTF Material Instance Library"))
-	bool bUseGLTFMaterialInstanceLibrary = false;
+	bool bUseGLTFMaterialInstanceLibrary = true;
 
 protected:
 	virtual void AdjustSettingsForContext(EInterchangePipelineContext ImportType, TObjectPtr<UObject> ReimportAsset) override;
