MaterialX Pipeline:
- Add mapping of MaterialX pbr nodedefs to Material Functions
- Refactor MaterialX pipeline and MaterialXCustomize with a better handling of the different mappings (surface shaders, bsdf, etc.)
- Adding Material Functions placeholders for the mapping, no implemented logic in it yet, just here for the overload.
- Removing rest of unused code of MaterialX in the generic material pipeline
- Hiding MaterialSettings and Material Functions from MaterialX pipeline
- Overriding Material Functions only from Project Settings -> Engine -> Interchange MaterialX

MaterialX Translator:
- Adding support of pbr <surface> node
- Connecting outputs based on their index for multi outputs, used by node with material functions like <artistic_ior>
- Adapting the data structures for matching nodes and inputs, adding an optionalkey based on the nodegroup, particularly useful for <mix>, <add> and <multiply> nodes on BSDFs

MaterialX SurfaceShaders:
- Avoid a potential crash, by looking for an an approximate node definition during matching node and inputs resolution, since the the nodes are altered, for example nodes like blur.

MaterialX Pipeline Customization
- Making Enums in MaterialXPipelineCustomization only ReadOnly text instead of dropdowns

Substrate:
- Adding Substrate material functions counterpart of surface shaders and BSDF in the Substrate folder
- Adding support of Substrate for <surface_unlit>, <UsdPreviewSurface>
- Adding a CVar and a function in the InterchangeImportModule to check the availability of Substrate in the project settings and to enable or disable its support in Interchange.

diff --git a/Engine/Plugins/Interchange/Editor/Source/Pipelines/Private/InterchangeMaterialXPipelineCustomizations.cpp b/Engine/Plugins/Interchange/Editor/Source/Pipelines/Private/InterchangeMaterialXPipelineCustomizations.cpp
--- a/Engine/Plugins/Interchange/Editor/Source/Pipelines/Private/InterchangeMaterialXPipelineCustomizations.cpp
+++ b/Engine/Plugins/Interchange/Editor/Source/Pipelines/Private/InterchangeMaterialXPipelineCustomizations.cpp
@@ -2,15 +2,15 @@
 
 #include "InterchangeMaterialXPipelineCustomizations.h"
 
-#include "InterchangeMaterialDefinitions.h"
 #include "InterchangeMaterialXPipeline.h"
-#include "Materials/MaterialFunction.h"
 
 #include "DetailLayoutBuilder.h"
+#include "Engine/RendererSettings.h"
+#include "Materials/MaterialFunction.h"
 #include "Modules/ModuleManager.h"
 #include "PropertyCustomizationHelpers.h"
-#include "Widgets/Text/STextBlock.h"
 #include "Widgets/Input/SEditableText.h"
+#include "Widgets/Text/STextBlock.h"
 
 TSharedRef<IDetailCustomization> FInterchangeMaterialXPipelineCustomization::MakeInstance()
 {
@@ -45,8 +45,6 @@ void FInterchangeMaterialXPipelineCustomization::CustomizeDetails(IDetailLayoutB
 	DetailsViewArgs.NameAreaSettings = FDetailsViewArgs::HideNameArea;
 	TSharedPtr<IDetailsView> DetailsView = PropertyEditorModule.CreateDetailView(DetailsViewArgs);
 
-	DetailsView->SetObject(Pipeline->MaterialXSettings);
-
 	MaterialXCategory.AddCustomRow(NSLOCTEXT("InterchangeMaterialXPipelineMaterialSubstitution::Message::Row", "MaterialXPredefinedMaterialSubstitutionMessageRow", "MaterialX Pipeline Material Substitution Message"))
 	.NameContent()
 	.VAlign(VAlign_Center)
@@ -60,8 +58,6 @@ void FInterchangeMaterialXPipelineCustomization::CustomizeDetails(IDetailLayoutB
 
 void FInterchangeMaterialXPipelineSettingsCustomization::CustomizeDetails(IDetailLayoutBuilder& DetailBuilder)
 {
-	using OnShouldFilterAssetFunc = bool (FInterchangeMaterialXPipelineSettingsCustomization::*)(const FAssetData&);
-
 	TArray<TWeakObjectPtr<UObject>> EditingObjects;
 	DetailBuilder.GetObjectsBeingCustomized(EditingObjects);
 	check(EditingObjects.Num() == 1);
@@ -73,97 +69,88 @@ void FInterchangeMaterialXPipelineSettingsCustomization::CustomizeDetails(IDetai
 		return;
 	}
 
-	TSharedRef< IPropertyHandle > PairingsHandle = DetailBuilder.GetProperty(GET_MEMBER_NAME_CHECKED(UMaterialXPipelineSettings, PredefinedSurfaceShaders));
-	if (!PairingsHandle->IsValidHandle())
+	auto Customize = [&DetailBuilder](const FName& MemberName, const FName& Category, const FText& DisplayName, uint8 EnumType)
 	{
-		return;
-	}
-
-	DetailBuilder.HideProperty(PairingsHandle);
-
-	IDetailCategoryBuilder& MaterialXPredefinedCategory = DetailBuilder.EditCategory("MaterialXPredefined");
+		TSharedRef< IPropertyHandle > PairingsHandle = DetailBuilder.GetProperty(MemberName);
+		if(!PairingsHandle->IsValidHandle())
+		{
+			return;
+		}
 
-	FText DisplayName{ NSLOCTEXT("InterchangeMaterialXPipelineSettingsCustomization", "CategoryDisplayName", "MaterialX Predefined Surface Shaders") };
-	MaterialXPredefinedCategory.SetDisplayName(DisplayName);
+		DetailBuilder.HideProperty(PairingsHandle);
 
-	uint32 NumChildren = 0;
-	PairingsHandle->GetNumChildren(NumChildren);
+		IDetailCategoryBuilder& MaterialXPredefinedCategory = DetailBuilder.EditCategory(Category, DisplayName);		
 
-	for(uint32 i = 0; i < NumChildren; ++i)
-	{
-		TSharedPtr<IPropertyHandle> ChildPropertyHandle = PairingsHandle->GetChildHandle(i);
-		TSharedPtr<IPropertyHandle> KeyPropertyHandle = ChildPropertyHandle->GetKeyHandle();
+		uint32 NumChildren = 0;
+		PairingsHandle->GetNumChildren(NumChildren);
 
-		UEnum* Enum = CastField<FEnumProperty>(KeyPropertyHandle->GetProperty())->GetEnum();
+		for(uint32 i = 0; i < NumChildren; ++i)
+		{
+			TSharedPtr<IPropertyHandle> ChildPropertyHandle = PairingsHandle->GetChildHandle(i);
+			TSharedPtr<IPropertyHandle> KeyPropertyHandle = ChildPropertyHandle->GetKeyHandle();
+			
+			UEnum* Enum = CastField<FEnumProperty>(KeyPropertyHandle->GetProperty())->GetEnum();
 
-		FText EnumString;
-		KeyPropertyHandle->GetValueAsDisplayText(EnumString);
+			uint8 EnumValue = 0;
+			KeyPropertyHandle->GetValue(EnumValue);
 
-		OnShouldFilterAssetFunc OnShouldFilterAsset = &FInterchangeMaterialXPipelineSettingsCustomization::OnShouldFilterAsset;
+			FText EnumString;
+			KeyPropertyHandle->GetValueAsDisplayText(EnumString);
 
-		uint8 EnumValue = 0;
-		KeyPropertyHandle->GetValue(EnumValue);
-		if(EnumString.ToString() == TEXT("Standard Surface"))
-		{
-			OnShouldFilterAsset = &FInterchangeMaterialXPipelineSettingsCustomization::OnShouldFilterAssetStandardSurface;
-		}
-		else if(EnumString.ToString() == TEXT("Standard Surface Transmission"))
-		{
-			OnShouldFilterAsset = &FInterchangeMaterialXPipelineSettingsCustomization::OnShouldFilterAssetStandardSurfaceTransmission;
-		}
-		else if(EnumString.ToString() == TEXT("Surface Unlit"))
-		{
-			OnShouldFilterAsset = &FInterchangeMaterialXPipelineSettingsCustomization::OnShouldFilterAssetSurfaceUnlit;
-		}
-		else if(EnumString.ToString() == TEXT("Usd Preview Surface"))
-		{
-			OnShouldFilterAsset = &FInterchangeMaterialXPipelineSettingsCustomization::OnShouldFilterAssetUsdPreviewSurface;
-		}
-
-		FDetailWidgetRow& DetailWidgetRow = MaterialXPredefinedCategory.AddCustomRow(DisplayName);
-		DetailWidgetRow
-		.NameContent()
-		[
-			SNew(SEditableText)
-			.Text(EnumString)
-			.IsReadOnly(true)
-			.ToolTipText(Enum->GetToolTipTextByIndex(EnumValue))
-		]
-		.ValueContent()
-		.MinDesiredWidth(1.0f)
-		.MaxDesiredWidth(1.0f)
-		[
-			SNew(SHorizontalBox)
-			+ SHorizontalBox::Slot()
-			.HAlign(HAlign_Fill)
+			FDetailWidgetRow& DetailWidgetRow = MaterialXPredefinedCategory.AddCustomRow(DisplayName);
+			DetailWidgetRow
+			.NameContent()
 			[
-				SNew(SObjectPropertyEntryBox)
-				.AllowedClass(UMaterialFunction::StaticClass())
-				.PropertyHandle(ChildPropertyHandle)
-				.OnShouldFilterAsset(this, OnShouldFilterAsset)
+				SNew(SEditableText)
+				.Text(EnumString)
+				.IsReadOnly(true)
+				.ToolTipText(Enum->GetToolTipTextByIndex(EnumValue))
 			]
-		];
-	}
-
-	DetailBuilder.HideCategory(GET_MEMBER_NAME_CHECKED(UMaterialXPipelineSettings, PredefinedSurfaceShaders));
-}
-
-bool FInterchangeMaterialXPipelineSettingsCustomization::OnShouldFilterAssetStandardSurface(const FAssetData & InAssetData)
-{
-	return UMaterialXPipelineSettings::ShouldFilterAssets(Cast<UMaterialFunction>(InAssetData.GetAsset()), UMaterialXPipelineSettings::StandardSurfaceInputs, UMaterialXPipelineSettings::StandardSurfaceOutputs);
+			.ValueContent()
+			.MinDesiredWidth(1.0f)
+			.MaxDesiredWidth(1.0f)
+			[
+				SNew(SHorizontalBox)
+				+ SHorizontalBox::Slot()
+				.HAlign(HAlign_Fill)
+				[
+					SNew(SObjectPropertyEntryBox)
+					.AllowedClass(UMaterialFunction::StaticClass())
+					.PropertyHandle(ChildPropertyHandle)
+					.OnShouldFilterAsset_Static(&OnShouldFilterAssetEnum, EnumType, EnumValue)
+				]
+			];
+		}
+	};
+
+	Customize(GET_MEMBER_NAME_CHECKED(UMaterialXPipelineSettings, PredefinedSurfaceShaders),
+			  TEXT("MaterialXPredefined | Surface Shaders"),
+			  NSLOCTEXT("InterchangeMaterialXPipelineSettingsCustomization", "MaterialXPredefined | SurfaceShaders", "Surface Shaders"),
+			  UMaterialXPipelineSettings::IndexSurfaceShaders);
+
+	Customize(GET_MEMBER_NAME_CHECKED(UMaterialXPipelineSettings, PredefinedBSDF),
+			  TEXT("MaterialXPredefined | BSDF"),
+			  NSLOCTEXT("InterchangeMaterialXPipelineSettingsCustomization", "MaterialXPredefined | BSDF", "Bidirectional Scattering Distribution Functions"),
+			  UMaterialXPipelineSettings::IndexBSDF);
+
+	Customize(GET_MEMBER_NAME_CHECKED(UMaterialXPipelineSettings, PredefinedEDF),
+			  TEXT("MaterialXPredefined | EDF"),
+			  NSLOCTEXT("InterchangeMaterialXPipelineSettingsCustomization", "MaterialXPredefined | EDF", "Emission Distribution Functions"),
+			  UMaterialXPipelineSettings::IndexEDF);
+
+	Customize(GET_MEMBER_NAME_CHECKED(UMaterialXPipelineSettings, PredefinedVDF),
+			  TEXT("MaterialXPredefined | VDF"),
+			  NSLOCTEXT("InterchangeMaterialXPipelineSettingsCustomization", "MaterialXPredefined | VDF", "Volume Distribution Functions"),
+			  UMaterialXPipelineSettings::IndexVDF);
 }
 
-bool FInterchangeMaterialXPipelineSettingsCustomization::OnShouldFilterAssetStandardSurfaceTransmission(const FAssetData& InAssetData)
+bool FInterchangeMaterialXPipelineSettingsCustomization::OnShouldFilterAssetEnum(const FAssetData& InAssetData, uint8 EnumType, uint8 EnumValue)
 {
-	return UMaterialXPipelineSettings::ShouldFilterAssets(Cast<UMaterialFunction>(InAssetData.GetAsset()), UMaterialXPipelineSettings::TransmissionSurfaceInputs, UMaterialXPipelineSettings::TransmissionSurfaceOutputs);
-}
+	UMaterialXPipelineSettings::FMaterialXSettings::ValueType * Settings = UMaterialXPipelineSettings::SettingsInputsOutputs.Find(UMaterialXPipelineSettings::ToEnumKey(EnumType, EnumValue));
 
-bool FInterchangeMaterialXPipelineSettingsCustomization::OnShouldFilterAssetSurfaceUnlit(const FAssetData& InAssetData)
-{
-	return UMaterialXPipelineSettings::ShouldFilterAssets(Cast<UMaterialFunction>(InAssetData.GetAsset()), UMaterialXPipelineSettings::SurfaceUnlitInputs, UMaterialXPipelineSettings::SurfaceUnlitOutputs);
-}
-
-bool FInterchangeMaterialXPipelineSettingsCustomization::OnShouldFilterAssetUsdPreviewSurface(const FAssetData& InAssetData)
-{
-	return UMaterialXPipelineSettings::ShouldFilterAssets(Cast<UMaterialFunction>(InAssetData.GetAsset()), UMaterialXPipelineSettings::UsdPreviewSurfaceInputs, UMaterialXPipelineSettings::UsdPreviewSurfaceOutputs);
+	if(Settings)
+	{
+		return UMaterialXPipelineSettings::ShouldFilterAssets(Cast<UMaterialFunction>(InAssetData.GetAsset()), Settings->Key, Settings->Value);
+	}
+	return false;
 }
diff --git a/Engine/Plugins/Interchange/Editor/Source/Pipelines/Private/InterchangeMaterialXPipelineCustomizations.h b/Engine/Plugins/Interchange/Editor/Source/Pipelines/Private/InterchangeMaterialXPipelineCustomizations.h
--- a/Engine/Plugins/Interchange/Editor/Source/Pipelines/Private/InterchangeMaterialXPipelineCustomizations.h
+++ b/Engine/Plugins/Interchange/Editor/Source/Pipelines/Private/InterchangeMaterialXPipelineCustomizations.h
@@ -23,11 +23,7 @@ public:
 	/** End IDetailCustomization interface */
 
 protected:
-	bool OnShouldFilterAsset(const FAssetData&) { return false; }
-	bool OnShouldFilterAssetStandardSurface(const FAssetData& InAssetData);
-	bool OnShouldFilterAssetStandardSurfaceTransmission(const FAssetData& InAssetData);
-	bool OnShouldFilterAssetSurfaceUnlit(const FAssetData& InAssetData);
-	bool OnShouldFilterAssetUsdPreviewSurface(const FAssetData& InAssetData);
+	static bool OnShouldFilterAssetEnum(const FAssetData&, uint8 EnumType, uint8 EnumValue);
 
 private:
 	TWeakObjectPtr<UMaterialXPipelineSettings> MaterialXSettings;
diff --git a/Engine/Plugins/Interchange/Runtime/Source/Import/Private/InterchangeImportModule.cpp b/Engine/Plugins/Interchange/Runtime/Source/Import/Private/InterchangeImportModule.cpp
--- a/Engine/Plugins/Interchange/Runtime/Source/Import/Private/InterchangeImportModule.cpp
+++ b/Engine/Plugins/Interchange/Runtime/Source/Import/Private/InterchangeImportModule.cpp
@@ -6,6 +6,7 @@
 #include "Animation/InterchangeAnimationTrackSetFactory.h"
 #include "CoreMinimal.h"
 #include "Engine/Engine.h"
+#include "Engine/RendererSettings.h"
 #include "Fbx/InterchangeFbxTranslator.h"
 #include "Gltf/InterchangeGltfTranslator.h"
 #include "InterchangeImportLog.h"
@@ -37,11 +38,26 @@
 
 DEFINE_LOG_CATEGORY(LogInterchangeImport);
 
+static bool GInterchangeEnableSubstrate = true;
+static FAutoConsoleVariableRef CCvarInterchangeEnableSubstrateSupport(
+	TEXT("Interchange.FeatureFlags.Import.Substrate"),
+	GInterchangeEnableSubstrate,
+	TEXT("Enable or disable support of Substrate with Interchange (only works if Substrate is enabled in the Project Settings). Enabled by default."),
+	ECVF_Default);
+
 class FInterchangeImportModule : public IInterchangeImportModule
 {
 	/** IModuleInterface implementation */
 	virtual void StartupModule() override;
 	virtual void ShutdownModule() override;
+	virtual bool IsSubstrateEnabled() const override
+	{
+		return bIsSubstrateEnabled && GInterchangeEnableSubstrate;
+	}
+
+private:
+
+	bool bIsSubstrateEnabled = false;
 };
 
 IMPLEMENT_MODULE(FInterchangeImportModule, InterchangeImport)
@@ -50,6 +66,8 @@ IMPLEMENT_MODULE(FInterchangeImportModule, InterchangeImport)
 
 void FInterchangeImportModule::StartupModule()
 {
+	bIsSubstrateEnabled = GetDefault<URendererSettings>() ? GetDefault<URendererSettings>()->bEnableSubstrate : false;
+
 	FInterchangeImportMaterialAsyncHelper& InterchangeMaterialAsyncHelper = FInterchangeImportMaterialAsyncHelper::GetInstance();
 
 	auto RegisterItems = []()
diff --git a/Engine/Plugins/Interchange/Runtime/Source/Import/Private/MaterialX/InterchangeMaterialXDefinitions.h b/Engine/Plugins/Interchange/Runtime/Source/Import/Private/MaterialX/InterchangeMaterialXDefinitions.h
--- a/Engine/Plugins/Interchange/Runtime/Source/Import/Private/MaterialX/InterchangeMaterialXDefinitions.h
+++ b/Engine/Plugins/Interchange/Runtime/Source/Import/Private/MaterialX/InterchangeMaterialXDefinitions.h
@@ -149,6 +149,24 @@ MATERIALX_NAMESPACE_BEGIN
 		}
 	}
 
+	namespace Surface
+	{
+		namespace Input
+		{
+			static constexpr const char* Bsdf = "bsdf";
+			static constexpr const char* Edf = "edf";
+			static constexpr const char* Opacity = "opacity";
+		}
+
+		namespace DefaultValue
+		{
+			namespace Float
+			{
+				constexpr float Opacity = 1.f;
+			}
+		}
+	}
+
 	namespace UsdPreviewSurface
 	{
 		namespace Input
@@ -358,10 +376,35 @@ MATERIALX_NAMESPACE_BEGIN
 		// Application
 		static constexpr const char* Time= "time";
 		// PBR
+		// BSDF
+		static constexpr const char* BurleyDiffuseBSDF= "burley_diffuse_bsdf";
+		static constexpr const char* ConductorBSDF = "conductor_bsdf";
+		static constexpr const char* DielectricBSDF = "dielectric_bsdf";
+		static constexpr const char* GeneralizedSchlickBSDF= "generalized_schlick_bsdf";
+		static constexpr const char* OrenNayarDiffuseBSDF = "oren_nayar_diffuse_bsdf";
+		static constexpr const char* SheenBSDF = "sheen_bsdf";
+		static constexpr const char* SubsurfaceBSDF= "subsurface_bsdf";
+		static constexpr const char* ThinFilmBSDF= "thin_film_bsdf";
+		static constexpr const char* TranslucentBSDF = "translucent_bsdf";
+		// EDF
+		static constexpr const char* ConicalEDF = "conical_edf";
+		static constexpr const char* MeasuredEDF = "measured_edf";
+		static constexpr const char* UniformEDF = "uniform_edf";
+		// VDF
+		static constexpr const char* AbsorptionVDF = "absorption_vdf";
+		static constexpr const char* AnisotropicVDF = "anisotropic_vdf";
+		// PBR Utility Nodes
+		static constexpr const char* ArtisticIOR = "artistic_ior";
+		static constexpr const char* Blackbody = "blackbody";
+		static constexpr const char* Layer = "layer";
+		static constexpr const char* RoughnessAnisotropy = "roughness_anisotropy";
+		static constexpr const char* RoughnessDual = "roughness_dual";
+		//Surface Shaders
 		static constexpr const char* GltfPbr = "gltf_pbr";
 		static constexpr const char* DisneyBSDF2012 = "disney_brdf_2012";
 		static constexpr const char* DisneyBSDF2015 = "disney_bsdf_2015";
 		static constexpr const char* StandardSurface = "standard_surface";
+		static constexpr const char* Surface = "surface";
 		static constexpr const char* UsdPreviewSurface = "UsdPreviewSurface";
 		// Shader
 		static constexpr const char* SurfaceUnlit = "surface_unlit";
@@ -378,6 +421,7 @@ MATERIALX_NAMESPACE_BEGIN
 		static constexpr const char* PointLight = "ND_point_light";
 		static constexpr const char* DirectionalLight = "ND_directional_light";
 		static constexpr const char* SpotLight = "ND_spot_light";
+		static constexpr const char* Surface = "ND_surface";
 	}
 
 	namespace Library
@@ -441,6 +485,8 @@ MATERIALX_NAMESPACE_BEGIN
 			template<typename T>
 			static const T NeutralOne = T{ 1 };
 		}
+
+		static constexpr const char* PBR = "pbr";
 	}
 
 	namespace Image
diff --git a/Engine/Plugins/Interchange/Runtime/Source/Import/Private/MaterialX/InterchangeMaterialXTranslator.cpp b/Engine/Plugins/Interchange/Runtime/Source/Import/Private/MaterialX/InterchangeMaterialXTranslator.cpp
--- a/Engine/Plugins/Interchange/Runtime/Source/Import/Private/MaterialX/InterchangeMaterialXTranslator.cpp
+++ b/Engine/Plugins/Interchange/Runtime/Source/Import/Private/MaterialX/InterchangeMaterialXTranslator.cpp
@@ -1,20 +1,15 @@
 // Copyright Epic Games, Inc. All Rights Reserved. 
 
 #include "MaterialX/InterchangeMaterialXTranslator.h"
-#include "InterchangeImportLog.h"
 #include "InterchangeManager.h"
+#include "MaterialX/MaterialXUtils/MaterialXManager.h"
 #include "Nodes/InterchangeSourceNode.h"
 #include "UObject/GCObjectScopeGuard.h"
 
-#include "Misc/PackageName.h"
-
 #if WITH_EDITOR
-#include "MaterialX/MaterialXUtils/MaterialXManager.h"
 #include "MaterialXFormat/Util.h"
-#include "MaterialXUtils/MaterialXSurfaceMaterial.h"
-#include "MaterialXUtils/MaterialXSpotLightShader.h"
-#include "MaterialXUtils/MaterialXPointLightShader.h"
-#include "MaterialXUtils/MaterialXDirectionalLightShader.h"
+#include "InterchangeMaterialXDefinitions.h"
+#include "MaterialX/MaterialXUtils/MaterialXBase.h"
 #endif
 
 #include UE_INLINE_GENERATED_CPP_BY_NAME(InterchangeMaterialXTranslator)
@@ -28,41 +23,6 @@ static FAutoConsoleVariableRef CCvarInterchangeEnableMaterialXImport(
 	TEXT("Whether MaterialX support is enabled."),
 	ECVF_Default);
 
-namespace UE::Interchange::MaterialX
-{
-	bool AreMaterialFunctionPackagesLoaded()
-	{
-		auto ArePackagesLoaded = [](const TArray<FString> & TextPaths) -> bool
-		{
-			bool bAllLoaded = true;
-			for(const FString & TextPath : TextPaths)
-			{
-				const FString FunctionPath{ FPackageName::ExportTextPathToObjectPath(TextPath) };
-				if(FPackageName::DoesPackageExist(FunctionPath))
-				{
-					if(!FSoftObjectPath(FunctionPath).TryLoad())
-					{
-						UE_LOG(LogInterchangeImport, Warning, TEXT("Couldn't load %s"), *FunctionPath);
-						bAllLoaded = false;
-					}
-				}
-				else
-				{
-					UE_LOG(LogInterchangeImport, Warning, TEXT("Couldn't find %s"), *FunctionPath);
-					bAllLoaded = false;
-				}
-			}
-
-			return bAllLoaded;
-		};
-
-		static const bool bPackagesLoaded = 
-			ArePackagesLoaded({ TEXT("MaterialFunction'/Engine/Functions/Engine_MaterialFunctions03/Procedurals/NormalFromHeightmap.NormalFromHeightmap'") });
-
-		return bPackagesLoaded;
-	}
-}
-
 EInterchangeTranslatorType UInterchangeMaterialXTranslator::GetTranslatorType() const
 {
 	return EInterchangeTranslatorType::Scenes;
diff --git a/Engine/Plugins/Interchange/Runtime/Source/Import/Private/MaterialX/MaterialXUtils/MaterialXBase.cpp b/Engine/Plugins/Interchange/Runtime/Source/Import/Private/MaterialX/MaterialXUtils/MaterialXBase.cpp
--- a/Engine/Plugins/Interchange/Runtime/Source/Import/Private/MaterialX/MaterialXUtils/MaterialXBase.cpp
+++ b/Engine/Plugins/Interchange/Runtime/Source/Import/Private/MaterialX/MaterialXUtils/MaterialXBase.cpp
@@ -10,6 +10,7 @@ namespace mx = MaterialX;
 
 FMaterialXBase::FMaterialXBase(UInterchangeBaseNodeContainer& BaseNodeContainer)
 	: NodeContainer(BaseNodeContainer)
+	, bIsSubstrateEnabled(FMaterialXManager::GetInstance().IsSubstrateEnabled())
 {}
 
 void FMaterialXBase::UpdateDocumentRecursively(MaterialX::GraphElementPtr Graph)
diff --git a/Engine/Plugins/Interchange/Runtime/Source/Import/Private/MaterialX/MaterialXUtils/MaterialXManager.cpp b/Engine/Plugins/Interchange/Runtime/Source/Import/Private/MaterialX/MaterialXUtils/MaterialXManager.cpp
--- a/Engine/Plugins/Interchange/Runtime/Source/Import/Private/MaterialX/MaterialXUtils/MaterialXManager.cpp
+++ b/Engine/Plugins/Interchange/Runtime/Source/Import/Private/MaterialX/MaterialXUtils/MaterialXManager.cpp
@@ -1,11 +1,13 @@
 // Copyright Epic Games, Inc. All Rights Reserved.
 
-#if WITH_EDITOR
 #include "MaterialX/MaterialXUtils/MaterialXManager.h"
+#include "InterchangeImportModule.h"
+#include "Misc/PackageName.h"
+#include "UObject/SoftObjectPath.h"
+
+#if WITH_EDITOR
 #include "MaterialX/InterchangeMaterialXDefinitions.h"
-#include "InterchangeMaterialDefinitions.h"
-#include "InterchangeShaderGraphNode.h"
-#include "MaterialX/MaterialXUtils/MaterialXSurfaceShaderAbstract.h"
+#include "MaterialX/MaterialXUtils/MaterialXSurfaceShader.h"
 #include "MaterialX/MaterialXUtils/MaterialXStandardSurfaceShader.h"
 #include "MaterialX/MaterialXUtils/MaterialXUsdPreviewSurfaceShader.h"
 #include "MaterialX/MaterialXUtils/MaterialXSurfaceUnlitShader.h"
@@ -14,77 +16,85 @@
 #include "MaterialX/MaterialXUtils/MaterialXSpotLightShader.h"
 #include "MaterialX/MaterialXUtils/MaterialXSurfaceMaterial.h"
 
+#include "Engine/RendererSettings.h"
+#include "InterchangeMaterialDefinitions.h"
+#include "InterchangeImportLog.h"
+
+
 namespace mx = MaterialX;
 
 //not a good solution to use semicolon because of drive disk on Windows
 const TCHAR FMaterialXManager::TexturePayloadSeparator = TEXT('{');
 
+#define MATERIALFUNCTION_PATH FString{bIsSubstrateEnabled ? TEXT("Substrate/") : TEXT("Functions/")}
+
 FMaterialXManager::FMaterialXManager()
 	: MatchingInputNames {
-		{{TEXT(""),                            TEXT("amplitude")}, TEXT("Amplitude")},
-		{{TEXT(""),                            TEXT("bg")},        TEXT("B")},
-		{{TEXT(""),                            TEXT("center")},    TEXT("Center")},
-		{{TEXT(""),                            TEXT("diminish")},  TEXT("Diminish")},
-		{{TEXT(""),                            TEXT("fg")},        TEXT("A")},
-		{{TEXT(""),                            TEXT("high")},      TEXT("Max")},
-		{{TEXT(""),                            TEXT("in")},        TEXT("Input")},
-		{{TEXT(""),                            TEXT("in1")},       TEXT("A")},
-		{{TEXT(""),                            TEXT("in2")},       TEXT("B")},
-		{{TEXT(""),                            TEXT("in3")},       TEXT("C")},
-		{{TEXT(""),                            TEXT("in4")},       TEXT("D")},
-		{{TEXT(""),                            TEXT("inlow")},     TEXT("InputLow")},
-		{{TEXT(""),                            TEXT("inhigh")},    TEXT("InputHigh")},
-		{{TEXT(""),                            TEXT("lacunarity")},TEXT("Lacunarity")},
-		{{TEXT(""),                            TEXT("low")},       TEXT("Min")},
-		{{TEXT(""),                            TEXT("lumacoeffs")},TEXT("LuminanceFactors")}, // for the moment not yet handled by Interchange, because of the attribute being an advanced pin
-		{{TEXT(""),                            TEXT("mix")},       TEXT("Alpha")},
-		{{TEXT(""),                            TEXT("offset")},    TEXT("Offset")},
-		{{TEXT(""),                            TEXT("position")},  TEXT("Position")},
-		{{TEXT(""),                            TEXT("texcoord")},  TEXT("Coordinates")},
-		{{TEXT(""),                            TEXT("octaves")},   TEXT("Octaves")},
-		{{TEXT(""),                            TEXT("outlow")},    TEXT("TargetLow")},
-		{{TEXT(""),                            TEXT("outhigh")},   TEXT("TargetHigh")},
-		{{TEXT(""),                            TEXT("valuel")},    TEXT("A")},
-		{{TEXT(""),                            TEXT("valuer")},    TEXT("B")},
-		{{TEXT(""),                            TEXT("valuet")},    TEXT("A")},
-		{{TEXT(""),                            TEXT("valueb")},    TEXT("B")},
-		{{TEXT(""),                            TEXT("valuetl")},   TEXT("A")},
-		{{TEXT(""),                            TEXT("valuetr")},   TEXT("B")},
-		{{TEXT(""),                            TEXT("valuebl")},   TEXT("C")},
-		{{TEXT(""),                            TEXT("valuebr")},   TEXT("D")},
-		{{TEXT(""),                            TEXT("value1")},    TEXT("A")},
-		{{TEXT(""),                            TEXT("value2")},    TEXT("B")},
-		{{MaterialX::Category::Atan2,          TEXT("in1")},       TEXT("Y")},
-		{{MaterialX::Category::Atan2,          TEXT("in2")},       TEXT("X")},
-		{{MaterialX::Category::HeightToNormal, TEXT("scale")},     UE::Interchange::Materials::Standard::Nodes::NormalFromHeightMap::Inputs::Intensity.ToString()},
-		{{MaterialX::Category::IfGreater,      TEXT("in1")},       TEXT("AGreaterThanB")},
-		{{MaterialX::Category::IfGreater,      TEXT("in2")},       TEXT("ALessThanB")}, //another input is added for the case equal see ConnectIfGreater
-		{{MaterialX::Category::IfGreaterEq,    TEXT("in1")},       TEXT("AGreaterThanB")}, //another is added for the case equal see ConnectIfGreaterEq
-		{{MaterialX::Category::IfGreaterEq,    TEXT("in2")},       TEXT("ALessThanB")},
-		{{MaterialX::Category::IfEqual,        TEXT("in1")},       TEXT("AEqualsB")},
-		{{MaterialX::Category::IfEqual,        TEXT("in2")},       TEXT("ALessThanB")},  // another input is added for the case greater see ConnectIfEqual
-		{{MaterialX::Category::Inside,         TEXT("in")},        TEXT("A")},			  // Inside is treated as a Multiply node
-		{{MaterialX::Category::Inside,         TEXT("mask")},      TEXT("B")},			  // Inside is treated as a Multiply node
-		{{MaterialX::Category::Invert,         TEXT("amount")},    TEXT("A")},
-		{{MaterialX::Category::Invert,         TEXT("in")},        TEXT("B")},
-		{{MaterialX::Category::Mix,            TEXT("fg")},        TEXT("B")},
-		{{MaterialX::Category::Mix,            TEXT("bg")},        TEXT("A")},
-		{{MaterialX::Category::Mix,            TEXT("mix")},		TEXT("Factor")},
-		{{MaterialX::Category::Noise3D,        TEXT("amplitude")}, TEXT("B")},            // The amplitude of the noise is connected to a multiply node
-		{{MaterialX::Category::Noise3D,        TEXT("pivot")},     TEXT("B")},            // The pivot of the noise is connected to a add node
-		{{MaterialX::Category::Normalize,      TEXT("in")},        TEXT("VectorInput")},
-		{{MaterialX::Category::Outside,        TEXT("in")},        TEXT("A")},				// Outside is treated as Multiply node
-		{{MaterialX::Category::Outside,        TEXT("mask")},      TEXT("B")},				// Outside is treated as Multiply node
-		{{MaterialX::Category::Place2D,        TEXT("pivot")},     TEXT("Pivot")},
-		{{MaterialX::Category::Place2D,        TEXT("rotate")},    TEXT("RotationAngle")},
-		{{MaterialX::Category::Power,          TEXT("in1")},       TEXT("Base")},
-		{{MaterialX::Category::Power,	       TEXT("in2")},       TEXT("Exponent")},
-		{{MaterialX::Category::Rotate2D,       TEXT("amount")},    TEXT("RotationAngle")},
-		{{MaterialX::Category::Rotate3D,       TEXT("amount")},    TEXT("RotationAngle")},
-		{{MaterialX::Category::Rotate3D,       TEXT("axis")},		TEXT("NormalizedRotationAxis")},
-		{{MaterialX::Category::Rotate3D,       TEXT("in")},        TEXT("Position")},
-		{{MaterialX::Category::Saturate,       TEXT("amount")},    TEXT("Fraction")},
-		{{MaterialX::Category::Smoothstep,     TEXT("in")},        TEXT("Value")},
+		{{TEXT(""),                            TEXT("amplitude")},   TEXT("Amplitude")},
+		{{TEXT(""),                            TEXT("bg")},          TEXT("B")},
+		{{TEXT(""),                            TEXT("center")},      TEXT("Center")},
+		{{TEXT(""),                            TEXT("diminish")},    TEXT("Diminish")},
+		{{TEXT(""),                            TEXT("fg")},          TEXT("A")},
+		{{TEXT(""),                            TEXT("high")},        TEXT("Max")},
+		{{TEXT(""),                            TEXT("in")},          TEXT("Input")},
+		{{TEXT(""),                            TEXT("in1")},         TEXT("A")},
+		{{TEXT(""),                            TEXT("in2")},         TEXT("B")},
+		{{TEXT(""),                            TEXT("in3")},         TEXT("C")},
+		{{TEXT(""),                            TEXT("in4")},         TEXT("D")},
+		{{TEXT(""),                            TEXT("inlow")},       TEXT("InputLow")},
+		{{TEXT(""),                            TEXT("inhigh")},      TEXT("InputHigh")},
+		{{TEXT(""),                            TEXT("lacunarity")},  TEXT("Lacunarity")},
+		{{TEXT(""),                            TEXT("low")},         TEXT("Min")},
+		{{TEXT(""),                            TEXT("lumacoeffs")},  TEXT("LuminanceFactors")}, // for the moment not yet handled by Interchange, because of the attribute being an advanced pin
+		{{TEXT(""),                            TEXT("mix")},         TEXT("Alpha")},
+		{{TEXT(""),                            TEXT("offset")},      TEXT("Offset")},
+		{{TEXT(""),                            TEXT("position")},    TEXT("Position")},
+		{{TEXT(""),                            TEXT("texcoord")},    TEXT("Coordinates")},
+		{{TEXT(""),                            TEXT("octaves")},     TEXT("Octaves")},
+		{{TEXT(""),                            TEXT("temperature")}, TEXT("Temp")},
+		{{TEXT(""),                            TEXT("outlow")},      TEXT("TargetLow")},
+		{{TEXT(""),                            TEXT("outhigh")},     TEXT("TargetHigh")},
+		{{TEXT(""),                            TEXT("valuel")},      TEXT("A")},
+		{{TEXT(""),                            TEXT("valuer")},      TEXT("B")},
+		{{TEXT(""),                            TEXT("valuet")},      TEXT("A")},
+		{{TEXT(""),                            TEXT("valueb")},      TEXT("B")},
+		{{TEXT(""),                            TEXT("valuetl")},     TEXT("A")},
+		{{TEXT(""),                            TEXT("valuetr")},     TEXT("B")},
+		{{TEXT(""),                            TEXT("valuebl")},     TEXT("C")},
+		{{TEXT(""),                            TEXT("valuebr")},     TEXT("D")},
+		{{TEXT(""),                            TEXT("value1")},      TEXT("A")},
+		{{TEXT(""),                            TEXT("value2")},      TEXT("B")},
+		{{MaterialX::Category::Atan2,          TEXT("in1")},         TEXT("Y")},
+		{{MaterialX::Category::Atan2,          TEXT("in2")},         TEXT("X")},
+		{{MaterialX::Category::HeightToNormal, TEXT("scale")},       UE::Interchange::Materials::Standard::Nodes::NormalFromHeightMap::Inputs::Intensity.ToString()},
+		{{MaterialX::Category::IfGreater,      TEXT("in1")},         TEXT("AGreaterThanB")},
+		{{MaterialX::Category::IfGreater,      TEXT("in2")},         TEXT("ALessThanB")}, //another input is added for the case equal see ConnectIfGreater
+		{{MaterialX::Category::IfGreaterEq,    TEXT("in1")},         TEXT("AGreaterThanB")}, //another is added for the case equal see ConnectIfGreaterEq
+		{{MaterialX::Category::IfGreaterEq,    TEXT("in2")},         TEXT("ALessThanB")},
+		{{MaterialX::Category::IfEqual,        TEXT("in1")},         TEXT("AEqualsB")},
+		{{MaterialX::Category::IfEqual,        TEXT("in2")},         TEXT("ALessThanB")},  // another input is added for the case greater see ConnectIfEqual
+		{{MaterialX::Category::Inside,         TEXT("in")},          TEXT("A")},			  // Inside is treated as a Multiply node
+		{{MaterialX::Category::Inside,         TEXT("mask")},        TEXT("B")},			  // Inside is treated as a Multiply node
+		{{MaterialX::Category::Invert,         TEXT("amount")},      TEXT("A")},
+		{{MaterialX::Category::Invert,         TEXT("in")},          TEXT("B")},
+		{{MaterialX::Category::Mix,            TEXT("fg")},          TEXT("B")},
+		{{MaterialX::Category::Mix,            TEXT("bg")},          TEXT("A")},
+		{{MaterialX::Category::Mix,            TEXT("mix")},		 TEXT("Factor")},
+		{{MaterialX::Category::Noise3D,        TEXT("amplitude")},   TEXT("B")},            // The amplitude of the noise is connected to a multiply node
+		{{MaterialX::Category::Noise3D,        TEXT("pivot")},       TEXT("B")},            // The pivot of the noise is connected to a add node
+		{{MaterialX::Category::Normalize,      TEXT("in")},          TEXT("VectorInput")},
+		{{MaterialX::Category::Outside,        TEXT("in")},          TEXT("A")},				// Outside is treated as Multiply node
+		{{MaterialX::Category::Outside,        TEXT("mask")},        TEXT("B")},				// Outside is treated as Multiply node
+		{{MaterialX::Category::Place2D,        TEXT("pivot")},       TEXT("Pivot")},
+		{{MaterialX::Category::Place2D,        TEXT("rotate")},      TEXT("RotationAngle")},
+		{{MaterialX::Category::Power,          TEXT("in1")},         TEXT("Base")},
+		{{MaterialX::Category::Power,	       TEXT("in2")},         TEXT("Exponent")},
+		{{MaterialX::Category::Rotate2D,       TEXT("amount")},      TEXT("RotationAngle")},
+		{{MaterialX::Category::Rotate3D,       TEXT("amount")},      TEXT("RotationAngle")},
+		{{MaterialX::Category::Rotate3D,       TEXT("axis")},		 TEXT("NormalizedRotationAxis")},
+		{{MaterialX::Category::Rotate3D,       TEXT("in")},          TEXT("Position")},
+		{{MaterialX::Category::Saturate,       TEXT("amount")},      TEXT("Fraction")},
+		{{MaterialX::Category::Smoothstep,     TEXT("in")},          TEXT("Value")},
 	}
 	, MatchingMaterialExpressions {
 		// Math nodes
@@ -143,6 +153,8 @@ FMaterialXManager::FMaterialXManager()
 		{MaterialX::Category::Combine2,     TEXT("AppendVector")},
 		{MaterialX::Category::Combine3,     TEXT("MaterialXAppend3Vector")},
 		{MaterialX::Category::Combine4,     TEXT("MaterialXAppend4Vector")},
+		// PBR
+		{MaterialX::Category::Blackbody,    TEXT("BlackBody")},
 		// Procedural2D nodes
 		{MaterialX::Category::Ramp4,        TEXT("MaterialXRamp4")},
 		{MaterialX::Category::RampLR,       TEXT("MaterialXRampLeftRight")},
@@ -189,12 +201,14 @@ FMaterialXManager::FMaterialXManager()
 		TEXT("Scale"),
 		TEXT("TargetLow"),
 		TEXT("TargetHigh"),
+		TEXT("Temp"),
 		TEXT("Value"),
 		TEXT("VectorInput"),
 		TEXT("X"),
 		TEXT("Y")
 	}
 	, MaterialXContainerDelegates{
+		{mx::Category::Surface, FMaterialXManager::FOnGetMaterialXInstance::CreateStatic(&FMaterialXSurfaceShader::MakeInstance)},
 		{mx::Category::SurfaceUnlit, FMaterialXManager::FOnGetMaterialXInstance::CreateStatic(&FMaterialXSurfaceUnlitShader::MakeInstance)},
 		{mx::Category::StandardSurface, FMaterialXManager::FOnGetMaterialXInstance::CreateStatic(&FMaterialXStandardSurfaceShader::MakeInstance)},
 		{mx::Category::UsdPreviewSurface, FMaterialXManager::FOnGetMaterialXInstance::CreateStatic(&FMaterialXUsdPreviewSurfaceShader::MakeInstance)},
@@ -203,7 +217,55 @@ FMaterialXManager::FMaterialXManager()
 		{mx::Category::SpotLight, FMaterialXManager::FOnGetMaterialXInstance::CreateStatic(&FMaterialXSpotLightShader::MakeInstance)},
 		{mx::SURFACE_MATERIAL_NODE_STRING.c_str(), FMaterialXManager::FOnGetMaterialXInstance::CreateStatic(&FMaterialXSurfaceMaterial::MakeInstance)}
 	}
-{}
+	, bIsSubstrateEnabled{ IInterchangeImportModule::IsAvailable() ? IInterchangeImportModule::Get().IsSubstrateEnabled() : false }
+{
+		MatchingMaterialFunctions = {
+			// BSDF Nodes
+			{mx::Category::BurleyDiffuseBSDF ,		TEXT("/Interchange/") + MATERIALFUNCTION_PATH + TEXT("MX_BurleyDiffuseBSDF.MX_BurleyDiffuseBSDF")},
+			{mx::Category::ConductorBSDF,			TEXT("/Interchange/") + MATERIALFUNCTION_PATH + TEXT("MX_ConductorBSDF.MX_ConductorBSDF")},
+			{mx::Category::DielectricBSDF,			TEXT("/Interchange/") + MATERIALFUNCTION_PATH + TEXT("MX_DielectricBSDF.MX_DielectricBSDF")},
+			{mx::Category::GeneralizedSchlickBSDF,	TEXT("/Interchange/") + MATERIALFUNCTION_PATH + TEXT("MX_GeneralizedSchlickBSDF.MX_GeneralizedSchlickBSDF")},
+			{mx::Category::OrenNayarDiffuseBSDF ,	TEXT("/Interchange/") + MATERIALFUNCTION_PATH + TEXT("MX_OrenNayarBSDF.MX_OrenNayarBSDF")},
+			{mx::Category::SheenBSDF,				TEXT("/Interchange/") + MATERIALFUNCTION_PATH + TEXT("MX_SheenBSDF.MX_SheenBSDF")},
+			{mx::Category::SubsurfaceBSDF ,			TEXT("/Interchange/") + MATERIALFUNCTION_PATH + TEXT("MX_SubsurfaceBSDF.MX_SubsurfaceBSDF")},
+			{mx::Category::ThinFilmBSDF,			TEXT("/Interchange/") + MATERIALFUNCTION_PATH + TEXT("MX_ThinFilmBSDF.MX_ThinFilmBSDF")},
+			{mx::Category::TranslucentBSDF,			TEXT("/Interchange/") + MATERIALFUNCTION_PATH + TEXT("MX_TranslucentBSDF.MX_TranslucentBSDF")},
+			// EDF Nodes
+			{mx::Category::ConicalEDF,	TEXT("/Interchange/") + MATERIALFUNCTION_PATH + TEXT("MX_ConicalEDF.MX_ConicalEDF")},
+			{mx::Category::MeasuredEDF, TEXT("/Interchange/") + MATERIALFUNCTION_PATH + TEXT("MX_MeasuredEDF.MX_MeasuredEDF")},
+			{mx::Category::UniformEDF,	TEXT("/Interchange/") + MATERIALFUNCTION_PATH + TEXT("MX_UniformEDF.MX_UniformEDF")},
+			// VDF Nodes
+			{mx::Category::AbsorptionVDF,  TEXT("/Interchange/") + MATERIALFUNCTION_PATH + TEXT("MX_AbsorptionVDF.MX_AbsorptionVDF")},
+			{mx::Category::AnisotropicVDF, TEXT("/Interchange/") + MATERIALFUNCTION_PATH + TEXT("MX_AnisotropicVDF.MX_AnisotropicVDF")},
+		};
+
+		if(bIsSubstrateEnabled)
+		{
+			// Vertical Layering
+			MatchingInputNames.Add({ TEXT(""), TEXT("top") }, TEXT("Top"));
+			MatchingInputNames.Add({ TEXT(""), TEXT("base") }, TEXT("Bottom"));
+			MaterialExpressionInputs.Add(TEXT("Top"));
+			MaterialExpressionInputs.Add(TEXT("Bottom"));
+			MatchingMaterialExpressions.Add({ MaterialX::Category::Layer }, TEXT("StrataVerticalLayering"));
+
+			// Horizontal Layering, MaterialX names it mix so we have to be careful not to take the lerp expression
+			MatchingInputNames.Add({ FKeyCategoryNodegroup{MaterialX::Category::Mix, MaterialX::NodeGroup::PBR}, TEXT("bg") }, TEXT("Background"));
+			MatchingInputNames.Add({ FKeyCategoryNodegroup{MaterialX::Category::Mix, MaterialX::NodeGroup::PBR}, TEXT("fg") }, TEXT("Foreground"));
+			MatchingInputNames.Add({ FKeyCategoryNodegroup{MaterialX::Category::Mix, MaterialX::NodeGroup::PBR}, TEXT("mix") }, TEXT("Mix"));
+			MaterialExpressionInputs.Add(TEXT("Background"));
+			MaterialExpressionInputs.Add(TEXT("Foreground"));
+			MaterialExpressionInputs.Add(TEXT("Mix"));
+			MatchingMaterialExpressions.Add({ MaterialX::Category::Mix, MaterialX::NodeGroup::PBR }, TEXT("StrataHorizontalMixing"));
+
+			// Add
+			MatchingMaterialExpressions.Add({ MaterialX::Category::Add, MaterialX::NodeGroup::PBR }, TEXT("StrataAdd"));
+
+			// Multiply
+			MatchingInputNames.Add({ FKeyCategoryNodegroup{MaterialX::Category::Multiply, MaterialX::NodeGroup::PBR}, TEXT("in2") }, TEXT("Weight"));
+			MaterialExpressionInputs.Add(TEXT("Weight"));
+			MatchingMaterialExpressions.Add({ MaterialX::Category::Multiply, MaterialX::NodeGroup::PBR }, TEXT("StrataWeight"));
+		}
+	}
 
 FMaterialXManager& FMaterialXManager::GetInstance()
 {
@@ -211,9 +273,9 @@ FMaterialXManager& FMaterialXManager::GetInstance()
 	return Instance;
 }
 
-const FString* FMaterialXManager::FindMatchingInput(const TPair<FString, FString>& CategoryInputKey) const
+const FString* FMaterialXManager::FindMatchingInput(const FString& CategoryKey, const FString& InputKey, const FString& NodeGroup) const
 {
-	return MatchingInputNames.Find(CategoryInputKey);
+	return MatchingInputNames.Find({ FKeyCategoryNodegroup{ CategoryKey, NodeGroup}, InputKey });
 }
 
 const FString* FMaterialXManager::FindMaterialExpressionInput(const FString& InputKey) const
@@ -221,9 +283,14 @@ const FString* FMaterialXManager::FindMaterialExpressionInput(const FString& Inp
 	return MaterialExpressionInputs.Find(InputKey);
 }
 
-const FString* FMaterialXManager::FindMatchingMaterialExpression(const FString& InputKey) const
+const FString* FMaterialXManager::FindMatchingMaterialExpression(const FString& CategoryKey, const FString& NodeGroup) const
 {
-	return MatchingMaterialExpressions.Find(InputKey);
+	return MatchingMaterialExpressions.Find({ CategoryKey, NodeGroup });
+}
+
+const FString* FMaterialXManager::FindMatchingMaterialFunction(const FString& CategoryKey) const
+{
+	return MatchingMaterialFunctions.Find(CategoryKey);
 }
 
 TSharedPtr<FMaterialXBase> FMaterialXManager::GetShaderTranslator(const FString& CategoryShader, UInterchangeBaseNodeContainer & NodeContainer)
@@ -246,4 +313,49 @@ void FMaterialXManager::RegisterMaterialXInstance(const FString& Category, FOnGe
 		MaterialXContainerDelegates.Add(Category, MaterialXInstanceDelegate);
 	}
 }
-#endif
\ No newline at end of file
+
+bool FMaterialXManager::IsSubstrateEnabled() const
+{
+	return bIsSubstrateEnabled;
+}
+#endif
+
+namespace UE::Interchange::MaterialX
+{
+	bool AreMaterialFunctionPackagesLoaded()
+	{
+#if WITH_EDITOR
+		auto ArePackagesLoaded = [](const TArray<FString>& TextPaths) -> bool
+		{
+			bool bAllLoaded = true;
+			for(const FString& TextPath : TextPaths)
+			{
+				const FString FunctionPath{ FPackageName::ExportTextPathToObjectPath(TextPath) };
+				if(FPackageName::DoesPackageExist(FunctionPath))
+				{
+					if(!FSoftObjectPath(FunctionPath).TryLoad())
+					{
+						UE_LOG(LogInterchangeImport, Warning, TEXT("Couldn't load %s"), *FunctionPath);
+						bAllLoaded = false;
+					}
+				}
+				else
+				{
+					UE_LOG(LogInterchangeImport, Warning, TEXT("Couldn't find %s"), *FunctionPath);
+					bAllLoaded = false;
+				}
+			}
+
+			return bAllLoaded;
+		};
+
+		static const bool bPackagesLoaded =	ArePackagesLoaded({ TEXT("MaterialFunction'/Engine/Functions/Engine_MaterialFunctions03/Procedurals/NormalFromHeightmap.NormalFromHeightmap'") });
+
+		return bPackagesLoaded;
+#else
+		return false;
+#endif
+	}
+}
+
+#undef MATERIALFUNCTION_PATH
\ No newline at end of file
diff --git a/Engine/Plugins/Interchange/Runtime/Source/Import/Private/MaterialX/MaterialXUtils/MaterialXStandardSurfaceShader.cpp b/Engine/Plugins/Interchange/Runtime/Source/Import/Private/MaterialX/MaterialXUtils/MaterialXStandardSurfaceShader.cpp
--- a/Engine/Plugins/Interchange/Runtime/Source/Import/Private/MaterialX/MaterialXUtils/MaterialXStandardSurfaceShader.cpp
+++ b/Engine/Plugins/Interchange/Runtime/Source/Import/Private/MaterialX/MaterialXUtils/MaterialXStandardSurfaceShader.cpp
@@ -2,7 +2,8 @@
 
 #if WITH_EDITOR
 #include "MaterialXStandardSurfaceShader.h"
-#include "Engine/RendererSettings.h"
+
+#include "Engine/EngineTypes.h"
 
 namespace mx = MaterialX;
 
@@ -25,8 +26,6 @@ void FMaterialXStandardSurfaceShader::Translate(mx::NodePtr StandardSurfaceNode)
 	
 	using namespace UE::Interchange::Materials;
 
-	bool bIsSubstrateEnabled = GetDefault<URendererSettings>()->bEnableSubstrate;
-
 	if(bIsSubstrateEnabled)
 	{
 		ConnectToSubstrateStandardSurface();
@@ -42,19 +41,7 @@ void FMaterialXStandardSurfaceShader::ConnectToStandardSurface()
 	using namespace UE::Interchange::Materials;
 	using namespace mx::StandardSurface;
 
-	UInterchangeFunctionCallShaderNode* StandardSurfaceShaderNode;
-	const FString NodeUID = UInterchangeShaderNode::MakeNodeUid(ANSI_TO_TCHAR(SurfaceShaderNode->getName().c_str()), FStringView{});
-	if(StandardSurfaceShaderNode = const_cast<UInterchangeFunctionCallShaderNode*>(Cast<UInterchangeFunctionCallShaderNode>(NodeContainer.GetNode(NodeUID))); !StandardSurfaceShaderNode)
-	{
-		const FString NodeName = SurfaceShaderNode->getName().c_str();
-		StandardSurfaceShaderNode = NewObject<UInterchangeFunctionCallShaderNode>(&NodeContainer);
-		StandardSurfaceShaderNode->InitializeNode(NodeUID, NodeName, EInterchangeNodeContainerType::TranslatedAsset);
-
-		StandardSurfaceShaderNode->SetCustomMaterialFunction(TEXT("/Interchange/Functions/MX_StandardSurface.MX_StandardSurface"));
-		NodeContainer.AddNode(StandardSurfaceShaderNode);
-
-		ShaderNodes.Add({ NodeName, SurfaceShaderNode->getNodeDef(mx::EMPTY_STRING, true)->getActiveOutputs()[0]->getName().c_str() }, StandardSurfaceShaderNode);
-	}
+	UInterchangeFunctionCallShaderNode* StandardSurfaceShaderNode = CreateFunctionCallShaderNode(SurfaceShaderNode->getName().c_str(), TEXT("/Interchange/Functions/MX_StandardSurface.MX_StandardSurface"));
 
 	// Inputs
 	//Base
@@ -194,19 +181,7 @@ void FMaterialXStandardSurfaceShader::ConnectToSubstrateStandardSurface()
 	using namespace UE::Interchange::Materials;
 	using namespace mx::StandardSurface;
 
-	UInterchangeFunctionCallShaderNode* StandardSurfaceShaderNode;
-	const FString NodeUID = UInterchangeShaderNode::MakeNodeUid(ANSI_TO_TCHAR(SurfaceShaderNode->getName().c_str()), FStringView{});
-	if(StandardSurfaceShaderNode = const_cast<UInterchangeFunctionCallShaderNode*>(Cast<UInterchangeFunctionCallShaderNode>(NodeContainer.GetNode(NodeUID))); !StandardSurfaceShaderNode)
-	{
-		const FString NodeName = SurfaceShaderNode->getName().c_str();
-		StandardSurfaceShaderNode = NewObject<UInterchangeFunctionCallShaderNode>(&NodeContainer);
-		StandardSurfaceShaderNode->InitializeNode(NodeUID, NodeName, EInterchangeNodeContainerType::TranslatedAsset);
-
-		StandardSurfaceShaderNode->SetCustomMaterialFunction(TEXT("/Interchange/Substrate/MX_StandardSurface.MX_StandardSurface"));
-		NodeContainer.AddNode(StandardSurfaceShaderNode);
-
-		ShaderNodes.Add({ NodeName, SurfaceShaderNode->getNodeDef(mx::EMPTY_STRING, true)->getActiveOutputs()[0]->getName().c_str() }, StandardSurfaceShaderNode);
-	}
+	UInterchangeFunctionCallShaderNode* StandardSurfaceShaderNode = CreateFunctionCallShaderNode(SurfaceShaderNode->getName().c_str(), TEXT("/Interchange/Substrate/MX_StandardSurface.MX_StandardSurface"));
 
 	// Inputs
 	//Base
@@ -334,6 +309,7 @@ void FMaterialXStandardSurfaceShader::ConnectToSubstrateStandardSurface()
 		if(UInterchangeShaderPortsAPI::HasInput(StandardSurfaceShaderNode, StandardSurface::Parameters::Opacity))
 		{
 			UInterchangeShaderPortsAPI::ConnectOuputToInputByName(ShaderGraphNode, SubstrateMaterial::Parameters::OpacityMask.ToString(), StandardSurfaceShaderNode->GetUniqueID(), StandardSurface::SubstrateMaterial::Outputs::Opacity.ToString());
+			ShaderGraphNode->SetCustomBlendMode(EBlendMode::BLEND_Masked);
 		}
 	}
 }
diff --git a/Engine/Plugins/Interchange/Runtime/Source/Import/Private/MaterialX/MaterialXUtils/MaterialXSurfaceShader.cpp b/Engine/Plugins/Interchange/Runtime/Source/Import/Private/MaterialX/MaterialXUtils/MaterialXSurfaceShader.cpp
new file mode 100644
--- /dev/null
+++ b/Engine/Plugins/Interchange/Runtime/Source/Import/Private/MaterialX/MaterialXUtils/MaterialXSurfaceShader.cpp
@@ -0,0 +1,54 @@
+// Copyright Epic Games, Inc. All Rights Reserved.
+
+#if WITH_EDITOR
+
+#include "MaterialXSurfaceShader.h"
+#include "InterchangeImportLog.h"
+
+namespace mx = MaterialX;
+
+FMaterialXSurfaceShader::FMaterialXSurfaceShader(UInterchangeBaseNodeContainer& BaseNodeContainer)
+	: FMaterialXSurfaceShaderAbstract(BaseNodeContainer)
+{
+	NodeDefinition = mx::NodeDefinition::Surface;
+}
+
+TSharedRef<FMaterialXBase> FMaterialXSurfaceShader::MakeInstance(UInterchangeBaseNodeContainer& BaseNodeContainer)
+{
+	TSharedRef<FMaterialXSurfaceShader > Result = MakeShared<FMaterialXSurfaceShader >(BaseNodeContainer);
+	Result->RegisterConnectNodeOutputToInputDelegates();
+	return Result;
+}
+
+void FMaterialXSurfaceShader::Translate(MaterialX::NodePtr SurfaceNode)
+{
+	this->SurfaceShaderNode = SurfaceNode;
+
+	UInterchangeFunctionCallShaderNode* FunctionSurfaceShaderNode = CreateFunctionCallShaderNode(SurfaceNode->getName().c_str(), bIsSubstrateEnabled ? TEXT("/Interchange/Functions/MX_Surface.MX_Surface") : TEXT("/Interchange/Substrate/MX_Surface.MX_Surface"));
+
+	using namespace UE::Interchange::Materials;
+
+	// BSDF
+	ConnectNodeOutputToInput(mx::Surface::Input::Bsdf, FunctionSurfaceShaderNode, Surface::Parameters::BSDF.ToString(), nullptr);
+
+	// EDF
+	ConnectNodeOutputToInput(mx::Surface::Input::Edf, FunctionSurfaceShaderNode, Surface::Parameters::EDF.ToString(), nullptr);
+
+	// Opacity
+	ConnectNodeOutputToInput(mx::Surface::Input::Opacity, FunctionSurfaceShaderNode, Surface::Parameters::Opacity.ToString(), mx::SurfaceUnlit::DefaultValue::Float::Opacity);
+
+	// Outputs
+	if(!bIsSubstrateEnabled)
+	{
+		UInterchangeShaderPortsAPI::ConnectDefaultOuputToInput(ShaderGraphNode, Common::Parameters::BxDF.ToString(), FunctionSurfaceShaderNode->GetUniqueID());
+	}
+	else
+	{
+		UInterchangeShaderPortsAPI::ConnectOuputToInputByName(ShaderGraphNode, SubstrateMaterial::Parameters::OpacityMask.ToString(), FunctionSurfaceShaderNode->GetUniqueID(), Surface::Substrate::Outputs::Surface.ToString());
+		if(UInterchangeShaderPortsAPI::HasInput(FunctionSurfaceShaderNode, Surface::Parameters::Opacity))
+		{
+			UInterchangeShaderPortsAPI::ConnectOuputToInputByName(ShaderGraphNode, SubstrateMaterial::Parameters::FrontMaterial.ToString(), FunctionSurfaceShaderNode->GetUniqueID(), Surface::Substrate::Outputs::Opacity.ToString());
+		}
+	}
+}
+#endif
\ No newline at end of file
diff --git a/Engine/Plugins/Interchange/Runtime/Source/Import/Private/MaterialX/MaterialXUtils/MaterialXSurfaceShader.h b/Engine/Plugins/Interchange/Runtime/Source/Import/Private/MaterialX/MaterialXUtils/MaterialXSurfaceShader.h
new file mode 100644
--- /dev/null
+++ b/Engine/Plugins/Interchange/Runtime/Source/Import/Private/MaterialX/MaterialXUtils/MaterialXSurfaceShader.h
@@ -0,0 +1,25 @@
+// Copyright Epic Games, Inc. All Rights Reserved.
+
+#pragma once
+
+#if WITH_EDITOR
+
+#include "MaterialX/MaterialXUtils/MaterialXSurfaceShaderAbstract.h"
+
+/** Handles <surface> node from the pbrlib, acts more or less as a generic surface shader */
+class FMaterialXSurfaceShader : public FMaterialXSurfaceShaderAbstract
+{
+protected:
+
+	template<typename T, ESPMode>
+	friend class SharedPointerInternals::TIntrusiveReferenceController;
+
+	FMaterialXSurfaceShader(UInterchangeBaseNodeContainer& BaseNodeContainer);
+
+public:
+
+	static TSharedRef<FMaterialXBase> MakeInstance(UInterchangeBaseNodeContainer& BaseNodeContainer);
+
+	virtual void Translate(MaterialX::NodePtr SurfaceNode) override;
+};
+#endif
\ No newline at end of file
diff --git a/Engine/Plugins/Interchange/Runtime/Source/Import/Private/MaterialX/MaterialXUtils/MaterialXSurfaceShaderAbstract.cpp b/Engine/Plugins/Interchange/Runtime/Source/Import/Private/MaterialX/MaterialXUtils/MaterialXSurfaceShaderAbstract.cpp
--- a/Engine/Plugins/Interchange/Runtime/Source/Import/Private/MaterialX/MaterialXUtils/MaterialXSurfaceShaderAbstract.cpp
+++ b/Engine/Plugins/Interchange/Runtime/Source/Import/Private/MaterialX/MaterialXUtils/MaterialXSurfaceShaderAbstract.cpp
@@ -143,7 +143,9 @@ bool FMaterialXSurfaceShaderAbstract::ConnectNodeGraphOutputToInput(MaterialX::I
 
 		for(mx::Edge Edge : Output->traverseGraph())
 		{
-			ConnectNodeCategoryOutputToInput(Edge, ShaderNode, ParentInputName);
+			Output->hasOutputString() ?
+				ConnectNodeCategoryOutputToInput(Edge, ShaderNode, ParentInputName, Output->getOutputString().c_str()) :
+				ConnectNodeCategoryOutputToInput(Edge, ShaderNode, ParentInputName);
 		}
 	}
 
@@ -154,34 +156,64 @@ bool FMaterialXSurfaceShaderAbstract::ConnectMatchingNodeOutputToInput(const FCo
 {
 	FMaterialXManager& Manager = FMaterialXManager::GetInstance();
 
+	auto GetIndexOutput = [&Connect]()
+	{
+		mx::NodeDefPtr NodeDef = Connect.UpstreamNode->getNodeDef(mx::EMPTY_STRING, true);
+		int Index = NodeDef->getChildIndex(TCHAR_TO_UTF8(*Connect.OutputName));
+		return Index < 0 ? Index : Index % NodeDef->getInputCount();
+	};
+
 	bool bIsConnected = false;
 
-	// First search a matching Material Expression
-	if(const FString* ShaderType = Manager.FindMatchingMaterialExpression(Connect.UpstreamNode->getCategory().c_str()))
+	auto ConnectOutputToInput = [&](const FString* ShaderType, auto* (FMaterialXSurfaceShaderAbstract::* CreateFunctionCallOrShaderNode)(const FString&, const FString&, const FString&), bool bFindMatchingInput = true)
 	{
 		UInterchangeShaderNode* OperatorNode = nullptr;
-		
-		OperatorNode = CreateShaderNode(Connect.UpstreamNode->getName().c_str(), *ShaderType);
+
+		//We don't take the node output here because it would cause the creation of a new node (output is meaningful with ComponentMaskNode/separate where we have to create a new expression
+		OperatorNode = (this->*CreateFunctionCallOrShaderNode)(Connect.UpstreamNode->getName().c_str(), *ShaderType, DefaultOutput);
 
 		for(mx::InputPtr Input : Connect.UpstreamNode->getInputs())
 		{
-			if(const FString* InputNameFound = Manager.FindMaterialExpressionInput(GetInputName(Input)))
+			if(!bFindMatchingInput)
+			{
+				AddAttributeFromValueOrInterface(Input, Input->getName().c_str(), OperatorNode);
+			}
+			else if(const FString* InputNameFound = Manager.FindMaterialExpressionInput(GetInputName(Input)))
 			{
 				AddAttributeFromValueOrInterface(Input, *InputNameFound, OperatorNode);
 			}
 		}
 
-		bIsConnected = UInterchangeShaderPortsAPI::ConnectDefaultOuputToInput(Connect.ParentShaderNode, Connect.InputChannelName, OperatorNode->GetUniqueID());
+		int IndexOutput = GetIndexOutput();
+
+		bIsConnected = IndexOutput < 0 ?
+			UInterchangeShaderPortsAPI::ConnectDefaultOuputToInput(Connect.ParentShaderNode, Connect.InputChannelName, OperatorNode->GetUniqueID()) :
+			UInterchangeShaderPortsAPI::ConnectOuputToInputByIndex(Connect.ParentShaderNode, Connect.InputChannelName, OperatorNode->GetUniqueID(), IndexOutput);
+	};
+
+	// First search a matching Material Expression
+	// search for a Material Expression based on the node group (essentially used for Substrate Mix)
+	if(const FString* ShaderType = Manager.FindMatchingMaterialExpression(Connect.UpstreamNode->getCategory().c_str(), Connect.UpstreamNode->getNodeDef(mx::EMPTY_STRING, true)->getNodeGroup().c_str()))
+	{
+		ConnectOutputToInput(ShaderType, &FMaterialXSurfaceShaderAbstract::CreateShaderNode);
+	}
+	else if((ShaderType = Manager.FindMatchingMaterialExpression(Connect.UpstreamNode->getCategory().c_str())))
+	{
+		ConnectOutputToInput(ShaderType, &FMaterialXSurfaceShaderAbstract::CreateShaderNode);
 	}
 	else if(FOnConnectNodeOutputToInput* Delegate = MatchingConnectNodeDelegates.Find(Connect.UpstreamNode->getCategory().c_str()))
 	{
 		bIsConnected = Delegate->ExecuteIfBound(Connect);
 	}
+	else if(const FString* FunctionPath = Manager.FindMatchingMaterialFunction(Connect.UpstreamNode->getCategory().c_str()))
+	{
+		ConnectOutputToInput(FunctionPath, &FMaterialXSurfaceShaderAbstract::CreateFunctionCallShaderNode, false);
+	}
 
 	return bIsConnected;
 }
 
-void FMaterialXSurfaceShaderAbstract::ConnectNodeCategoryOutputToInput(const MaterialX::Edge& Edge, UInterchangeShaderNode* ShaderNode, const FString& ParentInputName)
+void FMaterialXSurfaceShaderAbstract::ConnectNodeCategoryOutputToInput(const MaterialX::Edge& Edge, UInterchangeShaderNode* ShaderNode, const FString& ParentInputName, const FString& OutputName)
 {
 	if(mx::NodePtr UpstreamNode = Edge.getUpstreamElement()->asA<mx::Node>())
 	{
@@ -193,7 +225,7 @@ void FMaterialXSurfaceShaderAbstract::ConnectNodeCategoryOutputToInput(const Mat
 		//Replace the input's name by the one used in UE
 		SetMatchingInputsNames(UpstreamNode);
 
-		FString OutputChannelName = DefaultOutput;
+		FString OutputChannelName = OutputName;
 
 		if(mx::ElementPtr DownstreamElement = Edge.getDownstreamElement())
 		{
@@ -1141,6 +1173,25 @@ UInterchangeShaderNode* FMaterialXSurfaceShaderAbstract::CreateShaderNode(const
 	return Node;
 }
 
+UInterchangeFunctionCallShaderNode* FMaterialXSurfaceShaderAbstract::CreateFunctionCallShaderNode(const FString& NodeName, const FString& FunctionPath, const FString& OutputName)
+{
+	UInterchangeFunctionCallShaderNode* Node;
+
+	const FString NodeUID = UInterchangeShaderNode::MakeNodeUid(NodeName  + TEXT('_') + OutputName, FStringView{});
+
+	if(Node = const_cast<UInterchangeFunctionCallShaderNode*>(Cast<UInterchangeFunctionCallShaderNode>(NodeContainer.GetNode(NodeUID))); !Node)
+	{
+		Node = NewObject<UInterchangeFunctionCallShaderNode>(&NodeContainer);
+		Node->InitializeNode(NodeUID, NodeName, EInterchangeNodeContainerType::TranslatedAsset);
+		Node->SetCustomMaterialFunction(FunctionPath);
+		NodeContainer.AddNode(Node);
+
+		ShaderNodes.Add({ NodeName, OutputName}, Node);
+	}
+
+	return Node;
+}
+
 const FString& FMaterialXSurfaceShaderAbstract::GetMatchedInputName(MaterialX::NodePtr Node, MaterialX::InputPtr Input) const
 {
 	FMaterialXManager& Manager = FMaterialXManager::GetInstance();
@@ -1148,13 +1199,17 @@ const FString& FMaterialXSurfaceShaderAbstract::GetMatchedInputName(MaterialX::N
 	if(Input)
 	{
 		const FString NodeCategory{ Node->getCategory().c_str() };
-		const FString InputName{ GetInputName(Input) };
-		
-		if(const FString* Result = Manager.FindMatchingInput({ NodeCategory, InputName }))
+		const FString InputName{ GetInputName(Input) };	
+
+		if(const FString* Result = Manager.FindMatchingInput(NodeCategory, InputName, Node->getNodeDef(mx::EMPTY_STRING, true)->getNodeGroup().c_str()))
+		{
+			return *Result;
+		}
+		else if((Result = Manager.FindMatchingInput(NodeCategory, InputName)))
 		{
 			return *Result;
 		}
-		else if((Result = Manager.FindMatchingInput({ EmptyString, InputName })))
+		else if((Result = Manager.FindMatchingInput(EmptyString, InputName)))
 		{
 			return *Result;
 		}
diff --git a/Engine/Plugins/Interchange/Runtime/Source/Import/Private/MaterialX/MaterialXUtils/MaterialXSurfaceUnlitShader.cpp b/Engine/Plugins/Interchange/Runtime/Source/Import/Private/MaterialX/MaterialXUtils/MaterialXSurfaceUnlitShader.cpp
--- a/Engine/Plugins/Interchange/Runtime/Source/Import/Private/MaterialX/MaterialXUtils/MaterialXSurfaceUnlitShader.cpp
+++ b/Engine/Plugins/Interchange/Runtime/Source/Import/Private/MaterialX/MaterialXUtils/MaterialXSurfaceUnlitShader.cpp
@@ -3,6 +3,8 @@
 #if WITH_EDITOR
 #include "MaterialXSurfaceUnlitShader.h"
 
+#include "Engine/EngineTypes.h"
+
 namespace mx = MaterialX;
 
 FMaterialXSurfaceUnlitShader::FMaterialXSurfaceUnlitShader(UInterchangeBaseNodeContainer& BaseNodeContainer)
@@ -22,21 +24,7 @@ void FMaterialXSurfaceUnlitShader::Translate(mx::NodePtr SurfaceUnlitNode)
 {
 	this->SurfaceShaderNode = SurfaceUnlitNode;
 
-	UInterchangeFunctionCallShaderNode* SurfaceUnlitShaderNode;
-
-	const FString NodeUID = UInterchangeShaderNode::MakeNodeUid(ANSI_TO_TCHAR(SurfaceUnlitNode->getName().c_str()), FStringView{});
-
-	if(SurfaceUnlitShaderNode = const_cast<UInterchangeFunctionCallShaderNode*>(Cast<UInterchangeFunctionCallShaderNode>(NodeContainer.GetNode(NodeUID))); !SurfaceUnlitShaderNode)
-	{
-		const FString NodeName = SurfaceUnlitNode->getName().c_str();
-		SurfaceUnlitShaderNode = NewObject<UInterchangeFunctionCallShaderNode>(&NodeContainer);
-		SurfaceUnlitShaderNode->InitializeNode(NodeUID, NodeName, EInterchangeNodeContainerType::TranslatedAsset);
-
-		SurfaceUnlitShaderNode->SetCustomMaterialFunction(TEXT("/Interchange/Functions/MX_SurfaceUnlit.MX_SurfaceUnlit"));
-		NodeContainer.AddNode(SurfaceUnlitShaderNode);
-
-		ShaderNodes.Add({ NodeName, SurfaceUnlitNode->getNodeDef(mx::EMPTY_STRING, true)->getActiveOutputs()[0]->getName().c_str() }, SurfaceUnlitShaderNode);
-	}
+	UInterchangeFunctionCallShaderNode* SurfaceUnlitShaderNode = CreateFunctionCallShaderNode(SurfaceUnlitNode->getName().c_str(), TEXT("/Interchange/Functions/MX_SurfaceUnlit.MX_SurfaceUnlit"));
 
 	using namespace UE::Interchange::Materials;
 
@@ -55,18 +43,34 @@ void FMaterialXSurfaceUnlitShader::Translate(mx::NodePtr SurfaceUnlitNode)
 	//Transmission Color
 	ConnectNodeOutputToInput(mx::SurfaceUnlit::Input::TransmissionColor, SurfaceUnlitShaderNode, SurfaceUnlit::Parameters::TransmissionColor.ToString(), mx::SurfaceUnlit::DefaultValue::Color3::TransmissionColor);
 
-	// Outputs
-	UInterchangeShaderPortsAPI::ConnectOuputToInputByName(ShaderGraphNode, PBRMR::Parameters::EmissiveColor.ToString(), SurfaceUnlitShaderNode->GetUniqueID(), PBRMR::Parameters::EmissiveColor.ToString());
-
-	//We can't have both Opacity and Opacity Mask we need to make a choice	
-	if(UInterchangeShaderPortsAPI::HasInput(SurfaceUnlitShaderNode, SurfaceUnlit::Parameters::Transmission))
+	if(!bIsSubstrateEnabled)
 	{
-		UInterchangeShaderPortsAPI::ConnectOuputToInputByName(ShaderGraphNode, PBRMR::Parameters::Opacity.ToString(), SurfaceUnlitShaderNode->GetUniqueID(), PBRMR::Parameters::Opacity.ToString());
+		// Outputs
+		UInterchangeShaderPortsAPI::ConnectOuputToInputByName(ShaderGraphNode, PBRMR::Parameters::EmissiveColor.ToString(), SurfaceUnlitShaderNode->GetUniqueID(), PBRMR::Parameters::EmissiveColor.ToString());
+
+		//We can't have both Opacity and Opacity Mask we need to make a choice	
+		if(UInterchangeShaderPortsAPI::HasInput(SurfaceUnlitShaderNode, SurfaceUnlit::Parameters::Transmission))
+		{
+			UInterchangeShaderPortsAPI::ConnectOuputToInputByName(ShaderGraphNode, PBRMR::Parameters::Opacity.ToString(), SurfaceUnlitShaderNode->GetUniqueID(), PBRMR::Parameters::Opacity.ToString());
+		}
+		else if(UInterchangeShaderPortsAPI::HasInput(SurfaceUnlitShaderNode, SurfaceUnlit::Parameters::Opacity))
+		{
+			UInterchangeShaderPortsAPI::ConnectOuputToInputByName(ShaderGraphNode, PBRMR::Parameters::Opacity.ToString(), SurfaceUnlitShaderNode->GetUniqueID(), SurfaceUnlit::Outputs::OpacityMask.ToString());
+			ShaderGraphNode->SetCustomOpacityMaskClipValue(1.f); // just to connect to the opacity mask
+		}
 	}
-	else if(UInterchangeShaderPortsAPI::HasInput(SurfaceUnlitShaderNode, SurfaceUnlit::Parameters::Opacity))
+	else
 	{
-		UInterchangeShaderPortsAPI::ConnectOuputToInputByName(ShaderGraphNode, PBRMR::Parameters::Opacity.ToString(), SurfaceUnlitShaderNode->GetUniqueID(), TEXT("OpacityMask"));
-		ShaderGraphNode->SetCustomOpacityMaskClipValue(1.f); // just to connect to the opacity mask
+		UInterchangeShaderPortsAPI::ConnectOuputToInputByName(ShaderGraphNode, SubstrateMaterial::Parameters::FrontMaterial.ToString(), SurfaceUnlitShaderNode->GetUniqueID(), SurfaceUnlit::Substrate::Outputs::SurfaceUnlit.ToString());
+		if(UInterchangeShaderPortsAPI::HasInput(SurfaceUnlitShaderNode, SurfaceUnlit::Parameters::Transmission))
+		{
+			ShaderGraphNode->SetCustomBlendMode(EBlendMode::BLEND_TranslucentColoredTransmittance);
+		}
+		else if(UInterchangeShaderPortsAPI::HasInput(SurfaceUnlitShaderNode, SurfaceUnlit::Parameters::Opacity))
+		{
+			UInterchangeShaderPortsAPI::ConnectOuputToInputByName(ShaderGraphNode, SubstrateMaterial::Parameters::OpacityMask.ToString(), SurfaceUnlitShaderNode->GetUniqueID(), SurfaceUnlit::Substrate::Outputs::OpacityMask.ToString());
+			ShaderGraphNode->SetCustomBlendMode(EBlendMode::BLEND_Masked);
+		}
 	}
 }
 #endif
\ No newline at end of file
diff --git a/Engine/Plugins/Interchange/Runtime/Source/Import/Private/MaterialX/MaterialXUtils/MaterialXUsdPreviewSurfaceShader.cpp b/Engine/Plugins/Interchange/Runtime/Source/Import/Private/MaterialX/MaterialXUtils/MaterialXUsdPreviewSurfaceShader.cpp
--- a/Engine/Plugins/Interchange/Runtime/Source/Import/Private/MaterialX/MaterialXUtils/MaterialXUsdPreviewSurfaceShader.cpp
+++ b/Engine/Plugins/Interchange/Runtime/Source/Import/Private/MaterialX/MaterialXUtils/MaterialXUsdPreviewSurfaceShader.cpp
@@ -25,21 +25,8 @@ void FMaterialXUsdPreviewSurfaceShader::Translate(MaterialX::NodePtr UsdPreviewS
 
 	using namespace UE::Interchange::Materials;
 
-	UInterchangeFunctionCallShaderNode* UsdPreviewSurfaceShaderNode;
-
-	const FString NodeUID = UInterchangeShaderNode::MakeNodeUid(ANSI_TO_TCHAR(UsdPreviewSurfaceNode->getName().c_str()), FStringView{});
-
-	if(UsdPreviewSurfaceShaderNode = const_cast<UInterchangeFunctionCallShaderNode*>(Cast<UInterchangeFunctionCallShaderNode>(NodeContainer.GetNode(NodeUID))); !UsdPreviewSurfaceShaderNode)
-	{
-		const FString NodeName = UsdPreviewSurfaceNode->getName().c_str();
-		UsdPreviewSurfaceShaderNode = NewObject<UInterchangeFunctionCallShaderNode>(&NodeContainer);
-		UsdPreviewSurfaceShaderNode->InitializeNode(NodeUID, NodeName, EInterchangeNodeContainerType::TranslatedAsset);
-
-		UsdPreviewSurfaceShaderNode->SetCustomMaterialFunction(TEXT("/Interchange/Functions/MX_UsdPreviewSurface.MX_UsdPreviewSurface"));
-		NodeContainer.AddNode(UsdPreviewSurfaceShaderNode);
-
-		ShaderNodes.Add({ NodeName, UsdPreviewSurfaceNode->getNodeDef(mx::EMPTY_STRING, true)->getActiveOutputs()[0]->getName().c_str()}, UsdPreviewSurfaceShaderNode);
-	}
+	const FString FolderFunction = bIsSubstrateEnabled ? TEXT("Substrate/") : TEXT("Functions/");
+	UInterchangeFunctionCallShaderNode* UsdPreviewSurfaceShaderNode = CreateFunctionCallShaderNode(UsdPreviewSurfaceNode->getName().c_str(), TEXT("/Interchange/") + FolderFunction + TEXT("MX_UsdPreviewSurface.MX_UsdPreviewSurface"));
 
 	// Inputs
 	//Diffuse Color
@@ -81,22 +68,29 @@ void FMaterialXUsdPreviewSurfaceShader::Translate(MaterialX::NodePtr UsdPreviewS
 	//Occlusion
 	ConnectNodeOutputToInput(mx::UsdPreviewSurface::Input::Occlusion, UsdPreviewSurfaceShaderNode, UsdPreviewSurface::Parameters::Occlusion.ToString(), mx::UsdPreviewSurface::DefaultValue::Float::Occlusion);
 
+	if(!bIsSubstrateEnabled)
 	// Outputs
-	UInterchangeShaderPortsAPI::ConnectOuputToInputByName(ShaderGraphNode, PBRMR::Parameters::BaseColor.ToString(), UsdPreviewSurfaceShaderNode->GetUniqueID(), PBRMR::Parameters::BaseColor.ToString());
-	UInterchangeShaderPortsAPI::ConnectOuputToInputByName(ShaderGraphNode, PBRMR::Parameters::Metallic.ToString(), UsdPreviewSurfaceShaderNode->GetUniqueID(), PBRMR::Parameters::Metallic.ToString());
-	UInterchangeShaderPortsAPI::ConnectOuputToInputByName(ShaderGraphNode, PBRMR::Parameters::Specular.ToString(), UsdPreviewSurfaceShaderNode->GetUniqueID(), PBRMR::Parameters::Specular.ToString());
-	UInterchangeShaderPortsAPI::ConnectOuputToInputByName(ShaderGraphNode, PBRMR::Parameters::Roughness.ToString(), UsdPreviewSurfaceShaderNode->GetUniqueID(), PBRMR::Parameters::Roughness.ToString());
-	UInterchangeShaderPortsAPI::ConnectOuputToInputByName(ShaderGraphNode, PBRMR::Parameters::EmissiveColor.ToString(), UsdPreviewSurfaceShaderNode->GetUniqueID(), PBRMR::Parameters::EmissiveColor.ToString());
-	UInterchangeShaderPortsAPI::ConnectOuputToInputByName(ShaderGraphNode, PBRMR::Parameters::Normal.ToString(), UsdPreviewSurfaceShaderNode->GetUniqueID(), PBRMR::Parameters::Normal.ToString());
-	
-	if(UInterchangeShaderPortsAPI::HasInput(UsdPreviewSurfaceShaderNode, UsdPreviewSurface::Parameters::Opacity))
 	{
-		UInterchangeShaderPortsAPI::ConnectOuputToInputByName(ShaderGraphNode, PBRMR::Parameters::Opacity.ToString(), UsdPreviewSurfaceShaderNode->GetUniqueID(), PBRMR::Parameters::Opacity.ToString());
+		UInterchangeShaderPortsAPI::ConnectOuputToInputByName(ShaderGraphNode, PBRMR::Parameters::BaseColor.ToString(), UsdPreviewSurfaceShaderNode->GetUniqueID(), PBRMR::Parameters::BaseColor.ToString());
+		UInterchangeShaderPortsAPI::ConnectOuputToInputByName(ShaderGraphNode, PBRMR::Parameters::Metallic.ToString(), UsdPreviewSurfaceShaderNode->GetUniqueID(), PBRMR::Parameters::Metallic.ToString());
+		UInterchangeShaderPortsAPI::ConnectOuputToInputByName(ShaderGraphNode, PBRMR::Parameters::Specular.ToString(), UsdPreviewSurfaceShaderNode->GetUniqueID(), PBRMR::Parameters::Specular.ToString());
+		UInterchangeShaderPortsAPI::ConnectOuputToInputByName(ShaderGraphNode, PBRMR::Parameters::Roughness.ToString(), UsdPreviewSurfaceShaderNode->GetUniqueID(), PBRMR::Parameters::Roughness.ToString());
+		UInterchangeShaderPortsAPI::ConnectOuputToInputByName(ShaderGraphNode, PBRMR::Parameters::EmissiveColor.ToString(), UsdPreviewSurfaceShaderNode->GetUniqueID(), PBRMR::Parameters::EmissiveColor.ToString());
+		UInterchangeShaderPortsAPI::ConnectOuputToInputByName(ShaderGraphNode, PBRMR::Parameters::Normal.ToString(), UsdPreviewSurfaceShaderNode->GetUniqueID(), PBRMR::Parameters::Normal.ToString());
+
+		if(UInterchangeShaderPortsAPI::HasInput(UsdPreviewSurfaceShaderNode, UsdPreviewSurface::Parameters::Opacity))
+		{
+			UInterchangeShaderPortsAPI::ConnectOuputToInputByName(ShaderGraphNode, PBRMR::Parameters::Opacity.ToString(), UsdPreviewSurfaceShaderNode->GetUniqueID(), PBRMR::Parameters::Opacity.ToString());
+		}
+
+		UInterchangeShaderPortsAPI::ConnectOuputToInputByName(ShaderGraphNode, PBRMR::Parameters::Occlusion.ToString(), UsdPreviewSurfaceShaderNode->GetUniqueID(), PBRMR::Parameters::Occlusion.ToString());
+		UInterchangeShaderPortsAPI::ConnectOuputToInputByName(ShaderGraphNode, PBRMR::Parameters::Refraction.ToString(), UsdPreviewSurfaceShaderNode->GetUniqueID(), PBRMR::Parameters::Refraction.ToString());
+		UInterchangeShaderPortsAPI::ConnectOuputToInputByName(ShaderGraphNode, ClearCoat::Parameters::ClearCoat.ToString(), UsdPreviewSurfaceShaderNode->GetUniqueID(), ClearCoat::Parameters::ClearCoat.ToString());
+		UInterchangeShaderPortsAPI::ConnectOuputToInputByName(ShaderGraphNode, ClearCoat::Parameters::ClearCoatRoughness.ToString(), UsdPreviewSurfaceShaderNode->GetUniqueID(), ClearCoat::Parameters::ClearCoatRoughness.ToString());
+	}
+	else
+	{
+		UInterchangeShaderPortsAPI::ConnectDefaultOuputToInput(ShaderGraphNode, SubstrateMaterial::Parameters::FrontMaterial.ToString(), UsdPreviewSurfaceShaderNode->GetUniqueID());
 	}
-
-	UInterchangeShaderPortsAPI::ConnectOuputToInputByName(ShaderGraphNode, PBRMR::Parameters::Occlusion.ToString(), UsdPreviewSurfaceShaderNode->GetUniqueID(), PBRMR::Parameters::Occlusion.ToString());
-	UInterchangeShaderPortsAPI::ConnectOuputToInputByName(ShaderGraphNode, PBRMR::Parameters::Refraction.ToString(), UsdPreviewSurfaceShaderNode->GetUniqueID(), PBRMR::Parameters::Refraction.ToString());
-	UInterchangeShaderPortsAPI::ConnectOuputToInputByName(ShaderGraphNode, ClearCoat::Parameters::ClearCoat.ToString(), UsdPreviewSurfaceShaderNode->GetUniqueID(), ClearCoat::Parameters::ClearCoat.ToString());
-	UInterchangeShaderPortsAPI::ConnectOuputToInputByName(ShaderGraphNode, ClearCoat::Parameters::ClearCoatRoughness.ToString(), UsdPreviewSurfaceShaderNode->GetUniqueID(), ClearCoat::Parameters::ClearCoatRoughness.ToString());
 }
 #endif
\ No newline at end of file
diff --git a/Engine/Plugins/Interchange/Runtime/Source/Import/Public/InterchangeImportModule.h b/Engine/Plugins/Interchange/Runtime/Source/Import/Public/InterchangeImportModule.h
--- a/Engine/Plugins/Interchange/Runtime/Source/Import/Public/InterchangeImportModule.h
+++ b/Engine/Plugins/Interchange/Runtime/Source/Import/Public/InterchangeImportModule.h
@@ -34,5 +34,12 @@ public:
 	{
 		return FModuleManager::Get().IsModuleLoaded( "InterchangeImport" );
 	}
+
+	/**
+	 * Checks to see if Substrate support is enabled both in the Project Settings and for Interchange
+	 * 
+	 * @return True if Substrate is enabled
+	 */
+	virtual bool IsSubstrateEnabled() const { return false; };
 };
 
diff --git a/Engine/Plugins/Interchange/Runtime/Source/Import/Public/MaterialX/MaterialXUtils/MaterialXBase.h b/Engine/Plugins/Interchange/Runtime/Source/Import/Public/MaterialX/MaterialXUtils/MaterialXBase.h
--- a/Engine/Plugins/Interchange/Runtime/Source/Import/Public/MaterialX/MaterialXUtils/MaterialXBase.h
+++ b/Engine/Plugins/Interchange/Runtime/Source/Import/Public/MaterialX/MaterialXUtils/MaterialXBase.h
@@ -123,5 +123,7 @@ protected:
 
 	/** It's up to the derive class to initialize it with the proper node definition*/
 	const char* NodeDefinition; 
+
+	bool bIsSubstrateEnabled;
 };
 #endif
diff --git a/Engine/Plugins/Interchange/Runtime/Source/Import/Public/MaterialX/MaterialXUtils/MaterialXManager.h b/Engine/Plugins/Interchange/Runtime/Source/Import/Public/MaterialX/MaterialXUtils/MaterialXManager.h
--- a/Engine/Plugins/Interchange/Runtime/Source/Import/Public/MaterialX/MaterialXUtils/MaterialXManager.h
+++ b/Engine/Plugins/Interchange/Runtime/Source/Import/Public/MaterialX/MaterialXUtils/MaterialXManager.h
@@ -2,9 +2,10 @@
 
 #pragma once
 
+#include "CoreMinimal.h"
+
 #if WITH_EDITOR
 
-#include "CoreMinimal.h"
 #include "MaterialXCore/Node.h"
 
 class FMaterialXBase;
@@ -20,18 +21,24 @@ public:
 
 	static FMaterialXManager& GetInstance();
 
-	/** Find a corresponding Material Expression Input given a (Category, Input) pair*/
-	const FString* FindMatchingInput(const TPair<FString, FString>& CategoryInputKey) const;
+	/** Find a corresponding Material Expression Input given a (Category [NodeGroup], Input) pair*/
+	const FString* FindMatchingInput(const FString& CategoryKey, const FString& InputKey, const FString& NodeGroup = {}) const;
 
 	/** Find a stored Material Expression Input */
 	const FString* FindMaterialExpressionInput(const FString& InputKey) const;
 
 	/** Find a matching Material Expression given a MaterialX category*/
-	const FString* FindMatchingMaterialExpression(const FString& CategoryKey) const;
+	const FString* FindMatchingMaterialExpression(const FString& CategoryKey, const FString& NodeGroup = {}) const;
+
+	/** Find a matching Material Function given a MaterialX category*/
+	const FString* FindMatchingMaterialFunction(const FString& CategoryKey) const;
 
 	TSharedPtr<FMaterialXBase> GetShaderTranslator(const FString& CategoryShader, UInterchangeBaseNodeContainer& NodeContainer);
 
 	void RegisterMaterialXInstance(const FString& CategoryShader, FOnGetMaterialXInstance MaterialXInstanceDelegate);
+	
+	bool IsSubstrateEnabled() const;
+
 
 	FMaterialXManager(const FMaterialXManager&) = delete;
 	FMaterialXManager& operator=(const FMaterialXManager&) = delete;
@@ -40,18 +47,56 @@ public:
 
 private:
 
+	struct FKeyCategoryNodegroup
+	{
+		template<typename CategoryString>
+		FKeyCategoryNodegroup(CategoryString&& Category)
+			: Category{ std::forward<CategoryString>(Category) }
+		{}
+
+		template<typename CategoryString, typename NodeGroupString>
+		FKeyCategoryNodegroup(CategoryString&& Category, NodeGroupString&& NodeGroup)
+			: Category{ std::forward<CategoryString>(Category) }
+			, NodeGroup{ std::forward<CategoryString>(NodeGroup) }
+		{}
+
+		UE_NODISCARD FORCEINLINE bool operator==(const FKeyCategoryNodegroup& Rhs) const
+		{
+			return Category == Rhs.Category && NodeGroup == Rhs.NodeGroup;
+		}
+
+		friend FORCEINLINE uint32 GetTypeHash(const FKeyCategoryNodegroup& Key)
+		{
+			return HashCombine(GetTypeHash(Key.Category), GetTypeHash(Key.NodeGroup));
+		}
+
+		FString Category;
+		FString NodeGroup; // node group is optional, some nodes in MaterialX have different material expressions in UE
+	};
+
 	FMaterialXManager();
 
-	/** Given a MaterialX node (category - input), return the UE/Interchange input name.*/
-	TMap<TPair<FString, FString>, FString> MatchingInputNames;
+	/** Given a MaterialX node (category (optionally a nodegroup) - input), return the UE/Interchange input name.*/
+	TMap<TPair<FKeyCategoryNodegroup, FString>, FString> MatchingInputNames;
+
+	/** Given a MaterialX node category, optionally with a node group, return the UE material expression class name*/
+	TMap<FKeyCategoryNodegroup, FString> MatchingMaterialExpressions;
 
-	/** Given a MaterialX node category, return the UE material expression class name*/
-	TMap<FString, FString> MatchingMaterialExpressions;
+	/** Given a MaterialX node category, return the UE material function, used for BSDF nodes*/
+	TMap<FString, FString> MatchingMaterialFunctions;
 
 	/** The different inputs of material expression that we may encounter, the MaterialX Document is modified consequently regarding those*/
 	TSet<FString> MaterialExpressionInputs;
 
-	/** Container of a MaterialX document, to translate the different nodes based on their category, same reason as above concerning the use of STL*/
+	/** Container of a MaterialX document, to translate the different nodes based on their category*/
 	TMap<FString, FOnGetMaterialXInstance> MaterialXContainerDelegates;
+
+	bool bIsSubstrateEnabled;
 };
 #endif
+
+namespace UE::Interchange::MaterialX
+{
+	// Load necessary material functions, this function can only be called in the Game Thread
+	bool AreMaterialFunctionPackagesLoaded();
+}
\ No newline at end of file
diff --git a/Engine/Plugins/Interchange/Runtime/Source/Import/Public/MaterialX/MaterialXUtils/MaterialXSurfaceShaderAbstract.h b/Engine/Plugins/Interchange/Runtime/Source/Import/Public/MaterialX/MaterialXUtils/MaterialXSurfaceShaderAbstract.h
--- a/Engine/Plugins/Interchange/Runtime/Source/Import/Public/MaterialX/MaterialXUtils/MaterialXSurfaceShaderAbstract.h
+++ b/Engine/Plugins/Interchange/Runtime/Source/Import/Public/MaterialX/MaterialXUtils/MaterialXSurfaceShaderAbstract.h
@@ -173,7 +173,7 @@ protected:
 	 *
 	 * @return true if a shader node has been successfully created and is connected to the given input
 	 */
-	void ConnectNodeCategoryOutputToInput(const MaterialX::Edge& Edge, UInterchangeShaderNode* ParentShaderNode, const FString& InputChannelName);
+	void ConnectNodeCategoryOutputToInput(const MaterialX::Edge& Edge, UInterchangeShaderNode* ParentShaderNode, const FString& InputChannelName, const FString& OutputName = TEXT("out"));
 
 	/**
 	 * Create and Connect a node name directly connected from an input to a shader node
@@ -286,6 +286,17 @@ protected:
 	 */
 	UInterchangeShaderNode* CreateShaderNode(const FString& NodeName, const FString& ShaderType, const FString& OutputName = TEXT("out"));
 
+	/**
+	 * Helper function to create an InterchangeFunctionCallShaderNode
+	 *
+	 * @param NodeName - The name of the shader node
+	 * @param FunctionPath - The path to the Material Function we want to create
+	 * @param OutputName - The output name of the MaterialX node, default name is 'out' as stated by the standard library
+	 *
+	 * @return The shader node that was created
+	 */
+	UInterchangeFunctionCallShaderNode* CreateFunctionCallShaderNode(const FString& NodeName, const FString& FunctionPath, const FString& OutputName = TEXT("out"));
+
 	/**
 	 * Helper function to create an InterchangeTextureNode
 	 *
diff --git a/Engine/Plugins/Interchange/Runtime/Source/Nodes/Public/InterchangeMaterialDefinitions.h b/Engine/Plugins/Interchange/Runtime/Source/Nodes/Public/InterchangeMaterialDefinitions.h
--- a/Engine/Plugins/Interchange/Runtime/Source/Nodes/Public/InterchangeMaterialDefinitions.h
+++ b/Engine/Plugins/Interchange/Runtime/Source/Nodes/Public/InterchangeMaterialDefinitions.h
@@ -472,6 +472,32 @@ namespace Interchange
 			}
 		}
 
+		namespace Surface
+		{
+			const FName Name = TEXT("surface");
+
+			namespace Parameters
+			{
+				const FName BSDF = TEXT("bsdf");
+				const FName EDF = TEXT("edf");
+				const FName Opacity = TEXT("opacity");
+			}
+
+			namespace Outputs
+			{
+				const FName Surface = TEXT("Surface");
+			}
+
+			namespace Substrate
+			{
+				namespace Outputs
+				{
+					using namespace Surface::Outputs;
+					const FName Opacity = TEXT("Opacity");
+				}
+			}
+		}
+
 		namespace SurfaceUnlit
 		{
 			const FName Name = TEXT("surface_unlit");
@@ -484,6 +510,20 @@ namespace Interchange
 				const FName TransmissionColor = TEXT("transmission_color");
 				const FName Opacity = TEXT("opacity");
 			}
+
+			namespace Outputs
+			{
+				const FName OpacityMask = TEXT("OpacityMask");
+			}
+
+			namespace Substrate
+			{
+				namespace Outputs
+				{
+					using namespace SurfaceUnlit::Outputs;
+					const FName SurfaceUnlit = TEXT("Surface Unlit");
+				}
+			}
 		}
 
 		namespace UsdPreviewSurface
diff --git a/Engine/Plugins/Interchange/Runtime/Source/Pipelines/Private/InterchangeMaterialXPipeline.cpp b/Engine/Plugins/Interchange/Runtime/Source/Pipelines/Private/InterchangeMaterialXPipeline.cpp
--- a/Engine/Plugins/Interchange/Runtime/Source/Pipelines/Private/InterchangeMaterialXPipeline.cpp
+++ b/Engine/Plugins/Interchange/Runtime/Source/Pipelines/Private/InterchangeMaterialXPipeline.cpp
@@ -1,69 +1,103 @@
 // Copyright Epic Games, Inc. All Rights Reserved. 
 
 #include "InterchangeMaterialXPipeline.h"
-#include "InterchangeManager.h"
+
+#include "InterchangeImportModule.h"
 #include "InterchangeMaterialDefinitions.h"
 #include "InterchangeMaterialFactoryNode.h"
 #include "InterchangePipelineLog.h"
 
-#include "Engine/RendererSettings.h"
 #include "Materials/MaterialExpressionMaterialFunctionCall.h"
 #include "Materials/MaterialFunction.h"
 #include "Misc/PackageName.h"
 
+
+#define MATERIALX_FUNCTIONS_SUBSTRATE_PATH(Name, AssetName)  \
+	constexpr const TCHAR* Name##FunctionsPath = TEXT("/Interchange/Functions/") TEXT(AssetName) TEXT(".") TEXT(AssetName);  \
+	constexpr const TCHAR* Name##SubstratePath = TEXT("/Interchange/Substrate/") TEXT(AssetName) TEXT(".") TEXT(AssetName);
+
+#define MATERIALX_MATERIALFUNCTION_PATH(Name) \
+	MaterialXSettings->bIsSubstrateEnabled ? Name##FunctionsPath : Name##SubstratePath
+
+namespace
+{	
+	MATERIALX_FUNCTIONS_SUBSTRATE_PATH(StandardSurface, "MX_StandardSurface");
+	MATERIALX_FUNCTIONS_SUBSTRATE_PATH(TransmissionSurface, "MX_TransmissionSurface");
+	MATERIALX_FUNCTIONS_SUBSTRATE_PATH(SurfaceUnlit, "MX_SurfaceUnlit");
+	MATERIALX_FUNCTIONS_SUBSTRATE_PATH(Surface, "MX_Surface");
+	MATERIALX_FUNCTIONS_SUBSTRATE_PATH(UsdPreviewSurface, "MX_UsdPreviewSurface");
+	
+	MATERIALX_FUNCTIONS_SUBSTRATE_PATH(OrenNayarDiffuseBSDF, "MX_OrenNayarBSDF");
+	MATERIALX_FUNCTIONS_SUBSTRATE_PATH(BurleyDiffuseBSDF, "MX_BurleyDiffuseBSDF");
+	MATERIALX_FUNCTIONS_SUBSTRATE_PATH(DieletricBSDF, ".MX_DielectricBSDF");
+	MATERIALX_FUNCTIONS_SUBSTRATE_PATH(ConductorBSDF, "MX_ConductorBSDF");
+	MATERIALX_FUNCTIONS_SUBSTRATE_PATH(SheenBSDF, "MX_SheenBSDF");
+	MATERIALX_FUNCTIONS_SUBSTRATE_PATH(SubsurfaceBSDF, "MX_SubsurfaceBSDF");
+	MATERIALX_FUNCTIONS_SUBSTRATE_PATH(ThinFilmBSDF, "MX_ThinFilmBSDF");
+	MATERIALX_FUNCTIONS_SUBSTRATE_PATH(GeneralizedSchlickBSDF, "MX_GeneralizedSchlickBSDF");
+	MATERIALX_FUNCTIONS_SUBSTRATE_PATH(TranslucentBSDF, "MX_TranslucentBSDF");
+	
+	MATERIALX_FUNCTIONS_SUBSTRATE_PATH(UniformEDF, "MX_UniformEDF");
+	MATERIALX_FUNCTIONS_SUBSTRATE_PATH(ConicalEDF, "MX_ConicalEDF");
+	MATERIALX_FUNCTIONS_SUBSTRATE_PATH(MeasuredEDF, "MX_MeasuredEDF");
+	
+	MATERIALX_FUNCTIONS_SUBSTRATE_PATH(AbsorptionVDF, "MX_AbsorptionVDF");
+	MATERIALX_FUNCTIONS_SUBSTRATE_PATH(AnisotropicVDF, "MX_AnisotropicVDF");
+}
+
+TMap<FString, EMaterialXSettings> UInterchangeMaterialXPipeline::PathToEnumMapping;
 #if WITH_EDITOR
-TSet<FName> UMaterialXPipelineSettings::StandardSurfaceInputs;
-TSet<FName> UMaterialXPipelineSettings::StandardSurfaceOutputs;
-TSet<FName> UMaterialXPipelineSettings::TransmissionSurfaceInputs;
-TSet<FName> UMaterialXPipelineSettings::TransmissionSurfaceOutputs;
-TSet<FName> UMaterialXPipelineSettings::SurfaceUnlitInputs;
-TSet<FName> UMaterialXPipelineSettings::SurfaceUnlitOutputs;
-TSet<FName> UMaterialXPipelineSettings::UsdPreviewSurfaceInputs;
-TSet<FName> UMaterialXPipelineSettings::UsdPreviewSurfaceOutputs;
-#endif
-TMap<FString, EInterchangeMaterialXShaders> UInterchangeMaterialXPipeline::PathToEnumMapping;
+TMap<EMaterialXSettings, TPair<TSet<FName>, TSet<FName>>> UMaterialXPipelineSettings::SettingsInputsOutputs;
+#endif // WITH_EDITOR
 
 UMaterialXPipelineSettings::UMaterialXPipelineSettings()
 {
+#if WITH_EDITOR
 	if(HasAnyFlags(EObjectFlags::RF_ClassDefaultObject))
 	{
-		bIsSubstrateEnabled = GetDefault<URendererSettings>() ? GetDefault<URendererSettings>()->bEnableSubstrate : false;
-#if WITH_EDITOR
-		StandardSurfaceInputs =	{
-			UE::Interchange::Materials::StandardSurface::Parameters::Base,
-			UE::Interchange::Materials::StandardSurface::Parameters::BaseColor,
-			UE::Interchange::Materials::StandardSurface::Parameters::DiffuseRoughness,
-			UE::Interchange::Materials::StandardSurface::Parameters::Metalness,
-			UE::Interchange::Materials::StandardSurface::Parameters::Specular,
-			UE::Interchange::Materials::StandardSurface::Parameters::SpecularRoughness,
-			UE::Interchange::Materials::StandardSurface::Parameters::SpecularIOR,
-			UE::Interchange::Materials::StandardSurface::Parameters::SpecularAnisotropy,
-			UE::Interchange::Materials::StandardSurface::Parameters::SpecularRotation,
-			UE::Interchange::Materials::StandardSurface::Parameters::Subsurface,
-			UE::Interchange::Materials::StandardSurface::Parameters::SubsurfaceColor,
-			UE::Interchange::Materials::StandardSurface::Parameters::SubsurfaceRadius,
-			UE::Interchange::Materials::StandardSurface::Parameters::SubsurfaceScale,
-			UE::Interchange::Materials::StandardSurface::Parameters::Sheen,
-			UE::Interchange::Materials::StandardSurface::Parameters::SheenColor,
-			UE::Interchange::Materials::StandardSurface::Parameters::SheenRoughness,
-			UE::Interchange::Materials::StandardSurface::Parameters::Coat,
-			UE::Interchange::Materials::StandardSurface::Parameters::CoatColor,
-			UE::Interchange::Materials::StandardSurface::Parameters::CoatRoughness,
-			UE::Interchange::Materials::StandardSurface::Parameters::CoatNormal,
-			UE::Interchange::Materials::StandardSurface::Parameters::ThinFilmThickness,
-			UE::Interchange::Materials::StandardSurface::Parameters::Emission,
-			UE::Interchange::Materials::StandardSurface::Parameters::EmissionColor,
-			UE::Interchange::Materials::StandardSurface::Parameters::Normal,
-			UE::Interchange::Materials::StandardSurface::Parameters::Tangent,
-			bIsSubstrateEnabled ? UE::Interchange::Materials::StandardSurface::Parameters::SpecularColor : FName{},
-			bIsSubstrateEnabled ? UE::Interchange::Materials::StandardSurface::Parameters::CoatIOR : FName{},
-			bIsSubstrateEnabled ? UE::Interchange::Materials::StandardSurface::Parameters::CoatAnisotropy : FName{},
-			bIsSubstrateEnabled ? UE::Interchange::Materials::StandardSurface::Parameters::CoatRotation : FName{},
-			bIsSubstrateEnabled ? UE::Interchange::Materials::StandardSurface::Parameters::ThinFilmIOR : FName{},
-			bIsSubstrateEnabled ? UE::Interchange::Materials::StandardSurface::Parameters::Opacity : FName{},
-			bIsSubstrateEnabled ? UE::Interchange::Materials::StandardSurface::Parameters::SubsurfaceAnisotropy : FName{}
-		};
-		StandardSurfaceOutputs = {
+		bIsSubstrateEnabled = IInterchangeImportModule::IsAvailable() ? IInterchangeImportModule::Get().IsSubstrateEnabled() : false;
+
+		SettingsInputsOutputs = {
+			//Surface Shaders
+			{
+				UMaterialXPipelineSettings::ToEnumKey(EInterchangeMaterialXShaders::StandardSurface),
+				{
+				// StandardSurface Inputs
+				TSet<FName>{
+					UE::Interchange::Materials::StandardSurface::Parameters::Base,
+					UE::Interchange::Materials::StandardSurface::Parameters::BaseColor,
+					UE::Interchange::Materials::StandardSurface::Parameters::DiffuseRoughness,
+					UE::Interchange::Materials::StandardSurface::Parameters::Metalness,
+					UE::Interchange::Materials::StandardSurface::Parameters::Specular,
+					UE::Interchange::Materials::StandardSurface::Parameters::SpecularRoughness,
+					UE::Interchange::Materials::StandardSurface::Parameters::SpecularIOR,
+					UE::Interchange::Materials::StandardSurface::Parameters::SpecularAnisotropy,
+					UE::Interchange::Materials::StandardSurface::Parameters::SpecularRotation,
+					UE::Interchange::Materials::StandardSurface::Parameters::Subsurface,
+					UE::Interchange::Materials::StandardSurface::Parameters::SubsurfaceColor,
+					UE::Interchange::Materials::StandardSurface::Parameters::SubsurfaceRadius,
+					UE::Interchange::Materials::StandardSurface::Parameters::SubsurfaceScale,
+					UE::Interchange::Materials::StandardSurface::Parameters::Sheen,
+					UE::Interchange::Materials::StandardSurface::Parameters::SheenColor,
+					UE::Interchange::Materials::StandardSurface::Parameters::SheenRoughness,
+					UE::Interchange::Materials::StandardSurface::Parameters::Coat,
+					UE::Interchange::Materials::StandardSurface::Parameters::CoatColor,
+					UE::Interchange::Materials::StandardSurface::Parameters::CoatRoughness,
+					UE::Interchange::Materials::StandardSurface::Parameters::CoatNormal,
+					UE::Interchange::Materials::StandardSurface::Parameters::ThinFilmThickness,
+					UE::Interchange::Materials::StandardSurface::Parameters::Emission,
+					UE::Interchange::Materials::StandardSurface::Parameters::EmissionColor,
+					UE::Interchange::Materials::StandardSurface::Parameters::Normal,
+					UE::Interchange::Materials::StandardSurface::Parameters::Tangent,
+					bIsSubstrateEnabled ? UE::Interchange::Materials::StandardSurface::Parameters::SpecularColor : FName{},
+					bIsSubstrateEnabled ? UE::Interchange::Materials::StandardSurface::Parameters::CoatIOR : FName{},
+					bIsSubstrateEnabled ? UE::Interchange::Materials::StandardSurface::Parameters::CoatAnisotropy : FName{},
+					bIsSubstrateEnabled ? UE::Interchange::Materials::StandardSurface::Parameters::CoatRotation : FName{},
+					bIsSubstrateEnabled ? UE::Interchange::Materials::StandardSurface::Parameters::ThinFilmIOR : FName{},
+					bIsSubstrateEnabled ? UE::Interchange::Materials::StandardSurface::Parameters::Opacity : FName{},
+					bIsSubstrateEnabled ? UE::Interchange::Materials::StandardSurface::Parameters::SubsurfaceAnisotropy : FName{}
+				},
+			// StandardSurface Outputs
 			!bIsSubstrateEnabled ?
 			TSet<FName>{
 				TEXT("Base Color"), // MX_StandardSurface has BaseColor with a whitespace, this should be fixed in further release
@@ -86,139 +120,429 @@ UMaterialXPipelineSettings::UMaterialXPipelineSettings()
 				UE::Interchange::Materials::StandardSurface::SubstrateMaterial::Outputs::Opaque,
 				UE::Interchange::Materials::StandardSurface::SubstrateMaterial::Outputs::Opacity
 			}
+		}
+		},
+		{
+			UMaterialXPipelineSettings::ToEnumKey(EInterchangeMaterialXShaders::StandardSurfaceTransmission),
+			{
+				// StandardSurfaceTransmission Inputs
+				TSet<FName>{
+					UE::Interchange::Materials::StandardSurface::Parameters::Base,
+					UE::Interchange::Materials::StandardSurface::Parameters::BaseColor,
+					UE::Interchange::Materials::StandardSurface::Parameters::DiffuseRoughness,
+					UE::Interchange::Materials::StandardSurface::Parameters::Metalness,
+					UE::Interchange::Materials::StandardSurface::Parameters::Specular,
+					bIsSubstrateEnabled ? UE::Interchange::Materials::StandardSurface::Parameters::SpecularColor : FName{},
+					UE::Interchange::Materials::StandardSurface::Parameters::SpecularRoughness,
+					UE::Interchange::Materials::StandardSurface::Parameters::SpecularIOR,
+					UE::Interchange::Materials::StandardSurface::Parameters::SpecularAnisotropy,
+					UE::Interchange::Materials::StandardSurface::Parameters::SpecularRotation,
+					UE::Interchange::Materials::StandardSurface::Parameters::Transmission,
+					UE::Interchange::Materials::StandardSurface::Parameters::TransmissionColor,
+					UE::Interchange::Materials::StandardSurface::Parameters::TransmissionDepth,
+					UE::Interchange::Materials::StandardSurface::Parameters::TransmissionScatter,
+					UE::Interchange::Materials::StandardSurface::Parameters::TransmissionScatterAnisotropy,
+					UE::Interchange::Materials::StandardSurface::Parameters::TransmissionDispersion,
+					UE::Interchange::Materials::StandardSurface::Parameters::TransmissionExtraRoughness,
+					UE::Interchange::Materials::StandardSurface::Parameters::Subsurface,
+					UE::Interchange::Materials::StandardSurface::Parameters::SubsurfaceColor,
+					UE::Interchange::Materials::StandardSurface::Parameters::SubsurfaceRadius,
+					UE::Interchange::Materials::StandardSurface::Parameters::SubsurfaceScale,
+					UE::Interchange::Materials::StandardSurface::Parameters::Sheen,
+					UE::Interchange::Materials::StandardSurface::Parameters::SheenColor,
+					UE::Interchange::Materials::StandardSurface::Parameters::SheenRoughness,
+					UE::Interchange::Materials::StandardSurface::Parameters::Coat,
+					UE::Interchange::Materials::StandardSurface::Parameters::CoatColor,
+					UE::Interchange::Materials::StandardSurface::Parameters::CoatRoughness,
+					bIsSubstrateEnabled ? UE::Interchange::Materials::StandardSurface::Parameters::CoatAnisotropy : FName{},
+					bIsSubstrateEnabled ? UE::Interchange::Materials::StandardSurface::Parameters::CoatRotation : FName{},
+					bIsSubstrateEnabled ? UE::Interchange::Materials::StandardSurface::Parameters::CoatIOR : FName{},
+					UE::Interchange::Materials::StandardSurface::Parameters::CoatNormal,
+					UE::Interchange::Materials::StandardSurface::Parameters::ThinFilmThickness,
+					bIsSubstrateEnabled ? UE::Interchange::Materials::StandardSurface::Parameters::ThinFilmIOR : FName{},
+					UE::Interchange::Materials::StandardSurface::Parameters::Emission,
+					UE::Interchange::Materials::StandardSurface::Parameters::EmissionColor,
+					bIsSubstrateEnabled ? UE::Interchange::Materials::StandardSurface::Parameters::Opacity : FName{},
+					UE::Interchange::Materials::StandardSurface::Parameters::Normal,
+					UE::Interchange::Materials::StandardSurface::Parameters::Tangent
+				},
+				// StandardSurfaceTransmission Outputs
+				!bIsSubstrateEnabled ?
+				TSet<FName>{
+					UE::Interchange::Materials::PBRMR::Parameters::BaseColor,
+					UE::Interchange::Materials::PBRMR::Parameters::Metallic,
+					UE::Interchange::Materials::PBRMR::Parameters::Specular,
+					UE::Interchange::Materials::PBRMR::Parameters::Roughness,
+					UE::Interchange::Materials::PBRMR::Parameters::Anisotropy,
+					UE::Interchange::Materials::PBRMR::Parameters::EmissiveColor,
+					UE::Interchange::Materials::PBRMR::Parameters::Opacity,
+					UE::Interchange::Materials::PBRMR::Parameters::Normal,
+					UE::Interchange::Materials::PBRMR::Parameters::Tangent,
+					UE::Interchange::Materials::PBRMR::Parameters::Refraction,
+					UE::Interchange::Materials::ThinTranslucent::Parameters::TransmissionColor
+				}	:
+				TSet<FName>{UE::Interchange::Materials::StandardSurface::SubstrateMaterial::Outputs::Translucent}
+			}
+		},
+		{
+			UMaterialXPipelineSettings::ToEnumKey(EInterchangeMaterialXShaders::SurfaceUnlit),
+			{
+				// SurfaceUnlit Inputs
+				TSet<FName>{
+					UE::Interchange::Materials::SurfaceUnlit::Parameters::Emission,
+					UE::Interchange::Materials::SurfaceUnlit::Parameters::EmissionColor,
+					UE::Interchange::Materials::SurfaceUnlit::Parameters::Transmission,
+					UE::Interchange::Materials::SurfaceUnlit::Parameters::TransmissionColor,
+					UE::Interchange::Materials::SurfaceUnlit::Parameters::Opacity
+				},
+					// SurfaceUnlit Outputs
+					!bIsSubstrateEnabled ?
+					TSet<FName>{
+						UE::Interchange::Materials::Common::Parameters::EmissiveColor,
+						UE::Interchange::Materials::Common::Parameters::Opacity,
+						UE::Interchange::Materials::SurfaceUnlit::Outputs::OpacityMask
+						} :
+					TSet<FName>{
+						UE::Interchange::Materials::SurfaceUnlit::Substrate::Outputs::OpacityMask,
+						UE::Interchange::Materials::SurfaceUnlit::Substrate::Outputs::SurfaceUnlit
+					}
+				}
+			},
+			{
+				UMaterialXPipelineSettings::ToEnumKey(EInterchangeMaterialXShaders::UsdPreviewSurface),
+				{
+					// UsdPreviewSurface Inputs
+					TSet<FName>{
+						UE::Interchange::Materials::UsdPreviewSurface::Parameters::DiffuseColor,
+						UE::Interchange::Materials::UsdPreviewSurface::Parameters::EmissiveColor,
+						UE::Interchange::Materials::UsdPreviewSurface::Parameters::SpecularColor,
+						UE::Interchange::Materials::UsdPreviewSurface::Parameters::Metallic,
+						UE::Interchange::Materials::UsdPreviewSurface::Parameters::Roughness,
+						UE::Interchange::Materials::UsdPreviewSurface::Parameters::Clearcoat,
+						UE::Interchange::Materials::UsdPreviewSurface::Parameters::ClearcoatRoughness,
+						UE::Interchange::Materials::UsdPreviewSurface::Parameters::Opacity,
+						UE::Interchange::Materials::UsdPreviewSurface::Parameters::OpacityThreshold,
+						UE::Interchange::Materials::UsdPreviewSurface::Parameters::IOR,
+						UE::Interchange::Materials::UsdPreviewSurface::Parameters::Normal,
+						UE::Interchange::Materials::UsdPreviewSurface::Parameters::Displacement,
+						UE::Interchange::Materials::UsdPreviewSurface::Parameters::Occlusion
+					},
+							// UsdPreviewSurface Outputs
+							!bIsSubstrateEnabled ?
+							TSet<FName>{
+								UE::Interchange::Materials::PBRMR::Parameters::BaseColor,
+								UE::Interchange::Materials::PBRMR::Parameters::Metallic,
+								UE::Interchange::Materials::PBRMR::Parameters::Specular,
+								UE::Interchange::Materials::PBRMR::Parameters::Roughness,
+								UE::Interchange::Materials::PBRMR::Parameters::EmissiveColor,
+								UE::Interchange::Materials::PBRMR::Parameters::Opacity,
+								UE::Interchange::Materials::PBRMR::Parameters::Normal,
+								UE::Interchange::Materials::Common::Parameters::Refraction,
+								UE::Interchange::Materials::Common::Parameters::Occlusion,
+								UE::Interchange::Materials::ClearCoat::Parameters::ClearCoat,
+								UE::Interchange::Materials::ClearCoat::Parameters::ClearCoatRoughness,
+							} :
+							TSet<FName>{TEXT("Substrate UsdPreviewSurface"), TEXT("Opacity")}
+						}
+					},
+					{
+					UMaterialXPipelineSettings::ToEnumKey(EInterchangeMaterialXShaders::Surface),
+					{
+							// Surface Inputs
+							TSet<FName>{
+								UE::Interchange::Materials::Surface::Parameters::BSDF,
+								UE::Interchange::Materials::Surface::Parameters::EDF,
+								UE::Interchange::Materials::Surface::Parameters::Opacity
+							},
+								// Surface Outputs
+								bIsSubstrateEnabled ?
+								TSet<FName>{UE::Interchange::Materials::Surface::Outputs::Surface} :
+								TSet<FName>{UE::Interchange::Materials::Surface::Substrate::Outputs::Surface, UE::Interchange::Materials::Surface::Substrate::Outputs::Opacity}
+							}
+						},
+								// BSDF
+								{
+									UMaterialXPipelineSettings::ToEnumKey(EInterchangeMaterialXBSDF::OrenNayarDiffuse),
+									{
+										// OrenNayarDiffuse Inputs
+										TSet<FName>{
+											TEXT("weight"),
+											TEXT("color"),
+											TEXT("roughness"),
+											TEXT("normal")
+										},
+									// OrenNayarDiffuse Outputs
+									TSet<FName>{
+										TEXT("Output")
+									}
+								}
+							},
+							{
+								UMaterialXPipelineSettings::ToEnumKey(EInterchangeMaterialXBSDF::BurleyDiffuse),
+								{
+									// BurleyDiffuse Inputs
+									TSet<FName>{
+										TEXT("weight"),
+										TEXT("color"),
+										TEXT("roughness"),
+										TEXT("normal")
+									},
+									// BurleyDiffuse Outputs
+									TSet<FName>{
+										TEXT("Output")
+									}
+								}
+							},
+							{
+								UMaterialXPipelineSettings::ToEnumKey(EInterchangeMaterialXBSDF::Translucent),
+								{
+									// Translucent Inputs
+									TSet<FName>{
+										TEXT("weight"),
+										TEXT("color"),
+										TEXT("normal")
+									},
+									// Translucent Outputs
+									TSet<FName>{
+										TEXT("Output")
+									}
+								}
+							},
+							{
+								UMaterialXPipelineSettings::ToEnumKey(EInterchangeMaterialXBSDF::Dielectric),
+								{
+									// Dielectric Inputs
+									TSet<FName>{
+										TEXT("weight"),
+										TEXT("tint"),
+										TEXT("ior"),
+										TEXT("roughness"),
+										TEXT("normal"),
+										TEXT("tangent")
+									},
+									// Dielectric Outputs
+									TSet<FName>{
+										TEXT("Output")
+									}
+								}
+							},
+							{
+								UMaterialXPipelineSettings::ToEnumKey(EInterchangeMaterialXBSDF::Conductor),
+								{
+									// Conductor Inputs
+									TSet<FName>{
+										TEXT("weight"),
+										TEXT("ior"),
+										TEXT("extinction"),
+										TEXT("roughness"),
+										TEXT("normal"),
+										TEXT("tangent")
+									},
+									// Conductor Outputs
+									TSet<FName>{
+										TEXT("Output")
+									}
+								}
+							},
+							{
+								UMaterialXPipelineSettings::ToEnumKey(EInterchangeMaterialXBSDF::GeneralizedSchlick),
+								{
+									// GeneralizedSchlick Inputs
+									TSet<FName>{
+										TEXT("weight"),
+										TEXT("color0"),
+										TEXT("color90"),
+										TEXT("exponent"),
+										TEXT("roughness"),
+										TEXT("normal"),
+										TEXT("tangent")
+									},
+									// GeneralizedSchlick Outputs
+									TSet<FName>{
+										TEXT("Output")
+									}
+								}
+							},
+							{
+								UMaterialXPipelineSettings::ToEnumKey(EInterchangeMaterialXBSDF::Subsurface),
+								{
+									// Subsurface Inputs
+									TSet<FName>{
+										TEXT("weight"),
+										TEXT("color"),
+										TEXT("radius"),
+										TEXT("anisotropy"),
+										TEXT("normal")
+									},
+									// Subsurface Outputs
+									TSet<FName>{
+										TEXT("Output")
+									}
+								}
+							},
+							{
+								UMaterialXPipelineSettings::ToEnumKey(EInterchangeMaterialXBSDF::Sheen),
+								{
+									// Sheen Inputs
+									TSet<FName>{
+										TEXT("weight"),
+										TEXT("color"),
+										TEXT("roughness"),
+										TEXT("normal")
+									},
+									// Sheen Outputs
+									TSet<FName>{
+										TEXT("Output")
+									}
+								}
+							},
+							{
+								UMaterialXPipelineSettings::ToEnumKey(EInterchangeMaterialXBSDF::ThinFilm),
+								{
+									// ThinFilm Inputs
+									TSet<FName>{
+										TEXT("thickness"),
+										TEXT("ior")
+									},
+									// ThinFilm Outputs
+									TSet<FName>{
+										TEXT("Output")
+									}
+								}
+							},
+									// EDF
+									{
+										UMaterialXPipelineSettings::ToEnumKey(EInterchangeMaterialXEDF::Uniform),
+										{
+											// Uniform Inputs
+											TSet<FName>{
+												TEXT("color")
+											},
+									// Uniform Outputs
+									TSet<FName>{
+										TEXT("Output")
+									}
+								}
+							},
+							{
+								UMaterialXPipelineSettings::ToEnumKey(EInterchangeMaterialXEDF::Conical),
+								{
+									// Conical Inputs
+									TSet<FName>{
+										TEXT("color"),
+										TEXT("normal"),
+										TEXT("inner_angle"),
+										TEXT("outer_angle"),
+									},
+									// Conical Outputs
+									TSet<FName>{
+										TEXT("Output")
+									}
+								}
+							},
+							{
+								UMaterialXPipelineSettings::ToEnumKey(EInterchangeMaterialXEDF::Measured),
+								{
+									// Measured Inputs
+									TSet<FName>{
+										TEXT("color"),
+										TEXT("normal"),
+										TEXT("file")
+									},
+									// Measured Outputs
+									TSet<FName>{
+										TEXT("Output")
+									}
+								}
+							},
+									// VDF
+									{
+										UMaterialXPipelineSettings::ToEnumKey(EInterchangeMaterialXVDF::Absorption),
+										{
+											// Absorption Inputs
+											TSet<FName>{
+												TEXT("absorption")
+											},
+									// Absorption Outputs
+										TSet<FName>{
+											TEXT("Output")
+										}
+									}
+								},
+								{
+									UMaterialXPipelineSettings::ToEnumKey(EInterchangeMaterialXVDF::Anisotropic),
+									{
+										// Anisotropic Inputs
+										TSet<FName>{
+											TEXT("absorption"),
+											TEXT("scattering"),
+											TEXT("anisotropy"),
+										},
+										// Anisotropic Outputs
+										TSet<FName>{
+											TEXT("Output")
+										}
+									}
+								}
 		};
-		TransmissionSurfaceInputs = {
-			UE::Interchange::Materials::StandardSurface::Parameters::Base,
-			UE::Interchange::Materials::StandardSurface::Parameters::BaseColor,
-			UE::Interchange::Materials::StandardSurface::Parameters::DiffuseRoughness,
-			UE::Interchange::Materials::StandardSurface::Parameters::Metalness,
-			UE::Interchange::Materials::StandardSurface::Parameters::Specular,
-			bIsSubstrateEnabled ? UE::Interchange::Materials::StandardSurface::Parameters::SpecularColor : FName{},
-			UE::Interchange::Materials::StandardSurface::Parameters::SpecularRoughness,
-			UE::Interchange::Materials::StandardSurface::Parameters::SpecularIOR,
-			UE::Interchange::Materials::StandardSurface::Parameters::SpecularAnisotropy,
-			UE::Interchange::Materials::StandardSurface::Parameters::SpecularRotation,
-			UE::Interchange::Materials::StandardSurface::Parameters::Transmission,
-			UE::Interchange::Materials::StandardSurface::Parameters::TransmissionColor,
-			UE::Interchange::Materials::StandardSurface::Parameters::TransmissionDepth,
-			UE::Interchange::Materials::StandardSurface::Parameters::TransmissionScatter,
-			UE::Interchange::Materials::StandardSurface::Parameters::TransmissionScatterAnisotropy,
-			UE::Interchange::Materials::StandardSurface::Parameters::TransmissionDispersion,
-			UE::Interchange::Materials::StandardSurface::Parameters::TransmissionExtraRoughness,
-			UE::Interchange::Materials::StandardSurface::Parameters::Subsurface,
-			UE::Interchange::Materials::StandardSurface::Parameters::SubsurfaceColor,
-			UE::Interchange::Materials::StandardSurface::Parameters::SubsurfaceRadius,
-			UE::Interchange::Materials::StandardSurface::Parameters::SubsurfaceScale,
-			UE::Interchange::Materials::StandardSurface::Parameters::Sheen,
-			UE::Interchange::Materials::StandardSurface::Parameters::SheenColor,
-			UE::Interchange::Materials::StandardSurface::Parameters::SheenRoughness,
-			UE::Interchange::Materials::StandardSurface::Parameters::Coat,
-			UE::Interchange::Materials::StandardSurface::Parameters::CoatColor,
-			UE::Interchange::Materials::StandardSurface::Parameters::CoatRoughness,
-			bIsSubstrateEnabled ? UE::Interchange::Materials::StandardSurface::Parameters::CoatAnisotropy : FName{},
-			bIsSubstrateEnabled ? UE::Interchange::Materials::StandardSurface::Parameters::CoatRotation : FName{},
-			bIsSubstrateEnabled ? UE::Interchange::Materials::StandardSurface::Parameters::CoatIOR : FName{},
-			UE::Interchange::Materials::StandardSurface::Parameters::CoatNormal,
-			UE::Interchange::Materials::StandardSurface::Parameters::ThinFilmThickness,
-			bIsSubstrateEnabled ? UE::Interchange::Materials::StandardSurface::Parameters::ThinFilmIOR : FName{},
-			UE::Interchange::Materials::StandardSurface::Parameters::Emission,
-			UE::Interchange::Materials::StandardSurface::Parameters::EmissionColor,
-			bIsSubstrateEnabled ? UE::Interchange::Materials::StandardSurface::Parameters::Opacity : FName{},
-			UE::Interchange::Materials::StandardSurface::Parameters::Normal,
-			UE::Interchange::Materials::StandardSurface::Parameters::Tangent
-		};
-		TransmissionSurfaceOutputs = {
-			!bIsSubstrateEnabled ?
-			TSet<FName>{
-				UE::Interchange::Materials::PBRMR::Parameters::BaseColor,
-				UE::Interchange::Materials::PBRMR::Parameters::Metallic,
-				UE::Interchange::Materials::PBRMR::Parameters::Specular,
-				UE::Interchange::Materials::PBRMR::Parameters::Roughness,
-				UE::Interchange::Materials::PBRMR::Parameters::Anisotropy,
-				UE::Interchange::Materials::PBRMR::Parameters::EmissiveColor,
-				UE::Interchange::Materials::PBRMR::Parameters::Opacity,
-				UE::Interchange::Materials::PBRMR::Parameters::Normal,
-				UE::Interchange::Materials::PBRMR::Parameters::Tangent,
-				UE::Interchange::Materials::PBRMR::Parameters::Refraction,
-				UE::Interchange::Materials::ThinTranslucent::Parameters::TransmissionColor
-			}	:
-			TSet<FName>{UE::Interchange::Materials::StandardSurface::SubstrateMaterial::Outputs::Translucent}
-		};
-		SurfaceUnlitInputs = {
-			UE::Interchange::Materials::SurfaceUnlit::Parameters::Emission,
-			UE::Interchange::Materials::SurfaceUnlit::Parameters::EmissionColor,
-			UE::Interchange::Materials::SurfaceUnlit::Parameters::Transmission,
-			UE::Interchange::Materials::SurfaceUnlit::Parameters::TransmissionColor,
-			UE::Interchange::Materials::SurfaceUnlit::Parameters ::Opacity
-		};
-		SurfaceUnlitOutputs = {
-			UE::Interchange::Materials::Common::Parameters::EmissiveColor,
-			UE::Interchange::Materials::Common::Parameters::Opacity,
-			TEXT("OpacityMask")
-		};
-		UsdPreviewSurfaceInputs = {
-			UE::Interchange::Materials::UsdPreviewSurface::Parameters::DiffuseColor,
-			UE::Interchange::Materials::UsdPreviewSurface::Parameters::EmissiveColor,
-			UE::Interchange::Materials::UsdPreviewSurface::Parameters::SpecularColor,
-			UE::Interchange::Materials::UsdPreviewSurface::Parameters::Metallic,
-			UE::Interchange::Materials::UsdPreviewSurface::Parameters::Roughness,
-			UE::Interchange::Materials::UsdPreviewSurface::Parameters::Clearcoat,
-			UE::Interchange::Materials::UsdPreviewSurface::Parameters::ClearcoatRoughness,
-			UE::Interchange::Materials::UsdPreviewSurface::Parameters::Opacity,
-			UE::Interchange::Materials::UsdPreviewSurface::Parameters::OpacityThreshold,
-			UE::Interchange::Materials::UsdPreviewSurface::Parameters::IOR,
-			UE::Interchange::Materials::UsdPreviewSurface::Parameters::Normal,
-			UE::Interchange::Materials::UsdPreviewSurface::Parameters::Displacement,
-			UE::Interchange::Materials::UsdPreviewSurface::Parameters::Occlusion
-		};
-		UsdPreviewSurfaceOutputs = {
-			UE::Interchange::Materials::PBRMR::Parameters::BaseColor,
-			UE::Interchange::Materials::PBRMR::Parameters::Metallic,
-			UE::Interchange::Materials::PBRMR::Parameters::Specular,
-			UE::Interchange::Materials::PBRMR::Parameters::Roughness,
-			UE::Interchange::Materials::PBRMR::Parameters::EmissiveColor,
-			UE::Interchange::Materials::PBRMR::Parameters::Opacity,
-			UE::Interchange::Materials::PBRMR::Parameters::Normal,
-			UE::Interchange::Materials::Common::Parameters::Refraction,
-			UE::Interchange::Materials::Common::Parameters::Occlusion,
-			UE::Interchange::Materials::ClearCoat::Parameters::ClearCoat,
-			UE::Interchange::Materials::ClearCoat::Parameters::ClearCoatRoughness,
-		};
-#endif
 	}
+#endif // WITH_EDITOR
 }
 
 UInterchangeMaterialXPipeline::UInterchangeMaterialXPipeline()
-	: MaterialXSettings(UMaterialXPipelineSettings::StaticClass()->GetDefaultObject< UMaterialXPipelineSettings>())
+	: MaterialXSettings(UMaterialXPipelineSettings::StaticClass()->GetDefaultObject<UMaterialXPipelineSettings>())
 {
 	if(HasAnyFlags(EObjectFlags::RF_ClassDefaultObject))
 	{
-		// Init the PredefinedSurfaceShaders with Substrate assets, since the default value is set in BaseInterchange.ini and we have no way in the config file to conditionally init a property
-		if(MaterialXSettings->bIsSubstrateEnabled)
-		{
-			if(FString StandardSurfacePath = MaterialXSettings->GetAssetPathString(EInterchangeMaterialXShaders::StandardSurface); StandardSurfacePath == TEXT("/Interchange/Functions/MX_StandardSurface.MX_StandardSurface"))
-			{
-				MaterialXSettings->PredefinedSurfaceShaders.Add(EInterchangeMaterialXShaders::StandardSurface, FSoftObjectPath{ TEXT("/Interchange/Substrate/MX_StandardSurface.MX_StandardSurface") });
-			}
+#if WITH_EDITOR
+		PathToEnumMapping =	{
+			{MATERIALX_MATERIALFUNCTION_PATH(StandardSurface),		  UMaterialXPipelineSettings::ToEnumKey(EInterchangeMaterialXShaders::StandardSurface)},
+			{MATERIALX_MATERIALFUNCTION_PATH(TransmissionSurface),	  UMaterialXPipelineSettings::ToEnumKey(EInterchangeMaterialXShaders::StandardSurfaceTransmission)},
+			{MATERIALX_MATERIALFUNCTION_PATH(SurfaceUnlit),			  UMaterialXPipelineSettings::ToEnumKey(EInterchangeMaterialXShaders::SurfaceUnlit)},
+			{MATERIALX_MATERIALFUNCTION_PATH(Surface),				  UMaterialXPipelineSettings::ToEnumKey(EInterchangeMaterialXShaders::Surface)},
+			{MATERIALX_MATERIALFUNCTION_PATH(UsdPreviewSurface),	  UMaterialXPipelineSettings::ToEnumKey(EInterchangeMaterialXShaders::UsdPreviewSurface)},
 
-			if(FString TransmissionSurfacePath = MaterialXSettings->GetAssetPathString(EInterchangeMaterialXShaders::StandardSurfaceTransmission); TransmissionSurfacePath == TEXT("/Interchange/Functions/MX_TransmissionSurface.MX_TransmissionSurface"))
-			{
-				MaterialXSettings->PredefinedSurfaceShaders.Add(EInterchangeMaterialXShaders::StandardSurfaceTransmission, FSoftObjectPath{ TEXT("/Interchange/Substrate/MX_TransmissionSurface.MX_TransmissionSurface") });
-			}
-		}
+			{MATERIALX_MATERIALFUNCTION_PATH(OrenNayarDiffuseBSDF),	  UMaterialXPipelineSettings::ToEnumKey(EInterchangeMaterialXBSDF::OrenNayarDiffuse)},
+			{MATERIALX_MATERIALFUNCTION_PATH(BurleyDiffuseBSDF),	  UMaterialXPipelineSettings::ToEnumKey(EInterchangeMaterialXBSDF::BurleyDiffuse)},
+			{MATERIALX_MATERIALFUNCTION_PATH(DieletricBSDF),		  UMaterialXPipelineSettings::ToEnumKey(EInterchangeMaterialXBSDF::Dielectric)},
+			{MATERIALX_MATERIALFUNCTION_PATH(ConductorBSDF),		  UMaterialXPipelineSettings::ToEnumKey(EInterchangeMaterialXBSDF::Conductor)},
+			{MATERIALX_MATERIALFUNCTION_PATH(SheenBSDF),			  UMaterialXPipelineSettings::ToEnumKey(EInterchangeMaterialXBSDF::Sheen)},
+			{MATERIALX_MATERIALFUNCTION_PATH(SubsurfaceBSDF),		  UMaterialXPipelineSettings::ToEnumKey(EInterchangeMaterialXBSDF::Subsurface)},
+			{MATERIALX_MATERIALFUNCTION_PATH(ThinFilmBSDF),			  UMaterialXPipelineSettings::ToEnumKey(EInterchangeMaterialXBSDF::ThinFilm)},
+			{MATERIALX_MATERIALFUNCTION_PATH(GeneralizedSchlickBSDF), UMaterialXPipelineSettings::ToEnumKey(EInterchangeMaterialXBSDF::GeneralizedSchlick)},
+			{MATERIALX_MATERIALFUNCTION_PATH(TranslucentBSDF),		  UMaterialXPipelineSettings::ToEnumKey(EInterchangeMaterialXBSDF::Translucent)},
 
-		PathToEnumMapping =
-		{
-			{MaterialXSettings->bIsSubstrateEnabled ? TEXT("/Interchange/Substrate/MX_StandardSurface.MX_StandardSurface") : TEXT("/Interchange/Functions/MX_StandardSurface.MX_StandardSurface"), EInterchangeMaterialXShaders::StandardSurface},
-			{MaterialXSettings->bIsSubstrateEnabled ? TEXT("/Interchange/Substrate/MX_TransmissionSurface.MX_TransmissionSurface") : TEXT("/Interchange/Functions/MX_TransmissionSurface.MX_TransmissionSurface"), EInterchangeMaterialXShaders::StandardSurfaceTransmission},
-			{TEXT("/Interchange/Functions/MX_SurfaceUnlit.MX_SurfaceUnlit"),EInterchangeMaterialXShaders::SurfaceUnlit},
-			{TEXT("/Interchange/Functions/MX_UsdPreviewSurface.MX_UsdPreviewSurface"),EInterchangeMaterialXShaders::UsdPreviewSurface},
+			{MATERIALX_MATERIALFUNCTION_PATH(UniformEDF),			  UMaterialXPipelineSettings::ToEnumKey(EInterchangeMaterialXEDF::Uniform)},
+			{MATERIALX_MATERIALFUNCTION_PATH(ConicalEDF),			  UMaterialXPipelineSettings::ToEnumKey(EInterchangeMaterialXEDF::Conical)},
+			{MATERIALX_MATERIALFUNCTION_PATH(MeasuredEDF),			  UMaterialXPipelineSettings::ToEnumKey(EInterchangeMaterialXEDF::Measured)},
+
+			{MATERIALX_MATERIALFUNCTION_PATH(AbsorptionVDF),		  UMaterialXPipelineSettings::ToEnumKey(EInterchangeMaterialXVDF::Absorption)},
+			{MATERIALX_MATERIALFUNCTION_PATH(AnisotropicVDF),		  UMaterialXPipelineSettings::ToEnumKey(EInterchangeMaterialXVDF::Anisotropic)},
 		};
+
+		MaterialXSettings->InitPredefinedAssets();
+#endif // WITH_EDITOR
 	}
 
+
 	for (const TPair<EInterchangeMaterialXShaders, FSoftObjectPath>& Entry : MaterialXSettings->PredefinedSurfaceShaders)
 	{
-		PathToEnumMapping.FindOrAdd(Entry.Value.GetAssetPathString(), Entry.Key);
+		PathToEnumMapping.FindOrAdd(Entry.Value.GetAssetPathString(), EMaterialXSettings{ TInPlaceType<EInterchangeMaterialXShaders>{}, Entry.Key });
+	}
+
+	for(const TPair<EInterchangeMaterialXBSDF, FSoftObjectPath>& Entry : MaterialXSettings->PredefinedBSDF)
+	{
+		PathToEnumMapping.FindOrAdd(Entry.Value.GetAssetPathString(), EMaterialXSettings{ TInPlaceType<EInterchangeMaterialXBSDF>{}, Entry.Key });
+	}
+
+	for(const TPair<EInterchangeMaterialXEDF, FSoftObjectPath>& Entry : MaterialXSettings->PredefinedEDF)
+	{
+		PathToEnumMapping.FindOrAdd(Entry.Value.GetAssetPathString(), EMaterialXSettings{ TInPlaceType<EInterchangeMaterialXEDF>{}, Entry.Key });
+	}
+
+	for(const TPair<EInterchangeMaterialXVDF, FSoftObjectPath>& Entry : MaterialXSettings->PredefinedVDF)
+	{
+		PathToEnumMapping.FindOrAdd(Entry.Value.GetAssetPathString(), EMaterialXSettings{ TInPlaceType<EInterchangeMaterialXVDF>{}, Entry.Key });
 	}
 }
 
@@ -244,7 +568,7 @@ void UInterchangeMaterialXPipeline::ExecutePipeline(UInterchangeBaseNodeContaine
 		FString MaterialFunctionPath;
 		if (FactorNode->GetStringAttribute(MaterialFunctionMemberName, MaterialFunctionPath))
 		{
-			if (const EInterchangeMaterialXShaders* EnumPtr = PathToEnumMapping.Find(MaterialFunctionPath))
+			if (const EMaterialXSettings* EnumPtr = PathToEnumMapping.Find(MaterialFunctionPath))
 			{
 				FactorNode->AddStringAttribute(MaterialFunctionMemberName, MaterialXSettings->GetAssetPathString(*EnumPtr));
 			}
@@ -258,21 +582,21 @@ void UInterchangeMaterialXPipeline::ExecutePipeline(UInterchangeBaseNodeContaine
 
 bool UMaterialXPipelineSettings::AreRequiredPackagesLoaded()
 {
-	auto ArePackagesLoaded = [&](const TMap<EInterchangeMaterialXShaders, FSoftObjectPath>& ObjectPaths)
+	auto ArePackagesLoaded = [&](const auto& ObjectPaths)
 	{
 		bool bAllLoaded = true;
 
-		for (const TPair<EInterchangeMaterialXShaders, FSoftObjectPath>& Pair : ObjectPaths)
+		for(const auto& Pair : ObjectPaths)
 		{
-			const FSoftObjectPath& ObjectPath = Pair.Get<1>();
+			const FSoftObjectPath& ObjectPath = Pair.Value;
 
-			if (!ObjectPath.ResolveObject())
+			if(!ObjectPath.ResolveObject())
 			{
 				FString PackagePath = ObjectPath.GetLongPackageName();
-				if (FPackageName::DoesPackageExist(PackagePath))
+				if(FPackageName::DoesPackageExist(PackagePath))
 				{
 					UObject* Asset = ObjectPath.TryLoad();
-					if (!Asset)
+					if(!Asset)
 					{
 						UE_LOG(LogInterchangePipeline, Warning, TEXT("Couldn't load %s"), *PackagePath);
 						bAllLoaded = false;
@@ -280,21 +604,33 @@ bool UMaterialXPipelineSettings::AreRequiredPackagesLoaded()
 #if WITH_EDITOR
 					else
 					{
-						if (Pair.Get<EInterchangeMaterialXShaders>() == EInterchangeMaterialXShaders::StandardSurface)
+
+						using EnumT = decltype(Pair.Key);
+
+						static_assert(std::is_same_v<EnumT, EInterchangeMaterialXShaders> ||
+									  std::is_same_v<EnumT, EInterchangeMaterialXBSDF> ||
+									  std::is_same_v<EnumT, EInterchangeMaterialXEDF> ||
+									  std::is_same_v<EnumT, EInterchangeMaterialXVDF>,
+									  "Enum type not supported");
+
+						uint8 EnumType = IndexSurfaceShaders;
+
+						if constexpr(std::is_same_v<EnumT, EInterchangeMaterialXBSDF>)
 						{
-							bAllLoaded = !ShouldFilterAssets(Cast<UMaterialFunction>(Asset), StandardSurfaceInputs, StandardSurfaceOutputs);
+							EnumType = IndexBSDF;
 						}
-						else if (Pair.Get<EInterchangeMaterialXShaders>() == EInterchangeMaterialXShaders::StandardSurfaceTransmission)
+						else if constexpr(std::is_same_v<EnumT, EInterchangeMaterialXEDF>)
 						{
-							bAllLoaded = !ShouldFilterAssets(Cast<UMaterialFunction>(Asset), TransmissionSurfaceInputs, TransmissionSurfaceOutputs);
+							EnumType = IndexEDF;
 						}
-						else if (Pair.Get<EInterchangeMaterialXShaders>() == EInterchangeMaterialXShaders::SurfaceUnlit)
+						else if constexpr(std::is_same_v<EnumT, EInterchangeMaterialXVDF>)
 						{
-							bAllLoaded = !ShouldFilterAssets(Cast<UMaterialFunction>(Asset), SurfaceUnlitInputs, SurfaceUnlitOutputs);
+							EnumType = IndexVDF;
 						}
-						else if (Pair.Get<EInterchangeMaterialXShaders>() == EInterchangeMaterialXShaders::UsdPreviewSurface)
+
+						if(FMaterialXSettings::ValueType* Settings = SettingsInputsOutputs.Find(ToEnumKey(EnumType, uint8(Pair.Key))))
 						{
-							bAllLoaded = !ShouldFilterAssets(Cast<UMaterialFunction>(Asset), UsdPreviewSurfaceInputs, UsdPreviewSurfaceOutputs);
+							bAllLoaded = !ShouldFilterAssets(Cast<UMaterialFunction>(Asset), Settings->Key, Settings->Value);
 						}
 					}
 #endif // WITH_EDITOR
@@ -310,14 +646,94 @@ bool UMaterialXPipelineSettings::AreRequiredPackagesLoaded()
 		return bAllLoaded;
 	};
 
-	return ArePackagesLoaded(PredefinedSurfaceShaders);
+	return ArePackagesLoaded(PredefinedSurfaceShaders) && ArePackagesLoaded(PredefinedBSDF) && ArePackagesLoaded(PredefinedEDF) && ArePackagesLoaded(PredefinedVDF);
 }
 
-FString UMaterialXPipelineSettings::GetAssetPathString(EInterchangeMaterialXShaders ShaderType) const
+#if WITH_EDITOR
+void UMaterialXPipelineSettings::InitPredefinedAssets()
 {
-	if (const FSoftObjectPath* ObjectPath = PredefinedSurfaceShaders.Find(ShaderType))
+	if(bIsSubstrateEnabled)
 	{
-		return ObjectPath->GetAssetPathString();
+		TArray<TTuple<EMaterialXSettings, FString, FString>> MappingToSubstrate
+		{
+			{UMaterialXPipelineSettings::ToEnumKey(EInterchangeMaterialXShaders::Surface), SurfaceFunctionsPath, SurfaceSubstratePath},
+			{UMaterialXPipelineSettings::ToEnumKey(EInterchangeMaterialXShaders::StandardSurface), StandardSurfaceFunctionsPath, StandardSurfaceSubstratePath},
+			{UMaterialXPipelineSettings::ToEnumKey(EInterchangeMaterialXShaders::StandardSurfaceTransmission), TransmissionSurfaceFunctionsPath, TransmissionSurfaceSubstratePath},
+			{UMaterialXPipelineSettings::ToEnumKey(EInterchangeMaterialXShaders::SurfaceUnlit), SurfaceUnlitFunctionsPath, SurfaceUnlitSubstratePath},
+			{UMaterialXPipelineSettings::ToEnumKey(EInterchangeMaterialXShaders::UsdPreviewSurface), UsdPreviewSurfaceFunctionsPath, UsdPreviewSurfaceSubstratePath},
+
+			{UMaterialXPipelineSettings::ToEnumKey(EInterchangeMaterialXBSDF::OrenNayarDiffuse), OrenNayarDiffuseBSDFFunctionsPath, OrenNayarDiffuseBSDFSubstratePath},
+			{UMaterialXPipelineSettings::ToEnumKey(EInterchangeMaterialXBSDF::BurleyDiffuse), BurleyDiffuseBSDFFunctionsPath, BurleyDiffuseBSDFSubstratePath},
+			{UMaterialXPipelineSettings::ToEnumKey(EInterchangeMaterialXBSDF::Dielectric), DieletricBSDFFunctionsPath, DieletricBSDFSubstratePath},
+			{UMaterialXPipelineSettings::ToEnumKey(EInterchangeMaterialXBSDF::Conductor), ConductorBSDFFunctionsPath, ConductorBSDFSubstratePath},
+			{UMaterialXPipelineSettings::ToEnumKey(EInterchangeMaterialXBSDF::Sheen), SheenBSDFFunctionsPath, SheenBSDFSubstratePath},
+			{UMaterialXPipelineSettings::ToEnumKey(EInterchangeMaterialXBSDF::Subsurface), SubsurfaceBSDFFunctionsPath, SubsurfaceBSDFSubstratePath},
+			{UMaterialXPipelineSettings::ToEnumKey(EInterchangeMaterialXBSDF::ThinFilm), ThinFilmBSDFFunctionsPath, ThinFilmBSDFSubstratePath},
+			{UMaterialXPipelineSettings::ToEnumKey(EInterchangeMaterialXBSDF::GeneralizedSchlick), GeneralizedSchlickBSDFFunctionsPath, GeneralizedSchlickBSDFSubstratePath},
+			{UMaterialXPipelineSettings::ToEnumKey(EInterchangeMaterialXBSDF::Translucent), TranslucentBSDFFunctionsPath, TranslucentBSDFSubstratePath},
+
+			{UMaterialXPipelineSettings::ToEnumKey(EInterchangeMaterialXEDF::Uniform), UniformEDFFunctionsPath, UniformEDFSubstratePath},
+			{UMaterialXPipelineSettings::ToEnumKey(EInterchangeMaterialXEDF::Conical), ConicalEDFFunctionsPath, ConicalEDFSubstratePath},
+			{UMaterialXPipelineSettings::ToEnumKey(EInterchangeMaterialXEDF::Measured), MeasuredEDFFunctionsPath, MeasuredEDFSubstratePath},
+
+			{UMaterialXPipelineSettings::ToEnumKey(EInterchangeMaterialXVDF::Absorption), AbsorptionVDFFunctionsPath, AbsorptionVDFSubstratePath},
+			{UMaterialXPipelineSettings::ToEnumKey(EInterchangeMaterialXVDF::Anisotropic), AnisotropicVDFFunctionsPath, AnisotropicVDFSubstratePath},
+		};
+
+		for(auto Mapping : MappingToSubstrate)
+		{
+			const EMaterialXSettings& ShadersSettings = Mapping.Get<0>();
+			if(FString Path = GetAssetPathString(ShadersSettings); Path == Mapping.Get<1>())
+			{
+				if(ShadersSettings.IsType<EInterchangeMaterialXShaders>())
+				{
+					PredefinedSurfaceShaders.Add(ShadersSettings.Get<EInterchangeMaterialXShaders>(), FSoftObjectPath{ Mapping.Get<2>() });
+				}
+				else if(ShadersSettings.IsType<EInterchangeMaterialXBSDF>())
+				{
+					PredefinedBSDF.Add(ShadersSettings.Get<EInterchangeMaterialXBSDF>(), FSoftObjectPath{ Mapping.Get<2>() });
+				}
+				else if(ShadersSettings.IsType<EInterchangeMaterialXEDF>())
+				{
+					PredefinedEDF.Add(ShadersSettings.Get<EInterchangeMaterialXEDF>(), FSoftObjectPath{ Mapping.Get<2>() });
+				}
+				else if(ShadersSettings.IsType<EInterchangeMaterialXVDF>())
+				{
+					PredefinedVDF.Add(ShadersSettings.Get<EInterchangeMaterialXVDF>(), FSoftObjectPath{ Mapping.Get<2>() });
+				}
+			}
+		}
+	}
+}
+#endif // WITH_EDITOR
+
+FString UMaterialXPipelineSettings::GetAssetPathString(EMaterialXSettings EnumValue) const
+{
+	auto FindAssetPathString = [](const auto& PredefinedEnumPath, auto Enum) -> FString
+	{
+		if(const FSoftObjectPath* ObjectPath = PredefinedEnumPath.Find(Enum))
+		{
+			return ObjectPath->GetAssetPathString();
+		}
+
+		return {};
+	};
+
+	if(EnumValue.IsType<EInterchangeMaterialXShaders>())
+	{
+		return FindAssetPathString(PredefinedSurfaceShaders, EnumValue.Get<EInterchangeMaterialXShaders>());
+	}
+	else if(EnumValue.IsType<EInterchangeMaterialXBSDF>())
+	{
+		return FindAssetPathString(PredefinedBSDF, EnumValue.Get<EInterchangeMaterialXBSDF>());
+	}
+	else if(EnumValue.IsType<EInterchangeMaterialXEDF>())
+	{
+		return FindAssetPathString(PredefinedEDF, EnumValue.Get<EInterchangeMaterialXEDF>());
+	}
+	else if(EnumValue.IsType<EInterchangeMaterialXVDF>())
+	{
+		return FindAssetPathString(PredefinedVDF, EnumValue.Get<EInterchangeMaterialXVDF>());
 	}
 
 	return {};
@@ -355,5 +771,56 @@ bool UMaterialXPipelineSettings::ShouldFilterAssets(UMaterialFunction* Asset, co
 	// we allow at least one input of the same name, but we should have exactly the same outputs
 	return !(InputMatches > 0 && OutputMatches == Outputs.Num());
 }
+
+EMaterialXSettings UMaterialXPipelineSettings::ToEnumKey(uint8 EnumType, uint8 EnumValue)
+{
+	switch(EnumType)
+	{
+	case IndexBSDF:
+		return EMaterialXSettings{ TInPlaceType<EInterchangeMaterialXBSDF>{}, EInterchangeMaterialXBSDF{EnumValue} };
+
+	case IndexEDF:
+		return EMaterialXSettings{ TInPlaceType<EInterchangeMaterialXEDF>{}, EInterchangeMaterialXEDF{EnumValue} };
+
+	case IndexVDF:
+		return EMaterialXSettings{ TInPlaceType<EInterchangeMaterialXVDF>{}, EInterchangeMaterialXVDF{EnumValue} };
+
+	default:
+		return EMaterialXSettings{ TInPlaceType<EInterchangeMaterialXShaders>{}, EInterchangeMaterialXShaders{EnumValue} };
+	}
+}
 #endif // WITH_EDITOR
 
+namespace
+{
+	static uint8* GetMaterialXSettingsIndexValue(EMaterialXSettings Enum, SIZE_T& Index)
+	{
+		Index = Enum.GetIndex();
+		return
+			Index == 1 ? reinterpret_cast<uint8*>(Enum.TryGet<EInterchangeMaterialXBSDF>()) :
+			Index == 2 ? reinterpret_cast<uint8*>(Enum.TryGet<EInterchangeMaterialXEDF>()) :
+			Index == 3 ? reinterpret_cast<uint8*>(Enum.TryGet<EInterchangeMaterialXVDF>()) :
+			reinterpret_cast<uint8*>(Enum.TryGet<EInterchangeMaterialXShaders>());
+	}
+}
+
+uint32 GetTypeHash(EMaterialXSettings Key)
+{
+	SIZE_T Index;
+	uint8* UnderlyingValue = GetMaterialXSettingsIndexValue(Key, Index);
+	return HashCombine(Index, *UnderlyingValue);
+}
+
+bool operator==(EMaterialXSettings Lhs, EMaterialXSettings Rhs)
+{
+	SIZE_T LhsIndex;
+	uint8* LhsUnderlyingValue = GetMaterialXSettingsIndexValue(Lhs, LhsIndex);
+
+	SIZE_T RhsIndex;
+	uint8* RhsUnderlyingValue = GetMaterialXSettingsIndexValue(Rhs, RhsIndex);
+
+	return LhsIndex==RhsIndex && *LhsUnderlyingValue == *RhsUnderlyingValue;
+}
+
+#undef MATERIALX_FUNCTIONS_SUBSTRATE_PATH
+#undef MATERIALX_MATERIALFUNCTION_PATH
\ No newline at end of file
diff --git a/Engine/Plugins/Interchange/Runtime/Source/Pipelines/Public/InterchangeMaterialXPipeline.h b/Engine/Plugins/Interchange/Runtime/Source/Pipelines/Public/InterchangeMaterialXPipeline.h
--- a/Engine/Plugins/Interchange/Runtime/Source/Pipelines/Public/InterchangeMaterialXPipeline.h
+++ b/Engine/Plugins/Interchange/Runtime/Source/Pipelines/Public/InterchangeMaterialXPipeline.h
@@ -32,18 +32,92 @@ enum class EInterchangeMaterialXShaders : uint8
 	/** Default settings for Autodesk's Standard Surface shader	*/
 	StandardSurface,
 
-	/** Standard Surface shader used for translucency	*/
+	/** Standard Surface shader used for translucency */
 	StandardSurfaceTransmission,
 
 	/** Shader used for unlit surface*/
 	SurfaceUnlit,
 
-	/** Default settings for USD's Surface shader	*/
+	/** Default settings for USD's Surface shader */
 	UsdPreviewSurface,
 
-	MaxShaderCount
+	/** Construct a surface shader from scattering and emission distribution functions.*/
+	Surface,
+
+	MaxShaderCount UMETA(hidden)
+};
+
+#define INTERCHANGE_MATERIALX_EXPERIEMENTAL
+
+UENUM(BlueprintType)
+/** Data type representing a Bidirectional Scattering Distribution Function. */
+enum class EInterchangeMaterialXBSDF : uint8
+{
+	/** A BSDF node for diffuse reflections. */
+	OrenNayarDiffuse,
+
+	/** A BSDF node for Burley diffuse reflections. */
+	BurleyDiffuse,
+
+	/** A BSDF node for pure diffuse transmission. */
+	Translucent,
+
+	/** A reflection/transmission BSDF node based on a microfacet model and a Fresnel curve for dielectrics. */
+	Dielectric,
+
+	/** A reflection BSDF node based on a microfacet model and a Fresnel curve for conductors/metals. */
+	Conductor,
+
+	/** A reflection/transmission BSDF node based on a microfacet model and a generalized Schlick Fresnel curve. */
+	GeneralizedSchlick,
+
+	/** A subsurface scattering BSDF for true subsurface scattering. */
+	Subsurface,
+
+	/** A microfacet BSDF for the back-scattering properties of cloth-like materials. */
+	Sheen,
+
+	/** Adds an iridescent thin film layer over a microfacet base BSDF. */
+	ThinFilm,
+
+	MaxBSDFCount UMETA(hidden)
+};
+
+UENUM(BlueprintType)
+/** Data type representing an Emission Distribution Function. */
+enum class EInterchangeMaterialXEDF : uint8
+{
+	/** An EDF node for uniform emission. */
+	Uniform,
+
+	/** Constructs an EDF emitting light inside a cone around the normal direction. */
+	Conical,
+
+	/** Constructs an EDF emitting light according to a measured IES light profile. */
+	Measured,
+
+	MaxEDFCount UMETA(hidden)
+};
+
+UENUM(BlueprintType)
+/** Data type representing a Volume Distribution Function. */
+enum class EInterchangeMaterialXVDF : uint8
+{
+	/** Constructs a VDF for pure light absorption. */
+	Absorption,
+
+	/** Constructs a VDF scattering light for a participating medium, based on the Henyey-Greenstein phase function. */
+	Anisotropic,
+
+	MaxVDFCount UMETA(hidden)
 };
 
+using EMaterialXSettings = TVariant<EInterchangeMaterialXShaders, EInterchangeMaterialXBSDF, EInterchangeMaterialXEDF, EInterchangeMaterialXVDF>;
+
+uint32 INTERCHANGEPIPELINES_API GetTypeHash(EMaterialXSettings Key);
+
+bool INTERCHANGEPIPELINES_API operator==(EMaterialXSettings Lhs, EMaterialXSettings Rhs);
+
 UCLASS(config = Interchange, meta = (DisplayName = "Interchange MaterialX"))
 class INTERCHANGEPIPELINES_API UMaterialXPipelineSettings : public UDeveloperSettings
 {
@@ -52,31 +126,69 @@ class INTERCHANGEPIPELINES_API UMaterialXPipelineSettings : public UDeveloperSet
 public:
 	UMaterialXPipelineSettings();
 
-	UPROPERTY(EditAnywhere, config, Category = "MaterialXPredefined", meta = (DisplayName = "MaterialX Predefined Surface Shaders"))
+	bool AreRequiredPackagesLoaded();
+
+	FString GetAssetPathString(EMaterialXSettings EnumType) const;
+
+	template<typename EnumT>
+	FString GetAssetPathString(EnumT EnumValue) const
+	{
+		static_assert(std::is_same_v<EnumT, EInterchangeMaterialXShaders> ||
+					  std::is_same_v<EnumT, EInterchangeMaterialXBSDF> ||
+					  std::is_same_v<EnumT, EInterchangeMaterialXEDF> ||
+					  std::is_same_v<EnumT, EInterchangeMaterialXVDF>,
+					  "Enum type not supported");
+
+		return GetAssetPathString(EMaterialXSettings{ TInPlaceType<EnumT>{}, EnumValue });
+	}
+
+	UPROPERTY(EditAnywhere, config, Category = "MaterialXPredefined | Surface Shaders", meta = (DisplayName = "MaterialX Predefined Surface Shaders"))
 	TMap<EInterchangeMaterialXShaders, FSoftObjectPath> PredefinedSurfaceShaders;
 
-	bool AreRequiredPackagesLoaded();
+	UPROPERTY(EditAnywhere, config, Category = "MaterialXPredefined | BSDF", meta = (DisplayName = "MaterialX Predefined BSDF"))
+	TMap<EInterchangeMaterialXBSDF, FSoftObjectPath> PredefinedBSDF;
 
-	FString GetAssetPathString(EInterchangeMaterialXShaders ShaderType) const;
+	UPROPERTY(EditAnywhere, config, Category = "MaterialXPredefined | EDF", meta = (DisplayName = "MaterialX Predefined EDF"))
+	TMap<EInterchangeMaterialXEDF, FSoftObjectPath> PredefinedEDF;
 
-private:
+	UPROPERTY(EditAnywhere, config, Category = "MaterialXPredefined | VDF", meta = (DisplayName = "MaterialX Predefined VDF"))
+	TMap<EInterchangeMaterialXVDF, FSoftObjectPath> PredefinedVDF;
+
+#if WITH_EDITOR
+	/** Init the Predefined with Substrate assets, since the default value is set in BaseInterchange.ini and we have no way in the config file to conditionally init a property*/
+	void InitPredefinedAssets();
 
+private:
 	friend class FInterchangeMaterialXPipelineSettingsCustomization;
 	friend class UInterchangeMaterialXPipeline;
 
-	bool bIsSubstrateEnabled{ false };
+	using FMaterialXSettings = TMap<EMaterialXSettings, TPair<TSet<FName>, TSet<FName>>>;
 
-#if WITH_EDITOR
 	static bool ShouldFilterAssets(UMaterialFunction* Asset, const TSet<FName>& Inputs, const TSet<FName>& Outputs);
 
-	static TSet<FName> StandardSurfaceInputs;
-	static TSet<FName> StandardSurfaceOutputs;
-	static TSet<FName> TransmissionSurfaceInputs;
-	static TSet<FName> TransmissionSurfaceOutputs;
-	static TSet<FName> SurfaceUnlitInputs;
-	static TSet<FName> SurfaceUnlitOutputs;
-	static TSet<FName> UsdPreviewSurfaceInputs;
-	static TSet<FName> UsdPreviewSurfaceOutputs;
+	static EMaterialXSettings ToEnumKey(uint8 EnumType, uint8 EnumValue);
+
+	template<typename EnumT>
+	static EMaterialXSettings ToEnumKey(EnumT EnumValue)
+	{
+		static_assert(std::is_same_v<EnumT, EInterchangeMaterialXShaders> ||
+					  std::is_same_v<EnumT, EInterchangeMaterialXBSDF> ||
+					  std::is_same_v<EnumT, EInterchangeMaterialXEDF> ||
+					  std::is_same_v<EnumT, EInterchangeMaterialXVDF>,
+					  "Enum type not supported");
+
+		return EMaterialXSettings{ TInPlaceType<EnumT>{}, EnumValue };
+	}
+
+	/** The key is a combination of the index in the variant + the corresponding enum */
+	static FMaterialXSettings SettingsInputsOutputs;
+
+	static constexpr uint8 IndexSurfaceShaders = 0;
+	static constexpr uint8 IndexBSDF = 1;
+	static constexpr uint8 IndexEDF = 2;
+	static constexpr uint8 IndexVDF = 3;
+
+	bool bIsSubstrateEnabled{ false };
 #endif // WITH_EDITOR
 };
 
@@ -88,7 +200,6 @@ class INTERCHANGEPIPELINES_API UInterchangeMaterialXPipeline : public UInterchan
 	UInterchangeMaterialXPipeline();
 
 public:
-	
 	TObjectPtr<UMaterialXPipelineSettings> MaterialXSettings;
 
 protected:
@@ -102,5 +213,6 @@ protected:
 	}
 
 private:
-	static TMap<FString, EInterchangeMaterialXShaders> PathToEnumMapping;
-};
\ No newline at end of file
+
+	static TMap<FString, EMaterialXSettings> PathToEnumMapping;
+};
