Fix wrong typo of MX_DielectricBSDF Material Function

diff --git a/Engine/Plugins/Interchange/Runtime/Source/Pipelines/Private/InterchangeMaterialXPipeline.cpp b/Engine/Plugins/Interchange/Runtime/Source/Pipelines/Private/InterchangeMaterialXPipeline.cpp
--- a/Engine/Plugins/Interchange/Runtime/Source/Pipelines/Private/InterchangeMaterialXPipeline.cpp
+++ b/Engine/Plugins/Interchange/Runtime/Source/Pipelines/Private/InterchangeMaterialXPipeline.cpp
@@ -12,37 +12,37 @@
 #include "Misc/PackageName.h"
 
 
-#define MATERIALX_FUNCTIONS_SUBSTRATE_PATH(Name, AssetName)  \
-	constexpr const TCHAR* Name##FunctionsPath = TEXT("/Interchange/Functions/") TEXT(AssetName) TEXT(".") TEXT(AssetName);  \
-	constexpr const TCHAR* Name##SubstratePath = TEXT("/Interchange/Substrate/") TEXT(AssetName) TEXT(".") TEXT(AssetName);
+#define MATERIALX_FUNCTIONS_SUBSTRATE_PATH(Name)  \
+	constexpr const TCHAR* Name##FunctionsPath = TEXT("/Interchange/Functions/") TEXT("MX_") TEXT(#Name) TEXT(".") TEXT("MX_") TEXT(#Name);  \
+	constexpr const TCHAR* Name##SubstratePath = TEXT("/Interchange/Substrate/") TEXT("MX_") TEXT(#Name) TEXT(".") TEXT("MX_") TEXT(#Name);
 
 #define MATERIALX_MATERIALFUNCTION_PATH(Name) \
 	MaterialXSettings->bIsSubstrateEnabled ? Name##FunctionsPath : Name##SubstratePath
 
 namespace
 {	
-	MATERIALX_FUNCTIONS_SUBSTRATE_PATH(StandardSurface, "MX_StandardSurface");
-	MATERIALX_FUNCTIONS_SUBSTRATE_PATH(TransmissionSurface, "MX_TransmissionSurface");
-	MATERIALX_FUNCTIONS_SUBSTRATE_PATH(SurfaceUnlit, "MX_SurfaceUnlit");
-	MATERIALX_FUNCTIONS_SUBSTRATE_PATH(Surface, "MX_Surface");
-	MATERIALX_FUNCTIONS_SUBSTRATE_PATH(UsdPreviewSurface, "MX_UsdPreviewSurface");
+	MATERIALX_FUNCTIONS_SUBSTRATE_PATH(StandardSurface);
+	MATERIALX_FUNCTIONS_SUBSTRATE_PATH(TransmissionSurface);
+	MATERIALX_FUNCTIONS_SUBSTRATE_PATH(SurfaceUnlit);
+	MATERIALX_FUNCTIONS_SUBSTRATE_PATH(Surface);
+	MATERIALX_FUNCTIONS_SUBSTRATE_PATH(UsdPreviewSurface);
 	
-	MATERIALX_FUNCTIONS_SUBSTRATE_PATH(OrenNayarDiffuseBSDF, "MX_OrenNayarBSDF");
-	MATERIALX_FUNCTIONS_SUBSTRATE_PATH(BurleyDiffuseBSDF, "MX_BurleyDiffuseBSDF");
-	MATERIALX_FUNCTIONS_SUBSTRATE_PATH(DieletricBSDF, ".MX_DielectricBSDF");
-	MATERIALX_FUNCTIONS_SUBSTRATE_PATH(ConductorBSDF, "MX_ConductorBSDF");
-	MATERIALX_FUNCTIONS_SUBSTRATE_PATH(SheenBSDF, "MX_SheenBSDF");
-	MATERIALX_FUNCTIONS_SUBSTRATE_PATH(SubsurfaceBSDF, "MX_SubsurfaceBSDF");
-	MATERIALX_FUNCTIONS_SUBSTRATE_PATH(ThinFilmBSDF, "MX_ThinFilmBSDF");
-	MATERIALX_FUNCTIONS_SUBSTRATE_PATH(GeneralizedSchlickBSDF, "MX_GeneralizedSchlickBSDF");
-	MATERIALX_FUNCTIONS_SUBSTRATE_PATH(TranslucentBSDF, "MX_TranslucentBSDF");
+	MATERIALX_FUNCTIONS_SUBSTRATE_PATH(OrenNayarBSDF);
+	MATERIALX_FUNCTIONS_SUBSTRATE_PATH(BurleyDiffuseBSDF);
+	MATERIALX_FUNCTIONS_SUBSTRATE_PATH(DieletricBSDF);
+	MATERIALX_FUNCTIONS_SUBSTRATE_PATH(ConductorBSDF);
+	MATERIALX_FUNCTIONS_SUBSTRATE_PATH(SheenBSDF);
+	MATERIALX_FUNCTIONS_SUBSTRATE_PATH(SubsurfaceBSDF);
+	MATERIALX_FUNCTIONS_SUBSTRATE_PATH(ThinFilmBSDF);
+	MATERIALX_FUNCTIONS_SUBSTRATE_PATH(GeneralizedSchlickBSDF);
+	MATERIALX_FUNCTIONS_SUBSTRATE_PATH(TranslucentBSDF);
 	
-	MATERIALX_FUNCTIONS_SUBSTRATE_PATH(UniformEDF, "MX_UniformEDF");
-	MATERIALX_FUNCTIONS_SUBSTRATE_PATH(ConicalEDF, "MX_ConicalEDF");
-	MATERIALX_FUNCTIONS_SUBSTRATE_PATH(MeasuredEDF, "MX_MeasuredEDF");
+	MATERIALX_FUNCTIONS_SUBSTRATE_PATH(UniformEDF);
+	MATERIALX_FUNCTIONS_SUBSTRATE_PATH(ConicalEDF);
+	MATERIALX_FUNCTIONS_SUBSTRATE_PATH(MeasuredEDF);
 	
-	MATERIALX_FUNCTIONS_SUBSTRATE_PATH(AbsorptionVDF, "MX_AbsorptionVDF");
-	MATERIALX_FUNCTIONS_SUBSTRATE_PATH(AnisotropicVDF, "MX_AnisotropicVDF");
+	MATERIALX_FUNCTIONS_SUBSTRATE_PATH(AbsorptionVDF);
+	MATERIALX_FUNCTIONS_SUBSTRATE_PATH(AnisotropicVDF);
 }
 
 TMap<FString, EMaterialXSettings> UInterchangeMaterialXPipeline::PathToEnumMapping;
@@ -502,7 +502,7 @@ UInterchangeMaterialXPipeline::UInterchangeMaterialXPipeline()
 			{MATERIALX_MATERIALFUNCTION_PATH(Surface),				  UMaterialXPipelineSettings::ToEnumKey(EInterchangeMaterialXShaders::Surface)},
 			{MATERIALX_MATERIALFUNCTION_PATH(UsdPreviewSurface),	  UMaterialXPipelineSettings::ToEnumKey(EInterchangeMaterialXShaders::UsdPreviewSurface)},
 
-			{MATERIALX_MATERIALFUNCTION_PATH(OrenNayarDiffuseBSDF),	  UMaterialXPipelineSettings::ToEnumKey(EInterchangeMaterialXBSDF::OrenNayarDiffuse)},
+			{MATERIALX_MATERIALFUNCTION_PATH(OrenNayarBSDF),	      UMaterialXPipelineSettings::ToEnumKey(EInterchangeMaterialXBSDF::OrenNayarDiffuse)},
 			{MATERIALX_MATERIALFUNCTION_PATH(BurleyDiffuseBSDF),	  UMaterialXPipelineSettings::ToEnumKey(EInterchangeMaterialXBSDF::BurleyDiffuse)},
 			{MATERIALX_MATERIALFUNCTION_PATH(DieletricBSDF),		  UMaterialXPipelineSettings::ToEnumKey(EInterchangeMaterialXBSDF::Dielectric)},
 			{MATERIALX_MATERIALFUNCTION_PATH(ConductorBSDF),		  UMaterialXPipelineSettings::ToEnumKey(EInterchangeMaterialXBSDF::Conductor)},
@@ -662,7 +662,7 @@ void UMaterialXPipelineSettings::InitPredefinedAssets()
 			{UMaterialXPipelineSettings::ToEnumKey(EInterchangeMaterialXShaders::SurfaceUnlit), SurfaceUnlitFunctionsPath, SurfaceUnlitSubstratePath},
 			{UMaterialXPipelineSettings::ToEnumKey(EInterchangeMaterialXShaders::UsdPreviewSurface), UsdPreviewSurfaceFunctionsPath, UsdPreviewSurfaceSubstratePath},
 
-			{UMaterialXPipelineSettings::ToEnumKey(EInterchangeMaterialXBSDF::OrenNayarDiffuse), OrenNayarDiffuseBSDFFunctionsPath, OrenNayarDiffuseBSDFSubstratePath},
+			{UMaterialXPipelineSettings::ToEnumKey(EInterchangeMaterialXBSDF::OrenNayarDiffuse), OrenNayarBSDFFunctionsPath, OrenNayarBSDFSubstratePath},
 			{UMaterialXPipelineSettings::ToEnumKey(EInterchangeMaterialXBSDF::BurleyDiffuse), BurleyDiffuseBSDFFunctionsPath, BurleyDiffuseBSDFSubstratePath},
 			{UMaterialXPipelineSettings::ToEnumKey(EInterchangeMaterialXBSDF::Dielectric), DieletricBSDFFunctionsPath, DieletricBSDFSubstratePath},
 			{UMaterialXPipelineSettings::ToEnumKey(EInterchangeMaterialXBSDF::Conductor), ConductorBSDFFunctionsPath, ConductorBSDFSubstratePath},
