[Interchange - glTF] khr_materials_specular processing fixes

diff --git a/Engine/Plugins/Interchange/Runtime/Source/Import/Private/Gltf/InterchangeGLTFMaterial.cpp b/Engine/Plugins/Interchange/Runtime/Source/Import/Private/Gltf/InterchangeGLTFMaterial.cpp
--- a/Engine/Plugins/Interchange/Runtime/Source/Import/Private/Gltf/InterchangeGLTFMaterial.cpp
+++ b/Engine/Plugins/Interchange/Runtime/Source/Import/Private/Gltf/InterchangeGLTFMaterial.cpp
@@ -506,7 +506,7 @@ namespace UE::Interchange::GLTFMaterials
 					SetMap(GltfMaterial.Specular.SpecularMap, Inputs::SpecularTexture);
 
 					//SpecularFactor
-					SetScalar(Inputs::SpecularFactor, GltfMaterial.Specular.SpecularFactor, 1.f);
+					SetScalar(Inputs::SpecularFactor, GltfMaterial.Specular.SpecularFactor, 0.5); //(UE Specular default is 0.5).
 				}
 			}
 			else
