- Fix wrong tabulations in MaterialXPipeline
- Fix Functions/MX_Surface material function uncorrectly filtered due to output named "Output" instead of "Surface"

diff --git a/Engine/Plugins/Interchange/Runtime/Source/Pipelines/Private/InterchangeMaterialXPipeline.cpp b/Engine/Plugins/Interchange/Runtime/Source/Pipelines/Private/InterchangeMaterialXPipeline.cpp
--- a/Engine/Plugins/Interchange/Runtime/Source/Pipelines/Private/InterchangeMaterialXPipeline.cpp
+++ b/Engine/Plugins/Interchange/Runtime/Source/Pipelines/Private/InterchangeMaterialXPipeline.cpp
@@ -62,139 +62,139 @@ UMaterialXPipelineSettings::UMaterialXPipelineSettings()
 			{
 				UMaterialXPipelineSettings::ToEnumKey(EInterchangeMaterialXShaders::StandardSurface),
 				{
-				// StandardSurface Inputs
-				TSet<FName>{
-					UE::Interchange::Materials::StandardSurface::Parameters::Base,
-					UE::Interchange::Materials::StandardSurface::Parameters::BaseColor,
-					UE::Interchange::Materials::StandardSurface::Parameters::DiffuseRoughness,
-					UE::Interchange::Materials::StandardSurface::Parameters::Metalness,
-					UE::Interchange::Materials::StandardSurface::Parameters::Specular,
-					UE::Interchange::Materials::StandardSurface::Parameters::SpecularRoughness,
-					UE::Interchange::Materials::StandardSurface::Parameters::SpecularIOR,
-					UE::Interchange::Materials::StandardSurface::Parameters::SpecularAnisotropy,
-					UE::Interchange::Materials::StandardSurface::Parameters::SpecularRotation,
-					UE::Interchange::Materials::StandardSurface::Parameters::Subsurface,
-					UE::Interchange::Materials::StandardSurface::Parameters::SubsurfaceColor,
-					UE::Interchange::Materials::StandardSurface::Parameters::SubsurfaceRadius,
-					UE::Interchange::Materials::StandardSurface::Parameters::SubsurfaceScale,
-					UE::Interchange::Materials::StandardSurface::Parameters::Sheen,
-					UE::Interchange::Materials::StandardSurface::Parameters::SheenColor,
-					UE::Interchange::Materials::StandardSurface::Parameters::SheenRoughness,
-					UE::Interchange::Materials::StandardSurface::Parameters::Coat,
-					UE::Interchange::Materials::StandardSurface::Parameters::CoatColor,
-					UE::Interchange::Materials::StandardSurface::Parameters::CoatRoughness,
-					UE::Interchange::Materials::StandardSurface::Parameters::CoatNormal,
-					UE::Interchange::Materials::StandardSurface::Parameters::ThinFilmThickness,
-					UE::Interchange::Materials::StandardSurface::Parameters::Emission,
-					UE::Interchange::Materials::StandardSurface::Parameters::EmissionColor,
-					UE::Interchange::Materials::StandardSurface::Parameters::Normal,
-					UE::Interchange::Materials::StandardSurface::Parameters::Tangent,
-					bIsSubstrateEnabled ? UE::Interchange::Materials::StandardSurface::Parameters::SpecularColor : FName{},
-					bIsSubstrateEnabled ? UE::Interchange::Materials::StandardSurface::Parameters::CoatIOR : FName{},
-					bIsSubstrateEnabled ? UE::Interchange::Materials::StandardSurface::Parameters::CoatAnisotropy : FName{},
-					bIsSubstrateEnabled ? UE::Interchange::Materials::StandardSurface::Parameters::CoatRotation : FName{},
-					bIsSubstrateEnabled ? UE::Interchange::Materials::StandardSurface::Parameters::ThinFilmIOR : FName{},
-					bIsSubstrateEnabled ? UE::Interchange::Materials::StandardSurface::Parameters::Opacity : FName{},
-					bIsSubstrateEnabled ? UE::Interchange::Materials::StandardSurface::Parameters::SubsurfaceAnisotropy : FName{}
-				},
-			// StandardSurface Outputs
-			!bIsSubstrateEnabled ?
-			TSet<FName>{
-				TEXT("Base Color"), // MX_StandardSurface has BaseColor with a whitespace, this should be fixed in further release
-				UE::Interchange::Materials::PBRMR::Parameters::Metallic,
-				UE::Interchange::Materials::PBRMR::Parameters::Specular,
-				UE::Interchange::Materials::PBRMR::Parameters::Roughness,
-				UE::Interchange::Materials::PBRMR::Parameters::Anisotropy,
-				UE::Interchange::Materials::PBRMR::Parameters::EmissiveColor,
-				UE::Interchange::Materials::PBRMR::Parameters::Opacity,
-				UE::Interchange::Materials::PBRMR::Parameters::Normal,
-				UE::Interchange::Materials::PBRMR::Parameters::Tangent,
-				UE::Interchange::Materials::Sheen::Parameters::SheenRoughness,
-				UE::Interchange::Materials::Sheen::Parameters::SheenColor,
-				UE::Interchange::Materials::Subsurface::Parameters::SubsurfaceColor,
-				UE::Interchange::Materials::ClearCoat::Parameters::ClearCoat,
-				UE::Interchange::Materials::ClearCoat::Parameters::ClearCoatRoughness,
-				UE::Interchange::Materials::ClearCoat::Parameters::ClearCoatNormal
-			}	:
-			TSet<FName>{
-				UE::Interchange::Materials::StandardSurface::SubstrateMaterial::Outputs::Opaque,
-				UE::Interchange::Materials::StandardSurface::SubstrateMaterial::Outputs::Opacity
-			}
-		}
-		},
-		{
-			UMaterialXPipelineSettings::ToEnumKey(EInterchangeMaterialXShaders::StandardSurfaceTransmission),
+					// StandardSurface Inputs
+					TSet<FName>{
+						UE::Interchange::Materials::StandardSurface::Parameters::Base,
+						UE::Interchange::Materials::StandardSurface::Parameters::BaseColor,
+						UE::Interchange::Materials::StandardSurface::Parameters::DiffuseRoughness,
+						UE::Interchange::Materials::StandardSurface::Parameters::Metalness,
+						UE::Interchange::Materials::StandardSurface::Parameters::Specular,
+						UE::Interchange::Materials::StandardSurface::Parameters::SpecularRoughness,
+						UE::Interchange::Materials::StandardSurface::Parameters::SpecularIOR,
+						UE::Interchange::Materials::StandardSurface::Parameters::SpecularAnisotropy,
+						UE::Interchange::Materials::StandardSurface::Parameters::SpecularRotation,
+						UE::Interchange::Materials::StandardSurface::Parameters::Subsurface,
+						UE::Interchange::Materials::StandardSurface::Parameters::SubsurfaceColor,
+						UE::Interchange::Materials::StandardSurface::Parameters::SubsurfaceRadius,
+						UE::Interchange::Materials::StandardSurface::Parameters::SubsurfaceScale,
+						UE::Interchange::Materials::StandardSurface::Parameters::Sheen,
+						UE::Interchange::Materials::StandardSurface::Parameters::SheenColor,
+						UE::Interchange::Materials::StandardSurface::Parameters::SheenRoughness,
+						UE::Interchange::Materials::StandardSurface::Parameters::Coat,
+						UE::Interchange::Materials::StandardSurface::Parameters::CoatColor,
+						UE::Interchange::Materials::StandardSurface::Parameters::CoatRoughness,
+						UE::Interchange::Materials::StandardSurface::Parameters::CoatNormal,
+						UE::Interchange::Materials::StandardSurface::Parameters::ThinFilmThickness,
+						UE::Interchange::Materials::StandardSurface::Parameters::Emission,
+						UE::Interchange::Materials::StandardSurface::Parameters::EmissionColor,
+						UE::Interchange::Materials::StandardSurface::Parameters::Normal,
+						UE::Interchange::Materials::StandardSurface::Parameters::Tangent,
+						bIsSubstrateEnabled ? UE::Interchange::Materials::StandardSurface::Parameters::SpecularColor : FName{},
+						bIsSubstrateEnabled ? UE::Interchange::Materials::StandardSurface::Parameters::CoatIOR : FName{},
+						bIsSubstrateEnabled ? UE::Interchange::Materials::StandardSurface::Parameters::CoatAnisotropy : FName{},
+						bIsSubstrateEnabled ? UE::Interchange::Materials::StandardSurface::Parameters::CoatRotation : FName{},
+						bIsSubstrateEnabled ? UE::Interchange::Materials::StandardSurface::Parameters::ThinFilmIOR : FName{},
+						bIsSubstrateEnabled ? UE::Interchange::Materials::StandardSurface::Parameters::Opacity : FName{},
+						bIsSubstrateEnabled ? UE::Interchange::Materials::StandardSurface::Parameters::SubsurfaceAnisotropy : FName{}
+					},
+					// StandardSurface Outputs
+					!bIsSubstrateEnabled ?
+					TSet<FName>{
+						TEXT("Base Color"), // MX_StandardSurface has BaseColor with a whitespace, this should be fixed in further release
+						UE::Interchange::Materials::PBRMR::Parameters::Metallic,
+						UE::Interchange::Materials::PBRMR::Parameters::Specular,
+						UE::Interchange::Materials::PBRMR::Parameters::Roughness,
+						UE::Interchange::Materials::PBRMR::Parameters::Anisotropy,
+						UE::Interchange::Materials::PBRMR::Parameters::EmissiveColor,
+						UE::Interchange::Materials::PBRMR::Parameters::Opacity,
+						UE::Interchange::Materials::PBRMR::Parameters::Normal,
+						UE::Interchange::Materials::PBRMR::Parameters::Tangent,
+						UE::Interchange::Materials::Sheen::Parameters::SheenRoughness,
+						UE::Interchange::Materials::Sheen::Parameters::SheenColor,
+						UE::Interchange::Materials::Subsurface::Parameters::SubsurfaceColor,
+						UE::Interchange::Materials::ClearCoat::Parameters::ClearCoat,
+						UE::Interchange::Materials::ClearCoat::Parameters::ClearCoatRoughness,
+						UE::Interchange::Materials::ClearCoat::Parameters::ClearCoatNormal
+					}	:
+					TSet<FName>{
+						UE::Interchange::Materials::StandardSurface::SubstrateMaterial::Outputs::Opaque,
+						UE::Interchange::Materials::StandardSurface::SubstrateMaterial::Outputs::Opacity
+					}
+				}
+			},
 			{
-				// StandardSurfaceTransmission Inputs
-				TSet<FName>{
-					UE::Interchange::Materials::StandardSurface::Parameters::Base,
-					UE::Interchange::Materials::StandardSurface::Parameters::BaseColor,
-					UE::Interchange::Materials::StandardSurface::Parameters::DiffuseRoughness,
-					UE::Interchange::Materials::StandardSurface::Parameters::Metalness,
-					UE::Interchange::Materials::StandardSurface::Parameters::Specular,
-					bIsSubstrateEnabled ? UE::Interchange::Materials::StandardSurface::Parameters::SpecularColor : FName{},
-					UE::Interchange::Materials::StandardSurface::Parameters::SpecularRoughness,
-					UE::Interchange::Materials::StandardSurface::Parameters::SpecularIOR,
-					UE::Interchange::Materials::StandardSurface::Parameters::SpecularAnisotropy,
-					UE::Interchange::Materials::StandardSurface::Parameters::SpecularRotation,
-					UE::Interchange::Materials::StandardSurface::Parameters::Transmission,
-					UE::Interchange::Materials::StandardSurface::Parameters::TransmissionColor,
-					UE::Interchange::Materials::StandardSurface::Parameters::TransmissionDepth,
-					UE::Interchange::Materials::StandardSurface::Parameters::TransmissionScatter,
-					UE::Interchange::Materials::StandardSurface::Parameters::TransmissionScatterAnisotropy,
-					UE::Interchange::Materials::StandardSurface::Parameters::TransmissionDispersion,
-					UE::Interchange::Materials::StandardSurface::Parameters::TransmissionExtraRoughness,
-					UE::Interchange::Materials::StandardSurface::Parameters::Subsurface,
-					UE::Interchange::Materials::StandardSurface::Parameters::SubsurfaceColor,
-					UE::Interchange::Materials::StandardSurface::Parameters::SubsurfaceRadius,
-					UE::Interchange::Materials::StandardSurface::Parameters::SubsurfaceScale,
-					UE::Interchange::Materials::StandardSurface::Parameters::Sheen,
-					UE::Interchange::Materials::StandardSurface::Parameters::SheenColor,
-					UE::Interchange::Materials::StandardSurface::Parameters::SheenRoughness,
-					UE::Interchange::Materials::StandardSurface::Parameters::Coat,
-					UE::Interchange::Materials::StandardSurface::Parameters::CoatColor,
-					UE::Interchange::Materials::StandardSurface::Parameters::CoatRoughness,
-					bIsSubstrateEnabled ? UE::Interchange::Materials::StandardSurface::Parameters::CoatAnisotropy : FName{},
-					bIsSubstrateEnabled ? UE::Interchange::Materials::StandardSurface::Parameters::CoatRotation : FName{},
-					bIsSubstrateEnabled ? UE::Interchange::Materials::StandardSurface::Parameters::CoatIOR : FName{},
-					UE::Interchange::Materials::StandardSurface::Parameters::CoatNormal,
-					UE::Interchange::Materials::StandardSurface::Parameters::ThinFilmThickness,
-					bIsSubstrateEnabled ? UE::Interchange::Materials::StandardSurface::Parameters::ThinFilmIOR : FName{},
-					UE::Interchange::Materials::StandardSurface::Parameters::Emission,
-					UE::Interchange::Materials::StandardSurface::Parameters::EmissionColor,
-					bIsSubstrateEnabled ? UE::Interchange::Materials::StandardSurface::Parameters::Opacity : FName{},
-					UE::Interchange::Materials::StandardSurface::Parameters::Normal,
-					UE::Interchange::Materials::StandardSurface::Parameters::Tangent
-				},
-				// StandardSurfaceTransmission Outputs
-				!bIsSubstrateEnabled ?
-				TSet<FName>{
-					UE::Interchange::Materials::PBRMR::Parameters::BaseColor,
-					UE::Interchange::Materials::PBRMR::Parameters::Metallic,
-					UE::Interchange::Materials::PBRMR::Parameters::Specular,
-					UE::Interchange::Materials::PBRMR::Parameters::Roughness,
-					UE::Interchange::Materials::PBRMR::Parameters::Anisotropy,
-					UE::Interchange::Materials::PBRMR::Parameters::EmissiveColor,
-					UE::Interchange::Materials::PBRMR::Parameters::Opacity,
-					UE::Interchange::Materials::PBRMR::Parameters::Normal,
-					UE::Interchange::Materials::PBRMR::Parameters::Tangent,
-					UE::Interchange::Materials::PBRMR::Parameters::Refraction,
-					UE::Interchange::Materials::ThinTranslucent::Parameters::TransmissionColor
-				}	:
-				TSet<FName>{UE::Interchange::Materials::StandardSurface::SubstrateMaterial::Outputs::Translucent}
-			}
-		},
-		{
-			UMaterialXPipelineSettings::ToEnumKey(EInterchangeMaterialXShaders::SurfaceUnlit),
+				UMaterialXPipelineSettings::ToEnumKey(EInterchangeMaterialXShaders::StandardSurfaceTransmission),
+				{
+					// StandardSurfaceTransmission Inputs
+					TSet<FName>{
+						UE::Interchange::Materials::StandardSurface::Parameters::Base,
+						UE::Interchange::Materials::StandardSurface::Parameters::BaseColor,
+						UE::Interchange::Materials::StandardSurface::Parameters::DiffuseRoughness,
+						UE::Interchange::Materials::StandardSurface::Parameters::Metalness,
+						UE::Interchange::Materials::StandardSurface::Parameters::Specular,
+						bIsSubstrateEnabled ? UE::Interchange::Materials::StandardSurface::Parameters::SpecularColor : FName{},
+						UE::Interchange::Materials::StandardSurface::Parameters::SpecularRoughness,
+						UE::Interchange::Materials::StandardSurface::Parameters::SpecularIOR,
+						UE::Interchange::Materials::StandardSurface::Parameters::SpecularAnisotropy,
+						UE::Interchange::Materials::StandardSurface::Parameters::SpecularRotation,
+						UE::Interchange::Materials::StandardSurface::Parameters::Transmission,
+						UE::Interchange::Materials::StandardSurface::Parameters::TransmissionColor,
+						UE::Interchange::Materials::StandardSurface::Parameters::TransmissionDepth,
+						UE::Interchange::Materials::StandardSurface::Parameters::TransmissionScatter,
+						UE::Interchange::Materials::StandardSurface::Parameters::TransmissionScatterAnisotropy,
+						UE::Interchange::Materials::StandardSurface::Parameters::TransmissionDispersion,
+						UE::Interchange::Materials::StandardSurface::Parameters::TransmissionExtraRoughness,
+						UE::Interchange::Materials::StandardSurface::Parameters::Subsurface,
+						UE::Interchange::Materials::StandardSurface::Parameters::SubsurfaceColor,
+						UE::Interchange::Materials::StandardSurface::Parameters::SubsurfaceRadius,
+						UE::Interchange::Materials::StandardSurface::Parameters::SubsurfaceScale,
+						UE::Interchange::Materials::StandardSurface::Parameters::Sheen,
+						UE::Interchange::Materials::StandardSurface::Parameters::SheenColor,
+						UE::Interchange::Materials::StandardSurface::Parameters::SheenRoughness,
+						UE::Interchange::Materials::StandardSurface::Parameters::Coat,
+						UE::Interchange::Materials::StandardSurface::Parameters::CoatColor,
+						UE::Interchange::Materials::StandardSurface::Parameters::CoatRoughness,
+						bIsSubstrateEnabled ? UE::Interchange::Materials::StandardSurface::Parameters::CoatAnisotropy : FName{},
+						bIsSubstrateEnabled ? UE::Interchange::Materials::StandardSurface::Parameters::CoatRotation : FName{},
+						bIsSubstrateEnabled ? UE::Interchange::Materials::StandardSurface::Parameters::CoatIOR : FName{},
+						UE::Interchange::Materials::StandardSurface::Parameters::CoatNormal,
+						UE::Interchange::Materials::StandardSurface::Parameters::ThinFilmThickness,
+						bIsSubstrateEnabled ? UE::Interchange::Materials::StandardSurface::Parameters::ThinFilmIOR : FName{},
+						UE::Interchange::Materials::StandardSurface::Parameters::Emission,
+						UE::Interchange::Materials::StandardSurface::Parameters::EmissionColor,
+						bIsSubstrateEnabled ? UE::Interchange::Materials::StandardSurface::Parameters::Opacity : FName{},
+						UE::Interchange::Materials::StandardSurface::Parameters::Normal,
+						UE::Interchange::Materials::StandardSurface::Parameters::Tangent
+					},
+					// StandardSurfaceTransmission Outputs
+					!bIsSubstrateEnabled ?
+					TSet<FName>{
+						UE::Interchange::Materials::PBRMR::Parameters::BaseColor,
+						UE::Interchange::Materials::PBRMR::Parameters::Metallic,
+						UE::Interchange::Materials::PBRMR::Parameters::Specular,
+						UE::Interchange::Materials::PBRMR::Parameters::Roughness,
+						UE::Interchange::Materials::PBRMR::Parameters::Anisotropy,
+						UE::Interchange::Materials::PBRMR::Parameters::EmissiveColor,
+						UE::Interchange::Materials::PBRMR::Parameters::Opacity,
+						UE::Interchange::Materials::PBRMR::Parameters::Normal,
+						UE::Interchange::Materials::PBRMR::Parameters::Tangent,
+						UE::Interchange::Materials::PBRMR::Parameters::Refraction,
+						UE::Interchange::Materials::ThinTranslucent::Parameters::TransmissionColor
+					}	:
+					TSet<FName>{UE::Interchange::Materials::StandardSurface::SubstrateMaterial::Outputs::Translucent}
+				}
+			},
 			{
-				// SurfaceUnlit Inputs
-				TSet<FName>{
-					UE::Interchange::Materials::SurfaceUnlit::Parameters::Emission,
-					UE::Interchange::Materials::SurfaceUnlit::Parameters::EmissionColor,
-					UE::Interchange::Materials::SurfaceUnlit::Parameters::Transmission,
-					UE::Interchange::Materials::SurfaceUnlit::Parameters::TransmissionColor,
-					UE::Interchange::Materials::SurfaceUnlit::Parameters::Opacity
-				},
+				UMaterialXPipelineSettings::ToEnumKey(EInterchangeMaterialXShaders::SurfaceUnlit),
+				{
+					// SurfaceUnlit Inputs
+					TSet<FName>{
+						UE::Interchange::Materials::SurfaceUnlit::Parameters::Emission,
+						UE::Interchange::Materials::SurfaceUnlit::Parameters::EmissionColor,
+						UE::Interchange::Materials::SurfaceUnlit::Parameters::Transmission,
+						UE::Interchange::Materials::SurfaceUnlit::Parameters::TransmissionColor,
+						UE::Interchange::Materials::SurfaceUnlit::Parameters::Opacity
+					},
 					// SurfaceUnlit Outputs
 					!bIsSubstrateEnabled ?
 					TSet<FName>{
@@ -227,263 +227,263 @@ UMaterialXPipelineSettings::UMaterialXPipelineSettings()
 						UE::Interchange::Materials::UsdPreviewSurface::Parameters::Displacement,
 						UE::Interchange::Materials::UsdPreviewSurface::Parameters::Occlusion
 					},
-							// UsdPreviewSurface Outputs
-							!bIsSubstrateEnabled ?
-							TSet<FName>{
-								UE::Interchange::Materials::PBRMR::Parameters::BaseColor,
-								UE::Interchange::Materials::PBRMR::Parameters::Metallic,
-								UE::Interchange::Materials::PBRMR::Parameters::Specular,
-								UE::Interchange::Materials::PBRMR::Parameters::Roughness,
-								UE::Interchange::Materials::PBRMR::Parameters::EmissiveColor,
-								UE::Interchange::Materials::PBRMR::Parameters::Opacity,
-								UE::Interchange::Materials::PBRMR::Parameters::Normal,
-								UE::Interchange::Materials::Common::Parameters::Refraction,
-								UE::Interchange::Materials::Common::Parameters::Occlusion,
-								UE::Interchange::Materials::ClearCoat::Parameters::ClearCoat,
-								UE::Interchange::Materials::ClearCoat::Parameters::ClearCoatRoughness,
-							} :
-							TSet<FName>{TEXT("Substrate UsdPreviewSurface"), TEXT("Opacity")}
-						}
+					// UsdPreviewSurface Outputs
+					!bIsSubstrateEnabled ?
+					TSet<FName>{
+						UE::Interchange::Materials::PBRMR::Parameters::BaseColor,
+						UE::Interchange::Materials::PBRMR::Parameters::Metallic,
+						UE::Interchange::Materials::PBRMR::Parameters::Specular,
+						UE::Interchange::Materials::PBRMR::Parameters::Roughness,
+						UE::Interchange::Materials::PBRMR::Parameters::EmissiveColor,
+						UE::Interchange::Materials::PBRMR::Parameters::Opacity,
+						UE::Interchange::Materials::PBRMR::Parameters::Normal,
+						UE::Interchange::Materials::Common::Parameters::Refraction,
+						UE::Interchange::Materials::Common::Parameters::Occlusion,
+						UE::Interchange::Materials::ClearCoat::Parameters::ClearCoat,
+						UE::Interchange::Materials::ClearCoat::Parameters::ClearCoatRoughness,
+					} :
+					TSet<FName>{TEXT("Substrate UsdPreviewSurface"), TEXT("Opacity")}
+				}
+			},
+			{
+				UMaterialXPipelineSettings::ToEnumKey(EInterchangeMaterialXShaders::Surface),
+				{
+					// Surface Inputs
+					TSet<FName>{
+						UE::Interchange::Materials::Surface::Parameters::BSDF,
+						UE::Interchange::Materials::Surface::Parameters::EDF,
+						UE::Interchange::Materials::Surface::Parameters::Opacity
 					},
-					{
-					UMaterialXPipelineSettings::ToEnumKey(EInterchangeMaterialXShaders::Surface),
-					{
-							// Surface Inputs
-							TSet<FName>{
-								UE::Interchange::Materials::Surface::Parameters::BSDF,
-								UE::Interchange::Materials::Surface::Parameters::EDF,
-								UE::Interchange::Materials::Surface::Parameters::Opacity
-							},
-								// Surface Outputs
-								bIsSubstrateEnabled ?
-								TSet<FName>{UE::Interchange::Materials::Surface::Outputs::Surface} :
-								TSet<FName>{UE::Interchange::Materials::Surface::Substrate::Outputs::Surface, UE::Interchange::Materials::Surface::Substrate::Outputs::Opacity}
-							}
-						},
-								// BSDF
-								{
-									UMaterialXPipelineSettings::ToEnumKey(EInterchangeMaterialXBSDF::OrenNayarDiffuse),
-									{
-										// OrenNayarDiffuse Inputs
-										TSet<FName>{
-											TEXT("weight"),
-											TEXT("color"),
-											TEXT("roughness"),
-											TEXT("normal")
-										},
-									// OrenNayarDiffuse Outputs
-									TSet<FName>{
-										TEXT("Output")
-									}
-								}
-							},
-							{
-								UMaterialXPipelineSettings::ToEnumKey(EInterchangeMaterialXBSDF::BurleyDiffuse),
-								{
-									// BurleyDiffuse Inputs
-									TSet<FName>{
-										TEXT("weight"),
-										TEXT("color"),
-										TEXT("roughness"),
-										TEXT("normal")
-									},
-									// BurleyDiffuse Outputs
-									TSet<FName>{
-										TEXT("Output")
-									}
-								}
-							},
-							{
-								UMaterialXPipelineSettings::ToEnumKey(EInterchangeMaterialXBSDF::Translucent),
-								{
-									// Translucent Inputs
-									TSet<FName>{
-										TEXT("weight"),
-										TEXT("color"),
-										TEXT("normal")
-									},
-									// Translucent Outputs
-									TSet<FName>{
-										TEXT("Output")
-									}
-								}
-							},
-							{
-								UMaterialXPipelineSettings::ToEnumKey(EInterchangeMaterialXBSDF::Dielectric),
-								{
-									// Dielectric Inputs
-									TSet<FName>{
-										TEXT("weight"),
-										TEXT("tint"),
-										TEXT("ior"),
-										TEXT("roughness"),
-										TEXT("normal"),
-										TEXT("tangent")
-									},
-									// Dielectric Outputs
-									TSet<FName>{
-										TEXT("Output")
-									}
-								}
-							},
-							{
-								UMaterialXPipelineSettings::ToEnumKey(EInterchangeMaterialXBSDF::Conductor),
-								{
-									// Conductor Inputs
-									TSet<FName>{
-										TEXT("weight"),
-										TEXT("ior"),
-										TEXT("extinction"),
-										TEXT("roughness"),
-										TEXT("normal"),
-										TEXT("tangent")
-									},
-									// Conductor Outputs
-									TSet<FName>{
-										TEXT("Output")
-									}
-								}
-							},
-							{
-								UMaterialXPipelineSettings::ToEnumKey(EInterchangeMaterialXBSDF::GeneralizedSchlick),
-								{
-									// GeneralizedSchlick Inputs
-									TSet<FName>{
-										TEXT("weight"),
-										TEXT("color0"),
-										TEXT("color90"),
-										TEXT("exponent"),
-										TEXT("roughness"),
-										TEXT("normal"),
-										TEXT("tangent")
-									},
-									// GeneralizedSchlick Outputs
-									TSet<FName>{
-										TEXT("Output")
-									}
-								}
-							},
-							{
-								UMaterialXPipelineSettings::ToEnumKey(EInterchangeMaterialXBSDF::Subsurface),
-								{
-									// Subsurface Inputs
-									TSet<FName>{
-										TEXT("weight"),
-										TEXT("color"),
-										TEXT("radius"),
-										TEXT("anisotropy"),
-										TEXT("normal")
-									},
-									// Subsurface Outputs
-									TSet<FName>{
-										TEXT("Output")
-									}
-								}
-							},
-							{
-								UMaterialXPipelineSettings::ToEnumKey(EInterchangeMaterialXBSDF::Sheen),
-								{
-									// Sheen Inputs
-									TSet<FName>{
-										TEXT("weight"),
-										TEXT("color"),
-										TEXT("roughness"),
-										TEXT("normal")
-									},
-									// Sheen Outputs
-									TSet<FName>{
-										TEXT("Output")
-									}
-								}
-							},
-							{
-								UMaterialXPipelineSettings::ToEnumKey(EInterchangeMaterialXBSDF::ThinFilm),
-								{
-									// ThinFilm Inputs
-									TSet<FName>{
-										TEXT("thickness"),
-										TEXT("ior")
-									},
-									// ThinFilm Outputs
-									TSet<FName>{
-										TEXT("Output")
-									}
-								}
-							},
-									// EDF
-									{
-										UMaterialXPipelineSettings::ToEnumKey(EInterchangeMaterialXEDF::Uniform),
-										{
-											// Uniform Inputs
-											TSet<FName>{
-												TEXT("color")
-											},
-									// Uniform Outputs
-									TSet<FName>{
-										TEXT("Output")
-									}
-								}
-							},
-							{
-								UMaterialXPipelineSettings::ToEnumKey(EInterchangeMaterialXEDF::Conical),
-								{
-									// Conical Inputs
-									TSet<FName>{
-										TEXT("color"),
-										TEXT("normal"),
-										TEXT("inner_angle"),
-										TEXT("outer_angle"),
-									},
-									// Conical Outputs
-									TSet<FName>{
-										TEXT("Output")
-									}
-								}
-							},
-							{
-								UMaterialXPipelineSettings::ToEnumKey(EInterchangeMaterialXEDF::Measured),
-								{
-									// Measured Inputs
-									TSet<FName>{
-										TEXT("color"),
-										TEXT("normal"),
-										TEXT("file")
-									},
-									// Measured Outputs
-									TSet<FName>{
-										TEXT("Output")
-									}
-								}
-							},
-									// VDF
-									{
-										UMaterialXPipelineSettings::ToEnumKey(EInterchangeMaterialXVDF::Absorption),
-										{
-											// Absorption Inputs
-											TSet<FName>{
-												TEXT("absorption")
-											},
-									// Absorption Outputs
-										TSet<FName>{
-											TEXT("Output")
-										}
-									}
-								},
-								{
-									UMaterialXPipelineSettings::ToEnumKey(EInterchangeMaterialXVDF::Anisotropic),
-									{
-										// Anisotropic Inputs
-										TSet<FName>{
-											TEXT("absorption"),
-											TEXT("scattering"),
-											TEXT("anisotropy"),
-										},
-										// Anisotropic Outputs
-										TSet<FName>{
-											TEXT("Output")
-										}
-									}
-								}
+					// Surface Outputs
+					bIsSubstrateEnabled ?
+					TSet<FName>{UE::Interchange::Materials::Surface::Outputs::Surface} :
+					TSet<FName>{UE::Interchange::Materials::Surface::Substrate::Outputs::Surface, UE::Interchange::Materials::Surface::Substrate::Outputs::Opacity}
+				}
+			},
+			// BSDF
+			{
+				UMaterialXPipelineSettings::ToEnumKey(EInterchangeMaterialXBSDF::OrenNayarDiffuse),
+				{
+					// OrenNayarDiffuse Inputs
+					TSet<FName>{
+						TEXT("weight"),
+						TEXT("color"),
+						TEXT("roughness"),
+						TEXT("normal")
+					},
+					// OrenNayarDiffuse Outputs
+					TSet<FName>{
+						TEXT("Output")
+					}
+				}
+			},
+			{
+				UMaterialXPipelineSettings::ToEnumKey(EInterchangeMaterialXBSDF::BurleyDiffuse),
+				{
+					// BurleyDiffuse Inputs
+					TSet<FName>{
+						TEXT("weight"),
+						TEXT("color"),
+						TEXT("roughness"),
+						TEXT("normal")
+					},
+					// BurleyDiffuse Outputs
+					TSet<FName>{
+						TEXT("Output")
+					}
+				}
+			},
+			{
+				UMaterialXPipelineSettings::ToEnumKey(EInterchangeMaterialXBSDF::Translucent),
+				{
+					// Translucent Inputs
+					TSet<FName>{
+						TEXT("weight"),
+						TEXT("color"),
+						TEXT("normal")
+					},
+					// Translucent Outputs
+					TSet<FName>{
+						TEXT("Output")
+					}
+				}
+			},
+			{
+				UMaterialXPipelineSettings::ToEnumKey(EInterchangeMaterialXBSDF::Dielectric),
+				{
+					// Dielectric Inputs
+					TSet<FName>{
+						TEXT("weight"),
+						TEXT("tint"),
+						TEXT("ior"),
+						TEXT("roughness"),
+						TEXT("normal"),
+						TEXT("tangent")
+					},
+					// Dielectric Outputs
+					TSet<FName>{
+						TEXT("Output")
+					}
+				}
+			},
+			{
+				UMaterialXPipelineSettings::ToEnumKey(EInterchangeMaterialXBSDF::Conductor),
+				{
+					// Conductor Inputs
+					TSet<FName>{
+						TEXT("weight"),
+						TEXT("ior"),
+						TEXT("extinction"),
+						TEXT("roughness"),
+						TEXT("normal"),
+						TEXT("tangent")
+					},
+					// Conductor Outputs
+					TSet<FName>{
+						TEXT("Output")
+					}
+				}
+			},
+			{
+				UMaterialXPipelineSettings::ToEnumKey(EInterchangeMaterialXBSDF::GeneralizedSchlick),
+				{
+					// GeneralizedSchlick Inputs
+					TSet<FName>{
+						TEXT("weight"),
+						TEXT("color0"),
+						TEXT("color90"),
+						TEXT("exponent"),
+						TEXT("roughness"),
+						TEXT("normal"),
+						TEXT("tangent")
+					},
+					// GeneralizedSchlick Outputs
+					TSet<FName>{
+						TEXT("Output")
+					}
+				}
+			},
+			{
+				UMaterialXPipelineSettings::ToEnumKey(EInterchangeMaterialXBSDF::Subsurface),
+				{
+					// Subsurface Inputs
+					TSet<FName>{
+						TEXT("weight"),
+						TEXT("color"),
+						TEXT("radius"),
+						TEXT("anisotropy"),
+						TEXT("normal")
+					},
+					// Subsurface Outputs
+					TSet<FName>{
+						TEXT("Output")
+					}
+				}
+			},
+			{
+				UMaterialXPipelineSettings::ToEnumKey(EInterchangeMaterialXBSDF::Sheen),
+				{
+					// Sheen Inputs
+					TSet<FName>{
+						TEXT("weight"),
+						TEXT("color"),
+						TEXT("roughness"),
+						TEXT("normal")
+					},
+					// Sheen Outputs
+					TSet<FName>{
+						TEXT("Output")
+					}
+				}
+			},
+			{
+				UMaterialXPipelineSettings::ToEnumKey(EInterchangeMaterialXBSDF::ThinFilm),
+				{
+					// ThinFilm Inputs
+					TSet<FName>{
+						TEXT("thickness"),
+						TEXT("ior")
+					},
+					// ThinFilm Outputs
+					TSet<FName>{
+						TEXT("Output")
+					}
+				}
+			},
+			// EDF
+			{
+				UMaterialXPipelineSettings::ToEnumKey(EInterchangeMaterialXEDF::Uniform),
+				{
+					// Uniform Inputs
+					TSet<FName>{
+						TEXT("color")
+					},
+					// Uniform Outputs
+					TSet<FName>{
+						TEXT("Output")
+					}
+				}
+			},
+			{
+				UMaterialXPipelineSettings::ToEnumKey(EInterchangeMaterialXEDF::Conical),
+				{
+					// Conical Inputs
+					TSet<FName>{
+						TEXT("color"),
+						TEXT("normal"),
+						TEXT("inner_angle"),
+						TEXT("outer_angle"),
+					},
+					// Conical Outputs
+					TSet<FName>{
+						TEXT("Output")
+					}
+				}
+			},
+			{
+				UMaterialXPipelineSettings::ToEnumKey(EInterchangeMaterialXEDF::Measured),
+				{
+					// Measured Inputs
+					TSet<FName>{
+						TEXT("color"),
+						TEXT("normal"),
+						TEXT("file")
+					},
+					// Measured Outputs
+					TSet<FName>{
+						TEXT("Output")
+					}
+				}
+			},
+			// VDF
+			{
+				UMaterialXPipelineSettings::ToEnumKey(EInterchangeMaterialXVDF::Absorption),
+				{
+					// Absorption Inputs
+					TSet<FName>{
+						TEXT("absorption")
+					},
+					// Absorption Outputs
+					TSet<FName>{
+						TEXT("Output")
+					}
+				}
+			},
+			{
+				UMaterialXPipelineSettings::ToEnumKey(EInterchangeMaterialXVDF::Anisotropic),
+				{
+					// Anisotropic Inputs
+					TSet<FName>{
+						TEXT("absorption"),
+						TEXT("scattering"),
+						TEXT("anisotropy"),
+					},
+					// Anisotropic Outputs
+					TSet<FName>{
+						TEXT("Output")
+					}
+				}
+			}
 		};
 	}
 #endif // WITH_EDITOR
