MaterialX:
- Fixing weird rendering of materials, caused by material expression such as VertexNormalWS that should be in tangent spance when plugged to normal input

diff --git a/Engine/Plugins/Interchange/Runtime/Source/Import/Private/MaterialX/MaterialXUtils/MaterialXStandardSurfaceShader.cpp b/Engine/Plugins/Interchange/Runtime/Source/Import/Private/MaterialX/MaterialXUtils/MaterialXStandardSurfaceShader.cpp
--- a/Engine/Plugins/Interchange/Runtime/Source/Import/Private/MaterialX/MaterialXUtils/MaterialXStandardSurfaceShader.cpp
+++ b/Engine/Plugins/Interchange/Runtime/Source/Import/Private/MaterialX/MaterialXUtils/MaterialXStandardSurfaceShader.cpp
@@ -40,6 +40,7 @@ void FMaterialXStandardSurfaceShader::ConnectToStandardSurface()
 {
 	using namespace UE::Interchange::Materials;
 	using namespace mx::StandardSurface;
+	constexpr bool bInputInTangentSpace = true;
 
 	UInterchangeFunctionCallShaderNode* StandardSurfaceShaderNode = CreateFunctionCallShaderNode(SurfaceShaderNode->getName().c_str(), TEXT("/Interchange/Functions/MX_StandardSurface.MX_StandardSurface"));
 
@@ -102,7 +103,7 @@ void FMaterialXStandardSurfaceShader::ConnectToStandardSurface()
 	ConnectNodeOutputToInput(Input::CoatRoughness, StandardSurfaceShaderNode, StandardSurface::Parameters::CoatRoughness.ToString(), DefaultValue::Float::CoatRoughness);
 
 	//Coat Normal: No need to take the default input if there is no CoatNormal input
-	ConnectNodeOutputToInput(Input::CoatNormal, StandardSurfaceShaderNode, StandardSurface::Parameters::CoatNormal.ToString(), nullptr, TextureCompressionSettings::TC_Normalmap);
+	ConnectNodeOutputToInput(Input::CoatNormal, StandardSurfaceShaderNode, StandardSurface::Parameters::CoatNormal.ToString(), nullptr, bInputInTangentSpace);
 
 	//Thin Film Thickness
 	ConnectNodeOutputToInput(Input::ThinFilmThickness, StandardSurfaceShaderNode, StandardSurface::Parameters::ThinFilmThickness.ToString(), DefaultValue::Float::ThinFilmThickness);
@@ -114,10 +115,10 @@ void FMaterialXStandardSurfaceShader::ConnectToStandardSurface()
 	ConnectNodeOutputToInput(Input::EmissionColor, StandardSurfaceShaderNode, StandardSurface::Parameters::EmissionColor.ToString(), DefaultValue::Color3::EmissionColor);
 
 	//Normal: No need to take the default input if there is no Normal input
-	ConnectNodeOutputToInput(Input::Normal, StandardSurfaceShaderNode, StandardSurface::Parameters::Normal.ToString(), nullptr, TextureCompressionSettings::TC_Normalmap);
+	ConnectNodeOutputToInput(Input::Normal, StandardSurfaceShaderNode, StandardSurface::Parameters::Normal.ToString(), nullptr, bInputInTangentSpace);
 
 	//Tangent: No need to take the default input if there is no Tangent input
-	ConnectNodeOutputToInput(Input::Tangent, StandardSurfaceShaderNode, StandardSurface::Parameters::Tangent.ToString(), nullptr, TextureCompressionSettings::TC_Normalmap);
+	ConnectNodeOutputToInput(Input::Tangent, StandardSurfaceShaderNode, StandardSurface::Parameters::Tangent.ToString(), nullptr, bInputInTangentSpace);
 
 	//Transmission
 	ConnectNodeOutputToInput(Input::Transmission, StandardSurfaceShaderNode, StandardSurface::Parameters::Transmission.ToString(), DefaultValue::Float::Transmission);
@@ -180,6 +181,7 @@ void FMaterialXStandardSurfaceShader::ConnectToSubstrateStandardSurface()
 {
 	using namespace UE::Interchange::Materials;
 	using namespace mx::StandardSurface;
+	constexpr bool bInputInTangentSpace = true;
 
 	UInterchangeFunctionCallShaderNode* StandardSurfaceShaderNode = CreateFunctionCallShaderNode(SurfaceShaderNode->getName().c_str(), TEXT("/Interchange/Substrate/MX_StandardSurface.MX_StandardSurface"));
 
@@ -257,7 +259,7 @@ void FMaterialXStandardSurfaceShader::ConnectToSubstrateStandardSurface()
 	ConnectNodeOutputToInput(Input::CoatRotation, StandardSurfaceShaderNode, StandardSurface::Parameters::CoatRotation.ToString(), DefaultValue::Float::CoatAnisotropy);
 
 	//Coat Normal: No need to take the default input if there is no CoatNormal input
-	ConnectNodeOutputToInput(Input::CoatNormal, StandardSurfaceShaderNode, StandardSurface::Parameters::CoatNormal.ToString(), nullptr, TextureCompressionSettings::TC_Normalmap);
+	ConnectNodeOutputToInput(Input::CoatNormal, StandardSurfaceShaderNode, StandardSurface::Parameters::CoatNormal.ToString(), nullptr, bInputInTangentSpace);
 
 	//Thin Film Thickness
 	ConnectNodeOutputToInput(Input::ThinFilmThickness, StandardSurfaceShaderNode, StandardSurface::Parameters::ThinFilmThickness.ToString(), DefaultValue::Float::ThinFilmThickness);
@@ -272,10 +274,10 @@ void FMaterialXStandardSurfaceShader::ConnectToSubstrateStandardSurface()
 	ConnectNodeOutputToInput(Input::EmissionColor, StandardSurfaceShaderNode, StandardSurface::Parameters::EmissionColor.ToString(), DefaultValue::Color3::EmissionColor);
 
 	//Normal: No need to take the default input if there is no Normal input
-	ConnectNodeOutputToInput(Input::Normal, StandardSurfaceShaderNode, StandardSurface::Parameters::Normal.ToString(), nullptr, TextureCompressionSettings::TC_Normalmap);
+	ConnectNodeOutputToInput(Input::Normal, StandardSurfaceShaderNode, StandardSurface::Parameters::Normal.ToString(), nullptr, bInputInTangentSpace);
 
 	//Tangent: No need to take the default input if there is no Tangent input
-	ConnectNodeOutputToInput(Input::Tangent, StandardSurfaceShaderNode, StandardSurface::Parameters::Tangent.ToString(), nullptr, TextureCompressionSettings::TC_Normalmap);
+	ConnectNodeOutputToInput(Input::Tangent, StandardSurfaceShaderNode, StandardSurface::Parameters::Tangent.ToString(), nullptr, bInputInTangentSpace);
 
 	//Opacity
 	ConnectNodeOutputToInput(Input::Opacity, StandardSurfaceShaderNode, StandardSurface::Parameters::Opacity.ToString(), DefaultValue::Color3::Opacity);
diff --git a/Engine/Plugins/Interchange/Runtime/Source/Import/Private/MaterialX/MaterialXUtils/MaterialXSurfaceShaderAbstract.cpp b/Engine/Plugins/Interchange/Runtime/Source/Import/Private/MaterialX/MaterialXUtils/MaterialXSurfaceShaderAbstract.cpp
--- a/Engine/Plugins/Interchange/Runtime/Source/Import/Private/MaterialX/MaterialXUtils/MaterialXSurfaceShaderAbstract.cpp
+++ b/Engine/Plugins/Interchange/Runtime/Source/Import/Private/MaterialX/MaterialXUtils/MaterialXSurfaceShaderAbstract.cpp
@@ -23,7 +23,9 @@ FString FMaterialXSurfaceShaderAbstract::EmptyString{};
 FString FMaterialXSurfaceShaderAbstract::DefaultOutput{TEXT("out")};
 
 FMaterialXSurfaceShaderAbstract::FMaterialXSurfaceShaderAbstract(UInterchangeBaseNodeContainer& BaseNodeContainer)
-	: FMaterialXBase(BaseNodeContainer)
+	: FMaterialXBase{ BaseNodeContainer }
+	, ShaderGraphNode{ nullptr }
+	, bTangentSpaceInput{ false }
 {}
 
 bool FMaterialXSurfaceShaderAbstract::AddAttribute(MaterialX::InputPtr Input, const FString& InputChannelName, UInterchangeShaderNode* ShaderNode) const
@@ -676,166 +678,104 @@ void FMaterialXSurfaceShaderAbstract::ConnectOutsideInputToOutput(const FConnect
 	UInterchangeShaderPortsAPI::ConnectDefaultOuputToInput(NodeMultiply, TEXT("B"), NodeOneMinus->GetUniqueID());
 }
 
-void FMaterialXSurfaceShaderAbstract::ConnectPositionInputToOutput(const FConnectNode& Connect)
+UInterchangeShaderNode* FMaterialXSurfaceShaderAbstract::ConnectGeometryInputToOutput(const FConnectNode& Connect, const FString& ShaderType, const FString& TransformShaderType, const FString& TransformInput, const FString& TransformSourceType, int32 TransformSource, const FString& TransformType, int32 TransformSDestination, bool bIsVector)
 {
 	// MaterialX defines the space as: object, model, world
 	// model: The local coordinate space of the geometry, before any local deformations or global transforms have been applied.
 	// object: The local coordinate space of the geometry, after local deformations have been applied, but before any global transforms.
 	// world : The global coordinate space of the geometry, after local deformationsand global transforms have been applied.
 
-	// For the moment we don't have the distinction between model/object, so let's just create an UMaterialExpressionWorldPosition
-	// In case of model/object we need to add a TransformPoint from world to local space
-	UInterchangeShaderNode* PositionNode = CreateShaderNode(Connect.UpstreamNode->getName().c_str() + FString(TEXT("_Position")), TEXT("WorldPosition"));
-
-	// In case of the position node, it seems that the unit is different, we assume for now a conversion from mm -> m, even if UE by default is cm
-	// See standard_surface_marble_solid file, especially on the fractal3d node
-	UInterchangeShaderNode* UnitNode = CreateShaderNode(Connect.UpstreamNode->getName().c_str(), TEXT("Multiply"));
-	UnitNode->AddFloatAttribute(UInterchangeShaderPortsAPI::MakeInputValueKey(TEXT("B")), 0.001f);
-	UInterchangeShaderPortsAPI::ConnectDefaultOuputToInput(UnitNode, TEXT("A"), PositionNode->GetUniqueID());
+	// In case of model/object we need to add a TransformVector from world to local space
+	UInterchangeShaderNode* GeometryNode = CreateShaderNode(Connect.UpstreamNode->getName().c_str(), ShaderType);
 
 	UInterchangeShaderNode* NodeToConnectTo = Connect.ParentShaderNode;
 	FString InputToConnectTo = Connect.InputChannelName;
 
-	std::string Space = "world";
 	mx::InputPtr InputSpace = Connect.UpstreamNode->getInput("space");
 
-	if(InputSpace)
-	{
-		Space = InputSpace->getValueString();
-	}
-
 	//the default space defined by the nodedef is "object"
-	if(Space != "world" || !InputSpace)
+	bool bIsObjectSpace = (InputSpace && InputSpace->getValueString() != "world") || !InputSpace;
+
+	// We transform to Tangent Space only for Vector nodes
+	if(bTangentSpaceInput && bIsVector)
 	{
 		using namespace UE::Interchange::Materials::Standard::Nodes;
+		UInterchangeShaderNode* TransformTSNode = CreateShaderNode(Connect.UpstreamNode->getName().c_str() + FString(TEXT("_TransformTS")), TransformShaderType);
+		EMaterialVectorCoordTransformSource SpaceSource = bIsObjectSpace ? EMaterialVectorCoordTransformSource::TRANSFORMSOURCE_Local : EMaterialVectorCoordTransformSource::TRANSFORMSOURCE_World;
+		TransformTSNode->AddInt32Attribute(TransformSourceType, SpaceSource);
+		TransformTSNode->AddInt32Attribute(TransformType, EMaterialVectorCoordTransform::TRANSFORM_Tangent);
+		UInterchangeShaderPortsAPI::ConnectDefaultOuputToInput(NodeToConnectTo, InputToConnectTo, TransformTSNode->GetUniqueID());
+		NodeToConnectTo = TransformTSNode;
+		InputToConnectTo = TransformInput; //Same a TransformVector
+	}
 
-		UInterchangeShaderNode* TransformNode = CreateShaderNode(Connect.UpstreamNode->getName().c_str() + FString(TEXT("_Transform")), TransformPosition::Name.ToString());
+	if(bIsObjectSpace)
+	{
+		UInterchangeShaderNode* TransformNode = CreateShaderNode(Connect.UpstreamNode->getName().c_str() + FString(TEXT("_Transform")), TransformShaderType);
+		TransformNode->AddInt32Attribute(TransformSourceType, TransformSource);
+		TransformNode->AddInt32Attribute(TransformType, TransformSDestination);
+		UInterchangeShaderPortsAPI::ConnectDefaultOuputToInput(NodeToConnectTo, InputToConnectTo, TransformNode->GetUniqueID());
 		NodeToConnectTo = TransformNode;
-		InputToConnectTo = TransformPosition::Inputs::Input.ToString();
-		TransformNode->AddInt32Attribute(TransformPosition::Attributes::TransformSourceType.ToString(), EMaterialPositionTransformSource::TRANSFORMPOSSOURCE_World);
-		TransformNode->AddInt32Attribute(TransformPosition::Attributes::TransformType.ToString(), EMaterialPositionTransformSource::TRANSFORMPOSSOURCE_Local);
-		UInterchangeShaderPortsAPI::ConnectDefaultOuputToInput(Connect.ParentShaderNode, Connect.InputChannelName, TransformNode->GetUniqueID());
+		InputToConnectTo = TransformInput;
 	}
 
-	UInterchangeShaderPortsAPI::ConnectDefaultOuputToInput(NodeToConnectTo, InputToConnectTo, UnitNode->GetUniqueID());
+	UInterchangeShaderPortsAPI::ConnectDefaultOuputToInput(NodeToConnectTo, InputToConnectTo, GeometryNode->GetUniqueID());
+
+	return GeometryNode;
 }
 
-void FMaterialXSurfaceShaderAbstract::ConnectNormalInputToOutput(const FConnectNode& Connect)
+void FMaterialXSurfaceShaderAbstract::ConnectPositionInputToOutput(const FConnectNode& Connect)
 {
-	// MaterialX defines the space as: object, model, world
-	// model: The local coordinate space of the geometry, before any local deformations or global transforms have been applied.
-	// object: The local coordinate space of the geometry, after local deformations have been applied, but before any global transforms.
-	// world : The global coordinate space of the geometry, after local deformationsand global transforms have been applied.
-
-	// For the moment we don't have the distinction between model/object, so let's just create an UMaterialExpressionVertexNormalWS
-	// In case of model/object we need to add a TransformVector from world to local space
-	UInterchangeShaderNode* NormalNode = CreateShaderNode(Connect.UpstreamNode->getName().c_str(), TEXT("VertexNormalWS"));
-	UInterchangeShaderNode* NodeToConnectTo = Connect.ParentShaderNode;
-	FString InputToConnectTo = Connect.InputChannelName;
-
-	std::string Space = "world";
-	mx::InputPtr InputSpace = Connect.UpstreamNode->getInput("space");
-
-	if(InputSpace)
-	{
-		Space = InputSpace->getValueString();
-	}
-
-	//the default space defined by the nodedef is "object"
-	if(Space != "world" || !InputSpace)
-	{
-		using namespace UE::Interchange::Materials::Standard::Nodes;
+	using namespace UE::Interchange::Materials::Standard::Nodes;
+	UInterchangeShaderNode* UnitNode = ConnectGeometryInputToOutput(Connect, TEXT("Multiply"),
+																	TransformPosition::Name.ToString(),
+																	TransformPosition::Inputs::Input.ToString(),
+																	TransformPosition::Attributes::TransformSourceType.ToString(), EMaterialPositionTransformSource::TRANSFORMPOSSOURCE_World,
+																	TransformPosition::Attributes::TransformType.ToString(), EMaterialPositionTransformSource::TRANSFORMPOSSOURCE_Local,
+																	false);
 
-		UInterchangeShaderNode* TransformNode = CreateShaderNode(Connect.UpstreamNode->getName().c_str() + FString(TEXT("_Transform")), TransformVector::Name.ToString());
-		NodeToConnectTo = TransformNode;
-		InputToConnectTo = TransformVector::Inputs::Input.ToString();
-		TransformNode->AddInt32Attribute(TransformVector::Attributes::TransformSourceType.ToString(), EMaterialVectorCoordTransformSource::TRANSFORMSOURCE_World);
-		TransformNode->AddInt32Attribute(TransformVector::Attributes::TransformType.ToString(), EMaterialVectorCoordTransform::TRANSFORM_Local);
-		UInterchangeShaderPortsAPI::ConnectDefaultOuputToInput(Connect.ParentShaderNode, Connect.InputChannelName, TransformNode->GetUniqueID());
-	}
+	// In case of the position node, it seems that the unit is different, we assume for now a conversion from mm -> m, even if UE by default is cm
+	// See standard_surface_marble_solid file, especially on the fractal3d node
+	UInterchangeShaderNode* PositionNode = CreateShaderNode(Connect.UpstreamNode->getName().c_str() + FString(TEXT("_Position")), TEXT("WorldPosition"));
+	UnitNode->AddFloatAttribute(UInterchangeShaderPortsAPI::MakeInputValueKey(TEXT("B")), 0.001f);
+	UInterchangeShaderPortsAPI::ConnectDefaultOuputToInput(UnitNode, TEXT("A"), PositionNode->GetUniqueID());
+}
 
-	UInterchangeShaderPortsAPI::ConnectDefaultOuputToInput(NodeToConnectTo, InputToConnectTo, NormalNode->GetUniqueID());
+void FMaterialXSurfaceShaderAbstract::ConnectNormalInputToOutput(const FConnectNode& Connect)
+{
+	using namespace UE::Interchange::Materials::Standard::Nodes;
+	ConnectGeometryInputToOutput(Connect, TEXT("VertexNormalWS"),
+								 TransformVector::Name.ToString(),
+								 TransformVector::Inputs::Input.ToString(),
+								 TransformVector::Attributes::TransformSourceType.ToString(), EMaterialVectorCoordTransformSource::TRANSFORMSOURCE_World,
+								 TransformVector::Attributes::TransformType.ToString(), EMaterialVectorCoordTransform::TRANSFORM_Local);
 }
 
 void FMaterialXSurfaceShaderAbstract::ConnectTangentInputToOutput(const FConnectNode& Connect)
 {
-	// MaterialX defines the space as: object, model, world
-	// model: The local coordinate space of the geometry, before any local deformations or global transforms have been applied.
-	// object: The local coordinate space of the geometry, after local deformations have been applied, but before any global transforms.
-	// world : The global coordinate space of the geometry, after local deformationsand global transforms have been applied.
-
-	// For the moment we don't have the distinction between model/object, so let's just create an UMaterialExpressionVertexTangentWS
-	// In case of model/object we need to add a TransformVector from world to local space
-	UInterchangeShaderNode* TangentNode = CreateShaderNode(Connect.UpstreamNode->getName().c_str(), TEXT("VertexTangentWS"));
-	UInterchangeShaderNode* NodeToConnectTo = Connect.ParentShaderNode;
-	FString InputToConnectTo = Connect.InputChannelName;
-
-	std::string Space = "world";
-	mx::InputPtr InputSpace = Connect.UpstreamNode->getInput("space");
-
-	if(InputSpace)
-	{
-		Space = InputSpace->getValueString();
-	}
-
-	//the default space defined by the nodedef is "object"
-	if(Space != "world" || !InputSpace)
-	{
-		using namespace UE::Interchange::Materials::Standard::Nodes;
-
-		UInterchangeShaderNode* TransformNode = CreateShaderNode(Connect.UpstreamNode->getName().c_str() + FString(TEXT("_Transform")), TransformVector::Name.ToString());
-		NodeToConnectTo = TransformNode;
-		InputToConnectTo = TransformVector::Inputs::Input.ToString();
-		TransformNode->AddInt32Attribute(TransformVector::Attributes::TransformSourceType.ToString(), EMaterialVectorCoordTransformSource::TRANSFORMSOURCE_World);
-		TransformNode->AddInt32Attribute(TransformVector::Attributes::TransformType.ToString(), EMaterialVectorCoordTransform::TRANSFORM_Local);
-		UInterchangeShaderPortsAPI::ConnectDefaultOuputToInput(Connect.ParentShaderNode, Connect.InputChannelName, TransformNode->GetUniqueID());
-	}
-
-	UInterchangeShaderPortsAPI::ConnectDefaultOuputToInput(NodeToConnectTo, InputToConnectTo, TangentNode->GetUniqueID());
+	using namespace UE::Interchange::Materials::Standard::Nodes;
+	ConnectGeometryInputToOutput(Connect, TEXT("VertexTangentWS"),
+								 TransformVector::Name.ToString(),
+								 TransformVector::Inputs::Input.ToString(),
+								 TransformVector::Attributes::TransformSourceType.ToString(), EMaterialVectorCoordTransformSource::TRANSFORMSOURCE_World,
+								 TransformVector::Attributes::TransformType.ToString(), EMaterialVectorCoordTransform::TRANSFORM_Local);
 }
 
 void FMaterialXSurfaceShaderAbstract::ConnectBitangentInputToOutput(const FConnectNode& Connect)
 {
-	// MaterialX defines the space as: object, model, world
-	// model: The local coordinate space of the geometry, before any local deformations or global transforms have been applied.
-	// object: The local coordinate space of the geometry, after local deformations have been applied, but before any global transforms.
-	// world : The global coordinate space of the geometry, after local deformationsand global transforms have been applied.
+	using namespace UE::Interchange::Materials::Standard::Nodes;
+
+	UInterchangeShaderNode* BitangentNode = ConnectGeometryInputToOutput(Connect, TEXT("CrossProduct"),
+																		 TransformVector::Name.ToString(),
+																		 TransformVector::Inputs::Input.ToString(),
+																		 TransformVector::Attributes::TransformSourceType.ToString(), EMaterialVectorCoordTransformSource::TRANSFORMSOURCE_World,
+																		 TransformVector::Attributes::TransformType.ToString(), EMaterialVectorCoordTransform::TRANSFORM_Local);
 
-	// For the moment we don't have the distinction between model/object, so let's just do the cross product between the normal and the tangent
-	// In case of model/object we need to add a TransformVector from world to local space
 	UInterchangeShaderNode* NormalNode = CreateShaderNode(Connect.UpstreamNode->getName().c_str() + FString(TEXT("_Normal")), TEXT("VertexNormalWS"));
 	UInterchangeShaderNode* TangentNode = CreateShaderNode(Connect.UpstreamNode->getName().c_str() + FString(TEXT("_Tangent")), TEXT("VertexTangentWS"));
-	UInterchangeShaderNode* BitangentNode = CreateShaderNode(Connect.UpstreamNode->getName().c_str(), TEXT("CrossProduct"));
 
 	UInterchangeShaderPortsAPI::ConnectDefaultOuputToInput(BitangentNode, TEXT("A"), NormalNode->GetUniqueID());
 	UInterchangeShaderPortsAPI::ConnectDefaultOuputToInput(BitangentNode, TEXT("B"), TangentNode->GetUniqueID());
-
-
-	UInterchangeShaderNode* NodeToConnectTo = Connect.ParentShaderNode;
-	FString InputToConnectTo = Connect.InputChannelName;
-
-	std::string Space = "world";
-	mx::InputPtr InputSpace = Connect.UpstreamNode->getInput("space");
-
-	if(InputSpace)
-	{
-		Space = InputSpace->getValueString();
-	}
-
-	//the default space defined by the nodedef is "object"
-	if(Space != "world" || !InputSpace)
-	{
-		using namespace UE::Interchange::Materials::Standard::Nodes;
-
-		UInterchangeShaderNode* TransformNode = CreateShaderNode(Connect.UpstreamNode->getName().c_str() + FString(TEXT("_Transform")), TransformVector::Name.ToString());
-		NodeToConnectTo = TransformNode;
-		InputToConnectTo = TransformVector::Inputs::Input.ToString();
-		TransformNode->AddInt32Attribute(TransformVector::Attributes::TransformSourceType.ToString(), EMaterialVectorCoordTransformSource::TRANSFORMSOURCE_World);
-		TransformNode->AddInt32Attribute(TransformVector::Attributes::TransformType.ToString(), EMaterialVectorCoordTransform::TRANSFORM_Local);
-		UInterchangeShaderPortsAPI::ConnectDefaultOuputToInput(Connect.ParentShaderNode, Connect.InputChannelName, TransformNode->GetUniqueID());
-	}
-
-	UInterchangeShaderPortsAPI::ConnectDefaultOuputToInput(NodeToConnectTo, InputToConnectTo, BitangentNode->GetUniqueID());
 }
 
 void FMaterialXSurfaceShaderAbstract::ConnectTimeInputToOutput(const FConnectNode& Connect)
diff --git a/Engine/Plugins/Interchange/Runtime/Source/Import/Public/MaterialX/MaterialXUtils/MaterialXSurfaceShaderAbstract.h b/Engine/Plugins/Interchange/Runtime/Source/Import/Public/MaterialX/MaterialXUtils/MaterialXSurfaceShaderAbstract.h
--- a/Engine/Plugins/Interchange/Runtime/Source/Import/Public/MaterialX/MaterialXUtils/MaterialXSurfaceShaderAbstract.h
+++ b/Engine/Plugins/Interchange/Runtime/Source/Import/Public/MaterialX/MaterialXUtils/MaterialXSurfaceShaderAbstract.h
@@ -108,12 +108,12 @@ protected:
 	 * @param OptionalTextureCompression - Set the texture compression for all textures along the path of an input
 	 */
 	template<typename T>
-	bool ConnectNodeOutputToInput(const char* InputName, UInterchangeShaderNode* ShaderNode, const FString& InputShaderName, T DefaultValue, const TOptional<TextureCompressionSettings>& OptionalTextureCompression = TOptional<TextureCompressionSettings>{})
+	bool ConnectNodeOutputToInput(const char* InputName, UInterchangeShaderNode* ShaderNode, const FString& InputShaderName, T DefaultValue, bool bIsTangentSpaceInput = false)
 	{
 		MaterialX::DocumentPtr Document = SurfaceShaderNode->getDocument();
 		MaterialX::InputPtr Input = GetInput(SurfaceShaderNode, InputName);
 
-		TGuardValue<TOptional<TextureCompressionSettings>> InputTypeBeingProcessedGuard(TextureCompression, OptionalTextureCompression);
+		TGuardValue<bool>InputTypeBeingProcessedGuard(bTangentSpaceInput, bIsTangentSpaceInput);
 
 		bool bIsConnected = ConnectNodeGraphOutputToInput(Input, ShaderNode, InputShaderName);
 
@@ -370,9 +370,9 @@ protected:
 					TextureNode->SetCustomWrapV(WrapModeV);
 
 					// Encode the compression in the payloadKey
-					if(TextureCompression.IsSet())
+					if(bTangentSpaceInput)
 					{
-						TextureNode->SetPayLoadKey(TextureNode->GetPayLoadKey().GetValue() + FMaterialXManager::TexturePayloadSeparator + FString::FromInt(TextureCompression.GetValue()));
+						TextureNode->SetPayLoadKey(TextureNode->GetPayLoadKey().GetValue() + FMaterialXManager::TexturePayloadSeparator + FString::FromInt(TextureCompressionSettings::TC_Normalmap));
 					}
 				}
 			}
@@ -443,6 +443,9 @@ protected:
 	 */
 	void SetAttributeNewName(MaterialX::InputPtr Input, const char* NewName) const;
 
+private:
+	
+	UInterchangeShaderNode* ConnectGeometryInputToOutput(const FConnectNode& Connect, const FString& ShaderType, const FString& TransformShaderType, const FString& TransformInput, const FString& TransformSourceType, int32 TransformSource, const FString& TransformType, int32 TransformSDestination, bool bIsVector = true);
 
 protected:
 
@@ -464,6 +467,7 @@ protected:
 	/** Initialized by the material shader (e.g: surfacematerial), the derived class should only set the ShaderType */
 	UInterchangeShaderGraphNode* ShaderGraphNode;
 
-	TOptional<TextureCompressionSettings> TextureCompression;
+	/** Used for texture compression and transform to tangent space nodes coming from inputs such as coat_normal  */
+	bool bTangentSpaceInput;
 };
 #endif
\ No newline at end of file
