MaterialX:
Adding support for <artistic_ior>, <roughness_dual>, <roughness_anisotropy>, <glossiness_anisotropy> is automatically supported as of the inverse of roughness_anisotropy

diff --git a/Engine/Plugins/Interchange/Runtime/Source/Import/Private/MaterialX/MaterialXUtils/MaterialXManager.cpp b/Engine/Plugins/Interchange/Runtime/Source/Import/Private/MaterialX/MaterialXUtils/MaterialXManager.cpp
--- a/Engine/Plugins/Interchange/Runtime/Source/Import/Private/MaterialX/MaterialXUtils/MaterialXManager.cpp
+++ b/Engine/Plugins/Interchange/Runtime/Source/Import/Private/MaterialX/MaterialXUtils/MaterialXManager.cpp
@@ -16,7 +16,6 @@
 #include "MaterialX/MaterialXUtils/MaterialXSpotLightShader.h"
 #include "MaterialX/MaterialXUtils/MaterialXSurfaceMaterial.h"
 
-#include "Engine/RendererSettings.h"
 #include "InterchangeMaterialDefinitions.h"
 #include "InterchangeImportLog.h"
 
@@ -26,7 +25,7 @@ namespace mx = MaterialX;
 //not a good solution to use semicolon because of drive disk on Windows
 const TCHAR FMaterialXManager::TexturePayloadSeparator = TEXT('{');
 
-#define MATERIALFUNCTION_PATH FString{bIsSubstrateEnabled ? TEXT("Substrate/") : TEXT("Functions/")}
+#define MX_MATERIALFUNCTION(AssetName) TEXT("/Interchange/") + FString{bIsSubstrateEnabled ? TEXT("Substrate/") : TEXT("Functions/")} + TEXT(AssetName) TEXT(".") TEXT(AssetName)
 
 FMaterialXManager::FMaterialXManager()
 	: MatchingInputNames {
@@ -221,22 +220,26 @@ FMaterialXManager::FMaterialXManager()
 {
 		MatchingMaterialFunctions = {
 			// BSDF Nodes
-			{mx::Category::BurleyDiffuseBSDF ,		TEXT("/Interchange/") + MATERIALFUNCTION_PATH + TEXT("MX_BurleyDiffuseBSDF.MX_BurleyDiffuseBSDF")},
-			{mx::Category::ConductorBSDF,			TEXT("/Interchange/") + MATERIALFUNCTION_PATH + TEXT("MX_ConductorBSDF.MX_ConductorBSDF")},
-			{mx::Category::DielectricBSDF,			TEXT("/Interchange/") + MATERIALFUNCTION_PATH + TEXT("MX_DielectricBSDF.MX_DielectricBSDF")},
-			{mx::Category::GeneralizedSchlickBSDF,	TEXT("/Interchange/") + MATERIALFUNCTION_PATH + TEXT("MX_GeneralizedSchlickBSDF.MX_GeneralizedSchlickBSDF")},
-			{mx::Category::OrenNayarDiffuseBSDF ,	TEXT("/Interchange/") + MATERIALFUNCTION_PATH + TEXT("MX_OrenNayarBSDF.MX_OrenNayarBSDF")},
-			{mx::Category::SheenBSDF,				TEXT("/Interchange/") + MATERIALFUNCTION_PATH + TEXT("MX_SheenBSDF.MX_SheenBSDF")},
-			{mx::Category::SubsurfaceBSDF ,			TEXT("/Interchange/") + MATERIALFUNCTION_PATH + TEXT("MX_SubsurfaceBSDF.MX_SubsurfaceBSDF")},
-			{mx::Category::ThinFilmBSDF,			TEXT("/Interchange/") + MATERIALFUNCTION_PATH + TEXT("MX_ThinFilmBSDF.MX_ThinFilmBSDF")},
-			{mx::Category::TranslucentBSDF,			TEXT("/Interchange/") + MATERIALFUNCTION_PATH + TEXT("MX_TranslucentBSDF.MX_TranslucentBSDF")},
+			{mx::Category::BurleyDiffuseBSDF ,		MX_MATERIALFUNCTION("MX_BurleyDiffuseBSDF")},
+			{mx::Category::ConductorBSDF,			MX_MATERIALFUNCTION("MX_ConductorBSDF")},
+			{mx::Category::DielectricBSDF,			MX_MATERIALFUNCTION("MX_DielectricBSDF")},
+			{mx::Category::GeneralizedSchlickBSDF,	MX_MATERIALFUNCTION("MX_GeneralizedSchlickBSDF")},
+			{mx::Category::OrenNayarDiffuseBSDF ,	MX_MATERIALFUNCTION("MX_OrenNayarBSDF")},
+			{mx::Category::SheenBSDF,				MX_MATERIALFUNCTION("MX_SheenBSDF")},
+			{mx::Category::SubsurfaceBSDF ,			MX_MATERIALFUNCTION("MX_SubsurfaceBSDF")},
+			{mx::Category::ThinFilmBSDF,			MX_MATERIALFUNCTION("MX_ThinFilmBSDF")},
+			{mx::Category::TranslucentBSDF,			MX_MATERIALFUNCTION("MX_TranslucentBSDF")},
 			// EDF Nodes
-			{mx::Category::ConicalEDF,	TEXT("/Interchange/") + MATERIALFUNCTION_PATH + TEXT("MX_ConicalEDF.MX_ConicalEDF")},
-			{mx::Category::MeasuredEDF, TEXT("/Interchange/") + MATERIALFUNCTION_PATH + TEXT("MX_MeasuredEDF.MX_MeasuredEDF")},
-			{mx::Category::UniformEDF,	TEXT("/Interchange/") + MATERIALFUNCTION_PATH + TEXT("MX_UniformEDF.MX_UniformEDF")},
+			{mx::Category::ConicalEDF,				MX_MATERIALFUNCTION("MX_ConicalEDF")},
+			{mx::Category::MeasuredEDF,				MX_MATERIALFUNCTION("MX_MeasuredEDF")},
+			{mx::Category::UniformEDF,				MX_MATERIALFUNCTION("MX_UniformEDF")},
 			// VDF Nodes
-			{mx::Category::AbsorptionVDF,  TEXT("/Interchange/") + MATERIALFUNCTION_PATH + TEXT("MX_AbsorptionVDF.MX_AbsorptionVDF")},
-			{mx::Category::AnisotropicVDF, TEXT("/Interchange/") + MATERIALFUNCTION_PATH + TEXT("MX_AnisotropicVDF.MX_AnisotropicVDF")},
+			{mx::Category::AbsorptionVDF,			MX_MATERIALFUNCTION("MX_AbsorptionVDF")},
+			{mx::Category::AnisotropicVDF,			MX_MATERIALFUNCTION("MX_AnisotropicVDF")},
+			// Utility nodes
+			{mx::Category::ArtisticIOR,				TEXT("/Interchange/Functions/MX_Artistic_IOR.MX_Artistic_IOR")},
+			{mx::Category::RoughnessAnisotropy,		TEXT("/Interchange/Functions/MX_Roughness_Anisotropy.MX_Roughness_Anisotropy")},
+			{mx::Category::RoughnessDual,			TEXT("/Interchange/Functions/MX_Roughness_Dual.MX_Roughness_Dual")},
 		};
 
 		if(bIsSubstrateEnabled)
@@ -349,7 +352,10 @@ namespace UE::Interchange::MaterialX
 			return bAllLoaded;
 		};
 
-		static const bool bPackagesLoaded =	ArePackagesLoaded({ TEXT("MaterialFunction'/Engine/Functions/Engine_MaterialFunctions03/Procedurals/NormalFromHeightmap.NormalFromHeightmap'") });
+		static const bool bPackagesLoaded =	ArePackagesLoaded({ TEXT("MaterialFunction'/Engine/Functions/Engine_MaterialFunctions03/Procedurals/NormalFromHeightmap.NormalFromHeightmap'"),
+															    TEXT("MaterialFunction'/Interchange/Functions/MX_Artistic_IOR.MX_Artistic_IOR'"),
+															    TEXT("MaterialFunction'/Interchange/Functions/MX_Roughness_Anisotropy.MX_Roughness_Anisotropy'"), 
+															    TEXT("MaterialFunction'/Interchange/Functions/MX_Roughness_Dual.MX_Roughness_Dual'")});
 
 		return bPackagesLoaded;
 #else
@@ -358,4 +364,4 @@ namespace UE::Interchange::MaterialX
 	}
 }
 
-#undef MATERIALFUNCTION_PATH
\ No newline at end of file
+#undef MX_MATERIALFUNCTION
\ No newline at end of file
