MaterialX:
- Replace Strata to Substrate for material expressions, part of the recent global renaming of Strata into Substrate
- Fix wrong typo in MaterialXPipeline, 'Dieletric' to 'Dielectric'

diff --git a/Engine/Plugins/Interchange/Runtime/Source/Import/Private/MaterialX/MaterialXUtils/MaterialXManager.cpp b/Engine/Plugins/Interchange/Runtime/Source/Import/Private/MaterialX/MaterialXUtils/MaterialXManager.cpp
--- a/Engine/Plugins/Interchange/Runtime/Source/Import/Private/MaterialX/MaterialXUtils/MaterialXManager.cpp
+++ b/Engine/Plugins/Interchange/Runtime/Source/Import/Private/MaterialX/MaterialXUtils/MaterialXManager.cpp
@@ -249,7 +249,7 @@ FMaterialXManager::FMaterialXManager()
 			MatchingInputNames.Add({ TEXT(""), TEXT("base") }, TEXT("Bottom"));
 			MaterialExpressionInputs.Add(TEXT("Top"));
 			MaterialExpressionInputs.Add(TEXT("Bottom"));
-			MatchingMaterialExpressions.Add({ MaterialX::Category::Layer }, TEXT("StrataVerticalLayering"));
+			MatchingMaterialExpressions.Add({ MaterialX::Category::Layer }, TEXT("SubstrateVerticalLayering"));
 
 			// Horizontal Layering, MaterialX names it mix so we have to be careful not to take the lerp expression
 			MatchingInputNames.Add({ FKeyCategoryNodegroup{MaterialX::Category::Mix, MaterialX::NodeGroup::PBR}, TEXT("bg") }, TEXT("Background"));
@@ -258,15 +258,15 @@ FMaterialXManager::FMaterialXManager()
 			MaterialExpressionInputs.Add(TEXT("Background"));
 			MaterialExpressionInputs.Add(TEXT("Foreground"));
 			MaterialExpressionInputs.Add(TEXT("Mix"));
-			MatchingMaterialExpressions.Add({ MaterialX::Category::Mix, MaterialX::NodeGroup::PBR }, TEXT("StrataHorizontalMixing"));
+			MatchingMaterialExpressions.Add({ MaterialX::Category::Mix, MaterialX::NodeGroup::PBR }, TEXT("SubstrateHorizontalMixing"));
 
 			// Add
-			MatchingMaterialExpressions.Add({ MaterialX::Category::Add, MaterialX::NodeGroup::PBR }, TEXT("StrataAdd"));
+			MatchingMaterialExpressions.Add({ MaterialX::Category::Add, MaterialX::NodeGroup::PBR }, TEXT("SubstrateAdd"));
 
 			// Multiply
 			MatchingInputNames.Add({ FKeyCategoryNodegroup{MaterialX::Category::Multiply, MaterialX::NodeGroup::PBR}, TEXT("in2") }, TEXT("Weight"));
 			MaterialExpressionInputs.Add(TEXT("Weight"));
-			MatchingMaterialExpressions.Add({ MaterialX::Category::Multiply, MaterialX::NodeGroup::PBR }, TEXT("StrataWeight"));
+			MatchingMaterialExpressions.Add({ MaterialX::Category::Multiply, MaterialX::NodeGroup::PBR }, TEXT("SubstrateWeight"));
 		}
 	}
 
diff --git a/Engine/Plugins/Interchange/Runtime/Source/Pipelines/Private/InterchangeMaterialXPipeline.cpp b/Engine/Plugins/Interchange/Runtime/Source/Pipelines/Private/InterchangeMaterialXPipeline.cpp
--- a/Engine/Plugins/Interchange/Runtime/Source/Pipelines/Private/InterchangeMaterialXPipeline.cpp
+++ b/Engine/Plugins/Interchange/Runtime/Source/Pipelines/Private/InterchangeMaterialXPipeline.cpp
@@ -29,7 +29,7 @@ namespace
 	
 	MATERIALX_FUNCTIONS_SUBSTRATE_PATH(OrenNayarBSDF);
 	MATERIALX_FUNCTIONS_SUBSTRATE_PATH(BurleyDiffuseBSDF);
-	MATERIALX_FUNCTIONS_SUBSTRATE_PATH(DieletricBSDF);
+	MATERIALX_FUNCTIONS_SUBSTRATE_PATH(DielectricBSDF);
 	MATERIALX_FUNCTIONS_SUBSTRATE_PATH(ConductorBSDF);
 	MATERIALX_FUNCTIONS_SUBSTRATE_PATH(SheenBSDF);
 	MATERIALX_FUNCTIONS_SUBSTRATE_PATH(SubsurfaceBSDF);
@@ -504,7 +504,7 @@ UInterchangeMaterialXPipeline::UInterchangeMaterialXPipeline()
 
 			{MATERIALX_MATERIALFUNCTION_PATH(OrenNayarBSDF),	      UMaterialXPipelineSettings::ToEnumKey(EInterchangeMaterialXBSDF::OrenNayarDiffuse)},
 			{MATERIALX_MATERIALFUNCTION_PATH(BurleyDiffuseBSDF),	  UMaterialXPipelineSettings::ToEnumKey(EInterchangeMaterialXBSDF::BurleyDiffuse)},
-			{MATERIALX_MATERIALFUNCTION_PATH(DieletricBSDF),		  UMaterialXPipelineSettings::ToEnumKey(EInterchangeMaterialXBSDF::Dielectric)},
+			{MATERIALX_MATERIALFUNCTION_PATH(DielectricBSDF),		  UMaterialXPipelineSettings::ToEnumKey(EInterchangeMaterialXBSDF::Dielectric)},
 			{MATERIALX_MATERIALFUNCTION_PATH(ConductorBSDF),		  UMaterialXPipelineSettings::ToEnumKey(EInterchangeMaterialXBSDF::Conductor)},
 			{MATERIALX_MATERIALFUNCTION_PATH(SheenBSDF),			  UMaterialXPipelineSettings::ToEnumKey(EInterchangeMaterialXBSDF::Sheen)},
 			{MATERIALX_MATERIALFUNCTION_PATH(SubsurfaceBSDF),		  UMaterialXPipelineSettings::ToEnumKey(EInterchangeMaterialXBSDF::Subsurface)},
@@ -664,7 +664,7 @@ void UMaterialXPipelineSettings::InitPredefinedAssets()
 
 			{UMaterialXPipelineSettings::ToEnumKey(EInterchangeMaterialXBSDF::OrenNayarDiffuse), OrenNayarBSDFFunctionsPath, OrenNayarBSDFSubstratePath},
 			{UMaterialXPipelineSettings::ToEnumKey(EInterchangeMaterialXBSDF::BurleyDiffuse), BurleyDiffuseBSDFFunctionsPath, BurleyDiffuseBSDFSubstratePath},
-			{UMaterialXPipelineSettings::ToEnumKey(EInterchangeMaterialXBSDF::Dielectric), DieletricBSDFFunctionsPath, DieletricBSDFSubstratePath},
+			{UMaterialXPipelineSettings::ToEnumKey(EInterchangeMaterialXBSDF::Dielectric), DielectricBSDFFunctionsPath, DielectricBSDFSubstratePath},
 			{UMaterialXPipelineSettings::ToEnumKey(EInterchangeMaterialXBSDF::Conductor), ConductorBSDFFunctionsPath, ConductorBSDFSubstratePath},
 			{UMaterialXPipelineSettings::ToEnumKey(EInterchangeMaterialXBSDF::Sheen), SheenBSDFFunctionsPath, SheenBSDFSubstratePath},
 			{UMaterialXPipelineSettings::ToEnumKey(EInterchangeMaterialXBSDF::Subsurface), SubsurfaceBSDFFunctionsPath, SubsurfaceBSDFSubstratePath},
