Add automatic skeletal mesh detection if there is some animations

#rnx

diff --git a/Engine/Plugins/Interchange/Runtime/Source/Pipelines/Private/InterchangeGenericAssetsPipeline.cpp b/Engine/Plugins/Interchange/Runtime/Source/Pipelines/Private/InterchangeGenericAssetsPipeline.cpp
--- a/Engine/Plugins/Interchange/Runtime/Source/Pipelines/Private/InterchangeGenericAssetsPipeline.cpp
+++ b/Engine/Plugins/Interchange/Runtime/Source/Pipelines/Private/InterchangeGenericAssetsPipeline.cpp
@@ -138,7 +138,7 @@ void UInterchangeGenericAssetsPipeline::FilterPropertiesFromTranslatedData(UInte
 		MeshPipeline->FilterPropertiesFromTranslatedData(InBaseNodeContainer);
 		AnimationPipeline->FilterPropertiesFromTranslatedData(InBaseNodeContainer);
 
-		UInterchangePipelineMeshesUtilities* PipelineMeshesUtilities = UInterchangeGenericMeshPipeline::CreateMeshPipelineUtilities(InBaseNodeContainer, MeshPipeline);
+		UInterchangePipelineMeshesUtilities* PipelineMeshesUtilities = UInterchangeGenericMeshPipeline::CreateMeshPipelineUtilities(InBaseNodeContainer, MeshPipeline, CommonMeshesProperties->bAutoDetectMeshType);
 
 		TArray<FString> SkeletalMeshes;
 		PipelineMeshesUtilities->GetAllSkinnedMeshInstance(SkeletalMeshes);
diff --git a/Engine/Plugins/Interchange/Runtime/Source/Pipelines/Private/InterchangeGenericMeshPipeline.cpp b/Engine/Plugins/Interchange/Runtime/Source/Pipelines/Private/InterchangeGenericMeshPipeline.cpp
--- a/Engine/Plugins/Interchange/Runtime/Source/Pipelines/Private/InterchangeGenericMeshPipeline.cpp
+++ b/Engine/Plugins/Interchange/Runtime/Source/Pipelines/Private/InterchangeGenericMeshPipeline.cpp
@@ -5,10 +5,12 @@
 #include "Animation/Skeleton.h"
 #include "Engine/SkeletalMesh.h"
 #include "Engine/StaticMesh.h"
+#include "InterchangeAnimationTrackSetNode.h"
 #include "InterchangeGenericAssetsPipeline.h"
 #include "InterchangeMeshNode.h"
 #include "InterchangePipelineLog.h"
 #include "InterchangePipelineMeshesUtilities.h"
+#include "InterchangeSceneNode.h"
 #include "InterchangeSkeletalMeshFactoryNode.h"
 #include "InterchangeStaticMeshFactoryNode.h"
 #include "InterchangeSourceData.h"
@@ -149,13 +151,53 @@ bool UInterchangeGenericMeshPipeline::GetPropertyPossibleValues(const FName Prop
 
 #endif
 
-UInterchangePipelineMeshesUtilities* UInterchangeGenericMeshPipeline::CreateMeshPipelineUtilities(UInterchangeBaseNodeContainer* InBaseNodeContainer, const UInterchangeGenericMeshPipeline* Pipeline)
+UInterchangePipelineMeshesUtilities* UInterchangeGenericMeshPipeline::CreateMeshPipelineUtilities(UInterchangeBaseNodeContainer* InBaseNodeContainer
+	, const UInterchangeGenericMeshPipeline* Pipeline
+	, const bool bAutoDetectType)
 {
 	UInterchangePipelineMeshesUtilities* CreatedPipelineMeshesUtilities = UInterchangePipelineMeshesUtilities::CreateInterchangePipelineMeshesUtilities(InBaseNodeContainer);
 
+	bool bAutoDetectConvertStaticMeshToSkeletalMesh = false;
+	if (bAutoDetectType && Pipeline->CommonMeshesProperties->ForceAllMeshAsType == EInterchangeForceMeshType::IFMT_None)
+	{
+		TArray<FString> StaticMeshNodeUids;
+		bool bContainSkeletalMesh = false;
+		InBaseNodeContainer->IterateNodesOfType<UInterchangeMeshNode>([&bContainSkeletalMesh, &StaticMeshNodeUids](const FString& NodeUid, UInterchangeMeshNode* MeshNode)
+			{
+				if (!MeshNode->IsMorphTarget())
+				{
+					MeshNode->IsSkinnedMesh() ? bContainSkeletalMesh = true : StaticMeshNodeUids.Add(NodeUid);
+				}
+			});
+		
+		bool bContainAnimationNode = false;
+		if (!bContainSkeletalMesh && StaticMeshNodeUids.Num() > 0)
+		{
+			TMap<const UInterchangeSceneNode*, bool> CacheProcessSceneNodes;
+			InBaseNodeContainer->BreakableIterateNodesOfType<UInterchangeTransformAnimationTrackNode>([&InBaseNodeContainer, &bContainAnimationNode, &StaticMeshNodeUids, &CacheProcessSceneNodes](const FString& NodeUid, UInterchangeTransformAnimationTrackNode* AnimationNode)
+				{
+					FString SceneNodeUid;
+					if (AnimationNode->GetCustomActorDependencyUid(SceneNodeUid))
+					{
+						if (const UInterchangeSceneNode* SceneNode = Cast<UInterchangeSceneNode>(InBaseNodeContainer->GetNode(SceneNodeUid)))
+						{
+							if (IsImpactingAnyMeshesRecursive(SceneNode, InBaseNodeContainer, StaticMeshNodeUids, CacheProcessSceneNodes))
+							{
+								bContainAnimationNode = true;
+							}
+						}
+					}
+					return bContainAnimationNode;
+				});
+		}
+
+		//Auto detect some static mesh transform animations, we need to force the skeletal mesh type and recompute
+		bAutoDetectConvertStaticMeshToSkeletalMesh = bContainAnimationNode;
+	}
+
 	//Set the context option to use when querying the pipeline mesh utilities
 	FInterchangePipelineMeshesUtilitiesContext DataContext;
-	DataContext.bConvertStaticMeshToSkeletalMesh = (Pipeline->CommonMeshesProperties->ForceAllMeshAsType == EInterchangeForceMeshType::IFMT_SkeletalMesh);
+	DataContext.bConvertStaticMeshToSkeletalMesh = bAutoDetectConvertStaticMeshToSkeletalMesh || (Pipeline->CommonMeshesProperties->ForceAllMeshAsType == EInterchangeForceMeshType::IFMT_SkeletalMesh);
 	DataContext.bConvertSkeletalMeshToStaticMesh = (Pipeline->CommonMeshesProperties->ForceAllMeshAsType == EInterchangeForceMeshType::IFMT_StaticMesh);
 	DataContext.bConvertStaticsWithMorphTargetsToSkeletals = Pipeline->CommonSkeletalMeshesAndAnimationsProperties->bConvertStaticsWithMorphTargetsToSkeletals;
 	DataContext.bImportMeshesInBoneHierarchy = Pipeline->CommonSkeletalMeshesAndAnimationsProperties->bImportMeshesInBoneHierarchy;
@@ -178,7 +220,7 @@ void UInterchangeGenericMeshPipeline::ExecutePipeline(UInterchangeBaseNodeContai
 	{
 		SourceDatas.Add(SourceData);
 	}
-	PipelineMeshesUtilities = CreateMeshPipelineUtilities(BaseNodeContainer, this);
+	PipelineMeshesUtilities = CreateMeshPipelineUtilities(BaseNodeContainer, this, CommonMeshesProperties->bAutoDetectMeshType);
 
 	//Create skeletalmesh factory nodes
 	ExecutePreImportPipelineSkeletalMesh();
@@ -293,4 +335,37 @@ bool UInterchangeGenericMeshPipeline::DoClassesIncludeAllEditableStructPropertie
 	}
 	return bResult;
 }
-#endif
\ No newline at end of file
+#endif
+
+bool UInterchangeGenericMeshPipeline::IsImpactingAnyMeshesRecursive(const UInterchangeSceneNode* SceneNode
+	, const UInterchangeBaseNodeContainer* InBaseNodeContainer
+	, const TArray<FString>& StaticMeshNodeUids
+	, TMap<const UInterchangeSceneNode*, bool>& CacheProcessSceneNodes)
+{
+	bool& bIsImpactingCache = CacheProcessSceneNodes.FindOrAdd(SceneNode, false);
+	if (bIsImpactingCache)
+	{
+		return bIsImpactingCache;
+	}
+	FString AssetUid;
+	if (SceneNode->GetCustomAssetInstanceUid(AssetUid))
+	{
+		if (StaticMeshNodeUids.Contains(AssetUid))
+		{
+			bIsImpactingCache = true;
+			return true;
+		}
+	}
+	TArray<FString> Children = InBaseNodeContainer->GetNodeChildrenUids(SceneNode->GetUniqueID());
+	for (const FString& ChildUid : Children)
+	{
+		if (const UInterchangeSceneNode* ChildSceneNode = Cast<UInterchangeSceneNode>(InBaseNodeContainer->GetNode(ChildUid)))
+		{
+			if (IsImpactingAnyMeshesRecursive(ChildSceneNode, InBaseNodeContainer, StaticMeshNodeUids, CacheProcessSceneNodes))
+			{
+				return true;
+			}
+		}
+	}
+	return false;
+}
\ No newline at end of file
diff --git a/Engine/Plugins/Interchange/Runtime/Source/Pipelines/Public/InterchangeGenericAssetsPipelineSharedSettings.h b/Engine/Plugins/Interchange/Runtime/Source/Pipelines/Public/InterchangeGenericAssetsPipelineSharedSettings.h
--- a/Engine/Plugins/Interchange/Runtime/Source/Pipelines/Public/InterchangeGenericAssetsPipelineSharedSettings.h
+++ b/Engine/Plugins/Interchange/Runtime/Source/Pipelines/Public/InterchangeGenericAssetsPipelineSharedSettings.h
@@ -57,6 +57,13 @@ public:
  	UPROPERTY(EditAnywhere, BlueprintReadWrite, Category = "Common Meshes")
  	EInterchangeForceMeshType ForceAllMeshAsType = EInterchangeForceMeshType::IFMT_None;
 
+	/**
+	 * If enable, and some static mesh transforms are animated, the pipeline will convert the static mesh into rigid skeletal mesh.
+	 * Note: This setting will be ignored if ForceAllMeshAsType is not set to 'None'.
+	 */
+	UPROPERTY(EditAnywhere, BlueprintReadWrite, Category = "Common Meshes")
+	bool bAutoDetectMeshType = true;
+
 	/** If enable, meshes LODs will be imported. Note that it required the advanced bBakeMesh property to be enabled. */
 	UPROPERTY(EditAnywhere, BlueprintReadWrite, Category = "Common Meshes")
 	bool bImportLods = true;
@@ -111,6 +118,10 @@ public:
 		{
 			return true;
 		}
+		if (PropertyChangedEvent.Property->GetName() == GET_MEMBER_NAME_CHECKED(UInterchangeGenericCommonMeshesProperties, bAutoDetectMeshType))
+		{
+			return true;
+		}
 		if (PropertyChangedEvent.Property->GetName() == GET_MEMBER_NAME_CHECKED(UInterchangeGenericCommonMeshesProperties, bBakeMeshes))
 		{
 			return true;
diff --git a/Engine/Plugins/Interchange/Runtime/Source/Pipelines/Public/InterchangeGenericMeshPipeline.h b/Engine/Plugins/Interchange/Runtime/Source/Pipelines/Public/InterchangeGenericMeshPipeline.h
--- a/Engine/Plugins/Interchange/Runtime/Source/Pipelines/Public/InterchangeGenericMeshPipeline.h
+++ b/Engine/Plugins/Interchange/Runtime/Source/Pipelines/Public/InterchangeGenericMeshPipeline.h
@@ -15,6 +15,7 @@
 class UInterchangeGenericAssetsPipeline;
 class UInterchangeMeshNode;
 class UInterchangePipelineMeshesUtilities;
+class UInterchangeSceneNode;
 class UInterchangeSkeletalMeshFactoryNode;
 class UInterchangeSkeletalMeshLodDataNode;
 class UInterchangeSkeletonFactoryNode;
@@ -205,7 +206,9 @@ public:
 	virtual bool GetPropertyPossibleValues(const FName PropertyPath, TArray<FString>& PossibleValues) override;
 #endif
 
-	static UInterchangePipelineMeshesUtilities* CreateMeshPipelineUtilities(UInterchangeBaseNodeContainer* InBaseNodeContainer, const UInterchangeGenericMeshPipeline* Pipeline);
+	static UInterchangePipelineMeshesUtilities* CreateMeshPipelineUtilities(UInterchangeBaseNodeContainer* InBaseNodeContainer
+		, const UInterchangeGenericMeshPipeline* Pipeline
+		, const bool bAutoDetectType);
 
 protected:
 	virtual void ExecutePipeline(UInterchangeBaseNodeContainer* InBaseNodeContainer, const TArray<UInterchangeSourceData*>& InSourceDatas) override;
@@ -236,6 +239,11 @@ private:
 	/* Meshes utilities, to parse the translated graph and extract the meshes informations. */
 	TObjectPtr<UInterchangePipelineMeshesUtilities> PipelineMeshesUtilities = nullptr;
 
+	static bool IsImpactingAnyMeshesRecursive(const UInterchangeSceneNode* SceneNode
+		, const UInterchangeBaseNodeContainer* InBaseNodeContainer
+		, const TArray<FString>& StaticMeshNodeUids
+		, TMap<const UInterchangeSceneNode*, bool>& CacheProcessSceneNodes);
+
 	/************************************************************************/
 	/* Skeletal mesh API BEGIN                                              */
 
