Interchange: Fix animation issues
- Skeletal mesh factory now merge the payloadkey with the globaltransform to have a unique identifier to retrieve the payload result.
- Interchange fbx translator: triangulate all meshes before parsing the scene. This avoid bad pointer in case the fbx sdk triangulation re-assign the mesh resource.
- Interchange fbx translator: node with FbxNodeAttribute::eNull will not be consider a joint if there is no joint(FbxNodeAttribute::eSkeleton) in the parent hierarchy.
- Interchange fbx translator: the result payloads key for mesh now also merge the payloadkey with the globaltransform to have a unique identifier to retrieve the payload result.
Note this cl cannot go in 5.3 because some header have change.

#rnx

diff --git a/Engine/Plugins/Interchange/Runtime/Source/Dispatcher/Public/InterchangeCommands.h b/Engine/Plugins/Interchange/Runtime/Source/Dispatcher/Public/InterchangeCommands.h
--- a/Engine/Plugins/Interchange/Runtime/Source/Dispatcher/Public/InterchangeCommands.h
+++ b/Engine/Plugins/Interchange/Runtime/Source/Dispatcher/Public/InterchangeCommands.h
@@ -17,7 +17,7 @@ namespace UE
 		{
 		public:
 			//Major version should be updated when there is an existing API that has been change.
-			static int32 GetMajor() { return 15; }
+			static int32 GetMajor() { return 16; }
 			//Minor version must be updated when there is an addition to the API.
 			static int32 GetMinor() { return 0; }
 			//Patch version should be update if there is some bug fixes in the private code.
diff --git a/Engine/Plugins/Interchange/Runtime/Source/Import/Private/Mesh/InterchangeSkeletalMeshFactory.cpp b/Engine/Plugins/Interchange/Runtime/Source/Import/Private/Mesh/InterchangeSkeletalMeshFactory.cpp
--- a/Engine/Plugins/Interchange/Runtime/Source/Import/Private/Mesh/InterchangeSkeletalMeshFactory.cpp
+++ b/Engine/Plugins/Interchange/Runtime/Source/Import/Private/Mesh/InterchangeSkeletalMeshFactory.cpp
@@ -67,6 +67,15 @@ namespace UE
 				const UInterchangeSceneNode* SceneNode = nullptr;
 				TOptional<FTransform> SceneGlobalTransform;
 				FInterchangeMeshPayLoadKey TranslatorPayloadKey;
+				FString GetUniqueId() const
+				{
+					FString UniqueId = TranslatorPayloadKey.UniqueId;
+					if (SceneGlobalTransform.IsSet())
+					{
+						UniqueId += SceneGlobalTransform->ToString();
+					}
+					return UniqueId;
+				}
 			};
 
 			void FillMorphTargetMeshDescriptionsPerMorphTargetName(const FMeshNodeContext& MeshNodeContext
@@ -431,7 +440,7 @@ namespace UE
 				bool bImportMorphTarget = true;
 				SkeletalMeshFactoryNode->GetCustomImportMorphTarget(bImportMorphTarget);
 
-				TMap<FString, TFuture<TOptional<UE::Interchange::FMeshPayloadData>>> LodMeshPayloadPerTranslatorPayloadKey;
+				TMap<const FMeshNodeContext*, TFuture<TOptional<UE::Interchange::FMeshPayloadData>>> LodMeshPayloadPerTranslatorPayloadKey;
 				LodMeshPayloadPerTranslatorPayloadKey.Reserve(MeshReferences.Num());
 
 				TMap<FString, TOptional<UE::Interchange::FMeshPayloadData>> MorphTargetMeshDescriptionsPerMorphTargetName;
@@ -439,26 +448,50 @@ namespace UE
 
 				bool bImportVertexAttributes = false;
 				SkeletalMeshFactoryNode->GetCustomImportVertexAttributes(bImportVertexAttributes);
+				struct FInternalInstanceData
+				{
+					bool ScaleGreaterThenOne = false;
+					int32 Count = 0;
+					bool ShouldFetchWithTransform() const
+					{
+						return Count == 1 || ScaleGreaterThenOne;
+					}
+				};
+				TMap<FString, FInternalInstanceData> MeshInstancesDatas;
+				for (const FMeshNodeContext& MeshNodeContext : MeshReferences)
+				{
+					FInternalInstanceData& InstanceData = MeshInstancesDatas.FindOrAdd(MeshNodeContext.TranslatorPayloadKey.UniqueId);
+					InstanceData.Count++;
+					InstanceData.ScaleGreaterThenOne |= MeshNodeContext.SceneGlobalTransform->GetScale3D().GetAbs().GetMax() > 1.0;
+				}
 
 				for (const FMeshNodeContext& MeshNodeContext : MeshReferences)
 				{
+					const FInternalInstanceData& InstanceData = MeshInstancesDatas.FindChecked(MeshNodeContext.TranslatorPayloadKey.UniqueId);
+					FTransform ApplyTransformWhenFetchPayload = InstanceData.ShouldFetchWithTransform() ? MeshNodeContext.SceneGlobalTransform.Get(FTransform::Identity) : FTransform::Identity;
 					//Add the payload entry key, the payload data will be fill later in bulk by the translator
-					LodMeshPayloadPerTranslatorPayloadKey.Add(MeshNodeContext.TranslatorPayloadKey.UniqueId, MeshTranslatorPayloadInterface->GetMeshPayloadData(MeshNodeContext.TranslatorPayloadKey, MeshNodeContext.SceneGlobalTransform.Get(FTransform::Identity)));
+					LodMeshPayloadPerTranslatorPayloadKey.Add(&MeshNodeContext, MeshTranslatorPayloadInterface->GetMeshPayloadData(MeshNodeContext.TranslatorPayloadKey, ApplyTransformWhenFetchPayload));
 					//Count the morph target dependencies so we can reserve the right amount
 					MorphTargetCount += (bImportMorphTarget && MeshNodeContext.MeshNode) ? MeshNodeContext.MeshNode->GetMorphTargetDependeciesCount() : 0;
 				}
 				MorphTargetMeshDescriptionsPerMorphTargetName.Reserve(MorphTargetCount);
 
 				//Fill the lod mesh description using all combined mesh part
-				for (const FMeshNodeContext& MeshNodeContext : MeshReferences)
+				for(TPair<const FMeshNodeContext*, TFuture<TOptional<UE::Interchange::FMeshPayloadData>>>& MeshNodeContextAndFuture: LodMeshPayloadPerTranslatorPayloadKey)
 				{
+					if (!MeshNodeContextAndFuture.Key)
+					{
+						continue;
+					}
+					const FMeshNodeContext& MeshNodeContext = *MeshNodeContextAndFuture.Key;
 					TRACE_CPUPROFILER_EVENT_SCOPE("RetrieveAllSkeletalMeshPayloadsAndFillImportData::GetPayload")
-					TOptional<UE::Interchange::FMeshPayloadData> LodMeshPayload = LodMeshPayloadPerTranslatorPayloadKey.FindChecked(MeshNodeContext.TranslatorPayloadKey.UniqueId).Get();
+					TOptional<UE::Interchange::FMeshPayloadData> LodMeshPayload = MeshNodeContextAndFuture.Value.Get();
 					if (!LodMeshPayload.IsSet())
 					{
 						UE_LOG(LogInterchangeImport, Warning, TEXT("Invalid Skeletal mesh payload key [%s] SkeletalMesh asset %s"), *MeshNodeContext.TranslatorPayloadKey.UniqueId, *Arguments.AssetName);
 						continue;
 					}
+					
 					const int32 VertexOffset = LodMeshDescription.Vertices().Num();
 
 					FSkeletalMeshOperations::FSkeletalMeshAppendSettings SkeletalMeshAppendSettings;
@@ -475,8 +508,7 @@ namespace UE
 					{
 						//We need to rebind the mesh at time 0. Skeleton joint have the time zero transform, so we need to apply the skinning to the mesh
 						//With the skeleton transform at time zero
-						FTransform MeshGlobalTransform;
-						MeshGlobalTransform.SetIdentity();
+						FTransform MeshGlobalTransform = FTransform::Identity;
 						if (MeshNodeContext.SceneGlobalTransform.IsSet())
 						{
 							MeshGlobalTransform = MeshNodeContext.SceneGlobalTransform.GetValue();
@@ -539,8 +571,9 @@ namespace UE
 						}
 					}
 					
-					//The Mesh node parent bake transform was pass to the payload request.
-					AppendSettings.MeshTransform.Reset();
+					//The Mesh node parent bake transform can be pass to the payload request or not, it depend on the count of instance and the scale of the transform.
+					const FInternalInstanceData & InstanceData = MeshInstancesDatas.FindChecked(MeshNodeContext.TranslatorPayloadKey.UniqueId);
+					AppendSettings.MeshTransform = InstanceData.ShouldFetchWithTransform() ? FTransform::Identity : MeshNodeContext.SceneGlobalTransform.Get(FTransform::Identity);
 
 					FStaticMeshOperations::AppendMeshDescription(LodMeshPayload->MeshDescription, LodMeshDescription, AppendSettings);
 					if (MeshNodeContext.MeshNode->IsSkinnedMesh() || bIsRigidMesh)
diff --git a/Engine/Plugins/Interchange/Runtime/Source/Parsers/Fbx/Private/FbxMesh.cpp b/Engine/Plugins/Interchange/Runtime/Source/Parsers/Fbx/Private/FbxMesh.cpp
--- a/Engine/Plugins/Interchange/Runtime/Source/Parsers/Fbx/Private/FbxMesh.cpp
+++ b/Engine/Plugins/Interchange/Runtime/Source/Parsers/Fbx/Private/FbxMesh.cpp
@@ -576,9 +576,11 @@ bool FMeshDescriptionImporter::FillMeshDescriptionFromFbxMesh(FbxMesh* Mesh, con
 		}
 	}
 
-	if (!Mesh->IsTriangleMesh())
+	//Mesh must be triangulate when creating the payload context, we cannot change the Mesh pointer after
+	if(!Mesh->IsTriangleMesh())
 	{
-		const bool bReplace = true;
+		//Since we want to avoid deleting the pointer, we set the bReplace to false
+		constexpr bool bReplace = false;
 		FbxNodeAttribute* ConvertedNode = SDKGeometryConverter->Triangulate(Mesh, bReplace);
 
 		if (ConvertedNode != NULL && ConvertedNode->GetAttributeType() == FbxNodeAttribute::eMesh)
@@ -1570,6 +1572,35 @@ void GetMaterialIndex(FbxMesh* Mesh, TArray<int32>& MaterialIndexes)
 void FFbxMesh::AddAllMeshes(FbxScene* SDKScene, FbxGeometryConverter* SDKGeometryConverter, UInterchangeBaseNodeContainer& NodeContainer, TMap<FString, TSharedPtr<FPayloadContextBase>>& PayloadContexts)
 {
 	int32 GeometryCount = SDKScene->GetGeometryCount();
+	//Triangulate meshes
+	{
+		TArray<FbxMesh*> ToTriangulateMeshes;
+		ToTriangulateMeshes.Reserve(GeometryCount);
+		for (int32 GeometryIndex = 0; GeometryIndex < GeometryCount; ++GeometryIndex)
+		{
+			FbxGeometry* Geometry = SDKScene->GetGeometry(GeometryIndex);
+			if (Geometry->GetAttributeType() != FbxNodeAttribute::eMesh)
+			{
+				continue;
+			}
+			FbxMesh* Mesh = static_cast<FbxMesh*>(Geometry);
+			if (!Mesh)
+			{
+				continue;
+			}
+			if (!Mesh->IsTriangleMesh())
+			{
+				ToTriangulateMeshes.Add(Mesh);
+			}
+		}
+		for (FbxMesh* ToTriangulateMesh : ToTriangulateMeshes)
+		{
+			const bool bReplace = true;
+			SDKGeometryConverter->Triangulate(ToTriangulateMesh, bReplace);
+		}
+	}
+	//Now requery the triangulated geometries
+	GeometryCount = SDKScene->GetGeometryCount();
 	for (int32 GeometryIndex = 0; GeometryIndex < GeometryCount; ++GeometryIndex)
 	{
 		FbxGeometry* Geometry = SDKScene->GetGeometry(GeometryIndex);
@@ -1578,10 +1609,11 @@ void FFbxMesh::AddAllMeshes(FbxScene* SDKScene, FbxGeometryConverter* SDKGeometr
 			continue;
 		}
 		FbxMesh* Mesh = static_cast<FbxMesh*>(Geometry);
-		if (!Mesh)
+		if (!Mesh || !Mesh->IsTriangleMesh())
 		{
 			continue;
 		}
+
 		FString MeshName = Parser.GetFbxHelper()->GetMeshName(Mesh);
 		FString MeshUniqueID = Parser.GetFbxHelper()->GetMeshUniqueID(Mesh);
 		const UInterchangeMeshNode* ExistingMeshNode = Cast<UInterchangeMeshNode>(NodeContainer.GetNode(MeshUniqueID));
diff --git a/Engine/Plugins/Interchange/Runtime/Source/Parsers/Fbx/Private/FbxScene.cpp b/Engine/Plugins/Interchange/Runtime/Source/Parsers/Fbx/Private/FbxScene.cpp
--- a/Engine/Plugins/Interchange/Runtime/Source/Parsers/Fbx/Private/FbxScene.cpp
+++ b/Engine/Plugins/Interchange/Runtime/Source/Parsers/Fbx/Private/FbxScene.cpp
@@ -239,8 +239,9 @@ namespace UE
 						{
 							if (!DoesTheParentHierarchyContainJoints(Node->GetParent()))
 							{
-								//eNull node will be set has a transform and a joint specialized type
+								//eNull node not under any joint will be set has a transform
 								UnrealNode->AddSpecializedType(FSceneNodeStaticData::GetTransformSpecializeTypeString());
+								break;
 							}
 						}
 						//No break since the eNull act has a skeleton if possible
@@ -293,7 +294,6 @@ namespace UE
 				//If there is no attribute, make sure to threat the node like a joint if it's in the ForcejointNodes array
 				if (!bIsNodeContainJointAttribute && ForceJointNodes.Contains(Node))
 				{
-					UnrealNode->AddSpecializedType(FSceneNodeStaticData::GetTransformSpecializeTypeString());
 					ApplySkeletonAttribute();
 				}
 
@@ -655,6 +655,11 @@ namespace UE
 						switch (NodeAttribute->GetAttributeType())
 						{
 						case FbxNodeAttribute::eNull:
+							if (!DoesTheParentHierarchyContainJoints(Node->GetParent()))
+							{
+								//eNull node not under any joint are not joint
+								break;
+							}
 						case FbxNodeAttribute::eSkeleton:
 							bIsNodeContainJointAttribute = true;
 							ApplySkeletonAttribute();
diff --git a/Engine/Plugins/Interchange/Runtime/Source/Parsers/Fbx/Private/InterchangeFbxParser.cpp b/Engine/Plugins/Interchange/Runtime/Source/Parsers/Fbx/Private/InterchangeFbxParser.cpp
--- a/Engine/Plugins/Interchange/Runtime/Source/Parsers/Fbx/Private/InterchangeFbxParser.cpp
+++ b/Engine/Plugins/Interchange/Runtime/Source/Parsers/Fbx/Private/InterchangeFbxParser.cpp
@@ -104,14 +104,21 @@ namespace UE::Interchange
 		}
 	}
 
-	void FInterchangeFbxParser::FetchMeshPayload(const FString& PayloadKey, const FTransform& MeshGlobalTransform, const FString& ResultFolder)
+	FString FInterchangeFbxParser::FetchMeshPayload(const FString& PayloadKey, const FTransform& MeshGlobalTransform, const FString& ResultFolder)
 	{
 		check(FbxParserPrivate.IsValid());
 		ResultsContainer->Empty();
 		FString PayloadFilepathCopy;
+		FString ResultPayloadUniqueId = PayloadKey + MeshGlobalTransform.ToString();
+		//If we already have extract this mesh, no need to extract again
+		if (ResultPayloads.Contains(ResultPayloadUniqueId))
+		{
+			return ResultPayloadUniqueId;
+		}
+
 		{
 			FScopeLock Lock(&ResultPayloadsCriticalSection);
-			FString& PayloadFilepath = ResultPayloads.FindOrAdd(PayloadKey);
+			FString& PayloadFilepath = ResultPayloads.FindOrAdd(ResultPayloadUniqueId);
 			//To avoid file path with too many character, we hash the payloadKey so we have a deterministic length for the file path.
 			FString PayloadKeyHash = Private::HashString(PayloadKey);
 			PayloadFilepath = ResultFolder + TEXT("/") + PayloadKeyHash + FString::FromInt(UniqueIdCounter.IncrementExchange()) + TEXT(".payload");
@@ -124,18 +131,25 @@ namespace UE::Interchange
 			UInterchangeResultError_Generic* Error = AddMessage<UInterchangeResultError_Generic>();
 			Error->SourceAssetName = SourceFilename;
 			Error->Text = LOCTEXT("CantFetchPayload", "Cannot fetch FBX payload data.");
-			return;
 		}
+		return ResultPayloadUniqueId;
 	}
 
-	void FInterchangeFbxParser::FetchAnimationBakeTransformPayload(const FString& PayloadKey, const double BakeFrequency, const double RangeStartTime, const double RangeEndTime, const FString& ResultFolder)
+	FString FInterchangeFbxParser::FetchAnimationBakeTransformPayload(const FString& PayloadKey, const double BakeFrequency, const double RangeStartTime, const double RangeEndTime, const FString& ResultFolder)
 	{
 		check(FbxParserPrivate.IsValid());
 		ResultsContainer->Empty();
 		FString PayloadFilepathCopy;
+		FString ResultPayloadUniqueId = PayloadKey + FString::FromInt(static_cast<int32>(BakeFrequency*1000.0)) + FString::FromInt(static_cast<int32>(RangeStartTime * 1000.0)) + FString::FromInt(static_cast<int32>(RangeEndTime * 1000.0));
+		//If we already have extract this mesh, no need to extract again
+		if (ResultPayloads.Contains(ResultPayloadUniqueId))
+		{
+			return ResultPayloadUniqueId;
+		}
+
 		{
 			FScopeLock Lock(&ResultPayloadsCriticalSection);
-			FString& PayloadFilepath = ResultPayloads.FindOrAdd(PayloadKey);
+			FString& PayloadFilepath = ResultPayloads.FindOrAdd(ResultPayloadUniqueId);
 			//To avoid file path with too many character, we hash the payloadKey so we have a deterministic length for the file path.
 			FString PayloadKeyHash = Private::HashString(PayloadKey);
 			PayloadFilepath = ResultFolder + TEXT("/") + PayloadKeyHash + FString::FromInt(UniqueIdCounter.IncrementExchange()) + TEXT(".payload");
@@ -148,8 +162,8 @@ namespace UE::Interchange
 			UInterchangeResultError_Generic* Error = AddMessage<UInterchangeResultError_Generic>();
 			Error->SourceAssetName = SourceFilename;
 			Error->Text = LOCTEXT("CantFetchPayload", "Cannot fetch FBX payload data.");
-			return;
 		}
+		return ResultPayloadUniqueId;
 	}
 
 	TArray<FString> FInterchangeFbxParser::GetJsonLoadMessages() const
diff --git a/Engine/Plugins/Interchange/Runtime/Source/Parsers/Fbx/Public/InterchangeFbxParser.h b/Engine/Plugins/Interchange/Runtime/Source/Parsers/Fbx/Public/InterchangeFbxParser.h
--- a/Engine/Plugins/Interchange/Runtime/Source/Parsers/Fbx/Public/InterchangeFbxParser.h
+++ b/Engine/Plugins/Interchange/Runtime/Source/Parsers/Fbx/Public/InterchangeFbxParser.h
@@ -43,8 +43,9 @@ namespace UE
 			 * @param - PayloadKey is the key that describe the payload data to extract from the fbx file
 			 * @param - MeshGlobalTransform is the transform we want to apply to the mesh vertex
 			 * @param - ResultFolder is the folder where we must put any result file
+			 * @return - Return the 'ResultPayloads' key unique id. We cannot use only the payload key because the mesh global transform can be different.
 			 */
-			void FetchMeshPayload(const FString& PayloadKey, const FTransform& MeshGlobalTransform, const FString& ResultFolder);
+			FString FetchMeshPayload(const FString& PayloadKey, const FTransform& MeshGlobalTransform, const FString& ResultFolder);
 
 			/**
 			 * Extract bake transform animation payload data from the fbx, the key tell the translator what payload the client ask
@@ -53,8 +54,9 @@ namespace UE
 			 * @param - RangeStartTime is the start time of the bake
 			 * @param - RangeEndTime is the end time of the bake
 			 * @param - ResultFolder is the folder where we must put any result file
+			 * @return - Return the 'ResultPayloads' key unique id. We cannot use only the payload key because the bake parameter can be different.
 			 */
-			void FetchAnimationBakeTransformPayload(const FString& PayloadKey, const double BakeFrequency, const double RangeStartTime, const double RangeEndTime, const FString& ResultFolder);
+			FString FetchAnimationBakeTransformPayload(const FString& PayloadKey, const double BakeFrequency, const double RangeStartTime, const double RangeEndTime, const FString& ResultFolder);
 
 			FString GetResultFilepath() const { return ResultFilepath; }
 			FString GetResultPayloadFilepath(const FString& PayloadKey) const
diff --git a/Engine/Source/Programs/InterchangeWorker/Private/InterchangeWorkerImpl.cpp b/Engine/Source/Programs/InterchangeWorker/Private/InterchangeWorkerImpl.cpp
--- a/Engine/Source/Programs/InterchangeWorker/Private/InterchangeWorkerImpl.cpp
+++ b/Engine/Source/Programs/InterchangeWorker/Private/InterchangeWorkerImpl.cpp
@@ -317,9 +317,9 @@ ETaskState FInterchangeWorkerImpl::FetchFbxPayload(const FJsonFetchMeshPayloadCm
 	ETaskState ResultState = ETaskState::Unknown;
 	FString PayloadKey = FetchMeshPayloadCommand.GetPayloadKey();
 	FTransform MeshGlobalTransform = FetchMeshPayloadCommand.GetMeshGlobalTransform();
-	FbxParser.FetchMeshPayload(PayloadKey, MeshGlobalTransform, ResultFolder);
+	FString ResultPayloadsUniqueId = FbxParser.FetchMeshPayload(PayloadKey, MeshGlobalTransform, ResultFolder);
 	FJsonLoadSourceCmd::JsonResultParser ResultParser;
-	ResultParser.SetResultFilename(FbxParser.GetResultPayloadFilepath(PayloadKey));
+	ResultParser.SetResultFilename(FbxParser.GetResultPayloadFilepath(ResultPayloadsUniqueId));
 	OutJSonMessages = FbxParser.GetJsonLoadMessages();
 	OutJSonResult = ResultParser.ToJson();
 	ResultState = ETaskState::ProcessOk;
@@ -331,13 +331,13 @@ ETaskState FInterchangeWorkerImpl::FetchFbxPayload(const FJsonFetchAnimationBake
 	ETaskState ResultState = ETaskState::Unknown;
 	FString PayloadKey = FetchAnimationBakeTransformPayloadCommand.GetPayloadKey();
 
-	FbxParser.FetchAnimationBakeTransformPayload(PayloadKey
+	FString ResultPayloadsUniqueId = FbxParser.FetchAnimationBakeTransformPayload(PayloadKey
 		, FetchAnimationBakeTransformPayloadCommand.GetBakeFrequency()
 		, FetchAnimationBakeTransformPayloadCommand.GetRangeStartTime()
 		, FetchAnimationBakeTransformPayloadCommand.GetRangeEndTime()
 		, ResultFolder);
 	FJsonLoadSourceCmd::JsonResultParser ResultParser;
-	ResultParser.SetResultFilename(FbxParser.GetResultPayloadFilepath(PayloadKey));
+	ResultParser.SetResultFilename(FbxParser.GetResultPayloadFilepath(ResultPayloadsUniqueId));
 	OutJSonMessages = FbxParser.GetJsonLoadMessages();
 	OutJSonResult = ResultParser.ToJson();
 	ResultState = ETaskState::ProcessOk;
