Interchange: Make sure the name clash in fbx do not use the unique id of the node since there can be conflict here also.

#rnx

diff --git a/Engine/Plugins/Interchange/Runtime/Source/Parsers/Fbx/Private/FbxAPI.cpp b/Engine/Plugins/Interchange/Runtime/Source/Parsers/Fbx/Private/FbxAPI.cpp
--- a/Engine/Plugins/Interchange/Runtime/Source/Parsers/Fbx/Private/FbxAPI.cpp
+++ b/Engine/Plugins/Interchange/Runtime/Source/Parsers/Fbx/Private/FbxAPI.cpp
@@ -200,21 +200,23 @@ namespace UE
 				}
 
 				//////////////////////////////////////////////////////////////////////////
-				// Esnure Node Name Validity (uniqueness)
-				TSet<FString> NodeNames;
+				// Ensure Node Name Validity (uniqueness)
+				// Name clash must be global because unreal bones do not support name conflict (they are stored in an array, no hierarchy)
+				TMap<FString, int32> NodeNames;
 				for (int32 NodeIndex = 0; NodeIndex < SDKScene->GetNodeCount(); ++NodeIndex)
 				{
 					FbxNode* Node = SDKScene->GetNode(NodeIndex);
-
 					FString NodeName = Node->GetName();
-					FString NodeUniqueID = GetFbxHelper()->GetFbxNodeHierarchyName(Node);
-
-					if (NodeNames.Contains(NodeName))
+					if (int32* Count = NodeNames.Find(NodeName))
 					{
-						NodeName += TEXT("_") + FMD5::HashAnsiString(*NodeUniqueID);
+						(*Count)++;
+						NodeName += TEXT("_ncl_") + FString::FromInt(*Count);
 						Node->SetName(TCHAR_TO_UTF8(*NodeName));
 					}
-					NodeNames.Add(NodeName);
+					else
+					{
+						NodeNames.Add(NodeName, 0);
+					}
 				}
 			}
 		} //ns Private
