Use of Standard Surface substrate material function from Engine instead of the one in Interchange. Keeping the ones in Interchange for testing for the moment.

diff --git a/Engine/Plugins/Interchange/Runtime/Source/Import/Private/MaterialX/MaterialXUtils/MaterialXStandardSurfaceShader.cpp b/Engine/Plugins/Interchange/Runtime/Source/Import/Private/MaterialX/MaterialXUtils/MaterialXStandardSurfaceShader.cpp
--- a/Engine/Plugins/Interchange/Runtime/Source/Import/Private/MaterialX/MaterialXUtils/MaterialXStandardSurfaceShader.cpp
+++ b/Engine/Plugins/Interchange/Runtime/Source/Import/Private/MaterialX/MaterialXUtils/MaterialXStandardSurfaceShader.cpp
@@ -183,7 +183,7 @@ void FMaterialXStandardSurfaceShader::ConnectToSubstrateStandardSurface()
 	using namespace mx::StandardSurface;
 	constexpr bool bInputInTangentSpace = true;
 
-	UInterchangeFunctionCallShaderNode* StandardSurfaceShaderNode = CreateFunctionCallShaderNode(SurfaceShaderNode->getName().c_str(), TEXT("/Interchange/Substrate/MX_StandardSurface.MX_StandardSurface"));
+	UInterchangeFunctionCallShaderNode* StandardSurfaceShaderNode = CreateFunctionCallShaderNode(SurfaceShaderNode->getName().c_str(), TEXT("/Engine/Functions/Strata/Substrate-StandardSurface-Opaque.Substrate-StandardSurface-Opaque"));
 
 	// Inputs
 	//Base
@@ -216,21 +216,6 @@ void FMaterialXStandardSurfaceShader::ConnectToSubstrateStandardSurface()
 	//Specular Rotation
 	ConnectNodeOutputToInput(Input::SpecularRotation, StandardSurfaceShaderNode, StandardSurface::Parameters::SpecularRotation.ToString(), DefaultValue::Float::SpecularRotation);
 
-	//Subsurface
-	ConnectNodeOutputToInput(Input::Subsurface, StandardSurfaceShaderNode, StandardSurface::Parameters::Subsurface.ToString(), DefaultValue::Float::Subsurface);
-
-	//Subsurface Color
-	ConnectNodeOutputToInput(Input::SubsurfaceColor, StandardSurfaceShaderNode, StandardSurface::Parameters::SubsurfaceColor.ToString(), DefaultValue::Color3::SubsurfaceColor);
-
-	//Subsurface Radius
-	ConnectNodeOutputToInput(Input::SubsurfaceRadius, StandardSurfaceShaderNode, StandardSurface::Parameters::SubsurfaceRadius.ToString(), DefaultValue::Color3::SubsurfaceRadius);
-
-	//Subsurface Scale
-	ConnectNodeOutputToInput(Input::SubsurfaceScale, StandardSurfaceShaderNode, StandardSurface::Parameters::SubsurfaceScale.ToString(), DefaultValue::Float::SubsurfaceScale);
-
-	//Subsurface Anisotropy
-	ConnectNodeOutputToInput(Input::SubsurfaceAnisotropy, StandardSurfaceShaderNode, StandardSurface::Parameters::SubsurfaceAnisotropy.ToString(), DefaultValue::Float::SubsurfaceAnisotropy);
-
 	//Sheen
 	ConnectNodeOutputToInput(Input::Sheen, StandardSurfaceShaderNode, StandardSurface::Parameters::Sheen.ToString(), DefaultValue::Float::Sheen);
 
@@ -300,13 +285,28 @@ void FMaterialXStandardSurfaceShader::ConnectToSubstrateStandardSurface()
 	// Outputs
 	if(UInterchangeShaderPortsAPI::HasInput(StandardSurfaceShaderNode, StandardSurface::Parameters::Transmission))
 	{
-		StandardSurfaceShaderNode->SetCustomMaterialFunction(TEXT("/Interchange/Substrate/MX_TransmissionSurface.MX_TransmissionSurface"));
+		StandardSurfaceShaderNode->SetCustomMaterialFunction(TEXT("/Engine/Functions/Strata/Substrate-StandardSurface-Translucent.Substrate-StandardSurface-Translucent"));
 		ShaderGraphNode->SetCustomBlendMode(EBlendMode::BLEND_TranslucentColoredTransmittance);
 		UInterchangeShaderPortsAPI::ConnectOuputToInputByName(ShaderGraphNode, PBRMR::Parameters::Opacity.ToString(), StandardSurfaceShaderNode->GetUniqueID(), PBRMR::Parameters::Opacity.ToString());
 		UInterchangeShaderPortsAPI::ConnectOuputToInputByName(ShaderGraphNode, SubstrateMaterial::Parameters::FrontMaterial.ToString(), StandardSurfaceShaderNode->GetUniqueID(), StandardSurface::SubstrateMaterial::Outputs::Translucent.ToString());
 	}
 	else
 	{
+		//Subsurface
+		ConnectNodeOutputToInput(Input::Subsurface, StandardSurfaceShaderNode, StandardSurface::Parameters::Subsurface.ToString(), DefaultValue::Float::Subsurface);
+
+		//Subsurface Color
+		ConnectNodeOutputToInput(Input::SubsurfaceColor, StandardSurfaceShaderNode, StandardSurface::Parameters::SubsurfaceColor.ToString(), DefaultValue::Color3::SubsurfaceColor);
+
+		//Subsurface Radius
+		ConnectNodeOutputToInput(Input::SubsurfaceRadius, StandardSurfaceShaderNode, StandardSurface::Parameters::SubsurfaceRadius.ToString(), DefaultValue::Color3::SubsurfaceRadius);
+
+		//Subsurface Scale
+		ConnectNodeOutputToInput(Input::SubsurfaceScale, StandardSurfaceShaderNode, StandardSurface::Parameters::SubsurfaceScale.ToString(), DefaultValue::Float::SubsurfaceScale);
+
+		//Subsurface Anisotropy
+		ConnectNodeOutputToInput(Input::SubsurfaceAnisotropy, StandardSurfaceShaderNode, StandardSurface::Parameters::SubsurfaceAnisotropy.ToString(), DefaultValue::Float::SubsurfaceAnisotropy);
+
 		UInterchangeShaderPortsAPI::ConnectOuputToInputByName(ShaderGraphNode, SubstrateMaterial::Parameters::FrontMaterial.ToString(), StandardSurfaceShaderNode->GetUniqueID(), StandardSurface::SubstrateMaterial::Outputs::Opaque.ToString());
 		if(UInterchangeShaderPortsAPI::HasInput(StandardSurfaceShaderNode, StandardSurface::Parameters::Opacity))
 		{
diff --git a/Engine/Plugins/Interchange/Runtime/Source/Pipelines/Private/InterchangeMaterialXPipeline.cpp b/Engine/Plugins/Interchange/Runtime/Source/Pipelines/Private/InterchangeMaterialXPipeline.cpp
--- a/Engine/Plugins/Interchange/Runtime/Source/Pipelines/Private/InterchangeMaterialXPipeline.cpp
+++ b/Engine/Plugins/Interchange/Runtime/Source/Pipelines/Private/InterchangeMaterialXPipeline.cpp
@@ -11,7 +11,6 @@
 #include "Materials/MaterialFunction.h"
 #include "Misc/PackageName.h"
 
-
 #define MATERIALX_FUNCTIONS_SUBSTRATE_PATH(Name)  \
 	constexpr const TCHAR* Name##FunctionsPath = TEXT("/Interchange/Functions/") TEXT("MX_") TEXT(#Name) TEXT(".") TEXT("MX_") TEXT(#Name);  \
 	constexpr const TCHAR* Name##SubstratePath = TEXT("/Interchange/Substrate/") TEXT("MX_") TEXT(#Name) TEXT(".") TEXT("MX_") TEXT(#Name);
@@ -20,9 +19,11 @@
 	MaterialXSettings->bIsSubstrateEnabled ? Name##FunctionsPath : Name##SubstratePath
 
 namespace
-{	
-	MATERIALX_FUNCTIONS_SUBSTRATE_PATH(StandardSurface);
-	MATERIALX_FUNCTIONS_SUBSTRATE_PATH(TransmissionSurface);
+{
+	constexpr const TCHAR* StandardSurfaceFunctionsPath = TEXT("/Interchange/Functions/MX_StandardSurface.MX_StandardSurface");
+	constexpr const TCHAR* StandardSurfaceSubstratePath = TEXT("/Engine/Functions/Strata/Substrate-StandardSurface-Opaque.Substrate-StandardSurface-Opaque");
+	constexpr const TCHAR* TransmissionSurfaceFunctionsPath = TEXT("/Interchange/Functions/MX_TransmissionSurface.MX_TransmissionSurface");
+	constexpr const TCHAR* TransmissionSurfaceSubstratePath = TEXT("/Engine/Functions/Strata/Substrate-StandardSurface-Translucent.Substrate-StandardSurface-Translucent");
 	MATERIALX_FUNCTIONS_SUBSTRATE_PATH(SurfaceUnlit);
 	MATERIALX_FUNCTIONS_SUBSTRATE_PATH(Surface);
 	MATERIALX_FUNCTIONS_SUBSTRATE_PATH(UsdPreviewSurface);
@@ -144,10 +145,10 @@ UMaterialXPipelineSettings::UMaterialXPipelineSettings()
 						UE::Interchange::Materials::StandardSurface::Parameters::TransmissionScatterAnisotropy,
 						UE::Interchange::Materials::StandardSurface::Parameters::TransmissionDispersion,
 						UE::Interchange::Materials::StandardSurface::Parameters::TransmissionExtraRoughness,
-						UE::Interchange::Materials::StandardSurface::Parameters::Subsurface,
-						UE::Interchange::Materials::StandardSurface::Parameters::SubsurfaceColor,
-						UE::Interchange::Materials::StandardSurface::Parameters::SubsurfaceRadius,
-						UE::Interchange::Materials::StandardSurface::Parameters::SubsurfaceScale,
+						!bIsSubstrateEnabled ? UE::Interchange::Materials::StandardSurface::Parameters::Subsurface : FName{},
+						!bIsSubstrateEnabled ? UE::Interchange::Materials::StandardSurface::Parameters::SubsurfaceColor : FName{},
+						!bIsSubstrateEnabled ? UE::Interchange::Materials::StandardSurface::Parameters::SubsurfaceRadius : FName{},
+						!bIsSubstrateEnabled ? UE::Interchange::Materials::StandardSurface::Parameters::SubsurfaceScale : FName{},
 						UE::Interchange::Materials::StandardSurface::Parameters::Sheen,
 						UE::Interchange::Materials::StandardSurface::Parameters::SheenColor,
 						UE::Interchange::Materials::StandardSurface::Parameters::SheenRoughness,
