MaterialX:
- Add support of  'channels' attribute
- Fix crash in Heightmap node, caused by node type not being correct when we copy content of image node into it, hence the node def is nullptr. Node type is now changed before checking for the nodedef.

diff --git a/Engine/Plugins/Interchange/Runtime/Source/Import/Private/MaterialX/MaterialXUtils/MaterialXSurfaceShaderAbstract.cpp b/Engine/Plugins/Interchange/Runtime/Source/Import/Private/MaterialX/MaterialXUtils/MaterialXSurfaceShaderAbstract.cpp
--- a/Engine/Plugins/Interchange/Runtime/Source/Import/Private/MaterialX/MaterialXUtils/MaterialXSurfaceShaderAbstract.cpp
+++ b/Engine/Plugins/Interchange/Runtime/Source/Import/Private/MaterialX/MaterialXUtils/MaterialXSurfaceShaderAbstract.cpp
@@ -60,24 +60,42 @@ bool FMaterialXSurfaceShaderAbstract::AddAttribute(MaterialX::InputPtr Input, co
 	return false;
 }
 
-bool FMaterialXSurfaceShaderAbstract::AddAttributeFromValueOrInterface(MaterialX::InputPtr Input, const FString& InputChannelName, UInterchangeShaderNode* ShaderNode) const
+bool FMaterialXSurfaceShaderAbstract::AddAttributeFromValueOrInterface(MaterialX::InputPtr Input, const FString& InputChannelName, UInterchangeShaderNode* ShaderNode)
 {
+	bool bAttribute = false;
+
 	if(Input)
 	{
+		UInterchangeShaderNode* ShaderNodeToConnectTo = ShaderNode;
+		FString InputToConnectTo = InputChannelName;
+
+		if(Input->hasChannels())
+		{
+			using namespace UE::Interchange::Materials::Standard::Nodes;
+
+			UInterchangeShaderNode* SwizzleNode = CreateShaderNode(Input->getParent()->asA<mx::Node>()->getName().c_str() + FString{ TEXT("_Channels_") } + Input->getName().c_str(), Swizzle::Name.ToString());
+			SwizzleNode->AddStringAttribute(Swizzle::Attributes::Channels.ToString(), Input->getChannels().c_str());
+
+			UInterchangeShaderPortsAPI::ConnectDefaultOuputToInput(ShaderNode, InputChannelName, SwizzleNode->GetUniqueID());
+
+			ShaderNodeToConnectTo = SwizzleNode;
+			InputToConnectTo = TEXT("Input");
+		}
+
 		if(Input->hasValue())
 		{
-			return AddAttribute(Input, InputChannelName, ShaderNode);
+			bAttribute = AddAttribute(Input, InputToConnectTo, ShaderNodeToConnectTo);
 		}
 		else if(Input->hasInterfaceName())
 		{
 			if(mx::InputPtr InputInterface = Input->getInterfaceInput(); InputInterface->hasValue())
 			{
-				return AddAttribute(InputInterface, InputChannelName, ShaderNode);
+				bAttribute = AddAttribute(InputInterface, InputToConnectTo, ShaderNodeToConnectTo);
 			}
 		}
 	}
 
-	return false;
+	return bAttribute;
 }
 
 bool FMaterialXSurfaceShaderAbstract::AddFloatAttribute(MaterialX::InputPtr Input, const FString& InputChannelName, UInterchangeShaderNode* ShaderNode, float DefaultValue) const
@@ -229,6 +247,9 @@ void FMaterialXSurfaceShaderAbstract::ConnectNodeCategoryOutputToInput(const Mat
 
 		FString OutputChannelName = OutputName;
 
+		// Swizzle node for the attribute 'channels'
+		UInterchangeShaderNode* ChannelsNode = nullptr;
+
 		if(mx::ElementPtr DownstreamElement = Edge.getDownstreamElement())
 		{
 			if(mx::NodePtr DownstreamNode = DownstreamElement->asA<mx::Node>())
@@ -240,6 +261,13 @@ void FMaterialXSurfaceShaderAbstract::ConnectNodeCategoryOutputToInput(const Mat
 					{
 						OutputChannelName = ConnectedInput->getOutputString().c_str();
 					}
+
+					if(ConnectedInput->hasChannels())
+					{
+						using namespace UE::Interchange::Materials::Standard::Nodes;
+						ChannelsNode = CreateShaderNode((UpstreamNode->getName() + ConnectedInput->getName()).c_str() + FString{ TEXT("_Channels") }, Swizzle::Name.ToString());
+						ChannelsNode->AddStringAttribute(Swizzle::Attributes::Channels.ToString(), ConnectedInput->getChannels().c_str());
+					}
 				}
 
 				std::vector<mx::OutputPtr> Outputs{ DownstreamNode->getOutputs()};
@@ -253,6 +281,14 @@ void FMaterialXSurfaceShaderAbstract::ConnectNodeCategoryOutputToInput(const Mat
 				{
 					if(UInterchangeShaderNode** FoundNode = ShaderNodes.Find({ GetAttributeParentName(DownstreamNode), Outputs[i]->getName().c_str()}))
 					{
+						//Connect the swizzle node between the upstream and downstream node
+						if(ChannelsNode)
+						{
+							UInterchangeShaderPortsAPI::ConnectDefaultOuputToInput(*FoundNode, InputChannelName, ChannelsNode->GetUniqueID());
+							InputChannelName = TEXT("Input");
+							FoundNode = &ChannelsNode;
+						}
+
 						ParentShaderNodeOutputs.Emplace(*FoundNode);
 					}
 				}
@@ -882,12 +918,12 @@ void FMaterialXSurfaceShaderAbstract::ConnectHeightToNormalInputToOutput(const F
 			//we need to copy the content of the image node to this node
 			Connect.UpstreamNode->copyContentFrom(ConnectedNode);
 
-			SetMatchingInputsNames(Connect.UpstreamNode);
-
-			//the copy overwrite every attribute of the node, so we need to get them back, essentially the type and the renaming
+			// the copy overwrite every attribute of the node, so we need to get them back, essentially the type and the renaming
 			// the output is always a vec3
 			Connect.UpstreamNode->setType(mx::Type::Vector3);
 
+			SetMatchingInputsNames(Connect.UpstreamNode);
+
 			mx::NodeGraphPtr Graph = Connect.UpstreamNode->getParent()->asA<mx::NodeGraph>();
 			Graph->removeNode(ConnectedNode->getName());
 
@@ -932,11 +968,11 @@ void FMaterialXSurfaceShaderAbstract::ConnectBlurInputToOutput(const FConnectNod
 			//we need to copy the content of the image node to this node
 			Connect.UpstreamNode->copyContentFrom(ConnectedNode);
 
-			SetMatchingInputsNames(Connect.UpstreamNode);
-
 			//the copy overwrites every attribute of the node, so we need to get them back, essentially the type and the renaming
 			Connect.UpstreamNode->setType(NodeType);
 
+			SetMatchingInputsNames(Connect.UpstreamNode);
+
 			mx::NodeGraphPtr Graph = Connect.UpstreamNode->getParent()->asA<mx::NodeGraph>();
 			Graph->removeNode(ConnectedNode->getName());
 
diff --git a/Engine/Plugins/Interchange/Runtime/Source/Import/Public/MaterialX/MaterialXUtils/MaterialXSurfaceShaderAbstract.h b/Engine/Plugins/Interchange/Runtime/Source/Import/Public/MaterialX/MaterialXUtils/MaterialXSurfaceShaderAbstract.h
--- a/Engine/Plugins/Interchange/Runtime/Source/Import/Public/MaterialX/MaterialXUtils/MaterialXSurfaceShaderAbstract.h
+++ b/Engine/Plugins/Interchange/Runtime/Source/Import/Public/MaterialX/MaterialXUtils/MaterialXSurfaceShaderAbstract.h
@@ -60,7 +60,7 @@ protected:
 	 *
 	 * @return true if the attribute was successfully added
 	 */
-	bool AddAttributeFromValueOrInterface(MaterialX::InputPtr Input, const FString& InputChannelName, UInterchangeShaderNode* ShaderNode) const;
+	bool AddAttributeFromValueOrInterface(MaterialX::InputPtr Input, const FString& InputChannelName, UInterchangeShaderNode* ShaderNode);
 
 	/**
 	 * Add a float attribute to a shader node only if its value taken from the input is not equal to its default value. Return false if the attribute does not exist or if we cannot add it
