Interchange: Make sure the joint has both payload, the bake and the curve, so the factory can ask one or the other.

#rnx

diff --git a/Engine/Plugins/Interchange/Runtime/Source/Parsers/Fbx/Private/FbxScene.cpp b/Engine/Plugins/Interchange/Runtime/Source/Parsers/Fbx/Private/FbxScene.cpp
--- a/Engine/Plugins/Interchange/Runtime/Source/Parsers/Fbx/Private/FbxScene.cpp
+++ b/Engine/Plugins/Interchange/Runtime/Source/Parsers/Fbx/Private/FbxScene.cpp
@@ -681,13 +681,6 @@ namespace UE
 						SkeletalAnimationTrackNode = nullptr;
 
 						SkeletalAnimationAddedToContainer = false;
-
-						//if No Skeleton Attribute then it potentially has a Rigid animation:
-						//AnimationIndex does not change the Skeletal Attributes nor the Properties from which we gather the Rigid Animation Data:
-						if (AnimationIndex == 0)
-						{
-							AddRigidAnimation(Node, UnrealNode, NodeContainer, PayloadContexts);
-						}
 					}
 					else if (SkeletalAnimationTrackNode)
 					{
@@ -699,6 +692,12 @@ namespace UE
 							NodeContainer.AddNode(SkeletalAnimationTrackNode);
 						}
 					}
+
+					//Add the transform payload for all node
+					if (AnimationIndex == 0)
+					{
+						AddRigidAnimation(Node, UnrealNode, NodeContainer, PayloadContexts);
+					}
 				}
 			
 
diff --git a/Engine/Plugins/Interchange/Runtime/Source/Pipelines/Private/InterchangeGenericAnimationPipeline.cpp b/Engine/Plugins/Interchange/Runtime/Source/Pipelines/Private/InterchangeGenericAnimationPipeline.cpp
--- a/Engine/Plugins/Interchange/Runtime/Source/Pipelines/Private/InterchangeGenericAnimationPipeline.cpp
+++ b/Engine/Plugins/Interchange/Runtime/Source/Pipelines/Private/InterchangeGenericAnimationPipeline.cpp
@@ -143,10 +143,63 @@ void UInterchangeGenericAnimationPipeline::ExecutePipeline(UInterchangeBaseNodeC
 
 	if (!bSceneImport)
 	{
+		//Extract any skeleton node use by the skeletal mesh animation track
+		TArray<FString> SceneNodesUsedBySkeleton;
+		BaseNodeContainer->IterateNodesOfType<UInterchangeSkeletalAnimationTrackNode>([&](const FString& NodeUid, UInterchangeSkeletalAnimationTrackNode* Node)
+			{
+				if (Node)
+				{
+					FString SkeletonNodeUid;
+					if (Node->GetCustomSkeletonNodeUid(SkeletonNodeUid))
+					{
+						SceneNodesUsedBySkeleton.Add(SkeletonNodeUid);
+						BaseNodeContainer->IterateNodeChildren(SkeletonNodeUid, [&SceneNodesUsedBySkeleton](const UInterchangeBaseNode* Node)
+							{
+								if (const UInterchangeSceneNode* SceneNode = Cast<UInterchangeSceneNode>(Node))
+								{
+									SceneNodesUsedBySkeleton.Add(Node->GetUniqueID());
+								}
+							});
+					}
+				}
+			});
+
+		auto IsTrackOverrideBySkeletalMeshAnimation = [this, &SceneNodesUsedBySkeleton](TArray<FString>& AnimationTrackUids)
+			{
+				//Skip track node that is using one or more scene node used by any skeletal mesh skeleton.
+				bool bSomeTrackNodeUsedBySkeletalMesh = false;
+				for (const FString& AnimationTrackUid : AnimationTrackUids)
+				{
+					if (const UInterchangeTransformAnimationTrackNode* TransformTrackNode = Cast<UInterchangeTransformAnimationTrackNode>(BaseNodeContainer->GetNode(AnimationTrackUid)))
+					{
+						FString ActorNodeUid;
+						if (TransformTrackNode->GetCustomActorDependencyUid(ActorNodeUid))
+						{
+							if (SceneNodesUsedBySkeleton.Contains(ActorNodeUid))
+							{
+								bSomeTrackNodeUsedBySkeletalMesh = true;
+								break;
+							}
+						}
+					}
+				}
+				return bSomeTrackNodeUsedBySkeletalMesh;
+			};
 		//Support rigid mesh animation animation data  (UAnimSequence for rigid mesh)
 		for (UInterchangeAnimationTrackSetNode* TrackSetNode : TrackSetNodes)
 		{
-			if (!TrackSetNode) continue;
+			if (!TrackSetNode)
+			{
+				continue;
+			}
+
+			TArray<FString> AnimationTrackUids;
+			TrackSetNode->GetCustomAnimationTrackUids(AnimationTrackUids);
+			
+			if (IsTrackOverrideBySkeletalMeshAnimation(AnimationTrackUids))
+			{
+				continue;
+			}
 
 			UInterchangeSkeletalAnimationTrackNode* SkeletalAnimationNode = NewObject< UInterchangeSkeletalAnimationTrackNode >(BaseNodeContainer);
 			FString SkeletalAnimationNodeUid = "\\SkeletalAnimation\\ConvertedFromRigidAnimation\\" + TrackSetNode->GetUniqueID();
@@ -154,8 +207,7 @@ void UInterchangeGenericAnimationPipeline::ExecutePipeline(UInterchangeBaseNodeC
 
 			bool bCustomSkeletonNodeUidSet = false;
 
-			TArray<FString> AnimationTrackUids;
-			TrackSetNode->GetCustomAnimationTrackUids(AnimationTrackUids);
+			
 			float CustomFrameRate;
 			if (!TrackSetNode->GetCustomFrameRate(CustomFrameRate))
 			{
diff --git a/Engine/Plugins/Interchange/Runtime/Source/Pipelines/Private/InterchangeGenericScenesPipeline.cpp b/Engine/Plugins/Interchange/Runtime/Source/Pipelines/Private/InterchangeGenericScenesPipeline.cpp
--- a/Engine/Plugins/Interchange/Runtime/Source/Pipelines/Private/InterchangeGenericScenesPipeline.cpp
+++ b/Engine/Plugins/Interchange/Runtime/Source/Pipelines/Private/InterchangeGenericScenesPipeline.cpp
@@ -40,13 +40,39 @@
 
 namespace UE::Interchange::Private
 {
+	TArray<FString> GetAllActiveJoints(UInterchangeBaseNodeContainer* BaseNodeContainer)
+	{
+		TArray<FString> AllActiveJoints;
+		BaseNodeContainer->IterateNodesOfType<UInterchangeSkeletonFactoryNode>([&BaseNodeContainer, &AllActiveJoints](const FString& NodeUid, UInterchangeSkeletonFactoryNode* Node)
+			{
+				FString RootNodeUid;
+				if (Node->GetCustomRootJointUid(RootNodeUid))
+				{
+					AllActiveJoints.Add(RootNodeUid);
+					BaseNodeContainer->IterateNodeChildren(RootNodeUid, [&AllActiveJoints](const UInterchangeBaseNode* Node)
+						{
+							if (const UInterchangeSceneNode* SceneNode = Cast<UInterchangeSceneNode>(Node))
+							{
+								TArray<FString> SpecializeTypes;
+								SceneNode->GetSpecializedTypes(SpecializeTypes);
+								if (SpecializeTypes.Contains(UE::Interchange::FSceneNodeStaticData::GetJointSpecializeTypeString()))
+								{
+									AllActiveJoints.Add(Node->GetUniqueID());
+								}
+							}
+						});
+				}
+			});
+		return AllActiveJoints;
+	}
+
 	//Either a (TransformSpecialized || !JointSpecialized || RootJoint) can be a parent (only those get FactoryNodes) :
-	FString FindFactoryParentSceneNodeUid(UInterchangeBaseNodeContainer* BaseNodeContainer, const UInterchangeSceneNode* SceneNode)
+	FString FindFactoryParentSceneNodeUid(UInterchangeBaseNodeContainer* BaseNodeContainer, TArray<FString>& ActiveSkeletonUids, const UInterchangeSceneNode* SceneNode)
 	{
 		FString ParentUid = SceneNode->GetParentUid();
 		if (const UInterchangeSceneNode* ParentSceneNode = Cast<UInterchangeSceneNode>(BaseNodeContainer->GetNode(ParentUid)))
 		{
-			if (ParentSceneNode->IsSpecializedTypeContains(UE::Interchange::FSceneNodeStaticData::GetTransformSpecializeTypeString()))
+			if (!ActiveSkeletonUids.Contains(ParentUid) || ParentSceneNode->IsSpecializedTypeContains(UE::Interchange::FSceneNodeStaticData::GetTransformSpecializeTypeString()))
 			{
 				return ParentUid;
 			}
@@ -63,7 +89,7 @@ namespace UE::Interchange::Private
 						bool bParentsParentIsJoint = ParentSceneNode->IsSpecializedTypeContains(UE::Interchange::FSceneNodeStaticData::GetJointSpecializeTypeString());
 						if (bParentsParentIsJoint)
 						{
-							return FindFactoryParentSceneNodeUid(BaseNodeContainer, ParentSceneNode);
+							return FindFactoryParentSceneNodeUid(BaseNodeContainer, ActiveSkeletonUids, ParentSceneNode);
 						}
 						else
 						{
@@ -260,6 +286,9 @@ void UInterchangeGenericLevelPipeline::ExecutePipeline(UInterchangeBaseNodeConta
 	}
 #endif
 
+	/* Find all scene node that are active joint. Non active joint should be convert to actor if they are in a static mesh hierarchy */
+	TArray<FString> ActiveSkeletonUids = UE::Interchange::Private::GetAllActiveJoints(BaseNodeContainer);
+
 	for (const UInterchangeSceneNode* SceneNode : SceneNodes)
 	{
 		if (SceneNode)
@@ -273,14 +302,21 @@ void UInterchangeGenericLevelPipeline::ExecutePipeline(UInterchangeBaseNodeConta
 					bool bSkipNode = true;
 					if (SpecializeTypes.Contains(UE::Interchange::FSceneNodeStaticData::GetJointSpecializeTypeString()))
 					{
-						//check if its the root joint (we want to create an actor for the root joint)
-						FString CurrentNodesParentUid = SceneNode->GetParentUid();
-						const UInterchangeBaseNode* ParentNode = BaseNodeContainer->GetNode(CurrentNodesParentUid);
-						if (const UInterchangeSceneNode* ParentSceneNode = Cast<UInterchangeSceneNode>(ParentNode))
+						if(!ActiveSkeletonUids.Contains(SceneNode->GetUniqueID()))
+						{ 
+							bSkipNode = false;
+						}
+						else
 						{
-							if (!ParentSceneNode->IsSpecializedTypeContains(UE::Interchange::FSceneNodeStaticData::GetJointSpecializeTypeString()))
+							//check if its the root joint (we want to create an actor for the root joint)
+							FString CurrentNodesParentUid = SceneNode->GetParentUid();
+							const UInterchangeBaseNode* ParentNode = BaseNodeContainer->GetNode(CurrentNodesParentUid);
+							if (const UInterchangeSceneNode* ParentSceneNode = Cast<UInterchangeSceneNode>(ParentNode))
 							{
-								bSkipNode = false;
+								if (!ParentSceneNode->IsSpecializedTypeContains(UE::Interchange::FSceneNodeStaticData::GetJointSpecializeTypeString()))
+								{
+									bSkipNode = false;
+								}
 							}
 						}
 					}
@@ -406,7 +442,9 @@ void UInterchangeGenericLevelPipeline::ExecuteSceneNodePreImport(const FTransfor
 
 	if (!SceneNode->GetParentUid().IsEmpty())
 	{
-		FString ParentNodeUid = UE::Interchange::Private::FindFactoryParentSceneNodeUid(BaseNodeContainer, SceneNode);
+		/* Find all scene node that are active joint. Non active joint should be convert to actor if they are in a static mesh hierarchy */
+		TArray<FString> ActiveSkeletonUids = UE::Interchange::Private::GetAllActiveJoints(BaseNodeContainer);
+		FString ParentNodeUid = UE::Interchange::Private::FindFactoryParentSceneNodeUid(BaseNodeContainer, ActiveSkeletonUids, SceneNode);
 		if (ParentNodeUid != UInterchangeBaseNode::InvalidNodeUid())
 		{
 			FString ParentFactoryNodeUid = UInterchangeFactoryBaseNode::BuildFactoryNodeUid(ParentNodeUid);
