[Interchange - glTF] Defensive coding fixes for Animation processing.

#rnx

diff --git a/Engine/Plugins/Interchange/Runtime/Source/Import/Private/Gltf/InterchangeGltfAnimation.cpp b/Engine/Plugins/Interchange/Runtime/Source/Import/Private/Gltf/InterchangeGltfAnimation.cpp
--- a/Engine/Plugins/Interchange/Runtime/Source/Import/Private/Gltf/InterchangeGltfAnimation.cpp
+++ b/Engine/Plugins/Interchange/Runtime/Source/Import/Private/Gltf/InterchangeGltfAnimation.cpp
@@ -721,16 +721,21 @@ namespace UE::Interchange::Gltf::Private
 
 			void AcquireJointsWithBindPose(int32 CurrentIndex, TSet<int32>& Joints)
 			{
-				if (GLTFNodes[CurrentIndex].Type == GLTF::FNode::EType::Joint)
+				if (GLTFNodes.IsValidIndex(CurrentIndex))
 				{
-					if (GLTFNodes[CurrentIndex].bHasLocalBindPose)
-					{
-						Joints.Add(CurrentIndex);
-					}
+					const GLTF::FNode& CurrentNode = GLTFNodes[CurrentIndex];
 
-					for (int32 ChildIndex : GLTFNodes[CurrentIndex].Children)
+					if (CurrentNode.Type == GLTF::FNode::EType::Joint)
 					{
-						AcquireJointsWithBindPose(ChildIndex, Joints);
+						if (CurrentNode.bHasLocalBindPose)
+						{
+							Joints.Add(CurrentIndex);
+						}
+
+						for (int32 ChildIndex : CurrentNode.Children)
+						{
+							AcquireJointsWithBindPose(ChildIndex, Joints);
+						}
 					}
 				}
 			}
@@ -951,6 +956,11 @@ namespace UE::Interchange::Gltf::Private
 				{
 					int32 SkeletonRootIndex = AnimatedJointNodeIndices.Key;
 
+					if (!GLTFNodes.IsValidIndex(SkeletonRootIndex))
+					{
+						continue;
+					}
+
 					const FString* SkeletonUid = GLTFNodeToInterchangeUidMap.Find(&GLTFNodes[SkeletonRootIndex]);
 					if (!ensure(SkeletonUid))
 					{
