Adding GenerateHLSLExpression function for MaterialX material expressions:
- In
- Mask
- Matte
- Minus
- Out
- Over
- Overlay
- Plus
- Premult
- Screen

diff --git a/Engine/Plugins/Interchange/Runtime/Source/Import/Private/MaterialX/MaterialExpressions/MaterialExpressionIn.cpp b/Engine/Plugins/Interchange/Runtime/Source/Import/Private/MaterialX/MaterialExpressions/MaterialExpressionIn.cpp
--- a/Engine/Plugins/Interchange/Runtime/Source/Import/Private/MaterialX/MaterialExpressions/MaterialExpressionIn.cpp
+++ b/Engine/Plugins/Interchange/Runtime/Source/Import/Private/MaterialX/MaterialExpressions/MaterialExpressionIn.cpp
@@ -1,6 +1,8 @@
 // Copyright Epic Games, Inc. All Rights Reserved.
 #include "MaterialExpressionIn.h"
 #include "MaterialCompiler.h"
+#include "MaterialHLSLGenerator.h"
+#include "HLSLTree/HLSLTree.h"
 
 #include UE_INLINE_GENERATED_CPP_BY_NAME(MaterialExpressionIn)
 
@@ -51,6 +53,27 @@ void UMaterialExpressionMaterialXIn::GetCaption(TArray<FString>& OutCaptions) co
 {
 	OutCaptions.Add(TEXT("MaterialX In"));
 }
+
+bool UMaterialExpressionMaterialXIn::GenerateHLSLExpression(FMaterialHLSLGenerator& Generator, UE::HLSLTree::FScope& Scope, int32 OutputIndex, UE::HLSLTree::FExpression const*& OutExpression) const
+{
+	using namespace UE::HLSLTree;
+
+	const FExpression* ExpressionA = A.AcquireHLSLExpression(Generator, Scope);
+	const FExpression* ExpressionB = B.AcquireHLSLExpression(Generator, Scope);
+	const FExpression* ExpressionAlpha = Alpha.AcquireHLSLExpressionOrConstant(Generator, Scope, ConstAlpha);
+
+	if(!ExpressionA || !ExpressionB || !ExpressionAlpha)
+	{
+		return false;
+	}
+
+	FTree& Tree = Generator.GetTree();
+	OutExpression = Tree.NewLerp(ExpressionB,
+								 Tree.NewMul(ExpressionA, Tree.NewSwizzle(FSwizzleParameters(3), ExpressionB)),
+								 ExpressionAlpha);
+
+	return true;
+}
 #endif
 
 #undef LOCTEXT_NAMESPACE 
\ No newline at end of file
diff --git a/Engine/Plugins/Interchange/Runtime/Source/Import/Private/MaterialX/MaterialExpressions/MaterialExpressionIn.h b/Engine/Plugins/Interchange/Runtime/Source/Import/Private/MaterialX/MaterialExpressions/MaterialExpressionIn.h
--- a/Engine/Plugins/Interchange/Runtime/Source/Import/Private/MaterialX/MaterialExpressions/MaterialExpressionIn.h
+++ b/Engine/Plugins/Interchange/Runtime/Source/Import/Private/MaterialX/MaterialExpressions/MaterialExpressionIn.h
@@ -39,6 +39,7 @@ class UMaterialExpressionMaterialXIn: public UMaterialExpression
 #if WITH_EDITOR
 	virtual int32 Compile(class FMaterialCompiler* Compiler, int32 OutputIndex) override;
 	virtual void GetCaption(TArray<FString>& OutCaptions) const override;
+	virtual bool GenerateHLSLExpression(FMaterialHLSLGenerator& Generator, UE::HLSLTree::FScope& Scope, int32 OutputIndex, UE::HLSLTree::FExpression const*& OutExpression) const override;
 #endif
 	//~ End UMaterialExpressionMaterialX Interface
 };
diff --git a/Engine/Plugins/Interchange/Runtime/Source/Import/Private/MaterialX/MaterialExpressions/MaterialExpressionMask.cpp b/Engine/Plugins/Interchange/Runtime/Source/Import/Private/MaterialX/MaterialExpressions/MaterialExpressionMask.cpp
--- a/Engine/Plugins/Interchange/Runtime/Source/Import/Private/MaterialX/MaterialExpressions/MaterialExpressionMask.cpp
+++ b/Engine/Plugins/Interchange/Runtime/Source/Import/Private/MaterialX/MaterialExpressions/MaterialExpressionMask.cpp
@@ -1,6 +1,8 @@
 // Copyright Epic Games, Inc. All Rights Reserved.
 #include "MaterialExpressionMask.h"
 #include "MaterialCompiler.h"
+#include "MaterialHLSLGenerator.h"
+#include "HLSLTree/HLSLTree.h"
 
 #include UE_INLINE_GENERATED_CPP_BY_NAME(MaterialExpressionMask)
 
@@ -51,6 +53,27 @@ void UMaterialExpressionMaterialXMask::GetCaption(TArray<FString>& OutCaptions)
 {
 	OutCaptions.Add(TEXT("MaterialX Mask"));
 }
+
+bool UMaterialExpressionMaterialXMask::GenerateHLSLExpression(FMaterialHLSLGenerator& Generator, UE::HLSLTree::FScope& Scope, int32 OutputIndex, UE::HLSLTree::FExpression const*& OutExpression) const
+{
+	using namespace UE::HLSLTree;
+
+	const FExpression* ExpressionA = A.AcquireHLSLExpression(Generator, Scope);
+	const FExpression* ExpressionB = B.AcquireHLSLExpression(Generator, Scope);
+	const FExpression* ExpressionAlpha = Alpha.AcquireHLSLExpressionOrConstant(Generator, Scope, ConstAlpha);
+
+	if(!ExpressionA || !ExpressionB || !ExpressionAlpha)
+	{
+		return false;
+	}
+
+	FTree& Tree = Generator.GetTree();
+	OutExpression = Tree.NewLerp(ExpressionB,
+								 Tree.NewMul(ExpressionB, Tree.NewSwizzle(FSwizzleParameters(3), ExpressionA)),
+								 ExpressionAlpha);
+
+	return true;
+}
 #endif
 
 #undef LOCTEXT_NAMESPACE 
\ No newline at end of file
diff --git a/Engine/Plugins/Interchange/Runtime/Source/Import/Private/MaterialX/MaterialExpressions/MaterialExpressionMask.h b/Engine/Plugins/Interchange/Runtime/Source/Import/Private/MaterialX/MaterialExpressions/MaterialExpressionMask.h
--- a/Engine/Plugins/Interchange/Runtime/Source/Import/Private/MaterialX/MaterialExpressions/MaterialExpressionMask.h
+++ b/Engine/Plugins/Interchange/Runtime/Source/Import/Private/MaterialX/MaterialExpressions/MaterialExpressionMask.h
@@ -39,6 +39,7 @@ class UMaterialExpressionMaterialXMask : public UMaterialExpression
 #if WITH_EDITOR
 	virtual int32 Compile(class FMaterialCompiler* Compiler, int32 OutputIndex) override;
 	virtual void GetCaption(TArray<FString>& OutCaptions) const override;
+	virtual bool GenerateHLSLExpression(FMaterialHLSLGenerator& Generator, UE::HLSLTree::FScope& Scope, int32 OutputIndex, UE::HLSLTree::FExpression const*& OutExpression) const override;
 #endif
 	//~ End UMaterialExpressionMaterialX Interface
 };
diff --git a/Engine/Plugins/Interchange/Runtime/Source/Import/Private/MaterialX/MaterialExpressions/MaterialExpressionMatte.cpp b/Engine/Plugins/Interchange/Runtime/Source/Import/Private/MaterialX/MaterialExpressions/MaterialExpressionMatte.cpp
--- a/Engine/Plugins/Interchange/Runtime/Source/Import/Private/MaterialX/MaterialExpressions/MaterialExpressionMatte.cpp
+++ b/Engine/Plugins/Interchange/Runtime/Source/Import/Private/MaterialX/MaterialExpressions/MaterialExpressionMatte.cpp
@@ -1,6 +1,9 @@
 // Copyright Epic Games, Inc. All Rights Reserved.
 #include "MaterialExpressionMatte.h"
 #include "MaterialCompiler.h"
+#include "MaterialHLSLGenerator.h"
+#include "HLSLTree/HLSLTree.h"
+#include "HLSLTree/HLSLTreeCommon.h"
 
 #include UE_INLINE_GENERATED_CPP_BY_NAME(MaterialExpressionMatte)
 
@@ -60,6 +63,38 @@ void UMaterialExpressionMaterialXMatte::GetCaption(TArray<FString>& OutCaptions)
 {
 	OutCaptions.Add(TEXT("MaterialX Matte"));
 }
+
+bool UMaterialExpressionMaterialXMatte::GenerateHLSLExpression(FMaterialHLSLGenerator& Generator, UE::HLSLTree::FScope& Scope, int32 OutputIndex, UE::HLSLTree::FExpression const*& OutExpression) const
+{
+	using namespace UE::HLSLTree;
+
+	const FExpression* ExpressionA = A.AcquireHLSLExpression(Generator, Scope);
+	const FExpression* ExpressionB = B.AcquireHLSLExpression(Generator, Scope);
+	const FExpression* ExpressionAlpha = Alpha.AcquireHLSLExpressionOrConstant(Generator, Scope, ConstAlpha);
+
+	if(!ExpressionA || !ExpressionB || !ExpressionAlpha)
+	{
+		return false;
+	}
+
+	FTree& Tree = Generator.GetTree();
+	const FExpression* ExpressionAlphaA = Tree.NewSwizzle(FSwizzleParameters(3), ExpressionA);
+	const FExpression* ExpressionAlphaB = Tree.NewSwizzle(FSwizzleParameters(3), ExpressionB);
+	const FExpression* ExpressionRgbA = Tree.NewSwizzle(FSwizzleParameters(0, 1, 2), ExpressionA);
+	const FExpression* ExpressionRgbB = Tree.NewSwizzle(FSwizzleParameters(0, 1, 2), ExpressionB);
+
+	const FExpression* ExpressionOneMinusAlphaA = Tree.NewSub(Tree.NewConstant(1.f), ExpressionAlphaA);
+
+	const FExpression* ExpressionMatte = Tree.NewExpression<FExpressionAppend>(
+			Tree.NewAdd(Tree.NewMul(ExpressionRgbA, ExpressionAlphaA),
+						Tree.NewMul(ExpressionRgbB, ExpressionOneMinusAlphaA)),
+			Tree.NewAdd(ExpressionAlphaA,
+						Tree.NewMul(ExpressionAlphaB, ExpressionOneMinusAlphaA)));
+
+	OutExpression = Tree.NewLerp(ExpressionB, ExpressionMatte, ExpressionAlpha);
+
+	return true;
+}
 #endif
 
 #undef LOCTEXT_NAMESPACE 
\ No newline at end of file
diff --git a/Engine/Plugins/Interchange/Runtime/Source/Import/Private/MaterialX/MaterialExpressions/MaterialExpressionMatte.h b/Engine/Plugins/Interchange/Runtime/Source/Import/Private/MaterialX/MaterialExpressions/MaterialExpressionMatte.h
--- a/Engine/Plugins/Interchange/Runtime/Source/Import/Private/MaterialX/MaterialExpressions/MaterialExpressionMatte.h
+++ b/Engine/Plugins/Interchange/Runtime/Source/Import/Private/MaterialX/MaterialExpressions/MaterialExpressionMatte.h
@@ -39,6 +39,7 @@ class UMaterialExpressionMaterialXMatte : public UMaterialExpression
 #if WITH_EDITOR
 	virtual int32 Compile(class FMaterialCompiler* Compiler, int32 OutputIndex) override;
 	virtual void GetCaption(TArray<FString>& OutCaptions) const override;
+	virtual bool GenerateHLSLExpression(FMaterialHLSLGenerator& Generator, UE::HLSLTree::FScope& Scope, int32 OutputIndex, UE::HLSLTree::FExpression const*& OutExpression) const override;
 #endif
 	//~ End UMaterialExpressionMaterialX Interface
 };
diff --git a/Engine/Plugins/Interchange/Runtime/Source/Import/Private/MaterialX/MaterialExpressions/MaterialExpressionMinus.cpp b/Engine/Plugins/Interchange/Runtime/Source/Import/Private/MaterialX/MaterialExpressions/MaterialExpressionMinus.cpp
--- a/Engine/Plugins/Interchange/Runtime/Source/Import/Private/MaterialX/MaterialExpressions/MaterialExpressionMinus.cpp
+++ b/Engine/Plugins/Interchange/Runtime/Source/Import/Private/MaterialX/MaterialExpressions/MaterialExpressionMinus.cpp
@@ -1,6 +1,8 @@
 // Copyright Epic Games, Inc. All Rights Reserved.
 #include "MaterialExpressionMinus.h"
 #include "MaterialCompiler.h"
+#include "MaterialHLSLGenerator.h"
+#include "HLSLTree/HLSLTree.h"
 
 #include UE_INLINE_GENERATED_CPP_BY_NAME(MaterialExpressionMinus)
 
@@ -48,6 +50,28 @@ void UMaterialExpressionMaterialXMinus::GetCaption(TArray<FString>& OutCaptions)
 {
 	OutCaptions.Add(TEXT("MaterialX Minus"));
 }
+
+bool UMaterialExpressionMaterialXMinus::GenerateHLSLExpression(FMaterialHLSLGenerator& Generator, UE::HLSLTree::FScope& Scope, int32 OutputIndex, UE::HLSLTree::FExpression const*& OutExpression) const
+{
+	using namespace UE::HLSLTree;
+
+	const FExpression* ExpressionA = A.AcquireHLSLExpression(Generator, Scope);
+	const FExpression* ExpressionB = B.AcquireHLSLExpression(Generator, Scope);
+	const FExpression* ExpressionAlpha = Alpha.AcquireHLSLExpressionOrConstant(Generator, Scope, ConstAlpha);
+
+	if(!ExpressionA || !ExpressionB || !ExpressionAlpha)
+	{
+		return false;
+	}
+
+	FTree& Tree = Generator.GetTree();
+	
+	OutExpression = Tree.NewLerp(ExpressionB,
+								 Tree.NewSub(ExpressionA, ExpressionB),
+								 ExpressionAlpha);
+
+	return true;
+}
 #endif
 
 #undef LOCTEXT_NAMESPACE 
\ No newline at end of file
diff --git a/Engine/Plugins/Interchange/Runtime/Source/Import/Private/MaterialX/MaterialExpressions/MaterialExpressionMinus.h b/Engine/Plugins/Interchange/Runtime/Source/Import/Private/MaterialX/MaterialExpressions/MaterialExpressionMinus.h
--- a/Engine/Plugins/Interchange/Runtime/Source/Import/Private/MaterialX/MaterialExpressions/MaterialExpressionMinus.h
+++ b/Engine/Plugins/Interchange/Runtime/Source/Import/Private/MaterialX/MaterialExpressions/MaterialExpressionMinus.h
@@ -35,6 +35,7 @@ class UMaterialExpressionMaterialXMinus : public UMaterialExpression
 #if WITH_EDITOR
 	virtual int32 Compile(class FMaterialCompiler* Compiler, int32 OutputIndex) override;
 	virtual void GetCaption(TArray<FString>& OutCaptions) const override;
+	virtual bool GenerateHLSLExpression(FMaterialHLSLGenerator& Generator, UE::HLSLTree::FScope& Scope, int32 OutputIndex, UE::HLSLTree::FExpression const*& OutExpression) const override;
 #endif
 	//~ End UMaterialExpressionMaterialX Interface
 };
diff --git a/Engine/Plugins/Interchange/Runtime/Source/Import/Private/MaterialX/MaterialExpressions/MaterialExpressionOut.cpp b/Engine/Plugins/Interchange/Runtime/Source/Import/Private/MaterialX/MaterialExpressions/MaterialExpressionOut.cpp
--- a/Engine/Plugins/Interchange/Runtime/Source/Import/Private/MaterialX/MaterialExpressions/MaterialExpressionOut.cpp
+++ b/Engine/Plugins/Interchange/Runtime/Source/Import/Private/MaterialX/MaterialExpressions/MaterialExpressionOut.cpp
@@ -1,6 +1,8 @@
 // Copyright Epic Games, Inc. All Rights Reserved.
 #include "MaterialExpressionOut.h"
 #include "MaterialCompiler.h"
+#include "MaterialHLSLGenerator.h"
+#include "HLSLTree/HLSLTree.h"
 
 #include UE_INLINE_GENERATED_CPP_BY_NAME(MaterialExpressionOut)
 
@@ -54,6 +56,30 @@ void UMaterialExpressionMaterialXOut::GetCaption(TArray<FString>& OutCaptions) c
 {
 	OutCaptions.Add(TEXT("MaterialX Out"));
 }
+
+bool UMaterialExpressionMaterialXOut::GenerateHLSLExpression(FMaterialHLSLGenerator& Generator, UE::HLSLTree::FScope& Scope, int32 OutputIndex, UE::HLSLTree::FExpression const*& OutExpression) const
+{
+	using namespace UE::HLSLTree;
+
+	const FExpression* ExpressionA = A.AcquireHLSLExpression(Generator, Scope);
+	const FExpression* ExpressionB = B.AcquireHLSLExpression(Generator, Scope);
+	const FExpression* ExpressionAlpha = Alpha.AcquireHLSLExpressionOrConstant(Generator, Scope, ConstAlpha);
+
+	if(!ExpressionA || !ExpressionB || !ExpressionAlpha)
+	{
+		return false;
+	}
+
+	FTree& Tree = Generator.GetTree();
+	
+	const FExpression* ExpressionOut = Tree.NewMul(ExpressionA,
+												   Tree.NewSub(Tree.NewConstant(1.f),
+															   Tree.NewSwizzle(FSwizzleParameters(3), ExpressionB)));
+
+	OutExpression = Tree.NewLerp(ExpressionB, ExpressionOut, ExpressionAlpha);
+
+	return true;
+}
 #endif
 
 #undef LOCTEXT_NAMESPACE 
\ No newline at end of file
diff --git a/Engine/Plugins/Interchange/Runtime/Source/Import/Private/MaterialX/MaterialExpressions/MaterialExpressionOut.h b/Engine/Plugins/Interchange/Runtime/Source/Import/Private/MaterialX/MaterialExpressions/MaterialExpressionOut.h
--- a/Engine/Plugins/Interchange/Runtime/Source/Import/Private/MaterialX/MaterialExpressions/MaterialExpressionOut.h
+++ b/Engine/Plugins/Interchange/Runtime/Source/Import/Private/MaterialX/MaterialExpressions/MaterialExpressionOut.h
@@ -39,6 +39,7 @@ class UMaterialExpressionMaterialXOut : public UMaterialExpression
 #if WITH_EDITOR
 	virtual int32 Compile(class FMaterialCompiler* Compiler, int32 OutputIndex) override;
 	virtual void GetCaption(TArray<FString>& OutCaptions) const override;
+	virtual bool GenerateHLSLExpression(FMaterialHLSLGenerator& Generator, UE::HLSLTree::FScope& Scope, int32 OutputIndex, UE::HLSLTree::FExpression const*& OutExpression) const override;
 #endif
 	//~ End UMaterialExpressionMaterialX Interface
 };
diff --git a/Engine/Plugins/Interchange/Runtime/Source/Import/Private/MaterialX/MaterialExpressions/MaterialExpressionOver.cpp b/Engine/Plugins/Interchange/Runtime/Source/Import/Private/MaterialX/MaterialExpressions/MaterialExpressionOver.cpp
--- a/Engine/Plugins/Interchange/Runtime/Source/Import/Private/MaterialX/MaterialExpressions/MaterialExpressionOver.cpp
+++ b/Engine/Plugins/Interchange/Runtime/Source/Import/Private/MaterialX/MaterialExpressions/MaterialExpressionOver.cpp
@@ -1,6 +1,8 @@
 // Copyright Epic Games, Inc. All Rights Reserved.
 #include "MaterialExpressionOver.h"
 #include "MaterialCompiler.h"
+#include "MaterialHLSLGenerator.h"
+#include "HLSLTree/HLSLTree.h"
 
 #include UE_INLINE_GENERATED_CPP_BY_NAME(MaterialExpressionOver)
 
@@ -54,6 +56,32 @@ void UMaterialExpressionMaterialXOver::GetCaption(TArray<FString>& OutCaptions)
 {
 	OutCaptions.Add(TEXT("MaterialX Over"));
 }
+
+
+bool UMaterialExpressionMaterialXOver::GenerateHLSLExpression(FMaterialHLSLGenerator& Generator, UE::HLSLTree::FScope& Scope, int32 OutputIndex, UE::HLSLTree::FExpression const*& OutExpression) const
+{
+	using namespace UE::HLSLTree;
+
+	const FExpression* ExpressionA = A.AcquireHLSLExpression(Generator, Scope);
+	const FExpression* ExpressionB = B.AcquireHLSLExpression(Generator, Scope);
+	const FExpression* ExpressionAlpha = Alpha.AcquireHLSLExpressionOrConstant(Generator, Scope, ConstAlpha);
+
+	if(!ExpressionA || !ExpressionB || !ExpressionAlpha)
+	{
+		return false;
+	}
+
+	FTree& Tree = Generator.GetTree();
+
+	const FExpression* ExpressionOver = Tree.NewAdd(ExpressionA,
+												   Tree.NewMul(ExpressionB,
+															   Tree.NewSub(Tree.NewConstant(1.f),
+																		   Tree.NewSwizzle(FSwizzleParameters(3), ExpressionA))));
+
+	OutExpression = Tree.NewLerp(ExpressionB, ExpressionOver, ExpressionAlpha);
+
+	return true;
+}
 #endif
 
 #undef LOCTEXT_NAMESPACE 
\ No newline at end of file
diff --git a/Engine/Plugins/Interchange/Runtime/Source/Import/Private/MaterialX/MaterialExpressions/MaterialExpressionOver.h b/Engine/Plugins/Interchange/Runtime/Source/Import/Private/MaterialX/MaterialExpressions/MaterialExpressionOver.h
--- a/Engine/Plugins/Interchange/Runtime/Source/Import/Private/MaterialX/MaterialExpressions/MaterialExpressionOver.h
+++ b/Engine/Plugins/Interchange/Runtime/Source/Import/Private/MaterialX/MaterialExpressions/MaterialExpressionOver.h
@@ -39,6 +39,7 @@ class UMaterialExpressionMaterialXOver : public UMaterialExpression
 #if WITH_EDITOR
 	virtual int32 Compile(class FMaterialCompiler* Compiler, int32 OutputIndex) override;
 	virtual void GetCaption(TArray<FString>& OutCaptions) const override;
+	virtual bool GenerateHLSLExpression(FMaterialHLSLGenerator& Generator, UE::HLSLTree::FScope& Scope, int32 OutputIndex, UE::HLSLTree::FExpression const*& OutExpression) const override;
 #endif
 	//~ End UMaterialExpressionMaterialX Interface
 };
diff --git a/Engine/Plugins/Interchange/Runtime/Source/Import/Private/MaterialX/MaterialExpressions/MaterialExpressionOverlay.cpp b/Engine/Plugins/Interchange/Runtime/Source/Import/Private/MaterialX/MaterialExpressions/MaterialExpressionOverlay.cpp
--- a/Engine/Plugins/Interchange/Runtime/Source/Import/Private/MaterialX/MaterialExpressions/MaterialExpressionOverlay.cpp
+++ b/Engine/Plugins/Interchange/Runtime/Source/Import/Private/MaterialX/MaterialExpressions/MaterialExpressionOverlay.cpp
@@ -1,6 +1,8 @@
 // Copyright Epic Games, Inc. All Rights Reserved.
 #include "MaterialExpressionOverlay.h"
 #include "MaterialCompiler.h"
+#include "MaterialHLSLGenerator.h"
+#include "HLSLTree/HLSLTree.h"
 
 #include UE_INLINE_GENERATED_CPP_BY_NAME(MaterialExpressionOverlay)
 
@@ -63,6 +65,37 @@ void UMaterialExpressionMaterialXOverlay::GetCaption(TArray<FString>& OutCaption
 {
 	OutCaptions.Add(TEXT("MaterialX Overlay"));
 }
+
+bool UMaterialExpressionMaterialXOverlay::GenerateHLSLExpression(FMaterialHLSLGenerator& Generator, UE::HLSLTree::FScope& Scope, int32 OutputIndex, UE::HLSLTree::FExpression const*& OutExpression) const
+{
+	using namespace UE::HLSLTree;
+
+	const FExpression* ExpressionA = A.AcquireHLSLExpression(Generator, Scope);
+	const FExpression* ExpressionB = B.AcquireHLSLExpression(Generator, Scope);
+	const FExpression* ExpressionAlpha = Alpha.AcquireHLSLExpressionOrConstant(Generator, Scope, ConstAlpha);
+
+	if(!ExpressionA || !ExpressionB || !ExpressionAlpha)
+	{
+		return false;
+	}
+
+	FTree& Tree = Generator.GetTree();
+	const FExpression* ExpressionOne = Tree.NewConstant(1.f);
+
+	auto NewOneMinus = [&](const FExpression* Input)
+	{
+		return Tree.NewSub(ExpressionOne, Input);
+	};
+
+	const FExpression* ExpressionOverlay = Generator.GenerateBranch(Scope,
+																	Tree.NewLess(ExpressionA, Tree.NewConstant(0.5f)),
+																	Tree.NewMul(Tree.NewMul(Tree.NewConstant(2.f), ExpressionA), ExpressionB),
+																	NewOneMinus(Tree.NewMul(NewOneMinus(ExpressionA), NewOneMinus(ExpressionB))));
+
+	OutExpression = Tree.NewLerp(ExpressionB, ExpressionOverlay, ExpressionAlpha);
+
+	return true;
+}
 #endif
 
 #undef LOCTEXT_NAMESPACE 
\ No newline at end of file
diff --git a/Engine/Plugins/Interchange/Runtime/Source/Import/Private/MaterialX/MaterialExpressions/MaterialExpressionOverlay.h b/Engine/Plugins/Interchange/Runtime/Source/Import/Private/MaterialX/MaterialExpressions/MaterialExpressionOverlay.h
--- a/Engine/Plugins/Interchange/Runtime/Source/Import/Private/MaterialX/MaterialExpressions/MaterialExpressionOverlay.h
+++ b/Engine/Plugins/Interchange/Runtime/Source/Import/Private/MaterialX/MaterialExpressions/MaterialExpressionOverlay.h
@@ -36,6 +36,7 @@ class UMaterialExpressionMaterialXOverlay : public UMaterialExpression
 #if WITH_EDITOR
 	virtual int32 Compile(class FMaterialCompiler* Compiler, int32 OutputIndex) override;
 	virtual void GetCaption(TArray<FString>& OutCaptions) const override;
+	virtual bool GenerateHLSLExpression(FMaterialHLSLGenerator& Generator, UE::HLSLTree::FScope& Scope, int32 OutputIndex, UE::HLSLTree::FExpression const*& OutExpression) const override;
 #endif
 	//~ End UMaterialExpressionMaterialX Interface
 };
diff --git a/Engine/Plugins/Interchange/Runtime/Source/Import/Private/MaterialX/MaterialExpressions/MaterialExpressionPlus.cpp b/Engine/Plugins/Interchange/Runtime/Source/Import/Private/MaterialX/MaterialExpressions/MaterialExpressionPlus.cpp
--- a/Engine/Plugins/Interchange/Runtime/Source/Import/Private/MaterialX/MaterialExpressions/MaterialExpressionPlus.cpp
+++ b/Engine/Plugins/Interchange/Runtime/Source/Import/Private/MaterialX/MaterialExpressions/MaterialExpressionPlus.cpp
@@ -1,6 +1,8 @@
 // Copyright Epic Games, Inc. All Rights Reserved.
 #include "MaterialExpressionPlus.h"
 #include "MaterialCompiler.h"
+#include "MaterialHLSLGenerator.h"
+#include "HLSLTree/HLSLTree.h"
 
 #include UE_INLINE_GENERATED_CPP_BY_NAME(MaterialExpressionPlus)
 
@@ -48,6 +50,28 @@ void UMaterialExpressionMaterialXPlus::GetCaption(TArray<FString>& OutCaptions)
 {
 	OutCaptions.Add(TEXT("MaterialX Plus"));
 }
+
+bool UMaterialExpressionMaterialXPlus::GenerateHLSLExpression(FMaterialHLSLGenerator& Generator, UE::HLSLTree::FScope& Scope, int32 OutputIndex, UE::HLSLTree::FExpression const*& OutExpression) const
+{
+	using namespace UE::HLSLTree;
+
+	const FExpression* ExpressionA = A.AcquireHLSLExpression(Generator, Scope);
+	const FExpression* ExpressionB = B.AcquireHLSLExpression(Generator, Scope);
+	const FExpression* ExpressionAlpha = Alpha.AcquireHLSLExpressionOrConstant(Generator, Scope, ConstAlpha);
+
+	if(!ExpressionA || !ExpressionB || !ExpressionAlpha)
+	{
+		return false;
+	}
+
+	FTree& Tree = Generator.GetTree();
+
+	OutExpression = Tree.NewLerp(ExpressionB,
+								 Tree.NewAdd(ExpressionA, ExpressionB),
+								 ExpressionAlpha);
+
+	return false;
+}
 #endif
 
 #undef LOCTEXT_NAMESPACE 
\ No newline at end of file
diff --git a/Engine/Plugins/Interchange/Runtime/Source/Import/Private/MaterialX/MaterialExpressions/MaterialExpressionPlus.h b/Engine/Plugins/Interchange/Runtime/Source/Import/Private/MaterialX/MaterialExpressions/MaterialExpressionPlus.h
--- a/Engine/Plugins/Interchange/Runtime/Source/Import/Private/MaterialX/MaterialExpressions/MaterialExpressionPlus.h
+++ b/Engine/Plugins/Interchange/Runtime/Source/Import/Private/MaterialX/MaterialExpressions/MaterialExpressionPlus.h
@@ -35,6 +35,7 @@ class UMaterialExpressionMaterialXPlus : public UMaterialExpression
 #if WITH_EDITOR
 	virtual int32 Compile(class FMaterialCompiler* Compiler, int32 OutputIndex) override;
 	virtual void GetCaption(TArray<FString>& OutCaptions) const override;
+	virtual bool GenerateHLSLExpression(FMaterialHLSLGenerator& Generator, UE::HLSLTree::FScope& Scope, int32 OutputIndex, UE::HLSLTree::FExpression const*& OutExpression) const override;
 #endif
 	//~ End UMaterialExpressionMaterialX Interface
 };
diff --git a/Engine/Plugins/Interchange/Runtime/Source/Import/Private/MaterialX/MaterialExpressions/MaterialExpressionPremult.cpp b/Engine/Plugins/Interchange/Runtime/Source/Import/Private/MaterialX/MaterialExpressions/MaterialExpressionPremult.cpp
--- a/Engine/Plugins/Interchange/Runtime/Source/Import/Private/MaterialX/MaterialExpressions/MaterialExpressionPremult.cpp
+++ b/Engine/Plugins/Interchange/Runtime/Source/Import/Private/MaterialX/MaterialExpressions/MaterialExpressionPremult.cpp
@@ -1,6 +1,9 @@
 // Copyright Epic Games, Inc. All Rights Reserved.
 #include "MaterialExpressionPremult.h"
 #include "MaterialCompiler.h"
+#include "MaterialHLSLGenerator.h"
+#include "HLSLTree/HLSLTree.h"
+#include "HLSLTree/HLSLTreeCommon.h"
 
 #include UE_INLINE_GENERATED_CPP_BY_NAME(MaterialExpressionPremult)
 
@@ -42,6 +45,30 @@ void UMaterialExpressionMaterialXPremult::GetCaption(TArray<FString>& OutCaption
 {
 		OutCaptions.Add(TEXT("MaterialX Premult"));
 }
+
+bool UMaterialExpressionMaterialXPremult::GenerateHLSLExpression(FMaterialHLSLGenerator& Generator, UE::HLSLTree::FScope& Scope, int32 OutputIndex, UE::HLSLTree::FExpression const*& OutExpression) const
+{
+	using namespace UE::HLSLTree;
+
+	const FExpression* ExpressionInput = Input.AcquireHLSLExpression(Generator, Scope);
+
+	if(!ExpressionInput)
+	{
+		return false;
+	}
+
+	FTree& Tree = Generator.GetTree();
+
+	const FExpression* ExpressionAlpha = Tree.NewSwizzle(FSwizzleParameters(3), ExpressionInput);
+
+	OutExpression =
+		Tree.NewExpression<FExpressionAppend>(
+			Tree.NewMul(Tree.NewSwizzle(FSwizzleParameters(0, 1, 2), ExpressionInput),
+						ExpressionAlpha),
+			ExpressionAlpha);
+
+	return true;
+}
 #endif
 
 #undef LOCTEXT_NAMESPACE 
\ No newline at end of file
diff --git a/Engine/Plugins/Interchange/Runtime/Source/Import/Private/MaterialX/MaterialExpressions/MaterialExpressionPremult.h b/Engine/Plugins/Interchange/Runtime/Source/Import/Private/MaterialX/MaterialExpressions/MaterialExpressionPremult.h
--- a/Engine/Plugins/Interchange/Runtime/Source/Import/Private/MaterialX/MaterialExpressions/MaterialExpressionPremult.h
+++ b/Engine/Plugins/Interchange/Runtime/Source/Import/Private/MaterialX/MaterialExpressions/MaterialExpressionPremult.h
@@ -22,6 +22,7 @@ class UMaterialExpressionMaterialXPremult : public UMaterialExpression
 #if WITH_EDITOR
 	virtual int32 Compile(class FMaterialCompiler* Compiler, int32 OutputIndex) override;
 	virtual void GetCaption(TArray<FString>& OutCaptions) const override;
+	virtual bool GenerateHLSLExpression(FMaterialHLSLGenerator& Generator, UE::HLSLTree::FScope& Scope, int32 OutputIndex, UE::HLSLTree::FExpression const*& OutExpression) const override;
 #endif
 	//~ End UMaterialExpressionMaterialX Interface
 };
diff --git a/Engine/Plugins/Interchange/Runtime/Source/Import/Private/MaterialX/MaterialExpressions/MaterialExpressionScreen.cpp b/Engine/Plugins/Interchange/Runtime/Source/Import/Private/MaterialX/MaterialExpressions/MaterialExpressionScreen.cpp
--- a/Engine/Plugins/Interchange/Runtime/Source/Import/Private/MaterialX/MaterialExpressions/MaterialExpressionScreen.cpp
+++ b/Engine/Plugins/Interchange/Runtime/Source/Import/Private/MaterialX/MaterialExpressions/MaterialExpressionScreen.cpp
@@ -1,6 +1,8 @@
 // Copyright Epic Games, Inc. All Rights Reserved.
 #include "MaterialExpressionScreen.h"
 #include "MaterialCompiler.h"
+#include "MaterialHLSLGenerator.h"
+#include "HLSLTree/HLSLTree.h"
 
 #include UE_INLINE_GENERATED_CPP_BY_NAME(MaterialExpressionScreen)
 
@@ -56,6 +58,35 @@ void UMaterialExpressionMaterialXScreen::GetCaption(TArray<FString>& OutCaptions
 {
 	OutCaptions.Add(TEXT("MaterialX Screen"));
 }
+
+bool UMaterialExpressionMaterialXScreen::GenerateHLSLExpression(FMaterialHLSLGenerator& Generator, UE::HLSLTree::FScope& Scope, int32 OutputIndex, UE::HLSLTree::FExpression const*& OutExpression) const
+{
+	using namespace UE::HLSLTree;
+
+	const FExpression* ExpressionA = A.AcquireHLSLExpression(Generator, Scope);
+	const FExpression* ExpressionB = B.AcquireHLSLExpression(Generator, Scope);
+	const FExpression* ExpressionAlpha = Alpha.AcquireHLSLExpressionOrConstant(Generator, Scope, ConstAlpha);
+
+	if(!ExpressionA || !ExpressionB || !ExpressionAlpha)
+	{
+		return false;
+	}
+
+	FTree& Tree = Generator.GetTree();
+
+	const FExpression* ExpressionOne = Tree.NewConstant(1.f);
+
+	auto NewOneMinus = [&](const FExpression* Input)
+	{
+		return Tree.NewSub(ExpressionOne, Input);
+	};
+
+	OutExpression = Tree.NewLerp(ExpressionB,
+								 NewOneMinus(Tree.NewMul(NewOneMinus(ExpressionA), NewOneMinus(ExpressionB))),
+								 ExpressionAlpha);
+
+	return true;
+}
 #endif
 
 #undef LOCTEXT_NAMESPACE 
\ No newline at end of file
diff --git a/Engine/Plugins/Interchange/Runtime/Source/Import/Private/MaterialX/MaterialExpressions/MaterialExpressionScreen.h b/Engine/Plugins/Interchange/Runtime/Source/Import/Private/MaterialX/MaterialExpressions/MaterialExpressionScreen.h
--- a/Engine/Plugins/Interchange/Runtime/Source/Import/Private/MaterialX/MaterialExpressions/MaterialExpressionScreen.h
+++ b/Engine/Plugins/Interchange/Runtime/Source/Import/Private/MaterialX/MaterialExpressions/MaterialExpressionScreen.h
@@ -35,6 +35,7 @@ class UMaterialExpressionMaterialXScreen : public UMaterialExpression
 #if WITH_EDITOR
 	virtual int32 Compile(class FMaterialCompiler* Compiler, int32 OutputIndex) override;
 	virtual void GetCaption(TArray<FString>& OutCaptions) const override;
+	virtual bool GenerateHLSLExpression(FMaterialHLSLGenerator& Generator, UE::HLSLTree::FScope& Scope, int32 OutputIndex, UE::HLSLTree::FExpression const*& OutExpression) const override;
 #endif
 	//~ End UMaterialExpressionMaterialX Interface
 };
