Interchange: Add missing traces to be able to verify the performance.

#rnx

diff --git a/Engine/Plugins/Interchange/Runtime/Source/Import/Private/Animation/InterchangeAnimSequenceFactory.cpp b/Engine/Plugins/Interchange/Runtime/Source/Import/Private/Animation/InterchangeAnimSequenceFactory.cpp
--- a/Engine/Plugins/Interchange/Runtime/Source/Import/Private/Animation/InterchangeAnimSequenceFactory.cpp
+++ b/Engine/Plugins/Interchange/Runtime/Source/Import/Private/Animation/InterchangeAnimSequenceFactory.cpp
@@ -250,6 +250,8 @@ namespace UE::Interchange::Private
 	}
 
 	void RetrieveAnimationPayloads(UAnimSequence* AnimSequence
+		, UInterchangeAnimSequenceFactory::FBoneTrackData& BoneTrackData
+		, UInterchangeAnimSequenceFactory::FMorphTargetData& MorphTargetData
 		, const UInterchangeAnimSequenceFactoryNode* AnimSequenceFactoryNode
 		, const UInterchangeBaseNodeContainer* NodeContainer
 		, const UInterchangeSkeletonFactoryNode* SkeletonFactoryNode
@@ -274,6 +276,8 @@ namespace UE::Interchange::Private
 		TArray<FString> SkeletonNodes;
 		GetSkeletonSceneNodeFlatListRecursive(NodeContainer, SkeletonRootUid, SkeletonNodes);
 
+		TMap<FString, FInterchangeAnimationPayLoadKey> PayloadKeys;
+		AnimSequenceFactoryNode->GetSceneNodeAnimationPayloadKeys(PayloadKeys);
 		
 		const bool bShouldTransact = bIsReimporting;
 
@@ -293,17 +297,6 @@ namespace UE::Interchange::Private
 
 			const double BakeInterval = 1.0 / SampleRate;
 
-			TMap<FString, FInterchangeAnimationPayLoadKey> PayloadKeys;
-			AnimSequenceFactoryNode->GetSceneNodeAnimationPayloadKeys(PayloadKeys);
-
-			for (const TTuple<FString, FInterchangeAnimationPayLoadKey>& SceneNodeUidAndPayloadKey : PayloadKeys)
-			{
-				if (const UInterchangeSceneNode* SkeletonSceneNode = Cast<UInterchangeSceneNode>(NodeContainer->GetNode(SceneNodeUidAndPayloadKey.Key)))
-				{
-					AnimationPayloads.Add(SkeletonSceneNode, AnimSequenceTranslatorPayloadInterface->GetAnimationPayloadData(SceneNodeUidAndPayloadKey.Value, SampleRate, RangeStart, RangeEnd));
-				}
-			}
-
 			//This destroy all previously imported animation raw data
 			Controller.RemoveAllBoneTracks(bShouldTransact);
 
@@ -318,58 +311,14 @@ namespace UE::Interchange::Private
 				}
 			}
 
-			double MergedRangeEnd = RangeEnd;
-			double MergedRangeStart = RangeStart;
-			TMap<const UInterchangeSceneNode*, UE::Interchange::FAnimationPayloadData> PreProcessedAnimationPayloads;
-			for (const TTuple< const UInterchangeSceneNode*, TFuture<TOptional<UE::Interchange::FAnimationPayloadData>>>& AnimationPayload : AnimationPayloads)
-			{
-				const FName BoneName = FName(*(AnimationPayload.Key->GetDisplayLabel()));
-				const int32 BoneIndex = Skeleton->GetReferenceSkeleton().FindBoneIndex(BoneName);
-				if (BoneIndex == INDEX_NONE)
-				{
-					//Skip this bone, we did not found it in the skeleton
-					continue;
-				}
-				//If we are getting the root 
-				bool bApplyGlobalOffset = AnimationPayload.Key->GetUniqueID().Equals(SkeletonRootUid);
-
-				TOptional<UE::Interchange::FAnimationPayloadData> OptionalAnimationTransformPayload = AnimationPayload.Value.Get();
-				if (!OptionalAnimationTransformPayload.IsSet())
-				{
-					FString PayloadKey = PayloadKeys[AnimationPayload.Key->GetUniqueID()].UniqueId;
-					UE_LOG(LogInterchangeImport, Warning, TEXT("Invalid animation transform payload key [%s] AnimSequence asset %s"), *PayloadKey, *AssetName);
-					continue;
-				}
-
-				UE::Interchange::FAnimationPayloadData& AnimationTransformPayload = OptionalAnimationTransformPayload.GetValue();
-
-				if (AnimationTransformPayload.Type != EInterchangeAnimationPayLoadType::BAKED)
-				{
-					//Where Curve is null the LocalTransform should be used for the Baked Transform generation.
-					FTransform LocalTransform;
-					AnimationPayload.Key->GetCustomLocalTransform(LocalTransform);
-
-					//Currently only Curve -> Baked conversion (for LevelSequence->AnimSequence conversion by ForceMeshType Skeletal use case)
-					//and Curve -> Step Curve conversion (for custom attributes)
-					AnimationTransformPayload.CalculateDataFor(EInterchangeAnimationPayLoadType::BAKED, LocalTransform);
-					//Range End will be calculated as well:
-					if (MergedRangeEnd < AnimationTransformPayload.RangeEndTime)
-					{
-						MergedRangeEnd = AnimationTransformPayload.RangeEndTime;
-					}
-				}
-
-				PreProcessedAnimationPayloads.Add(AnimationPayload.Key, AnimationTransformPayload);
-			}
-
-			const double SequenceLength = FMath::Max<double>(MergedRangeEnd - MergedRangeStart, MINIMUM_ANIMATION_LENGTH);
+			const double SequenceLength = FMath::Max<double>(BoneTrackData.MergedRangeEnd - BoneTrackData.MergedRangeStart, MINIMUM_ANIMATION_LENGTH);
 			int32 FrameCount = FMath::RoundToInt32(SequenceLength * SampleRate);
 			int32 BakeKeyCount = FrameCount + 1;
 			const FFrameRate ResampleFrameRate(SampleRate, 1);
 			Controller.SetFrameRate(ResampleFrameRate, bShouldTransact);
 			Controller.SetNumberOfFrames(FrameCount, bShouldTransact);
 
-			for (TTuple< const UInterchangeSceneNode*, UE::Interchange::FAnimationPayloadData>& AnimationPayload : PreProcessedAnimationPayloads)
+			for (TTuple< const UInterchangeSceneNode*, UE::Interchange::FAnimationPayloadData>& AnimationPayload : BoneTrackData.PreProcessedAnimationPayloads)
 			{
 				const FName BoneName = FName(*(AnimationPayload.Key->GetDisplayLabel()));
 				UE::Interchange::FAnimationPayloadData& AnimationTransformPayload = AnimationPayload.Value;
@@ -547,74 +496,18 @@ namespace UE::Interchange::Private
 
 			//Import morph target curves
 			{
-				TMap<FString, FInterchangeAnimationPayLoadKey> MorphTargetNodeAnimationPayloads;
-				AnimSequenceFactoryNode->GetMorphTargetNodeAnimationPayloadKeys(MorphTargetNodeAnimationPayloads);
-
-				TMap<FString, TFuture<TOptional<UE::Interchange::FAnimationPayloadData>>> AnimationCurvesPayloads;
-				TMap<FString, FString> AnimationCurveMorphTargetNodeNames;
-
-				for (const TPair<FString, FInterchangeAnimationPayLoadKey>& MorphTargetNodeUidAnimationPayload : MorphTargetNodeAnimationPayloads)
+				for (TPair<FString, UE::Interchange::FAnimationPayloadData>& CurveNameAndPayload : MorphTargetData.CurvesPayloads)
 				{
-					FString PayloadKey = MorphTargetNodeUidAnimationPayload.Value.UniqueId;
-					if (PayloadKey.Len() == 0)
-					{
-						continue;
-					}
-					if (const UInterchangeMeshNode* MorphTargetNode = Cast<UInterchangeMeshNode>(NodeContainer->GetNode(MorphTargetNodeUidAnimationPayload.Key)))
-					{
-						if (MorphTargetNodeUidAnimationPayload.Value.Type == EInterchangeAnimationPayLoadType::MORPHTARGETCURVEWEIGHTINSTANCE)
-						{
-							AnimationCurvesPayloads.Add(PayloadKey, Async(EAsyncExecution::TaskGraph, [&MorphTargetNodeUidAnimationPayload]
-								{
-									TOptional<UE::Interchange::FAnimationPayloadData> Result;
-									UE::Interchange::FAnimationPayloadData AnimationPayLoadData(MorphTargetNodeUidAnimationPayload.Value.Type);
-
-									TArray<FString> PayLoadKeys;
-									MorphTargetNodeUidAnimationPayload.Value.UniqueId.ParseIntoArray(PayLoadKeys, TEXT(":"));
-
-									if (PayLoadKeys.Num() != 2)
-									{
-										return Result;
-									}
-
-									float Weight;
-									LexFromString(Weight, *PayLoadKeys[1]);
-
-									AnimationPayLoadData.Curves.SetNum(1);
-									FRichCurve& Curve = AnimationPayLoadData.Curves[0];
-									Curve.AddKey(0, Weight);
-
-									Result.Emplace(AnimationPayLoadData);
-
-									return Result;
-								}));
-						}
-						else
-						{
-							AnimationCurvesPayloads.Add(PayloadKey, AnimSequenceTranslatorPayloadInterface->GetAnimationPayloadData(MorphTargetNodeUidAnimationPayload.Value));
-						}
-						AnimationCurveMorphTargetNodeNames.Add(PayloadKey, MorphTargetNode->GetDisplayLabel());
-					}
-				}
-
-				for (TPair<FString, TFuture<TOptional<UE::Interchange::FAnimationPayloadData>>>& CurveNameAndPayload : AnimationCurvesPayloads)
-				{
-					TOptional<UE::Interchange::FAnimationPayloadData> AnimationCurvePayload = CurveNameAndPayload.Value.Get();
-					if (!AnimationCurvePayload.IsSet())
-					{
-						UE_LOG(LogInterchangeImport, Warning, TEXT("Invalid animation morph target curve payload key [%s] AnimSequence asset %s"), *CurveNameAndPayload.Key, *AssetName);
-						continue;
-					}
-					FAnimationPayloadData& CurvePayload = AnimationCurvePayload.GetValue();
+					UE::Interchange::FAnimationPayloadData& AnimationCurvePayload = CurveNameAndPayload.Value;
 					if (bRemoveCurveRedundantKeys)
 					{
-						for (FRichCurve& RichCurve : CurvePayload.Curves)
+						for (FRichCurve& RichCurve : AnimationCurvePayload.Curves)
 						{
 							RichCurve.RemoveRedundantAutoTangentKeys(SMALL_NUMBER);
 						}
 					}
 					constexpr int32 CurveFlags = 0;
-					CreateMorphTargetCurve(AnimSequence, CurvePayload.Curves, AnimationCurveMorphTargetNodeNames.FindChecked(CurveNameAndPayload.Key), CurveFlags, bDoNotImportCurveWithZero, bAddCurveMetadataToSkeleton, bShouldTransact);
+					CreateMorphTargetCurve(AnimSequence, AnimationCurvePayload.Curves, MorphTargetData.CurveNodeNamePerPayloadKey.FindChecked(CurveNameAndPayload.Key), CurveFlags, bDoNotImportCurveWithZero, bAddCurveMetadataToSkeleton, bShouldTransact);
 				}
 			}
 
@@ -813,13 +706,15 @@ UClass* UInterchangeAnimSequenceFactory::GetFactoryClass() const
 
 UInterchangeFactoryBase::FImportAssetResult UInterchangeAnimSequenceFactory::BeginImportAsset_GameThread(const FImportAssetObjectParams& Arguments)
 {
-	FImportAssetResult ImportAssetResult;
+	TRACE_CPUPROFILER_EVENT_SCOPE("UInterchangeAnimSequenceFactory::BeginImportAsset_GameThread");
+
+	UInterchangeFactoryBase::FImportAssetResult ImportAssetResult;
 #if !WITH_EDITOR || !WITH_EDITORONLY_DATA
 
 	UE_LOG(LogInterchangeImport, Error, TEXT("Cannot import animsequence asset in runtime, this is an editor only feature."));
 	return ImportAssetResult;
 #else
-	UAnimSequence* AnimSequence = nullptr;
+	UAnimSequence* NewAnimSequence = nullptr;
 	if (!Arguments.AssetNode || !Arguments.AssetNode->GetObjectClass()->IsChildOf(GetFactoryClass()))
 	{
 		return ImportAssetResult;
@@ -830,7 +725,14 @@ UInterchangeFactoryBase::FImportAssetResult UInterchangeAnimSequenceFactory::Beg
 	{
 		return ImportAssetResult;
 	}
-	
+
+	const IInterchangeAnimationPayloadInterface* AnimSequenceTranslatorPayloadInterface = Cast<IInterchangeAnimationPayloadInterface>(Arguments.Translator);
+	if (!AnimSequenceTranslatorPayloadInterface)
+	{
+		UE_LOG(LogInterchangeImport, Error, TEXT("Cannot import AnimSequence, the translator do not implement the IInterchangeAnimationPayloadInterface."));
+		return ImportAssetResult;
+	}
+
 	//Verify if the bone track animation is valid (sequence length versus framerate ...)
 	if(!IsBoneTrackAnimationValid(AnimSequenceFactoryNode, Arguments))
 	{
@@ -850,69 +752,301 @@ UInterchangeFactoryBase::FImportAssetResult UInterchangeAnimSequenceFactory::Beg
 	// create a new material or overwrite existing asset, if possible
 	if (!ExistingAsset)
 	{
-		AnimSequence = NewObject<UAnimSequence>(Arguments.Parent, *Arguments.AssetName, RF_Public | RF_Standalone);
+		NewAnimSequence = NewObject<UAnimSequence>(Arguments.Parent, *Arguments.AssetName, RF_Public | RF_Standalone);
 	}
 	else
 	{
 		//This is a reimport, we are just re-updating the source data
-		AnimSequence = Cast<UAnimSequence>(ExistingAsset);
+		NewAnimSequence = Cast<UAnimSequence>(ExistingAsset);
 	}
 
-	if (!AnimSequence)
+	if (!NewAnimSequence)
 	{
 		UE_LOG(LogInterchangeImport, Warning, TEXT("Could not create AnimSequence asset %s"), *Arguments.AssetName);
 		return ImportAssetResult;
 	}
 
-	AnimSequenceFactoryNode->SetCustomReferenceObject(FSoftObjectPath(AnimSequence));
+	AnimSequenceFactoryNode->SetCustomReferenceObject(FSoftObjectPath(NewAnimSequence));
+
+	NewAnimSequence->PreEditChange(nullptr);
+
+	//Verify if the bone track animation is valid (sequence length versus framerate ...)
+	if (!IsBoneTrackAnimationValid(AnimSequenceFactoryNode, Arguments))
+	{
+		return ImportAssetResult;
+	}
+
+	FString SkeletonUid;
+	if (!AnimSequenceFactoryNode->GetCustomSkeletonFactoryNodeUid(SkeletonUid))
+	{
+		UE_LOG(LogInterchangeImport, Warning, TEXT("Could not create AnimSequence asset %s, because there is no skeleton."), *Arguments.AssetName);
+		return ImportAssetResult;
+	}
+
+	const UInterchangeSkeletonFactoryNode* SkeletonFactoryNode = Cast<UInterchangeSkeletonFactoryNode>(Arguments.NodeContainer->GetNode(SkeletonUid));
+	if (!SkeletonFactoryNode)
+	{
+		UE_LOG(LogInterchangeImport, Warning, TEXT("Invalid skeleton factory node, the skeleton factory node is obligatory to import this animsequence [%s]!"), *Arguments.AssetName);
+		return ImportAssetResult;
+	}
+	FSoftObjectPath SkeletonFactoryNodeReferenceObject;
+	SkeletonFactoryNode->GetCustomReferenceObject(SkeletonFactoryNodeReferenceObject);
+
+	USkeleton* Skeleton = nullptr;
+
+	FSoftObjectPath SpecifiedSkeleton;
+	AnimSequenceFactoryNode->GetCustomSkeletonSoftObjectPath(SpecifiedSkeleton);
+	if (Skeleton == nullptr)
+	{
+		UObject* SkeletonObject = nullptr;
+
+		if (SpecifiedSkeleton.IsValid())
+		{
+			SkeletonObject = SpecifiedSkeleton.TryLoad();
+		}
+		else if (SkeletonFactoryNodeReferenceObject.IsValid())
+		{
+			SkeletonObject = SkeletonFactoryNodeReferenceObject.TryLoad();
+		}
+
+		if (SkeletonObject)
+		{
+			Skeleton = Cast<USkeleton>(SkeletonObject);
+
+		}
+
+		if (!ensure(Skeleton))
+		{
+			UE_LOG(LogInterchangeImport, Warning, TEXT("Invalid Skeleton when importing animation sequence asset %s"), *Arguments.AssetName);
+			return ImportAssetResult;
+		}
+	}
 
-	AnimSequence->PreEditChange(nullptr);
+	NewAnimSequence->SetSkeleton(Skeleton);
 
+	AnimSequence = NewAnimSequence;
+	ImportAssetResult.ImportedObject = NewAnimSequence;
 
-	ImportAssetResult.ImportedObject = ImportObjectSourceData(Arguments);
 	return ImportAssetResult;
 #endif //else !WITH_EDITOR || !WITH_EDITORONLY_DATA
 }
 
-UObject* UInterchangeAnimSequenceFactory::ImportObjectSourceData(const FImportAssetObjectParams& Arguments)
+UInterchangeFactoryBase::FImportAssetResult UInterchangeAnimSequenceFactory::ImportAsset_Async(const FImportAssetObjectParams& Arguments)
 {
-#if !WITH_EDITOR || !WITH_EDITORONLY_DATA
+	TRACE_CPUPROFILER_EVENT_SCOPE("UInterchangeAnimSequenceFactory::ImportAsset_Async");
 
-	UE_LOG(LogInterchangeImport, Error, TEXT("Cannot import AnimSequence asset in runtime, this is an editor only feature."));
-	return nullptr;
+	UInterchangeFactoryBase::FImportAssetResult ImportAssetResult;
+#if !WITH_EDITOR || !WITH_EDITORONLY_DATA
 
+	UE_LOG(LogInterchangeImport, Error, TEXT("Cannot import animsequence asset in runtime, this is an editor only feature."));
+	return ImportAssetResult;
 #else
 
-	if (!Arguments.AssetNode || !Arguments.AssetNode->GetObjectClass()->IsChildOf(GetFactoryClass()))
+	//The game thread part should have verified all the data, so no need to do extra log
+	if (!AnimSequence)
+	{
+		return ImportAssetResult;
+	}
+	const IInterchangeAnimationPayloadInterface* AnimSequenceTranslatorPayloadInterface = Cast<IInterchangeAnimationPayloadInterface>(Arguments.Translator);
+	if (!AnimSequenceTranslatorPayloadInterface)
 	{
-		return nullptr;
+		return ImportAssetResult;
 	}
 
 	UInterchangeAnimSequenceFactoryNode* AnimSequenceFactoryNode = Cast<UInterchangeAnimSequenceFactoryNode>(Arguments.AssetNode);
 	if (AnimSequenceFactoryNode == nullptr)
 	{
-		return nullptr;
+		return ImportAssetResult;
 	}
 
-	//Verify if the bone track animation is valid (sequence length versus framerate ...)
-	if (!IsBoneTrackAnimationValid(AnimSequenceFactoryNode, Arguments))
+	const bool bIsReImport = (Arguments.ReimportObject != nullptr);
+
+	bool bImportBoneTracks = false;
+	AnimSequenceFactoryNode->GetCustomImportBoneTracks(bImportBoneTracks);
+	if (bImportBoneTracks)
 	{
-		return nullptr;
+		USkeleton* Skeleton = AnimSequence->GetSkeleton();
+		check(Skeleton);
+
+		TMap<const UInterchangeSceneNode*, TFuture<TOptional<UE::Interchange::FAnimationPayloadData>>> AnimationPayloads;
+
+		//Get the sample rate, default to 30Hz in case the attribute is missing
+		double SampleRate = 30.0;
+		AnimSequenceFactoryNode->GetCustomImportBoneTracksSampleRate(SampleRate);
+
+		double RangeStart = 0.0;
+		AnimSequenceFactoryNode->GetCustomImportBoneTracksRangeStart(RangeStart);
+
+		double RangeEnd = 1.0 / SampleRate; //One frame duration per default
+		AnimSequenceFactoryNode->GetCustomImportBoneTracksRangeStop(RangeEnd);
+
+
+		TMap<FString, FInterchangeAnimationPayLoadKey> PayloadKeys;
+		AnimSequenceFactoryNode->GetSceneNodeAnimationPayloadKeys(PayloadKeys);
+
+		for (const TTuple<FString, FInterchangeAnimationPayLoadKey>& SceneNodeUidAndPayloadKey : PayloadKeys)
+		{
+			if (const UInterchangeSceneNode* SkeletonSceneNode = Cast<UInterchangeSceneNode>(Arguments.NodeContainer->GetNode(SceneNodeUidAndPayloadKey.Key)))
+			{
+				AnimationPayloads.Add(SkeletonSceneNode, AnimSequenceTranslatorPayloadInterface->GetAnimationPayloadData(SceneNodeUidAndPayloadKey.Value, SampleRate, RangeStart, RangeEnd));
+			}
+		}
+		
+		BoneTrackData.MergedRangeEnd = RangeEnd;
+		BoneTrackData.MergedRangeStart = RangeStart;
+		BoneTrackData.PreProcessedAnimationPayloads.Reset();
+		for (const TTuple< const UInterchangeSceneNode*, TFuture<TOptional<UE::Interchange::FAnimationPayloadData>>>& AnimationPayload : AnimationPayloads)
+		{
+			const FName BoneName = FName(*(AnimationPayload.Key->GetDisplayLabel()));
+			const int32 BoneIndex = Skeleton->GetReferenceSkeleton().FindBoneIndex(BoneName);
+			if (BoneIndex == INDEX_NONE)
+			{
+				//Skip this bone, we did not found it in the skeleton
+				continue;
+			}
+
+			TOptional<UE::Interchange::FAnimationPayloadData> OptionalAnimationTransformPayload = AnimationPayload.Value.Get();
+			if (!OptionalAnimationTransformPayload.IsSet())
+			{
+				FString PayloadKey = PayloadKeys[AnimationPayload.Key->GetUniqueID()].UniqueId;
+				UE_LOG(LogInterchangeImport, Warning, TEXT("Invalid animation transform payload key [%s] AnimSequence asset %s"), *PayloadKey, *Arguments.AssetName);
+				continue;
+			}
+
+			UE::Interchange::FAnimationPayloadData& AnimationTransformPayload = OptionalAnimationTransformPayload.GetValue();
+
+			if (AnimationTransformPayload.Type != EInterchangeAnimationPayLoadType::BAKED)
+			{
+				//Where Curve is null the LocalTransform should be used for the Baked Transform generation.
+				FTransform LocalTransform;
+				AnimationPayload.Key->GetCustomLocalTransform(LocalTransform);
+
+				//Currently only Curve -> Baked conversion (for LevelSequence->AnimSequence conversion by ForceMeshType Skeletal use case)
+				//and Curve -> Step Curve conversion (for custom attributes)
+				AnimationTransformPayload.CalculateDataFor(EInterchangeAnimationPayLoadType::BAKED, LocalTransform);
+				//Range End will be calculated as well:
+				if (BoneTrackData.MergedRangeEnd < AnimationTransformPayload.RangeEndTime)
+				{
+					BoneTrackData.MergedRangeEnd = AnimationTransformPayload.RangeEndTime;
+				}
+			}
+
+			BoneTrackData.PreProcessedAnimationPayloads.Add(AnimationPayload.Key, AnimationTransformPayload);
+		}
+	}
+
+	//Import morph target curves
+	{
+		TMap<FString, FInterchangeAnimationPayLoadKey> MorphTargetNodeAnimationPayloads;
+		AnimSequenceFactoryNode->GetMorphTargetNodeAnimationPayloadKeys(MorphTargetNodeAnimationPayloads);
+
+		TMap<FString, TFuture<TOptional<UE::Interchange::FAnimationPayloadData>>> AnimationCurvesPayloads;
+
+		for (const TPair<FString, FInterchangeAnimationPayLoadKey>& MorphTargetNodeUidAnimationPayload : MorphTargetNodeAnimationPayloads)
+		{
+			FString PayloadKey = MorphTargetNodeUidAnimationPayload.Value.UniqueId;
+			if (PayloadKey.Len() == 0)
+			{
+				continue;
+			}
+			if (const UInterchangeMeshNode* MorphTargetNode = Cast<UInterchangeMeshNode>(Arguments.NodeContainer->GetNode(MorphTargetNodeUidAnimationPayload.Key)))
+			{
+				if (MorphTargetNodeUidAnimationPayload.Value.Type == EInterchangeAnimationPayLoadType::MORPHTARGETCURVEWEIGHTINSTANCE)
+				{
+					AnimationCurvesPayloads.Add(PayloadKey, Async(EAsyncExecution::TaskGraph, [&MorphTargetNodeUidAnimationPayload]
+						{
+							TOptional<UE::Interchange::FAnimationPayloadData> Result;
+							UE::Interchange::FAnimationPayloadData AnimationPayLoadData(MorphTargetNodeUidAnimationPayload.Value.Type);
+
+							TArray<FString> PayLoadKeys;
+							MorphTargetNodeUidAnimationPayload.Value.UniqueId.ParseIntoArray(PayLoadKeys, TEXT(":"));
+
+							if (PayLoadKeys.Num() != 2)
+							{
+								return Result;
+							}
+
+							float Weight;
+							LexFromString(Weight, *PayLoadKeys[1]);
+
+							AnimationPayLoadData.Curves.SetNum(1);
+							FRichCurve& Curve = AnimationPayLoadData.Curves[0];
+							Curve.AddKey(0, Weight);
+
+							Result.Emplace(AnimationPayLoadData);
+
+							return Result;
+						}));
+				}
+				else
+				{
+					AnimationCurvesPayloads.Add(PayloadKey, AnimSequenceTranslatorPayloadInterface->GetAnimationPayloadData(MorphTargetNodeUidAnimationPayload.Value));
+				}
+				MorphTargetData.CurveNodeNamePerPayloadKey.Add(PayloadKey, MorphTargetNode->GetDisplayLabel());
+			}
+		}
+
+		for (TPair<FString, TFuture<TOptional<UE::Interchange::FAnimationPayloadData>>>& CurveNameAndPayload : AnimationCurvesPayloads)
+		{
+			TOptional<UE::Interchange::FAnimationPayloadData> AnimationCurvePayload = CurveNameAndPayload.Value.Get();
+			if (!AnimationCurvePayload.IsSet())
+			{
+				UE_LOG(LogInterchangeImport, Warning, TEXT("Invalid animation morph target curve payload key [%s] AnimSequence asset %s"), *CurveNameAndPayload.Key, *Arguments.AssetName);
+				continue;
+			}
+			UE::Interchange::FAnimationPayloadData& CurvePayload = AnimationCurvePayload.GetValue();
+			MorphTargetData.CurvesPayloads.Add(CurveNameAndPayload.Key, MoveTemp(CurvePayload));
+		}
+		AnimationCurvesPayloads.Empty();
+	}
+
+	//Getting the file Hash will cache it into the source data
+	Arguments.SourceData->GetFileContentHash();
+
+	return ImportAssetResult;
+#endif //else !WITH_EDITOR || !WITH_EDITORONLY_DATA
+}
+
+UInterchangeFactoryBase::FImportAssetResult UInterchangeAnimSequenceFactory::EndImportAsset_GameThread(const FImportAssetObjectParams& Arguments)
+{
+	TRACE_CPUPROFILER_EVENT_SCOPE("UInterchangeAnimSequenceFactory::ImportAsset_Async");
+
+	UInterchangeFactoryBase::FImportAssetResult ImportAssetResult;
+#if !WITH_EDITOR || !WITH_EDITORONLY_DATA
+
+	UE_LOG(LogInterchangeImport, Error, TEXT("Cannot import animsequence asset in runtime, this is an editor only feature."));
+	return ImportAssetResult;
+#else
+	if (!AnimSequence)
+	{
+		return ImportAssetResult;
+	}
+
+	if (!Arguments.AssetNode || !Arguments.AssetNode->GetObjectClass()->IsChildOf(GetFactoryClass()))
+	{
+		return ImportAssetResult;
+	}
+
+	UInterchangeAnimSequenceFactoryNode* AnimSequenceFactoryNode = Cast<UInterchangeAnimSequenceFactoryNode>(Arguments.AssetNode);
+	if (AnimSequenceFactoryNode == nullptr)
+	{
+		return ImportAssetResult;
 	}
 
 	FString SkeletonUid;
 	if (!AnimSequenceFactoryNode->GetCustomSkeletonFactoryNodeUid(SkeletonUid))
 	{
 		//Do not create a empty anim sequence, we need skeleton that contain animation
-		return nullptr;
+		return ImportAssetResult;
 	}
 
 	const UInterchangeSkeletonFactoryNode* SkeletonFactoryNode = Cast<UInterchangeSkeletonFactoryNode>(Arguments.NodeContainer->GetNode(SkeletonUid));
 	if (!SkeletonFactoryNode)
 	{
 		UE_LOG(LogInterchangeImport, Warning, TEXT("Invalid skeleton factory node, the skeleton factory node is obligatory to import this animsequence [%s]!"), *Arguments.AssetName);
-		return nullptr;
+		return ImportAssetResult;
 	}
+
 	FSoftObjectPath SkeletonFactoryNodeReferenceObject;
 	SkeletonFactoryNode->GetCustomReferenceObject(SkeletonFactoryNodeReferenceObject);
 
@@ -942,7 +1076,7 @@ UObject* UInterchangeAnimSequenceFactory::ImportObjectSourceData(const FImportAs
 		if (!ensure(Skeleton))
 		{
 			UE_LOG(LogInterchangeImport, Warning, TEXT("Invalid Skeleton when importing animation sequence asset %s"), *Arguments.AssetName);
-			return nullptr;
+			return ImportAssetResult;
 		}
 	}
 
@@ -950,22 +1084,7 @@ UObject* UInterchangeAnimSequenceFactory::ImportObjectSourceData(const FImportAs
 	if (!AnimSequenceTranslatorPayloadInterface)
 	{
 		UE_LOG(LogInterchangeImport, Error, TEXT("Cannot import AnimSequence, the translator do not implement the IInterchangeAnimationPayloadInterface."));
-		return nullptr;
-	}
-
-	UObject* AnimSequenceObject = UE::Interchange::FFactoryCommon::AsyncFindObject(AnimSequenceFactoryNode, GetFactoryClass(), Arguments.Parent, Arguments.AssetName);
-
-	if (!AnimSequenceObject)
-	{
-		UE_LOG(LogInterchangeImport, Error, TEXT("Could not import the AnimSequence asset %s, because the asset do not exist."), *Arguments.AssetName);
-		return nullptr;
-	}
-
-	UAnimSequence* AnimSequence = Cast<UAnimSequence>(AnimSequenceObject);
-	if (!ensure(AnimSequence))
-	{
-		UE_LOG(LogInterchangeImport, Error, TEXT("Could not cast to AnimSequence asset %s"), *Arguments.AssetName);
-		return nullptr;
+		return ImportAssetResult;
 	}
 
 	const bool bIsReImport = (Arguments.ReimportObject != nullptr);
@@ -986,7 +1105,6 @@ UObject* UInterchangeAnimSequenceFactory::ImportObjectSourceData(const FImportAs
 		const bool bShouldTransact = bIsReImport;
 		IAnimationDataController& Controller = AnimSequence->GetController();
 		Controller.OpenBracket(NSLOCTEXT("InterchangeAnimSequenceFactory", "ImportAnimationInterchange_Bracket", "Importing Animation (Interchange)"), bShouldTransact);
-		AnimSequence->SetSkeleton(Skeleton);
 		Controller.InitializeModel();
 		AnimSequence->ImportFileFramerate = SampleRate;
 		AnimSequence->ImportResampleFramerate = SampleRate;
@@ -994,6 +1112,8 @@ UObject* UInterchangeAnimSequenceFactory::ImportObjectSourceData(const FImportAs
 
 		TArray<FString> CurvesNotFound;
 		UE::Interchange::Private::RetrieveAnimationPayloads(AnimSequence
+			, BoneTrackData
+			, MorphTargetData
 			, AnimSequenceFactoryNode
 			, Arguments.NodeContainer
 			, SkeletonFactoryNode
@@ -1037,18 +1157,15 @@ UObject* UInterchangeAnimSequenceFactory::ImportObjectSourceData(const FImportAs
 		UE::Interchange::FFactoryCommon::ApplyReimportStrategyToAsset(AnimSequence, PreviousNode, CurrentNode, AnimSequenceFactoryNode);
 	}
 
-	//Getting the file Hash will cache it into the source data
-	Arguments.SourceData->GetFileContentHash();
-
-	//The interchange completion task (call in the GameThread after the factories pass), will call PostEditChange which will trig another asynchronous system that will build all material in parallel
-	return AnimSequenceObject;
-
+	return ImportAssetResult;
 #endif //else !WITH_EDITOR || !WITH_EDITORONLY_DATA
 }
 
 /* This function is call in the completion task on the main thread, use it to call main thread post creation step for your assets*/
 void UInterchangeAnimSequenceFactory::SetupObject_GameThread(const FSetupObjectParams& Arguments)
 {
+	TRACE_CPUPROFILER_EVENT_SCOPE("UInterchangeAnimSequenceFactory::SetupObject_GameThread");
+
 	check(IsInGameThread());
 	Super::SetupObject_GameThread(Arguments);
 
@@ -1057,7 +1174,7 @@ void UInterchangeAnimSequenceFactory::SetupObject_GameThread(const FSetupObjectP
 	if (ensure(Arguments.ImportedObject && Arguments.SourceData))
 	{
 		// We must call the Update of the asset source file in the main thread because UAssetImportData::Update execute some delegate we do not control
-		UAnimSequence* AnimSequence = CastChecked<UAnimSequence>(Arguments.ImportedObject);
+		check(AnimSequence == CastChecked<UAnimSequence>(Arguments.ImportedObject));
 
 		UAssetImportData* ImportDataPtr = AnimSequence->AssetImportData;
 		UE::Interchange::FFactoryCommon::FUpdateImportAssetDataParameters UpdateImportAssetDataParameters(AnimSequence, ImportDataPtr, Arguments.SourceData, Arguments.NodeUniqueID, Arguments.NodeContainer, Arguments.OriginalPipelines);
@@ -1070,7 +1187,7 @@ void UInterchangeAnimSequenceFactory::SetupObject_GameThread(const FSetupObjectP
 bool UInterchangeAnimSequenceFactory::GetSourceFilenames(const UObject* Object, TArray<FString>& OutSourceFilenames) const
 {
 #if WITH_EDITORONLY_DATA
-	if (const UAnimSequence* AnimSequence = Cast<UAnimSequence>(Object))
+	if (ensure(AnimSequence && (AnimSequence == Cast<UAnimSequence>(Object))))
 	{
 		return UE::Interchange::FFactoryCommon::GetSourceFilenames(AnimSequence->AssetImportData.Get(), OutSourceFilenames);
 	}
@@ -1082,7 +1199,7 @@ bool UInterchangeAnimSequenceFactory::GetSourceFilenames(const UObject* Object,
 bool UInterchangeAnimSequenceFactory::SetSourceFilename(const UObject* Object, const FString& SourceFilename, int32 SourceIndex) const
 {
 #if WITH_EDITORONLY_DATA
-	if (const UAnimSequence* AnimSequence = Cast<UAnimSequence>(Object))
+	if (ensure(AnimSequence && (AnimSequence == Cast<UAnimSequence>(Object))))
 	{
 		return UE::Interchange::FFactoryCommon::SetSourceFilename(AnimSequence->AssetImportData.Get(), SourceFilename, SourceIndex);
 	}
diff --git a/Engine/Plugins/Interchange/Runtime/Source/Import/Private/Animation/InterchangeAnimationTrackSetFactory.cpp b/Engine/Plugins/Interchange/Runtime/Source/Import/Private/Animation/InterchangeAnimationTrackSetFactory.cpp
--- a/Engine/Plugins/Interchange/Runtime/Source/Import/Private/Animation/InterchangeAnimationTrackSetFactory.cpp
+++ b/Engine/Plugins/Interchange/Runtime/Source/Import/Private/Animation/InterchangeAnimationTrackSetFactory.cpp
@@ -662,6 +662,8 @@ UClass* UInterchangeAnimationTrackSetFactory::GetFactoryClass() const
 
 UInterchangeFactoryBase::FImportAssetResult UInterchangeAnimationTrackSetFactory::BeginImportAsset_GameThread(const FImportAssetObjectParams& Arguments)
 {
+	TRACE_CPUPROFILER_EVENT_SCOPE("UInterchangeAnimationTrackSetFactory::BeginImportAsset_GameThread");
+
 	FImportAssetResult ImportAssetResult;
 #if !WITH_EDITOR || !WITH_EDITORONLY_DATA
 
@@ -745,6 +747,8 @@ UInterchangeFactoryBase::FImportAssetResult UInterchangeAnimationTrackSetFactory
 
 UObject* UInterchangeAnimationTrackSetFactory::ImportObjectSourceData(const FImportAssetObjectParams& Arguments)
 {
+	TRACE_CPUPROFILER_EVENT_SCOPE("UInterchangeAnimationTrackSetFactory::ImportObjectSourceData");
+
 #if !WITH_EDITOR || !WITH_EDITORONLY_DATA
 	// TODO: Can we import ULevelSequence at runtime
 	UE_LOG(LogInterchangeImport, Error, TEXT("Cannot import LevelSequence asset in runtime, this is an editor only feature."));
@@ -824,6 +828,8 @@ UObject* UInterchangeAnimationTrackSetFactory::ImportObjectSourceData(const FImp
 /* This function is call in the completion task on the main thread, use it to call main thread post creation step for your assets*/
 void UInterchangeAnimationTrackSetFactory::SetupObject_GameThread(const FSetupObjectParams& Arguments)
 {
+	TRACE_CPUPROFILER_EVENT_SCOPE("UInterchangeAnimationTrackSetFactory::SetupObject_GameThread");
+
 	check(IsInGameThread());
 	Super::SetupObject_GameThread(Arguments);
 
diff --git a/Engine/Plugins/Interchange/Runtime/Source/Import/Private/Fbx/InterchangeFbxTranslator.cpp b/Engine/Plugins/Interchange/Runtime/Source/Import/Private/Fbx/InterchangeFbxTranslator.cpp
--- a/Engine/Plugins/Interchange/Runtime/Source/Import/Private/Fbx/InterchangeFbxTranslator.cpp
+++ b/Engine/Plugins/Interchange/Runtime/Source/Import/Private/Fbx/InterchangeFbxTranslator.cpp
@@ -83,6 +83,9 @@ TArray<FString> UInterchangeFbxTranslator::GetSupportedFormats() const
 
 bool UInterchangeFbxTranslator::Translate(UInterchangeBaseNodeContainer& BaseNodeContainer) const
 {
+	TRACE_CPUPROFILER_EVENT_SCOPE("UInterchangeFbxTranslator::Translate");
+	//Make sure the hash is compute here in asynchronous mode
+	GetSourceData()->GetFileContentHash();
 	FString Filename = GetSourceData()->GetFilename();
 	if (!FPaths::FileExists(Filename))
 	{
diff --git a/Engine/Plugins/Interchange/Runtime/Source/Import/Private/Material/InterchangeMaterialFactory.cpp b/Engine/Plugins/Interchange/Runtime/Source/Import/Private/Material/InterchangeMaterialFactory.cpp
--- a/Engine/Plugins/Interchange/Runtime/Source/Import/Private/Material/InterchangeMaterialFactory.cpp
+++ b/Engine/Plugins/Interchange/Runtime/Source/Import/Private/Material/InterchangeMaterialFactory.cpp
@@ -686,6 +686,7 @@ UClass* UInterchangeMaterialFactory::GetFactoryClass() const
 
 UInterchangeFactoryBase::FImportAssetResult UInterchangeMaterialFactory::BeginImportAsset_GameThread(const FImportAssetObjectParams& Arguments)
 {
+	TRACE_CPUPROFILER_EVENT_SCOPE("UInterchangeMaterialFactory::BeginImportAsset_GameThread");
 	FImportAssetResult ImportAssetResult;
 	UObject* Material = nullptr;
 
@@ -890,6 +891,7 @@ UInterchangeFactoryBase::FImportAssetResult UInterchangeMaterialFactory::BeginIm
 
 UInterchangeFactoryBase::FImportAssetResult UInterchangeMaterialFactory::ImportAsset_Async(const FImportAssetObjectParams& Arguments)
 {
+	TRACE_CPUPROFILER_EVENT_SCOPE("UInterchangeMaterialFactory::ImportAsset_Async");
 	FImportAssetResult ImportAssetResult;
 	ImportAssetResult.bIsFactorySkipAsset = bSkipImport;
 
@@ -969,6 +971,7 @@ UInterchangeFactoryBase::FImportAssetResult UInterchangeMaterialFactory::ImportA
 /* This function is call in the completion task on the main thread, use it to call main thread post creation step for your assets*/
 void UInterchangeMaterialFactory::SetupObject_GameThread(const FSetupObjectParams& Arguments)
 {
+	TRACE_CPUPROFILER_EVENT_SCOPE("UInterchangeMaterialFactory::SetupObject_GameThread");
 	check(IsInGameThread());
 	Super::SetupObject_GameThread(Arguments);
 
@@ -983,7 +986,7 @@ void UInterchangeMaterialFactory::SetupObject_GameThread(const FSetupObjectParam
 		UMaterialInterface* ImportedMaterialInterface = CastChecked<UMaterialInterface>(Arguments.ImportedObject);
 
 #if WITH_EDITOR
-		//Update the samplers type in case the textures were changed during their PreImportPreCompletedCallback
+		//Update the samplers type in case the textures were changed during their SetupObject_GameThread
 		if (UMaterial* ImportedMaterial = Cast<UMaterial>(ImportedMaterialInterface))
 		{
 			for (UMaterialExpression* Expression : ImportedMaterial->GetExpressions())
@@ -1036,6 +1039,8 @@ bool UInterchangeMaterialFactory::SetSourceFilename(const UObject* Object, const
 #if WITH_EDITOR
 void UInterchangeMaterialFactory::SetupMaterial(UMaterial* Material, const FImportAssetObjectParams& Arguments, const UInterchangeBaseMaterialFactoryNode* BaseMaterialFactoryNode)
 {
+	TRACE_CPUPROFILER_EVENT_SCOPE("UInterchangeMaterialFactory::SetupMaterial");
+
 	using namespace UE::Interchange::Materials;
 	using namespace UE::Interchange::MaterialFactory::Internal;
 
@@ -1538,6 +1543,7 @@ void UInterchangeMaterialFactory::SetupMaterial(UMaterial* Material, const FImpo
 
 void UInterchangeMaterialFactory::SetupMaterialInstance(UMaterialInstance& MaterialInstance, const UInterchangeBaseNodeContainer& NodeContainer, const UInterchangeMaterialInstanceFactoryNode& FactoryNode, bool bResetInstance)
 {
+	TRACE_CPUPROFILER_EVENT_SCOPE("UInterchangeMaterialFactory::SetupMaterialInstance");
 	using namespace UE::Interchange::MaterialFactory::Internal;
 
 	if (bResetInstance)
@@ -1729,6 +1735,7 @@ UClass* UInterchangeMaterialFunctionFactory::GetFactoryClass() const
 
 UInterchangeFactoryBase::FImportAssetResult UInterchangeMaterialFunctionFactory::BeginImportAsset_GameThread(const FImportAssetObjectParams& Arguments)
 {
+	TRACE_CPUPROFILER_EVENT_SCOPE("UInterchangeMaterialFunctionFactory::BeginImportAsset_GameThread");
 	FImportAssetResult ImportAssetResult;
 	UObject* Material = nullptr;
 
@@ -1814,6 +1821,7 @@ UInterchangeFactoryBase::FImportAssetResult UInterchangeMaterialFunctionFactory:
 
 UInterchangeFactoryBase::FImportAssetResult UInterchangeMaterialFunctionFactory::ImportAsset_Async(const FImportAssetObjectParams& Arguments)
 {
+	TRACE_CPUPROFILER_EVENT_SCOPE("UInterchangeMaterialFunctionFactory::ImportAsset_Async");
 	FImportAssetResult ImportAssetResult;
 	ImportAssetResult.bIsFactorySkipAsset = bSkipImport;
 
@@ -1869,6 +1877,7 @@ UInterchangeFactoryBase::FImportAssetResult UInterchangeMaterialFunctionFactory:
 /* This function is call in the completion task on the main thread, use it to call main thread post creation step for your assets*/
 void UInterchangeMaterialFunctionFactory::SetupObject_GameThread(const FSetupObjectParams& Arguments)
 {
+	TRACE_CPUPROFILER_EVENT_SCOPE("UInterchangeMaterialFunctionFactory::SetupObject_GameThread");
 	Super::SetupObject_GameThread(Arguments);
 
 	if (bSkipImport)
@@ -1883,7 +1892,7 @@ void UInterchangeMaterialFunctionFactory::SetupObject_GameThread(const FSetupObj
 
 		if (UMaterialFunctionEditorOnlyData* Data = ImportedMaterialFunction->GetEditorOnlyData())
 		{
-			//Update the samplers type in case the textures were changed during their PreImportPreCompletedCallback
+			//Update the samplers type in case the textures were changed during their SetupObject_GameThread
 			for (UMaterialExpression* Expression : Data->ExpressionCollection.Expressions)
 			{
 				if (UMaterialExpressionTextureBase* TextureSample = Cast<UMaterialExpressionTextureBase>(Expression))
@@ -1899,6 +1908,7 @@ void UInterchangeMaterialFunctionFactory::SetupObject_GameThread(const FSetupObj
 #if WITH_EDITOR
 void UInterchangeMaterialFunctionFactory::SetupMaterial(UMaterialFunction* MaterialFunction, const FImportAssetObjectParams& Arguments, const UInterchangeMaterialFunctionFactoryNode* MaterialFunctionFactoryNode)
 {
+	TRACE_CPUPROFILER_EVENT_SCOPE("UInterchangeMaterialFunctionFactory::SetupMaterial");
 	using namespace UE::Interchange::MaterialFactory::Internal;
 	using namespace UE::Interchange::Materials;
 
diff --git a/Engine/Plugins/Interchange/Runtime/Source/Import/Private/Mesh/InterchangePhysicsAssetFactory.cpp b/Engine/Plugins/Interchange/Runtime/Source/Import/Private/Mesh/InterchangePhysicsAssetFactory.cpp
--- a/Engine/Plugins/Interchange/Runtime/Source/Import/Private/Mesh/InterchangePhysicsAssetFactory.cpp
+++ b/Engine/Plugins/Interchange/Runtime/Source/Import/Private/Mesh/InterchangePhysicsAssetFactory.cpp
@@ -18,6 +18,7 @@ UClass* UInterchangePhysicsAssetFactory::GetFactoryClass() const
 
 UInterchangeFactoryBase::FImportAssetResult UInterchangePhysicsAssetFactory::BeginImportAsset_GameThread(const FImportAssetObjectParams& Arguments)
 {
+	TRACE_CPUPROFILER_EVENT_SCOPE("UInterchangePhysicsAssetFactory::BeginImportAsset_GameThread");
 	FImportAssetResult ImportAssetResult;
 	UObject* PhysicsAsset = nullptr;
 
@@ -70,6 +71,7 @@ UInterchangeFactoryBase::FImportAssetResult UInterchangePhysicsAssetFactory::Beg
 
 UInterchangeFactoryBase::FImportAssetResult UInterchangePhysicsAssetFactory::ImportAsset_Async(const FImportAssetObjectParams& Arguments)
 {
+	TRACE_CPUPROFILER_EVENT_SCOPE("UInterchangePhysicsAssetFactory::ImportAsset_Async");
 	FImportAssetResult ImportAssetResult;
 #if !WITH_EDITORONLY_DATA
 
diff --git a/Engine/Plugins/Interchange/Runtime/Source/Import/Private/Mesh/InterchangeSkeletalMeshFactory.cpp b/Engine/Plugins/Interchange/Runtime/Source/Import/Private/Mesh/InterchangeSkeletalMeshFactory.cpp
--- a/Engine/Plugins/Interchange/Runtime/Source/Import/Private/Mesh/InterchangeSkeletalMeshFactory.cpp
+++ b/Engine/Plugins/Interchange/Runtime/Source/Import/Private/Mesh/InterchangeSkeletalMeshFactory.cpp
@@ -1510,8 +1510,18 @@ UInterchangeFactoryBase::FImportAssetResult UInterchangeSkeletalMeshFactory::End
 	const UClass* SkeletalMeshClass = SkeletalMeshFactoryNode->GetObjectClass();
 	check(SkeletalMeshClass && SkeletalMeshClass->IsChildOf(GetFactoryClass()));
 
-	//Get the skeletal mesh asset
-	USkeletalMesh* SkeletalMesh = Cast<USkeletalMesh>(StaticFindObject(nullptr, Arguments.Parent, *Arguments.AssetName));
+	//Get the skeletal mesh asset from the factory node or a find if the node was not set properly
+	USkeletalMesh* SkeletalMesh = nullptr;
+	FSoftObjectPath ReferenceObject;
+	if (SkeletalMeshFactoryNode->GetCustomReferenceObject(ReferenceObject))
+	{
+		SkeletalMesh = Cast<USkeletalMesh>(ReferenceObject.TryLoad());
+	}
+	if (!SkeletalMesh)
+	{
+		SkeletalMesh = Cast<USkeletalMesh>(StaticFindObject(nullptr, Arguments.Parent, *Arguments.AssetName));
+	}
+	
 	if (!ensure(SkeletalMesh))
 	{
 		if (Arguments.ReimportObject == nullptr)
@@ -1551,7 +1561,8 @@ UInterchangeFactoryBase::FImportAssetResult UInterchangeSkeletalMeshFactory::End
 
 	if (USkeleton* SkeletonReference = ImportAssetObjectData.SkeletonReference)
 	{
-		if ((!ImportAssetObjectData.bApplyGeometryOnly || !ImportAssetObjectData.bIsReImport) && !SkeletonReference->MergeAllBonesToBoneTree(SkeletalMesh))
+		constexpr bool bShowProgress = false;
+		if ((!ImportAssetObjectData.bApplyGeometryOnly || !ImportAssetObjectData.bIsReImport) && !SkeletonReference->MergeAllBonesToBoneTree(SkeletalMesh, bShowProgress))
 		{
 			TUniqueFunction<bool()> RecreateSkeleton = [this, WeakSkeletalMesh = TWeakObjectPtr<USkeletalMesh>(SkeletalMesh), WeakSkeleton = TWeakObjectPtr<USkeleton>(SkeletonReference)]()
 			{
@@ -1618,7 +1629,7 @@ UInterchangeFactoryBase::FImportAssetResult UInterchangeSkeletalMeshFactory::End
 								for (const USkeletalMesh* ExtraSkeletalMesh : OtherSkeletalMeshUsingSkeleton)
 								{
 									// merge still can fail
-									if (!SkeletonPtr->MergeAllBonesToBoneTree(ExtraSkeletalMesh))
+									if (!SkeletonPtr->MergeAllBonesToBoneTree(ExtraSkeletalMesh, bShowProgress))
 									{
 										FMessageDialog::Open(EAppMsgType::Ok,
 											FText::Format(NSLOCTEXT("InterchangeSkeletalMeshFactory", "SkeletonRegenError_RemergingBones", "Failed to merge SkeletalMesh '{0}'."), FText::FromString(ExtraSkeletalMesh->GetName())));
@@ -1712,7 +1723,7 @@ void UInterchangeSkeletalMeshFactory::Cancel()
 /* This function is call in the completion task on the main thread, use it to call main thread post creation step for your assets*/
 void UInterchangeSkeletalMeshFactory::SetupObject_GameThread(const FSetupObjectParams& Arguments)
 {
-	TRACE_CPUPROFILER_EVENT_SCOPE("UInterchangeSkeletalMeshFactory::PreImportPreCompletedCallback")
+	TRACE_CPUPROFILER_EVENT_SCOPE("UInterchangeSkeletalMeshFactory::SetupObject_GameThread")
 	check(IsInGameThread());
 	Super::SetupObject_GameThread(Arguments);
 
@@ -1762,6 +1773,7 @@ void UInterchangeSkeletalMeshFactory::SetupObject_GameThread(const FSetupObjectP
 				{
 					EInterchangeSkeletalMeshContentType ImportContentType = EInterchangeSkeletalMeshContentType::All;
 					SkeletalMeshFactoryNode->GetCustomImportContentType(ImportContentType);
+					FMD5Hash SourceFileHash = Arguments.SourceData->GetFileContentHash().Get(FMD5Hash());
 					const FString& NewSourceFilename = Arguments.SourceData->GetFilename();
 					const int32 NewSourceIndex = GetSourceIndexFromContentType(ImportContentType);
 					//NewSourceIndex should be 0, 1 or 2 (All, Geo, Skinning)
@@ -1773,12 +1785,15 @@ void UInterchangeSkeletalMeshFactory::SetupObject_GameThread(const FSetupObjectP
 						FString SourceLabel = GetSourceLabelFromSourceIndex(SourceIndex);
 						if (SourceIndex == NewSourceIndex)
 						{
-							AssetImportData->ScriptedAddFilename(NewSourceFilename, SourceIndex, SourceLabel);
-							//Do not add any extra path if import both geo and skinning
-							if(SourceIndex == GetSourceIndexFromContentType(EInterchangeSkeletalMeshContentType::All))
+							if (SourceIndex == GetSourceIndexFromContentType(EInterchangeSkeletalMeshContentType::All))
 							{
+								AssetImportData->Update(NewSourceFilename, SourceFileHash);
 								break;
 							}
+							else
+							{
+								AssetImportData->ScriptedAddFilename(NewSourceFilename, SourceIndex, SourceLabel);
+							}
 						}
 						else
 						{
diff --git a/Engine/Plugins/Interchange/Runtime/Source/Import/Private/Mesh/InterchangeSkeletonFactory.cpp b/Engine/Plugins/Interchange/Runtime/Source/Import/Private/Mesh/InterchangeSkeletonFactory.cpp
--- a/Engine/Plugins/Interchange/Runtime/Source/Import/Private/Mesh/InterchangeSkeletonFactory.cpp
+++ b/Engine/Plugins/Interchange/Runtime/Source/Import/Private/Mesh/InterchangeSkeletonFactory.cpp
@@ -19,6 +19,8 @@ UClass* UInterchangeSkeletonFactory::GetFactoryClass() const
 
 UInterchangeFactoryBase::FImportAssetResult UInterchangeSkeletonFactory::BeginImportAsset_GameThread(const FImportAssetObjectParams& Arguments)
 {
+	TRACE_CPUPROFILER_EVENT_SCOPE("UInterchangeSkeletonFactory::BeginImportAsset_GameThread");
+
 	FImportAssetResult ImportAssetResult;
 
 	USkeleton* Skeleton = nullptr;
diff --git a/Engine/Plugins/Interchange/Runtime/Source/Import/Private/Mesh/InterchangeStaticMeshFactory.cpp b/Engine/Plugins/Interchange/Runtime/Source/Import/Private/Mesh/InterchangeStaticMeshFactory.cpp
--- a/Engine/Plugins/Interchange/Runtime/Source/Import/Private/Mesh/InterchangeStaticMeshFactory.cpp
+++ b/Engine/Plugins/Interchange/Runtime/Source/Import/Private/Mesh/InterchangeStaticMeshFactory.cpp
@@ -147,6 +147,8 @@ namespace UE::Interchange::Private::StaticMesh
 
 UInterchangeFactoryBase::FImportAssetResult UInterchangeStaticMeshFactory::BeginImportAsset_GameThread(const FImportAssetObjectParams& Arguments)
 {
+	TRACE_CPUPROFILER_EVENT_SCOPE("UInterchangeStaticMeshFactory::BeginImportAsset_GameThread");
+
 	FImportAssetResult ImportAssetResult;
 	UStaticMesh* StaticMesh = nullptr;
 	if (!Arguments.AssetNode || !Arguments.AssetNode->GetObjectClass()->IsChildOf(GetFactoryClass()))
@@ -215,6 +217,8 @@ UInterchangeFactoryBase::FImportAssetResult UInterchangeStaticMeshFactory::Begin
 
 UInterchangeFactoryBase::FImportAssetResult UInterchangeStaticMeshFactory::ImportAsset_Async(const FImportAssetObjectParams& Arguments)
 {
+	TRACE_CPUPROFILER_EVENT_SCOPE("UInterchangeStaticMeshFactory::ImportAsset_Async");
+
 	FImportAssetResult ImportAssetResult;
 	if (!Arguments.AssetNode || !Arguments.AssetNode->GetObjectClass()->IsChildOf(GetFactoryClass()))
 	{
@@ -479,6 +483,8 @@ UInterchangeFactoryBase::FImportAssetResult UInterchangeStaticMeshFactory::Impor
 
 UInterchangeFactoryBase::FImportAssetResult UInterchangeStaticMeshFactory::EndImportAsset_GameThread(const FImportAssetObjectParams& Arguments)
 {
+	TRACE_CPUPROFILER_EVENT_SCOPE("UInterchangeStaticMeshFactory::EndImportAsset_GameThread");
+
 	FImportAssetResult ImportAssetResult;
 	if (!Arguments.AssetNode || !Arguments.AssetNode->GetObjectClass()->IsChildOf(GetFactoryClass()))
 	{
@@ -749,6 +755,8 @@ void UInterchangeStaticMeshFactory::SetupSourceModelsSettings(UStaticMesh& Stati
 /* This function is call in the completion task on the main thread, use it to call main thread post creation step for your assets */
 void UInterchangeStaticMeshFactory::SetupObject_GameThread(const FSetupObjectParams& Arguments)
 {
+	TRACE_CPUPROFILER_EVENT_SCOPE("UInterchangeStaticMeshFactory::SetupObject_GameThread");
+
 	check(IsInGameThread());
 	Super::SetupObject_GameThread(Arguments);
 
@@ -770,6 +778,8 @@ void UInterchangeStaticMeshFactory::SetupObject_GameThread(const FSetupObjectPar
 
 TArray<UInterchangeStaticMeshFactory::FMeshPayload> UInterchangeStaticMeshFactory::GetMeshPayloads(const FImportAssetObjectParams& Arguments, const TArray<FString>& MeshUids) const
 {
+	TRACE_CPUPROFILER_EVENT_SCOPE("UInterchangeStaticMeshFactory::GetMeshPayloads");
+
 	TArray<FMeshPayload> Payloads;
 	Payloads.Reserve(MeshUids.Num());
 
diff --git a/Engine/Plugins/Interchange/Runtime/Source/Import/Private/Scene/InterchangeSceneImportAssetFactory.cpp b/Engine/Plugins/Interchange/Runtime/Source/Import/Private/Scene/InterchangeSceneImportAssetFactory.cpp
--- a/Engine/Plugins/Interchange/Runtime/Source/Import/Private/Scene/InterchangeSceneImportAssetFactory.cpp
+++ b/Engine/Plugins/Interchange/Runtime/Source/Import/Private/Scene/InterchangeSceneImportAssetFactory.cpp
@@ -85,6 +85,7 @@ bool UInterchangeSceneImportAssetFactory::SetSourceFilename(const UObject* Objec
 
 UInterchangeFactoryBase::FImportAssetResult UInterchangeSceneImportAssetFactory::BeginImportAsset_GameThread(const FImportAssetObjectParams& Arguments)
 {
+	TRACE_CPUPROFILER_EVENT_SCOPE("UInterchangeSceneImportAssetFactory::BeginImportAsset_GameThread");
 	using namespace UE::Interchange::Private::InterchangeSceneImportAssetFactory;
 
 	UClass* TargetClass = GetFactoryClass();
@@ -114,6 +115,7 @@ UInterchangeFactoryBase::FImportAssetResult UInterchangeSceneImportAssetFactory:
 
 void UInterchangeSceneImportAssetFactory::SetupObject_GameThread(const FSetupObjectParams& Arguments)
 {
+	TRACE_CPUPROFILER_EVENT_SCOPE("UInterchangeSceneImportAssetFactory::SetupObject_GameThread");
 	check(IsInGameThread());
 	Super::SetupObject_GameThread(Arguments);
 
@@ -135,6 +137,7 @@ void UInterchangeSceneImportAssetFactory::SetupObject_GameThread(const FSetupObj
 
 void UInterchangeSceneImportAssetFactory::FinalizeObject_GameThread(const FSetupObjectParams& Arguments)
 {
+	TRACE_CPUPROFILER_EVENT_SCOPE("UInterchangeSceneImportAssetFactory::FinalizeObject_GameThread");
 	check(IsInGameThread());
 	Super::FinalizeObject_GameThread(Arguments);
 
diff --git a/Engine/Plugins/Interchange/Runtime/Source/Import/Private/Scene/InterchangeSceneVariantSetsFactory.cpp b/Engine/Plugins/Interchange/Runtime/Source/Import/Private/Scene/InterchangeSceneVariantSetsFactory.cpp
--- a/Engine/Plugins/Interchange/Runtime/Source/Import/Private/Scene/InterchangeSceneVariantSetsFactory.cpp
+++ b/Engine/Plugins/Interchange/Runtime/Source/Import/Private/Scene/InterchangeSceneVariantSetsFactory.cpp
@@ -296,6 +296,7 @@ UClass* UInterchangeSceneVariantSetsFactory::GetFactoryClass() const
 
 UInterchangeFactoryBase::FImportAssetResult UInterchangeSceneVariantSetsFactory::BeginImportAsset_GameThread(const FImportAssetObjectParams& Arguments)
 {
+	TRACE_CPUPROFILER_EVENT_SCOPE("UInterchangeSceneVariantSetsFactory::BeginImportAsset_GameThread");
 	FImportAssetResult ImportAssetResult;
 #if !WITH_EDITOR || !WITH_EDITORONLY_DATA
 
@@ -373,6 +374,7 @@ UInterchangeFactoryBase::FImportAssetResult UInterchangeSceneVariantSetsFactory:
 
 UObject* UInterchangeSceneVariantSetsFactory::ImportObjectSourceData(const FImportAssetObjectParams& Arguments, ULevelVariantSets* LevelVariantSets)
 {
+	TRACE_CPUPROFILER_EVENT_SCOPE("UInterchangeSceneVariantSetsFactory::ImportObjectSourceData");
 #if !WITH_EDITOR || !WITH_EDITORONLY_DATA
 	// TODO: Can we import ULevelVariantSets at runtime
 	UE_LOG(LogInterchangeImport, Error, TEXT("Cannot import LevelVariantSets asset in runtime, this is an editor only feature."));
@@ -420,6 +422,7 @@ UObject* UInterchangeSceneVariantSetsFactory::ImportObjectSourceData(const FImpo
 
 void UInterchangeSceneVariantSetsFactory::SetupObject_GameThread(const FSetupObjectParams& Arguments)
 {
+	TRACE_CPUPROFILER_EVENT_SCOPE("UInterchangeSceneVariantSetsFactory::SetupObject_GameThread");
 	check(IsInGameThread());
 	Super::SetupObject_GameThread(Arguments);
 
diff --git a/Engine/Plugins/Interchange/Runtime/Source/Import/Private/Texture/InterchangeTextureFactory.cpp b/Engine/Plugins/Interchange/Runtime/Source/Import/Private/Texture/InterchangeTextureFactory.cpp
--- a/Engine/Plugins/Interchange/Runtime/Source/Import/Private/Texture/InterchangeTextureFactory.cpp
+++ b/Engine/Plugins/Interchange/Runtime/Source/Import/Private/Texture/InterchangeTextureFactory.cpp
@@ -1531,6 +1531,7 @@ UClass* UInterchangeTextureFactory::GetFactoryClass() const
 
 UInterchangeFactoryBase::FImportAssetResult UInterchangeTextureFactory::BeginImportAsset_GameThread(const FImportAssetObjectParams& Arguments)
 {
+	TRACE_CPUPROFILER_EVENT_SCOPE("UInterchangeTextureFactory::BeginImportAsset_GameThread");
 	using namespace  UE::Interchange::Private::InterchangeTextureFactory;
 	FImportAssetResult ImportAssetResult;
 	UTexture* Texture = nullptr;
@@ -1785,7 +1786,7 @@ UInterchangeFactoryBase::FImportAssetResult UInterchangeTextureFactory::ImportAs
 /* This function is call in the completion task on the main thread, use it to call main thread post creation step for your assets*/
 void UInterchangeTextureFactory::SetupObject_GameThread(const FSetupObjectParams& Arguments)
 {
-	TRACE_CPUPROFILER_EVENT_SCOPE(UInterchangeTextureFactory::BeginPreCompletedCallback);
+	TRACE_CPUPROFILER_EVENT_SCOPE(UInterchangeTextureFactory::SetupObject_GameThread);
 
 	check(IsInGameThread());
 
@@ -1881,6 +1882,7 @@ void UInterchangeTextureFactory::SetupObject_GameThread(const FSetupObjectParams
 
 void UInterchangeTextureFactory::CheckForInvalidResolutions(UE::Interchange::Private::InterchangeTextureFactory::FTexturePayloadVariant& InPayloadVariant, const UInterchangeSourceData* SourceData, const UInterchangeTextureFactoryNode* TextureFactoryNode)
 {
+	TRACE_CPUPROFILER_EVENT_SCOPE(UInterchangeTextureFactory::CheckForInvalidResolutions);
 	using namespace UE::Interchange;
 
 	auto AddErrorMessage = [this, &TextureFactoryNode](FString&& InSourceAssetName, const FText& ErrorMessage) -> UInterchangeResultError_Generic*
diff --git a/Engine/Plugins/Interchange/Runtime/Source/Import/Public/Animation/InterchangeAnimSequenceFactory.h b/Engine/Plugins/Interchange/Runtime/Source/Import/Public/Animation/InterchangeAnimSequenceFactory.h
--- a/Engine/Plugins/Interchange/Runtime/Source/Import/Public/Animation/InterchangeAnimSequenceFactory.h
+++ b/Engine/Plugins/Interchange/Runtime/Source/Import/Public/Animation/InterchangeAnimSequenceFactory.h
@@ -3,6 +3,7 @@
 #pragma once
 
 #include "CoreMinimal.h"
+#include "InterchangeCommonAnimationPayload.h"
 #include "InterchangeFactoryBase.h"
 #include "UObject/Object.h"
 #include "UObject/ObjectMacros.h"
@@ -11,6 +12,7 @@
 
 class UAnimSequence;
 class UInterchangeAnimSequenceFactoryNode;
+class UInterchangeSceneNode;
 
 UCLASS(BlueprintType)
 class INTERCHANGEIMPORT_API UInterchangeAnimSequenceFactory : public UInterchangeFactoryBase
@@ -23,16 +25,43 @@ public:
 
 	virtual UClass* GetFactoryClass() const override;
 	virtual EInterchangeFactoryAssetType GetFactoryAssetType() override { return EInterchangeFactoryAssetType::Animations; }
+	
 	virtual FImportAssetResult BeginImportAsset_GameThread(const FImportAssetObjectParams& Arguments) override;
+	
+	virtual FImportAssetResult ImportAsset_Async(const FImportAssetObjectParams& Arguments) override;
+	
+	virtual FImportAssetResult EndImportAsset_GameThread(const FImportAssetObjectParams& Arguments) override;
+
 	virtual void SetupObject_GameThread(const FSetupObjectParams& Arguments) override;
 
 	virtual bool GetSourceFilenames(const UObject* Object, TArray<FString>& OutSourceFilenames) const override;
 	virtual bool SetSourceFilename(const UObject* Object, const FString& SourceFilename, int32 SourceIndex) const override;
+	
+	struct FBoneTrackData
+	{
+		TMap<const UInterchangeSceneNode*, UE::Interchange::FAnimationPayloadData> PreProcessedAnimationPayloads;
+		double MergedRangeStart = 0.0;
+		double MergedRangeEnd = 0.0;
+	};
+
+	struct FMorphTargetData
+	{
+		TMap<FString, UE::Interchange::FAnimationPayloadData> CurvesPayloads;
+		TMap<FString, FString> CurveNodeNamePerPayloadKey;
+	};
 
 private:
-	UObject* ImportObjectSourceData(const FImportAssetObjectParams& Arguments);
 	bool IsBoneTrackAnimationValid(const UInterchangeAnimSequenceFactoryNode* AnimSequenceFactoryNode, const FImportAssetObjectParams& Arguments);
 
+	//The imported AnimSequence
+	TObjectPtr<UAnimSequence> AnimSequence = nullptr;
+
+	//Bone track animations payload data
+	FBoneTrackData BoneTrackData;
+
+	//Morph target curves payload data
+	FMorphTargetData MorphTargetData;
+
 	// Interchange factory base interface end
 	//////////////////////////////////////////////////////////////////////////
 };
diff --git a/Engine/Plugins/Interchange/Runtime/Source/Pipelines/Private/InterchangeGenericSkeletalMeshPipeline.cpp b/Engine/Plugins/Interchange/Runtime/Source/Pipelines/Private/InterchangeGenericSkeletalMeshPipeline.cpp
--- a/Engine/Plugins/Interchange/Runtime/Source/Pipelines/Private/InterchangeGenericSkeletalMeshPipeline.cpp
+++ b/Engine/Plugins/Interchange/Runtime/Source/Pipelines/Private/InterchangeGenericSkeletalMeshPipeline.cpp
@@ -566,6 +566,7 @@ void UInterchangeGenericMeshPipeline::PostImportSkeletalMesh(UObject* CreatedAss
 
 void UInterchangeGenericMeshPipeline::PostImportPhysicsAssetImport(UObject* CreatedAsset, const UInterchangeFactoryBaseNode* FactoryNode)
 {
+	TRACE_CPUPROFILER_EVENT_SCOPE("UInterchangeGenericMeshPipeline::PostImportPhysicsAssetImport");
 #if WITH_EDITOR
 	if (!bCreatePhysicsAsset || !BaseNodeContainer)
 	{
@@ -592,9 +593,12 @@ void UInterchangeGenericMeshPipeline::PostImportPhysicsAssetImport(UObject* Crea
 					{
 						auto CreateFromSkeletalMeshLambda = [CreatedPhysicsAsset, SkeletalMesh]()
 						{
+							TRACE_CPUPROFILER_EVENT_SCOPE("UInterchangeGenericMeshPipeline::PostImportPhysicsAssetImport::CreateFromSkeletalMeshLambda");
 							FPhysAssetCreateParams NewBodyData;
 							FText CreationErrorMessage;
-							if (!FPhysicsAssetUtils::CreateFromSkeletalMesh(CreatedPhysicsAsset, SkeletalMesh, NewBodyData, CreationErrorMessage))
+							constexpr bool bSetToMesh = true;
+							constexpr bool bShowProgress = false;
+							if (!FPhysicsAssetUtils::CreateFromSkeletalMesh(CreatedPhysicsAsset, SkeletalMesh, NewBodyData, CreationErrorMessage, bSetToMesh, bShowProgress))
 							{
 								//TODO: Log an error
 							}
diff --git a/Engine/Source/Developer/PhysicsUtilities/Private/PhysicsAssetUtils.cpp b/Engine/Source/Developer/PhysicsUtilities/Private/PhysicsAssetUtils.cpp
--- a/Engine/Source/Developer/PhysicsUtilities/Private/PhysicsAssetUtils.cpp
+++ b/Engine/Source/Developer/PhysicsUtilities/Private/PhysicsAssetUtils.cpp
@@ -203,7 +203,7 @@ int32 CalculateCommonRootBoneIndex(const FReferenceSkeleton& RefSkeleton, const
 	return RootBoneIndex;
 }
 
-bool CreateFromSkeletalMeshInternal(UPhysicsAsset* PhysicsAsset, USkeletalMesh* SkelMesh, const FPhysAssetCreateParams& Params, const FSkinnedBoneTriangleCache& TriangleCache)
+bool CreateFromSkeletalMeshInternal(UPhysicsAsset* PhysicsAsset, USkeletalMesh* SkelMesh, const FPhysAssetCreateParams& Params, const FSkinnedBoneTriangleCache& TriangleCache, bool bShowProgress)
 {
 	// For each bone, get the vertices most firmly attached to it.
 	TArray<FBoneVertInfo> Infos;
@@ -340,8 +340,8 @@ bool CreateFromSkeletalMeshInternal(UPhysicsAsset* PhysicsAsset, USkeletalMesh*
 	// Finally, iterate through all the bones and create bodies when needed
 
 	const bool bCanCreateConstraints = CanCreateConstraints();
-	FScopedSlowTask SlowTask((float)NumBones * 2, FText(), IsInGameThread());
-	if (IsInGameThread())
+	FScopedSlowTask SlowTask((float)NumBones * 2, FText(), bShowProgress&& IsInGameThread());
+	if (bShowProgress && IsInGameThread())
 	{
 		SlowTask.MakeDialog();
 	}
@@ -356,7 +356,7 @@ bool CreateFromSkeletalMeshInternal(UPhysicsAsset* PhysicsAsset, USkeletalMesh*
 			// Go ahead and make this bone physical.
 			FName BoneName = SkelMesh->GetRefSkeleton().GetBoneName(BoneIndex);
 
-			if (IsInGameThread())
+			if (bShowProgress && IsInGameThread())
 			{
 				SlowTask.EnterProgressFrame(1.0f, FText::Format(NSLOCTEXT("PhysicsAssetEditor", "ResetCollsionStepInfo", "Generating collision for {0}"), FText::FromName(BoneName)));
 			}
@@ -469,7 +469,7 @@ bool CreateFromSkeletalMeshInternal(UPhysicsAsset* PhysicsAsset, USkeletalMesh*
 		FBodyInstance* BodyInstance = Bodies[BodyIdx];
 		if(BodyInstance && BodyInstance->BodySetup.IsValid())
 		{
-			if (IsInGameThread())
+			if (bShowProgress && IsInGameThread())
 			{
 				SlowTask.EnterProgressFrame(1.0f, FText::Format(NSLOCTEXT("PhysicsAssetEditor", "ResetCollsionStepInfoOverlaps", "Fixing overlaps for {0}"), FText::FromName(BodyInstance->BodySetup->BoneName)));
 			}
@@ -493,7 +493,7 @@ bool CreateFromSkeletalMeshInternal(UPhysicsAsset* PhysicsAsset, USkeletalMesh*
 	return NumBodies > 0;
 }
 
-bool CreateFromSkeletalMesh(UPhysicsAsset* PhysicsAsset, USkeletalMesh* SkelMesh, const FPhysAssetCreateParams& Params, FText& OutErrorMessage, bool bSetToMesh /*= true*/)
+bool CreateFromSkeletalMesh(UPhysicsAsset* PhysicsAsset, USkeletalMesh* SkelMesh, const FPhysAssetCreateParams& Params, FText& OutErrorMessage, bool bSetToMesh /*= true*/, bool bShowProgress /*= true */)
 {
 	PhysicsAsset->PreviewSkeletalMesh = SkelMesh;
 
@@ -507,14 +507,14 @@ bool CreateFromSkeletalMesh(UPhysicsAsset* PhysicsAsset, USkeletalMesh* SkelMesh
 		TriangleCache.BuildCache();
 	}
 
-	bool bSuccess = CreateFromSkeletalMeshInternal(PhysicsAsset, SkelMesh, Params, TriangleCache);
+	bool bSuccess = CreateFromSkeletalMeshInternal(PhysicsAsset, SkelMesh, Params, TriangleCache, bShowProgress);
 	if (!bSuccess)
 	{
 		// try lower minimum bone size 
 		FPhysAssetCreateParams LocalParams = Params;
 		LocalParams.MinBoneSize = 1.f;
 
-		bSuccess = CreateFromSkeletalMeshInternal(PhysicsAsset, SkelMesh, LocalParams, TriangleCache);
+		bSuccess = CreateFromSkeletalMeshInternal(PhysicsAsset, SkelMesh, LocalParams, TriangleCache, bShowProgress);
 
 		if(!bSuccess)
 		{
diff --git a/Engine/Source/Developer/PhysicsUtilities/Public/PhysicsAssetUtils.h b/Engine/Source/Developer/PhysicsUtilities/Public/PhysicsAssetUtils.h
--- a/Engine/Source/Developer/PhysicsUtilities/Public/PhysicsAssetUtils.h
+++ b/Engine/Source/Developer/PhysicsUtilities/Public/PhysicsAssetUtils.h
@@ -133,7 +133,7 @@ namespace FPhysicsAssetUtils
 	 * @param	OutErrorMessage		Additional error information
 	 * @param	bSetToMesh			Whether or not to apply the physics asset to SkelMesh immediately
 	 */
-	PHYSICSUTILITIES_API bool CreateFromSkeletalMesh(UPhysicsAsset* PhysicsAsset, USkeletalMesh* SkelMesh, const FPhysAssetCreateParams& Params, FText& OutErrorMessage, bool bSetToMesh = true);
+	PHYSICSUTILITIES_API bool CreateFromSkeletalMesh(UPhysicsAsset* PhysicsAsset, USkeletalMesh* SkelMesh, const FPhysAssetCreateParams& Params, FText& OutErrorMessage, bool bSetToMesh = true, bool bShowProgress = true);
 
 	/** Replaces any collision already in the BodySetup with an auto-generated one using the parameters provided.
 	 * 
diff --git a/Engine/Source/Editor/UnrealEd/Classes/ThumbnailRendering/MaterialInstanceThumbnailRenderer.h b/Engine/Source/Editor/UnrealEd/Classes/ThumbnailRendering/MaterialInstanceThumbnailRenderer.h
--- a/Engine/Source/Editor/UnrealEd/Classes/ThumbnailRendering/MaterialInstanceThumbnailRenderer.h
+++ b/Engine/Source/Editor/UnrealEd/Classes/ThumbnailRendering/MaterialInstanceThumbnailRenderer.h
@@ -23,6 +23,7 @@ class UMaterialInstanceThumbnailRenderer : public UDefaultSizedThumbnailRenderer
 
 	// Begin UThumbnailRenderer Object
 	virtual void Draw(UObject* Object, int32 X, int32 Y, uint32 Width, uint32 Height, FRenderTarget*, FCanvas* Canvas, bool bAdditionalViewFamily) override;
+	virtual bool CanVisualizeAsset(UObject* Object) override;
 	// End UThumbnailRenderer Object
 
 	// UObject implementation
diff --git a/Engine/Source/Editor/UnrealEd/Classes/ThumbnailRendering/TextureThumbnailRenderer.h b/Engine/Source/Editor/UnrealEd/Classes/ThumbnailRendering/TextureThumbnailRenderer.h
--- a/Engine/Source/Editor/UnrealEd/Classes/ThumbnailRendering/TextureThumbnailRenderer.h
+++ b/Engine/Source/Editor/UnrealEd/Classes/ThumbnailRendering/TextureThumbnailRenderer.h
@@ -24,6 +24,7 @@ class UTextureThumbnailRenderer : public UThumbnailRenderer
 	// Begin UThumbnailRenderer Object
 	UNREALED_API virtual void GetThumbnailSize(UObject* Object, float Zoom, uint32& OutWidth, uint32& OutHeight) const override;
 	UNREALED_API virtual void Draw(UObject* Object, int32 X, int32 Y, uint32 Width, uint32 Height, FRenderTarget*, FCanvas* Canvas, bool bAdditionalViewFamily) override;
+	UNREALED_API virtual bool CanVisualizeAsset(UObject* Object) override;
 	// End UThumbnailRenderer Object
 };
 
diff --git a/Engine/Source/Editor/UnrealEd/Private/ThumbnailRendering/MaterialInstanceThumbnailRenderer.cpp b/Engine/Source/Editor/UnrealEd/Private/ThumbnailRendering/MaterialInstanceThumbnailRenderer.cpp
--- a/Engine/Source/Editor/UnrealEd/Private/ThumbnailRendering/MaterialInstanceThumbnailRenderer.cpp
+++ b/Engine/Source/Editor/UnrealEd/Private/ThumbnailRendering/MaterialInstanceThumbnailRenderer.cpp
@@ -102,6 +102,18 @@ void UMaterialInstanceThumbnailRenderer::Draw(UObject* Object, int32 X, int32 Y,
 	}
 }
 
+bool UMaterialInstanceThumbnailRenderer::CanVisualizeAsset(UObject* Object)
+{
+	if (UMaterialInterface* MatInst = Cast<UMaterialInterface>(Object))
+	{
+		if (MatInst->IsCompiling())
+		{
+			return false;
+		}
+	}
+	return true;
+}
+
 void UMaterialInstanceThumbnailRenderer::BeginDestroy()
 {
 	if (ThumbnailScene != nullptr)
diff --git a/Engine/Source/Editor/UnrealEd/Private/ThumbnailRendering/PhysicsAssetThumbnailRenderer.cpp b/Engine/Source/Editor/UnrealEd/Private/ThumbnailRendering/PhysicsAssetThumbnailRenderer.cpp
--- a/Engine/Source/Editor/UnrealEd/Private/ThumbnailRendering/PhysicsAssetThumbnailRenderer.cpp
+++ b/Engine/Source/Editor/UnrealEd/Private/ThumbnailRendering/PhysicsAssetThumbnailRenderer.cpp
@@ -25,7 +25,7 @@ bool UPhysicsAssetThumbnailRenderer::CanVisualizeAsset(UObject* Object)
 		return false;
 	}
 	USkeletalMesh* SkeletalMesh = PhysicsAsset->PreviewSkeletalMesh.LoadSynchronous();
-	const bool bValidRenderData = SkeletalMesh && !SkeletalMesh->IsCompiling() && SkeletalMesh->GetResourceForRendering();
+	const bool bValidRenderData = SkeletalMesh && SkeletalMesh->IsReadyToRenderInThumbnail();
 	if (!bValidRenderData)
 	{
 		return false;
@@ -39,7 +39,7 @@ void UPhysicsAssetThumbnailRenderer::Draw(UObject* Object, int32 X, int32 Y, uin
 	if (PhysicsAsset)
 	{
 		USkeletalMesh* SkeletalMesh = PhysicsAsset->PreviewSkeletalMesh.LoadSynchronous();
-		const bool bValidRenderData = SkeletalMesh && !SkeletalMesh->IsCompiling() && SkeletalMesh->GetResourceForRendering();
+		const bool bValidRenderData = SkeletalMesh && SkeletalMesh->IsReadyToRenderInThumbnail();
 		if (!bValidRenderData)
 		{
 			return;
diff --git a/Engine/Source/Editor/UnrealEd/Private/ThumbnailRendering/SkeletalMeshThumbnailRenderer.cpp b/Engine/Source/Editor/UnrealEd/Private/ThumbnailRendering/SkeletalMeshThumbnailRenderer.cpp
--- a/Engine/Source/Editor/UnrealEd/Private/ThumbnailRendering/SkeletalMeshThumbnailRenderer.cpp
+++ b/Engine/Source/Editor/UnrealEd/Private/ThumbnailRendering/SkeletalMeshThumbnailRenderer.cpp
@@ -50,7 +50,7 @@ bool USkeletalMeshThumbnailRenderer::CanVisualizeAsset(UObject* Object)
 		return false;
 	}
 	USkeletalMesh* SkeletalMesh = Cast<USkeletalMesh>(Object);
-	const bool bValidRenderData = SkeletalMesh && !SkeletalMesh->IsCompiling() && SkeletalMesh->GetResourceForRendering();
+	const bool bValidRenderData = SkeletalMesh && SkeletalMesh->IsReadyToRenderInThumbnail();
 	return bValidRenderData;
 }
 
diff --git a/Engine/Source/Editor/UnrealEd/Private/ThumbnailRendering/SkeletonThumbnailRenderer.cpp b/Engine/Source/Editor/UnrealEd/Private/ThumbnailRendering/SkeletonThumbnailRenderer.cpp
--- a/Engine/Source/Editor/UnrealEd/Private/ThumbnailRendering/SkeletonThumbnailRenderer.cpp
+++ b/Engine/Source/Editor/UnrealEd/Private/ThumbnailRendering/SkeletonThumbnailRenderer.cpp
@@ -17,18 +17,20 @@ USkeletonThumbnailRenderer::USkeletonThumbnailRenderer(const FObjectInitializer&
 
 bool USkeletonThumbnailRenderer::CanVisualizeAsset(UObject* Object)
 {
-	USkeleton* Skeleton = Cast<USkeleton>(Object);
-
-	constexpr bool bFindIfNotSet = true;
-	if (USkeletalMesh* SkeletalMesh = Skeleton->GetPreviewMesh(bFindIfNotSet))
+	//The const version of GetPreviewMesh function simply return the preview mesh, it wont verify the skeleton compatibility.
+	//If we are not compiling call the function that are verifying the compatibility.
+	const USkeleton* ConstSkeleton = Cast<USkeleton>(Object);
+	USkeletalMesh* SkeletalMesh = ConstSkeleton->GetPreviewMesh();
+	if (!SkeletalMesh || !SkeletalMesh->IsCompiling())
 	{
-		if (SkeletalMesh && (SkeletalMesh->IsCompiling() || SkeletalMesh->GetResourceForRendering() == nullptr))
-		{
-			return false;
-		}
+		USkeleton* Skeleton = Cast<USkeleton>(Object);
+		//This case can stall this thread (GameThread) since it will verify this skeleton compatibility with the preview skeletal mesh.
+		//If the skeletal mesh is compiling a wait until compile will happen when doing USkeletalMesh::GetSkeleton()
+		constexpr bool bFindIfNotSet = true;
+		SkeletalMesh = Skeleton->GetPreviewMesh(bFindIfNotSet);
 	}
-
-	return true;
+	//Only visualize thumbnail if the skeletal mesh preview exist and is not compiling and have valid render data
+	return (SkeletalMesh && SkeletalMesh->IsReadyToRenderInThumbnail());
 }
 
 void USkeletonThumbnailRenderer::Draw(UObject* Object, int32 X, int32 Y, uint32 Width, uint32 Height, FRenderTarget* RenderTarget, FCanvas* Canvas, bool bAdditionalViewFamily)
diff --git a/Engine/Source/Editor/UnrealEd/Private/ThumbnailRendering/TextureThumbnailRenderer.cpp b/Engine/Source/Editor/UnrealEd/Private/ThumbnailRendering/TextureThumbnailRenderer.cpp
--- a/Engine/Source/Editor/UnrealEd/Private/ThumbnailRendering/TextureThumbnailRenderer.cpp
+++ b/Engine/Source/Editor/UnrealEd/Private/ThumbnailRendering/TextureThumbnailRenderer.cpp
@@ -200,3 +200,9 @@ void UTextureThumbnailRenderer::Draw(UObject* Object, int32 X, int32 Y, uint32 W
 		}
 	}
 }
+
+bool UTextureThumbnailRenderer::CanVisualizeAsset(UObject* Object)
+{
+	UTexture* Texture = Cast<UTexture>(Object);
+	return (Texture && !Texture->IsCompiling());
+}
\ No newline at end of file
diff --git a/Engine/Source/Runtime/Engine/Classes/Animation/Skeleton.h b/Engine/Source/Runtime/Engine/Classes/Animation/Skeleton.h
--- a/Engine/Source/Runtime/Engine/Classes/Animation/Skeleton.h
+++ b/Engine/Source/Runtime/Engine/Classes/Animation/Skeleton.h
@@ -806,7 +806,7 @@ public:
 	 * 
 	 * @return true if success
 	 */
-	ENGINE_API bool MergeBonesToBoneTree(const USkinnedAsset* InSkinnedAsset, const TArray<int32> &RequiredRefBones);
+	ENGINE_API bool MergeBonesToBoneTree(const USkinnedAsset* InSkinnedAsset, const TArray<int32> &RequiredRefBones, bool bShowProgress = true);
 
 	/** 
 	 * Merge all Bones to BoneTrees if not exists
@@ -818,7 +818,7 @@ public:
 	 * 
 	 * @return true if success
 	 */
-	ENGINE_API bool MergeAllBonesToBoneTree(const USkinnedAsset* InSkinnedAsset);
+	ENGINE_API bool MergeAllBonesToBoneTree(const USkinnedAsset* InSkinnedAsset, bool bShowProgress = true);
 
 	/** 
 	 * Merge has failed, then Recreate BoneTree
@@ -974,7 +974,7 @@ protected:
 	FOnSkeletonHierarchyChangedMulticaster OnSkeletonHierarchyChanged;
 
 	/** Call this when the skeleton has changed to fix dependent assets */
-	ENGINE_API void HandleSkeletonHierarchyChange();
+	ENGINE_API void HandleSkeletonHierarchyChange(bool bShowProgress = true);
 
 public:
 	typedef FOnSkeletonHierarchyChangedMulticaster::FDelegate FOnSkeletonHierarchyChanged;
diff --git a/Engine/Source/Runtime/Engine/Classes/Engine/SkeletalMesh.h b/Engine/Source/Runtime/Engine/Classes/Engine/SkeletalMesh.h
--- a/Engine/Source/Runtime/Engine/Classes/Engine/SkeletalMesh.h
+++ b/Engine/Source/Runtime/Engine/Classes/Engine/SkeletalMesh.h
@@ -747,6 +747,8 @@ public:
 		return false;
 	}
 
+	ENGINE_API bool IsReadyToRenderInThumbnail() const;
+
 	/* Return true if this skeletalmesh was never build since its creation. USkinnedAsset interface. */
 	ENGINE_API virtual bool IsInitialBuildDone() const override;
 
diff --git a/Engine/Source/Runtime/Engine/Classes/Materials/Material.h b/Engine/Source/Runtime/Engine/Classes/Materials/Material.h
--- a/Engine/Source/Runtime/Engine/Classes/Materials/Material.h
+++ b/Engine/Source/Runtime/Engine/Classes/Materials/Material.h
@@ -1360,6 +1360,10 @@ public:
 #endif
 	ENGINE_API virtual bool IsComplete() const override;
 
+#if WITH_EDITOR
+	ENGINE_API virtual bool IsCompiling() const override;
+#endif
+
 #if WITH_EDITORONLY_DATA
 	ENGINE_API virtual bool IterateDependentFunctions(TFunctionRef<bool(UMaterialFunctionInterface*)> Predicate) const override;
 	ENGINE_API virtual void GetDependentFunctions(TArray<class UMaterialFunctionInterface*>& DependentFunctions) const override;
diff --git a/Engine/Source/Runtime/Engine/Classes/Materials/MaterialInstance.h b/Engine/Source/Runtime/Engine/Classes/Materials/MaterialInstance.h
--- a/Engine/Source/Runtime/Engine/Classes/Materials/MaterialInstance.h
+++ b/Engine/Source/Runtime/Engine/Classes/Materials/MaterialInstance.h
@@ -975,6 +975,10 @@ public:
 #endif
 	ENGINE_API virtual bool IsComplete() const override;
 
+#if WITH_EDITOR
+	ENGINE_API virtual bool IsCompiling() const override;
+#endif
+
 	/** Tracking of in-flight uniform expression cache update operations for the material instance, for thread safety destroying the resource. */
 	void StartCacheUniformExpressions() const;
 	void FinishCacheUniformExpressions() const;
diff --git a/Engine/Source/Runtime/Engine/Classes/Materials/MaterialInterface.h b/Engine/Source/Runtime/Engine/Classes/Materials/MaterialInterface.h
--- a/Engine/Source/Runtime/Engine/Classes/Materials/MaterialInterface.h
+++ b/Engine/Source/Runtime/Engine/Classes/Materials/MaterialInterface.h
@@ -30,6 +30,7 @@
 #include "RHIFeatureLevel.h"
 #include "PSOPrecache.h"
 #include "StaticParameterSet.h"
+#include "Interfaces/Interface_AsyncCompilation.h"
 
 #if UE_ENABLE_INCLUDE_ORDER_DEPRECATED_IN_5_2
 #include "RHI.h"
@@ -251,7 +252,7 @@ public:
 };
 
 UCLASS(abstract, BlueprintType, MinimalAPI, HideCategories = (Thumbnail))
-class UMaterialInterface : public UObject, public IBlendableInterface, public IInterface_AssetUserData
+class UMaterialInterface : public UObject, public IBlendableInterface, public IInterface_AssetUserData, public IInterface_AsyncCompilation
 {
 	GENERATED_UCLASS_BODY()
 
@@ -997,6 +998,14 @@ public:
 	*/
 	virtual bool IsComplete() const { return true; }
 
+	/** IInterface_AsyncCompilation begin*/
+#if WITH_EDITOR
+	ENGINE_API virtual bool IsCompiling() const override { return false; };
+#else
+	FORCEINLINE bool IsCompiling() const { return false; }
+#endif
+	/** IInterface_AsyncCompilation end*/
+
 	/** @brief Checks to see if this material has all its shaders cached and if not, will perform a synchronous compilation of those.
 	*
 	* Materials are not guaranteed to have all their shaders compiled after loading and using this function before a draw will ensure that the material will not render until it's 
diff --git a/Engine/Source/Runtime/Engine/Private/Animation/Skeleton.cpp b/Engine/Source/Runtime/Engine/Private/Animation/Skeleton.cpp
--- a/Engine/Source/Runtime/Engine/Private/Animation/Skeleton.cpp
+++ b/Engine/Source/Runtime/Engine/Private/Animation/Skeleton.cpp
@@ -608,6 +608,7 @@ bool USkeleton::RenameMarkerName(FName InOldName, FName InNewName)
 
 bool USkeleton::DoesParentChainMatch(int32 StartBoneIndex, const USkinnedAsset* InSkinnedAsset) const
 {
+	TRACE_CPUPROFILER_EVENT_SCOPE("USkeleton::DoesParentChainMatch");
 	const FReferenceSkeleton& SkeletonRefSkel = ReferenceSkeleton;
 	const FReferenceSkeleton& MeshRefSkel = InSkinnedAsset->GetRefSkeleton();
 
@@ -653,6 +654,7 @@ bool USkeleton::DoesParentChainMatch(int32 StartBoneIndex, const USkinnedAsset*
 
 bool USkeleton::IsCompatibleMesh(const USkinnedAsset* InSkinnedAsset, bool bDoParentChainCheck) const
 {
+	TRACE_CPUPROFILER_EVENT_SCOPE("USkeleton::IsCompatibleMesh");
 	// at least % of bone should match 
 	int32 NumOfBoneMatches = 0;
 
@@ -899,8 +901,9 @@ bool USkeleton::RecreateBoneTree(USkinnedAsset* InSkinnedAsset)
 	return false;
 }
 
-bool USkeleton::MergeAllBonesToBoneTree(const USkinnedAsset* InSkinnedAsset)
+bool USkeleton::MergeAllBonesToBoneTree(const USkinnedAsset* InSkinnedAsset, bool bShowProgress /*= true*/)
 {
+	TRACE_CPUPROFILER_EVENT_SCOPE("USkeleton::MergeAllBonesToBoneTree");
 	if( InSkinnedAsset )
 	{
 		TArray<int32> RequiredBoneIndices;
@@ -916,7 +919,7 @@ bool USkeleton::MergeAllBonesToBoneTree(const USkinnedAsset* InSkinnedAsset)
 		if( RequiredBoneIndices.Num() > 0 )
 		{
 			// merge bones to the selected skeleton
-			return MergeBonesToBoneTree( InSkinnedAsset, RequiredBoneIndices );
+			return MergeBonesToBoneTree( InSkinnedAsset, RequiredBoneIndices, bShowProgress);
 		}
 	}
 
@@ -925,6 +928,7 @@ bool USkeleton::MergeAllBonesToBoneTree(const USkinnedAsset* InSkinnedAsset)
 
 bool USkeleton::CreateReferenceSkeletonFromMesh(const USkinnedAsset* InSkinnedAsset, const TArray<int32> & RequiredRefBones)
 {
+	TRACE_CPUPROFILER_EVENT_SCOPE("USkeleton::CreateReferenceSkeletonFromMesh");
 	// Filter list, we only want bones that have their parents present in this array.
 	TArray<int32> FilteredRequiredBones; 
 	FAnimationRuntime::ExcludeBonesWithNoParents(RequiredRefBones, InSkinnedAsset->GetRefSkeleton(), FilteredRequiredBones);
@@ -964,7 +968,7 @@ bool USkeleton::CreateReferenceSkeletonFromMesh(const USkinnedAsset* InSkinnedAs
 }
 
 
-bool USkeleton::MergeBonesToBoneTree(const USkinnedAsset* InSkinnedAsset, const TArray<int32> & RequiredRefBones)
+bool USkeleton::MergeBonesToBoneTree(const USkinnedAsset* InSkinnedAsset, const TArray<int32> & RequiredRefBones, bool bShowProgress /*= true*/)
 {
 	// see if it needs all animation data to remap - only happens when bone structure CHANGED - added
 	bool bSuccess = false;
@@ -983,6 +987,7 @@ bool USkeleton::MergeBonesToBoneTree(const USkinnedAsset* InSkinnedAsset, const
 		// can we play? - hierarchy matches
 		if( IsCompatibleMesh(InSkinnedAsset) )
 		{
+			TRACE_CPUPROFILER_EVENT_SCOPE("USkeleton::MergeBonesToBoneTree::CompatibleBranch");
 			// Exclude bones who do not have a parent.
 			TArray<int32> FilteredRequiredBones;
 			FAnimationRuntime::ExcludeBonesWithNoParents(RequiredRefBones, InSkinnedAsset->GetRefSkeleton(), FilteredRequiredBones);
@@ -1022,7 +1027,7 @@ bool USkeleton::MergeBonesToBoneTree(const USkinnedAsset* InSkinnedAsset, const
 	if (bShouldHandleHierarchyChange)
 	{
 #if WITH_EDITOR
-		HandleSkeletonHierarchyChange();
+		HandleSkeletonHierarchyChange(bShowProgress);
 #endif
 	}
 
@@ -1347,7 +1352,14 @@ USkeletalMesh* USkeleton::GetAssetPreviewMesh(UObject* InAsset)
 
 	if (!PreviewMesh)
 	{
-		PreviewMesh = GetPreviewMesh(false);
+		//The const version avoid verifying the skeleton compatibility, which can stall the thread
+		const USkeleton* ThisSkeleton = this;
+		PreviewMesh = ThisSkeleton->GetPreviewMesh();
+		if (PreviewMesh && !PreviewMesh->IsCompiling())
+		{
+			//Verify the compatibility only if we are not building
+			PreviewMesh = GetPreviewMesh(false);
+		}
 	}
 
 	return PreviewMesh;
@@ -1397,8 +1409,9 @@ void USkeleton::RemoveBonesFromSkeleton( const TArray<FName>& BonesToRemove, boo
 	}
 }
 
-void USkeleton::HandleSkeletonHierarchyChange()
+void USkeleton::HandleSkeletonHierarchyChange(bool bShowProgress /*= true*/)
 {
+	TRACE_CPUPROFILER_EVENT_SCOPE("USkeleton::HandleSkeletonHierarchyChange");
 	MarkPackageDirty();
 
 	RegenerateGuid();
@@ -1433,15 +1446,21 @@ void USkeleton::HandleSkeletonHierarchyChange()
 		}
 	}
 
-	FScopedSlowTask SlowTask((float)NumLoadedAssets, LOCTEXT("HandleSkeletonHierarchyChange", "Rebuilding animations..."));
-	SlowTask.MakeDialog();
+	FScopedSlowTask SlowTask((float)NumLoadedAssets, LOCTEXT("HandleSkeletonHierarchyChange", "Rebuilding animations..."), bShowProgress);
+	if (bShowProgress)
+	{
+		SlowTask.MakeDialog();
+	}
 
 	for (TObjectIterator<UAnimationAsset> It; It; ++It)
 	{
 		UAnimationAsset* CurrentAnimation = *It;
 		if (CurrentAnimation->GetSkeleton() == this)
 		{
-			SlowTask.EnterProgressFrame(1.0f, FText::Format(LOCTEXT("HandleSkeletonHierarchyChange_Format", "Rebuilding Animation: {0}"), FText::FromString(CurrentAnimation->GetName())));
+			if (bShowProgress)
+			{
+				SlowTask.EnterProgressFrame(1.0f, FText::Format(LOCTEXT("HandleSkeletonHierarchyChange_Format", "Rebuilding Animation: {0}"), FText::FromString(CurrentAnimation->GetName())));
+			}
 
 			CurrentAnimation->ValidateSkeleton();
 		}
diff --git a/Engine/Source/Runtime/Engine/Private/EditorFramework/AssetImportData.cpp b/Engine/Source/Runtime/Engine/Private/EditorFramework/AssetImportData.cpp
--- a/Engine/Source/Runtime/Engine/Private/EditorFramework/AssetImportData.cpp
+++ b/Engine/Source/Runtime/Engine/Private/EditorFramework/AssetImportData.cpp
@@ -135,6 +135,7 @@ void UAssetImportData::UpdateFilenameOnly(const FString& InPath, int32 Index)
 
 void UAssetImportData::AddFileName(const FString& InPath, int32 Index, FString SourceFileLabel /*= FString()*/)
 {
+	TRACE_CPUPROFILER_EVENT_SCOPE("UAssetImportData::AddFileName");
 	FAssetImportInfo Old = SourceData;
 
 	// Reset our current data
@@ -177,6 +178,7 @@ void UAssetImportData::AddFileName(const FString& InPath, int32 Index, FString S
 
 void UAssetImportData::SetSourceFiles(TArray<FAssetImportInfo::FSourceFile>&& SourceFiles)
 {
+	TRACE_CPUPROFILER_EVENT_SCOPE("UAssetImportData::SetSourceFiles");
 	FAssetImportInfo Old = SourceData;
 
 	for (FAssetImportInfo::FSourceFile& SourceFile : SourceFiles)
@@ -199,6 +201,7 @@ void UAssetImportData::SetSourceFiles(TArray<FAssetImportInfo::FSourceFile>&& So
 
 void UAssetImportData::Update(const FString& InPath, FMD5Hash *Md5Hash/* = nullptr*/)
 {
+	TRACE_CPUPROFILER_EVENT_SCOPE("UAssetImportData::Update");
 	FAssetImportInfo Old = SourceData;
 	SourceData.SourceFiles.Reset();
 	for (int32 SourceIndex = 0; SourceIndex < Old.SourceFiles.Num(); ++SourceIndex)
diff --git a/Engine/Source/Runtime/Engine/Private/Materials/Material.cpp b/Engine/Source/Runtime/Engine/Private/Materials/Material.cpp
--- a/Engine/Source/Runtime/Engine/Private/Materials/Material.cpp
+++ b/Engine/Source/Runtime/Engine/Private/Materials/Material.cpp
@@ -2623,6 +2623,32 @@ bool UMaterial::IsComplete() const
 	return bComplete;
 }
 
+#if WITH_EDITOR
+bool UMaterial::IsCompiling() const
+{
+	bool bIsCompiling = false;
+	if (FApp::CanEverRender())
+	{
+		const EMaterialQualityLevel::Type ActiveQualityLevel = GetCachedScalabilityCVars().MaterialQualityLevel;
+		uint32 FeatureLevelsToCompile = GetFeatureLevelsToCompileForRendering();
+
+		while (FeatureLevelsToCompile != 0)
+		{
+			const ERHIFeatureLevel::Type FeatureLevel = (ERHIFeatureLevel::Type)FBitSet::GetAndClearNextBit(FeatureLevelsToCompile);
+			const EShaderPlatform ShaderPlatform = GShaderPlatformForFeatureLevel[FeatureLevel];
+
+			FMaterialResource* CurrentResource = FindMaterialResource(MaterialResources, FeatureLevel, ActiveQualityLevel, true);
+			if (CurrentResource && !CurrentResource->IsCompilationFinished())
+			{
+				bIsCompiling = true;
+				break;
+			}
+		}
+	}
+	return bIsCompiling;
+}
+#endif
+
 FGraphEventArray UMaterial::PrecachePSOs(const FPSOPrecacheVertexFactoryDataList& VertexFactoryDataList, const FPSOPrecacheParams& InPreCacheParams, EPSOPrecachePriority Priority, TArray<FMaterialPSOPrecacheRequestID>& OutMaterialPSORequestIDs)
 {
 	FGraphEventArray GraphEvents;
diff --git a/Engine/Source/Runtime/Engine/Private/Materials/MaterialInstance.cpp b/Engine/Source/Runtime/Engine/Private/Materials/MaterialInstance.cpp
--- a/Engine/Source/Runtime/Engine/Private/Materials/MaterialInstance.cpp
+++ b/Engine/Source/Runtime/Engine/Private/Materials/MaterialInstance.cpp
@@ -2594,6 +2594,32 @@ bool UMaterialInstance::IsComplete() const
 	return bComplete;
 }
 
+#if WITH_EDITOR
+bool UMaterialInstance::IsCompiling() const
+{
+	bool bIsCompiling = false;
+	if (bHasStaticPermutationResource && FApp::CanEverRender())
+	{
+		uint32 FeatureLevelsToCompile = GetFeatureLevelsToCompileForRendering();
+		const EMaterialQualityLevel::Type ActiveQualityLevel = GetCachedScalabilityCVars().MaterialQualityLevel;
+
+		while (FeatureLevelsToCompile != 0)
+		{
+			const ERHIFeatureLevel::Type FeatureLevel = (ERHIFeatureLevel::Type)FBitSet::GetAndClearNextBit(FeatureLevelsToCompile);
+			const EShaderPlatform ShaderPlatform = GShaderPlatformForFeatureLevel[FeatureLevel];
+
+			FMaterialResource* CurrentResource = FindMaterialResource(StaticPermutationMaterialResources, FeatureLevel, ActiveQualityLevel, true);
+			if (CurrentResource && !CurrentResource->IsCompilationFinished())
+			{
+				bIsCompiling = true;
+				break;
+			}
+		}
+	}
+	return bIsCompiling;
+}
+#endif
+
 #if WITH_EDITOR
 bool UMaterialInstance::SetMaterialLayers(const FMaterialLayersFunctions& LayersValue)
 {
diff --git a/Engine/Source/Runtime/Engine/Private/ShaderCompiler/ShaderCompiler.cpp b/Engine/Source/Runtime/Engine/Private/ShaderCompiler/ShaderCompiler.cpp
--- a/Engine/Source/Runtime/Engine/Private/ShaderCompiler/ShaderCompiler.cpp
+++ b/Engine/Source/Runtime/Engine/Private/ShaderCompiler/ShaderCompiler.cpp
@@ -6016,7 +6016,10 @@ void FShaderCompilingManager::BlockOnShaderMapCompletion(const TArray<int32>& Sh
 		}
 
 		FScopedSlowTask SlowTask(NumJobs, FText::Format(LOCTEXT("BlockOnShaderMapCompletion", "Compiling Shaders ({0})"), NumJobs), GIsEditor && !IsRunningCommandlet() && GPlayInEditorID == INDEX_NONE);
-		SlowTask.MakeDialog(false, true);
+		if (NumJobs > 0)
+		{
+			SlowTask.MakeDialog(false, true);
+		}
 
 		int32 NumPendingJobs = 0;
 		// Keep track of previous number of pending jobs so we can update the slow task with the amount of work done.
@@ -6143,7 +6146,10 @@ void FShaderCompilingManager::BlockOnAllShaderMapCompletion(TMap<int32, FShaderM
 		}
 
 		FScopedSlowTask SlowTask(NumJobs, FText::Format(LOCTEXT("BlockOnAllShaderMapCompletion", "Compiling Shaders ({0})"), NumJobs), NumJobs && GIsEditor && !IsRunningCommandlet());
-		SlowTask.MakeDialog(false, true);
+		if (NumJobs > 0)
+		{
+			SlowTask.MakeDialog(false, true);
+		}
 
 		int32 NumPendingJobs = 0;
 		// Keep track of previous number of pending jobs so we can update the slow task with the amount of work done.
diff --git a/Engine/Source/Runtime/Engine/Private/SkeletalMesh.cpp b/Engine/Source/Runtime/Engine/Private/SkeletalMesh.cpp
--- a/Engine/Source/Runtime/Engine/Private/SkeletalMesh.cpp
+++ b/Engine/Source/Runtime/Engine/Private/SkeletalMesh.cpp
@@ -515,6 +515,25 @@ void USkeletalMesh::ValidateBoundsExtension()
 
 #if WITH_EDITOR
 
+bool USkeletalMesh::IsReadyToRenderInThumbnail() const
+{
+	if (IsCompiling() || !GetResourceForRendering())
+	{
+		return false;
+	}
+
+	//Since skeletal mesh use material, we want to avoid drawing thumbnail when shader are compiling
+	for (const FSkeletalMaterial& SkeletalMaterial : GetMaterials())
+	{
+		if (SkeletalMaterial.MaterialInterface && SkeletalMaterial.MaterialInterface->IsCompiling())
+		{
+			return false;
+		}
+	}
+
+	return true;
+}
+
 bool USkeletalMesh::IsInitialBuildDone() const
 {
 	//We are consider built if we have a valid lod model and a valid inline reduction cache
diff --git a/Engine/Source/Runtime/Interchange/Core/Public/InterchangeSourceData.h b/Engine/Source/Runtime/Interchange/Core/Public/InterchangeSourceData.h
--- a/Engine/Source/Runtime/Interchange/Core/Public/InterchangeSourceData.h
+++ b/Engine/Source/Runtime/Interchange/Core/Public/InterchangeSourceData.h
@@ -43,6 +43,7 @@ public:
 	/** Return the hash of the content point by Filename. */
 	FORCEINLINE TOptional<FMD5Hash> GetFileContentHash() const
 	{
+		TRACE_CPUPROFILER_EVENT_SCOPE("UInterchangeSourceData::GetFileContentHash");
 		if(!FileContentHashCache.IsSet())
 		{
 			ComputeFileContentHashCache();
