[Interchange - glTF] Fixing ClearCoat BottomNormal material connections.

#rnx

diff --git a/Engine/Plugins/Interchange/Runtime/Source/Import/Public/Gltf/InterchangeGLTFMaterial.h b/Engine/Plugins/Interchange/Runtime/Source/Import/Public/Gltf/InterchangeGLTFMaterial.h
--- a/Engine/Plugins/Interchange/Runtime/Source/Import/Public/Gltf/InterchangeGLTFMaterial.h
+++ b/Engine/Plugins/Interchange/Runtime/Source/Import/Public/Gltf/InterchangeGLTFMaterial.h
@@ -142,7 +142,7 @@ namespace UE::Interchange::GLTFMaterials
 			TPair<FString, TArray<FString>>(
 				TEXT("/Interchange/gltf/MaterialBodies/MF_ClearCoat_Body.MF_ClearCoat_Body"),
 				TArray<FString>{
-					TEXT("ClearCoatBottomNormal"),
+					TEXT("ClearCoatNormal"),
 					TEXT("BaseColor"),
 					TEXT("Metallic"),
 					TEXT("Specular"),
