[Interchange - glTF] glTF BindPose parsing fixes for NonRoot Skeletals.

#rnx

diff --git a/Engine/Plugins/Interchange/Runtime/Source/Import/Private/Gltf/InterchangeGltfMesh.cpp b/Engine/Plugins/Interchange/Runtime/Source/Import/Private/Gltf/InterchangeGltfMesh.cpp
--- a/Engine/Plugins/Interchange/Runtime/Source/Import/Private/Gltf/InterchangeGltfMesh.cpp
+++ b/Engine/Plugins/Interchange/Runtime/Source/Import/Private/Gltf/InterchangeGltfMesh.cpp
@@ -73,19 +73,8 @@ namespace UE::Interchange::Gltf::Private
 
 	bool GetSkeletalMeshDescriptionForPayLoadKey(const GLTF::FAsset& GltfAsset, const FString& PayLoadKey, const FTransform& MeshGlobalTransform, FMeshDescription& MeshDescription, TArray<FString>* OutJointUniqueNames)
 	{
-		TArray<FString> BakingAndPayloadKey;
-		PayLoadKey.ParseIntoArray(BakingAndPayloadKey, TEXT(";"));
-
-		bool bBakeSkinJointTransform = false;
-		if (BakingAndPayloadKey.Num() == 2)
-		{
-			LexFromString(bBakeSkinJointTransform, *BakingAndPayloadKey[0]);
-		}
-
-		FString PayloadKeyToParse = BakingAndPayloadKey.Num() == 2 ? BakingAndPayloadKey[1] : PayLoadKey;
-
 		TArray<FString> PayLoadKeys;
-		PayloadKeyToParse.ParseIntoArray(PayLoadKeys, TEXT(":"));
+		PayLoadKey.ParseIntoArray(PayLoadKeys, TEXT(":"));
 		TMap<int32, TArray<int32>> MeshIndexToSkinIndicesMap;
 
 		int32 MeshIndex = 0;
@@ -139,27 +128,6 @@ namespace UE::Interchange::Gltf::Private
 				continue;
 			}
 			const GLTF::FSkinInfo& Skin = GltfAsset.Skins[SkinIndex];
-			FStaticMeshAttributes StaticMeshAttributes(SkinnedMeshDescription);
-
-			//for instanced meshes we need to bake the transforms of the joints:
-			if ((bBakeSkinJointTransform || SkinIndices.Num() > 1) && GltfAsset.Skins[SkinIndex].Joints.Num() > 0)
-			{
-				FTransform TransformLocalToWorld3d(FTransform::Identity);
-
-				CalculateJointTransformations(GltfAsset.Nodes, GetRootNodeIndex(GltfAsset, GltfAsset.Skins[SkinIndex].Joints), TransformLocalToWorld3d);
-
-				FTransform3f TransformLocalToWorld(TransformLocalToWorld3d);
-				TVertexAttributesRef<FVector3f> VertexPositions = StaticMeshAttributes.GetVertexPositions();
-				FVertexArray VertexIndices = SkinnedMeshDescription.Vertices();
-				for (FVertexID VertexID : SkinnedMeshDescription.Vertices().GetElementIDs())
-				{
-					FVector3f VertexPosition = VertexPositions[VertexID];
-
-					VertexPosition = TransformLocalToWorld.TransformPosition(VertexPosition);
-
-					VertexPositions[VertexID] = VertexPosition;
-				}
-			}
 
 			FSkeletalMeshAttributes SkeletalMeshAttributes(SkinnedMeshDescription);
 			SkeletalMeshAttributes.Register();
diff --git a/Engine/Plugins/Interchange/Runtime/Source/Import/Private/Gltf/InterchangeGltfTranslator.cpp b/Engine/Plugins/Interchange/Runtime/Source/Import/Private/Gltf/InterchangeGltfTranslator.cpp
--- a/Engine/Plugins/Interchange/Runtime/Source/Import/Private/Gltf/InterchangeGltfTranslator.cpp
+++ b/Engine/Plugins/Interchange/Runtime/Source/Import/Private/Gltf/InterchangeGltfTranslator.cpp
@@ -1176,11 +1176,9 @@ void UInterchangeGLTFTranslator::HandleGltfSkeletons(UInterchangeBaseNodeContain
 				SkeletalMeshNode->SetSkeletonDependencyUid(*SkeletonNodeUid);
 			}
 
-			bool bBakeSkinJointTransform = RootSkinJointNodeIndex != RootJointNodeIndex || RootJointNode.ParentIndex == -1;
-
 			//generate payload key:
 			//of template:
-			//"LexToString(bBakeSkinJointTransform)";"LexToString(SkinnedMeshNode.MeshIndex | (SkinnedMeshNode.Skindex << 16))":"LexToString(SkinnedMeshNode.MeshIndex | (SkinnedMeshNode.Skindex << 16))".....
+			//"LexToString(SkinnedMeshNode.MeshIndex | (SkinnedMeshNode.Skindex << 16))":"LexToString(SkinnedMeshNode.MeshIndex | (SkinnedMeshNode.Skindex << 16))".....
 			FString Payload = TEXT("");
 			for (int32 SkinnedMeshIndex : RootJointToSkinnedMeshNodes.Value)
 			{
@@ -1192,7 +1190,6 @@ void UInterchangeGLTFTranslator::HandleGltfSkeletons(UInterchangeBaseNodeContain
 				
 				Payload += LexToString(SkinnedMeshNode.MeshIndex | (SkinnedMeshNode.Skindex << 16));
 			}
-			Payload = LexToString(bBakeSkinJointTransform) + TEXT(";") + Payload;
 			SkeletalMeshNode->SetPayLoadKey(Payload, EInterchangeMeshPayLoadType::SKELETAL);
 
 			//set the mesh actor node's custom asset instance uid to the new duplicated mesh
diff --git a/Engine/Plugins/Interchange/Runtime/Source/Parsers/GLTFCore/Private/GLTF/GLTFReader.cpp b/Engine/Plugins/Interchange/Runtime/Source/Parsers/GLTFCore/Private/GLTF/GLTFReader.cpp
--- a/Engine/Plugins/Interchange/Runtime/Source/Parsers/GLTFCore/Private/GLTF/GLTFReader.cpp
+++ b/Engine/Plugins/Interchange/Runtime/Source/Parsers/GLTFCore/Private/GLTF/GLTFReader.cpp
@@ -1155,6 +1155,22 @@ namespace GLTF
 			}
 		}
 	}
+
+	void GenerateGlobalTransform(const TArray<FNode>& Nodes, int32 CurrentIndex, FTransform& GlobalTransform, const int32& SkeletonCommonRootIndex/*PivotPoint*/)
+	{
+		if (Nodes.IsValidIndex(CurrentIndex))
+		{
+			const FNode& CurrentNode = Nodes[CurrentIndex];
+
+			if (CurrentIndex != SkeletonCommonRootIndex)
+			{
+				GenerateGlobalTransform(Nodes, CurrentNode.ParentIndex, GlobalTransform, SkeletonCommonRootIndex);
+			}
+
+			GlobalTransform = CurrentNode.Transform * GlobalTransform;
+		}
+	}
+
 	void FFileReader::GenerateLocalBindPosesPerSkinIndices() const
 	{
 		for (size_t SkinIndex = 0; SkinIndex < Asset->Skins.Num(); SkinIndex++)
@@ -1165,17 +1181,46 @@ namespace GLTF
 				for (size_t JointCounter = 0; JointCounter < Skin.Joints.Num(); JointCounter++)
 				{
 					FNode& CurrentNode = Asset->Nodes[Skin.Joints[JointCounter]];
-					if (CurrentNode.ParentIndex != INDEX_NONE)
+
+					if (CurrentNode.ParentIndex != INDEX_NONE &&
+						Asset->Nodes.IsValidIndex(CurrentNode.ParentIndex) &&
+						Asset->Nodes[CurrentNode.ParentIndex].Type == FNode::EType::Joint &&
+						(Asset->Nodes[CurrentNode.ParentIndex].SkinIndexToGlobalInverseBindTransform.Contains(SkinIndex) || Asset->Nodes[CurrentNode.ParentIndex].SkinIndexToGlobalInverseBindTransform.Num() > 0)
+						)
 					{
-						if (Asset->Nodes[CurrentNode.ParentIndex].Type == FNode::EType::Joint
-							&& Asset->Nodes[CurrentNode.ParentIndex].SkinIndexToGlobalInverseBindTransform.Contains(SkinIndex))
+						FNode& ParentNode = Asset->Nodes[CurrentNode.ParentIndex];
+
+						//LocalBindPose; //bind pose would be CurrentNode.GlobalInverseBindTransform.Inverse() * ParentNode.GlobalInverseBindTransform
+						FTransform ParentGlobalInverseBindTransform;
+						if (ParentNode.SkinIndexToGlobalInverseBindTransform.Contains(SkinIndex))
 						{
-							//LocalBindPose; //bind pose would be CurrentNode.GlobalInverseBindTransform.Inverse() * ParentNode.GlobalInverseBindTransform
-							FTransform ParentGlobalInverseBindTransform = Asset->Nodes[CurrentNode.ParentIndex].SkinIndexToGlobalInverseBindTransform[SkinIndex];
-							FTransform LocalBindPose = CurrentNode.SkinIndexToGlobalInverseBindTransform[SkinIndex].Inverse() * ParentGlobalInverseBindTransform;
+							ParentGlobalInverseBindTransform = ParentNode.SkinIndexToGlobalInverseBindTransform[SkinIndex];
+						}
+						else
+						{
+							//Scenario is that the a Skin is instantiated at the end of another skin
+							//(Prime example is the RecursiveSkeleton gltf sample file.)
+							ParentGlobalInverseBindTransform = ParentNode.SkinIndexToGlobalInverseBindTransform.begin().Value();
+
+							if (ParentNode.SkinIndexToGlobalInverseBindTransform.Num() > 1)
+							{
+								Messages.Emplace(EMessageSeverity::Error, FString::Printf(TEXT("The same Joint [%s] is used in multiple Skins, which is currently not supported."), *ParentNode.Name));
+							}
+						}
+						
+						FTransform LocalBindPose = CurrentNode.SkinIndexToGlobalInverseBindTransform[SkinIndex].Inverse() * ParentGlobalInverseBindTransform;
 
-							CurrentNode.SkinIndexToLocalBindPose.Add(SkinIndex, LocalBindPose);
+						CurrentNode.SkinIndexToLocalBindPose.Add(SkinIndex, LocalBindPose);
+					}
+					else
+					{
+						FTransform ParentGlobalTransform;
+						if (Skin.Skeleton != INDEX_NONE && Skin.Skeleton != CurrentNode.Index)
+						{
+							GenerateGlobalTransform(Asset->Nodes, CurrentNode.ParentIndex, ParentGlobalTransform, Skin.Skeleton);
 						}
+						FTransform LocalBindPose = CurrentNode.SkinIndexToGlobalInverseBindTransform[SkinIndex].Inverse() * ParentGlobalTransform.Inverse();
+						CurrentNode.SkinIndexToLocalBindPose.Add(SkinIndex, LocalBindPose);
 					}
 				}
 			}
