Fixed usage of TRACE_CPUPROFILER_EVENT_SCOPE macros in Interchange plugin.

* TRACE_CPUPROFILER_EVENT_SCOPE(ScopeName) --> to be used with a plain text as scope name (not a string!). If the parameter is a string ("abc" or TEXT("abc")), the quotes (and the TEXT prefix) will be included in the actual timer name.
* TRACE_CPUPROFILER_EVENT_SCOPE_STR("scope name") --> to be used with a static string (const ANSICHAR* or const TCHAR*); ex: when the scope name needs to include a space char
* TRACE_CPUPROFILER_EVENT_SCOPE_TEXT(*MyScopeName) --> to be used with a dynamic string (const ANSICHAR* or const TCHAR*)
See also the comments for these macros in ProfilingDebugging\CpuProfilerTrace.h.

diff --git a/Engine/Plugins/Interchange/Runtime/Source/Import/Private/Fbx/InterchangeFbxTranslator.cpp b/Engine/Plugins/Interchange/Runtime/Source/Import/Private/Fbx/InterchangeFbxTranslator.cpp
--- a/Engine/Plugins/Interchange/Runtime/Source/Import/Private/Fbx/InterchangeFbxTranslator.cpp
+++ b/Engine/Plugins/Interchange/Runtime/Source/Import/Private/Fbx/InterchangeFbxTranslator.cpp
@@ -83,7 +83,7 @@ TArray<FString> UInterchangeFbxTranslator::GetSupportedFormats() const
 
 bool UInterchangeFbxTranslator::Translate(UInterchangeBaseNodeContainer& BaseNodeContainer) const
 {
-	TRACE_CPUPROFILER_EVENT_SCOPE("UInterchangeFbxTranslator::Translate");
+	TRACE_CPUPROFILER_EVENT_SCOPE(UInterchangeFbxTranslator::Translate);
 	//Make sure the hash is compute here in asynchronous mode
 	GetSourceData()->GetFileContentHash();
 	FString Filename = GetSourceData()->GetFilename();
@@ -213,7 +213,7 @@ TFuture<TOptional<UE::Interchange::FMeshPayloadData>> UInterchangeFbxTranslator:
 	FString JsonCommand = CreateFetchMeshPayloadFbxCommand(PayLoadKey.UniqueId, MeshGlobalTransform);
 	const int32 CreatedTaskIndex = Dispatcher->AddTask(JsonCommand, FInterchangeDispatcherTaskCompleted::CreateLambda([this, Promise, PayLoadKey](const int32 TaskIndex)
 	{
-		TRACE_CPUPROFILER_EVENT_SCOPE("UInterchangeFbxTranslator::GetStaticMeshPayloadData::Dispatcher->AddTaskDone")
+		TRACE_CPUPROFILER_EVENT_SCOPE(UInterchangeFbxTranslator::GetStaticMeshPayloadData::Dispatcher->AddTaskDone)
 		UE::Interchange::ETaskState TaskState;
 		FString JsonResult;
 		TArray<FString> JsonMessages;
@@ -321,7 +321,7 @@ TFuture<TOptional<UE::Interchange::FAnimationPayloadData>> UInterchangeFbxTransl
 
 	const int32 CreatedTaskIndex = Dispatcher->AddTask(JsonCommand, FInterchangeDispatcherTaskCompleted::CreateLambda([this, Promise, PayLoadKey](const int32 TaskIndex)
 		{
-			TRACE_CPUPROFILER_EVENT_SCOPE("UInterchangeFbxTranslator::GetAnimationCurvePayloadData::Dispatcher->AddTaskDone")
+			TRACE_CPUPROFILER_EVENT_SCOPE(UInterchangeFbxTranslator::GetAnimationCurvePayloadData::Dispatcher->AddTaskDone)
 			UE::Interchange::ETaskState TaskState;
 			FString JsonResult;
 			TArray<FString> JsonMessages;
diff --git a/Engine/Plugins/Interchange/Runtime/Source/Import/Private/Mesh/InterchangeSkeletalMeshFactory.cpp b/Engine/Plugins/Interchange/Runtime/Source/Import/Private/Mesh/InterchangeSkeletalMeshFactory.cpp
--- a/Engine/Plugins/Interchange/Runtime/Source/Import/Private/Mesh/InterchangeSkeletalMeshFactory.cpp
+++ b/Engine/Plugins/Interchange/Runtime/Source/Import/Private/Mesh/InterchangeSkeletalMeshFactory.cpp
@@ -85,7 +85,7 @@ namespace UE
 																 , const UInterchangeBaseNodeContainer* NodeContainer
 																 , FString AssetName)
 			{
-				TRACE_CPUPROFILER_EVENT_SCOPE("FillMorphTargetMeshDescriptionsPerMorphTargetName")
+				TRACE_CPUPROFILER_EVENT_SCOPE(FillMorphTargetMeshDescriptionsPerMorphTargetName)
 				TArray<FString> MorphTargetUids;
 				MeshNodeContext.MeshNode->GetMorphTargetDependencies(MorphTargetUids);
 				TMap<FString, TFuture<TOptional<UE::Interchange::FMeshPayloadData>>> TempMorphTargetMeshDescriptionsPerMorphTargetName;
@@ -147,7 +147,7 @@ namespace UE
 
 			void CopyMorphTargetsMeshDescriptionToSkeletalMeshImportData(const TMap<FString, TOptional<UE::Interchange::FMeshPayloadData>>& LodMorphTargetMeshDescriptions, FSkeletalMeshImportData& DestinationSkeletalMeshImportData)
 			{
-				TRACE_CPUPROFILER_EVENT_SCOPE("CopyMorphTargetsMeshDescriptionToSkeletalMeshImportData")
+				TRACE_CPUPROFILER_EVENT_SCOPE(CopyMorphTargetsMeshDescriptionToSkeletalMeshImportData)
 				const int32 OriginalMorphTargetCount = LodMorphTargetMeshDescriptions.Num();
 				TArray<FString> Keys;
 				int32 MorphTargetCount = 0;
@@ -270,7 +270,7 @@ namespace UE
 
 			const UInterchangeSceneNode* RecursiveFindJointByName(const UInterchangeBaseNodeContainer* NodeContainer, const FString& ParentJointNodeId, const FString& JointName)
 			{
-				TRACE_CPUPROFILER_EVENT_SCOPE("RecursiveFindJointByName")
+				TRACE_CPUPROFILER_EVENT_SCOPE(RecursiveFindJointByName)
 				if (const UInterchangeSceneNode* JointNode = Cast<UInterchangeSceneNode>(NodeContainer->GetNode(ParentJointNodeId)))
 				{
 					if (JointNode->GetDisplayLabel().Equals(JointName))
@@ -295,7 +295,7 @@ namespace UE
 				, const FString& RootJointNodeId
 				, const FTransform& MeshGlobalTransform)
 			{
-				TRACE_CPUPROFILER_EVENT_SCOPE("SkinVertexPositionToTimeZero")
+				TRACE_CPUPROFILER_EVENT_SCOPE(SkinVertexPositionToTimeZero)
 				FMeshDescription& MeshDescription = LodMeshPayload.MeshDescription;
 				const int32 VertexCount = MeshDescription.Vertices().Num();
 				const TArray<FString>& JointNames = LodMeshPayload.JointNames;
@@ -423,7 +423,7 @@ namespace UE
 																  , const UInterchangeBaseNodeContainer* NodeContainer
 																  , const FString& RootJointNodeId)
 			{
-				TRACE_CPUPROFILER_EVENT_SCOPE("RetrieveAllSkeletalMeshPayloadsAndFillImportData")
+				TRACE_CPUPROFILER_EVENT_SCOPE(RetrieveAllSkeletalMeshPayloadsAndFillImportData)
 				if (!MeshTranslatorPayloadInterface)
 				{
 					return;
@@ -484,7 +484,7 @@ namespace UE
 						continue;
 					}
 					const FMeshNodeContext& MeshNodeContext = *MeshNodeContextAndFuture.Key;
-					TRACE_CPUPROFILER_EVENT_SCOPE("RetrieveAllSkeletalMeshPayloadsAndFillImportData::GetPayload")
+					TRACE_CPUPROFILER_EVENT_SCOPE(RetrieveAllSkeletalMeshPayloadsAndFillImportData::GetPayload)
 					TOptional<UE::Interchange::FMeshPayloadData> LodMeshPayload = MeshNodeContextAndFuture.Value.Get();
 					if (!LodMeshPayload.IsSet())
 					{
@@ -498,7 +498,7 @@ namespace UE
 					SkeletalMeshAppendSettings.bAppendVertexAttributes = bImportVertexAttributes;
 					SkeletalMeshAppendSettings.SourceVertexIDOffset = VertexOffset;
 					{
-						TRACE_CPUPROFILER_EVENT_SCOPE("RetrieveAllSkeletalMeshPayloadsAndFillImportData::CompactPayload")
+						TRACE_CPUPROFILER_EVENT_SCOPE(RetrieveAllSkeletalMeshPayloadsAndFillImportData::CompactPayload)
 						FElementIDRemappings ElementIDRemappings;
 						LodMeshPayload->MeshDescription.Compact(ElementIDRemappings);
 					}
@@ -600,7 +600,7 @@ namespace UE
 
 			void ProcessImportMeshInfluences(const int32 WedgeCount, TArray<SkeletalMeshImportData::FRawBoneInfluence>& Influences)
 			{
-				TRACE_CPUPROFILER_EVENT_SCOPE("ProcessImportMeshInfluences")
+				TRACE_CPUPROFILER_EVENT_SCOPE(ProcessImportMeshInfluences)
 				// Sort influences by vertex index.
 				struct FCompareVertexIndex
 				{
@@ -796,7 +796,7 @@ namespace UE
 
 			void RemapSkeletalMeshVertexColorToImportData(const USkeletalMesh* SkeletalMesh, const int32 LODIndex, FSkeletalMeshImportData* SkelMeshImportData)
 			{
-				TRACE_CPUPROFILER_EVENT_SCOPE("RemapSkeletalMeshVertexColorToImportData")
+				TRACE_CPUPROFILER_EVENT_SCOPE(RemapSkeletalMeshVertexColorToImportData)
 				//Make sure we have all the source data we need to do the remap
 				if (!SkeletalMesh->GetImportedModel() || !SkeletalMesh->GetImportedModel()->LODModels.IsValidIndex(LODIndex) || !SkeletalMesh->GetHasVertexColors())
 				{
@@ -909,7 +909,7 @@ UClass* UInterchangeSkeletalMeshFactory::GetFactoryClass() const
 
 UInterchangeFactoryBase::FImportAssetResult UInterchangeSkeletalMeshFactory::BeginImportAsset_GameThread(const FImportAssetObjectParams& Arguments)
 {
-	TRACE_CPUPROFILER_EVENT_SCOPE("UInterchangeSkeletalMeshFactory::BeginImportAssetObject_GameThread")
+	TRACE_CPUPROFILER_EVENT_SCOPE(UInterchangeSkeletalMeshFactory::BeginImportAssetObject_GameThread)
 
 	FImportAssetResult ImportAssetResult;
 
@@ -981,7 +981,7 @@ UInterchangeFactoryBase::FImportAssetResult UInterchangeSkeletalMeshFactory::Beg
 
 UInterchangeFactoryBase::FImportAssetResult UInterchangeSkeletalMeshFactory::ImportAsset_Async(const FImportAssetObjectParams& Arguments)
 {
-	TRACE_CPUPROFILER_EVENT_SCOPE("UInterchangeSkeletalMeshFactory::CreateAsset")
+	TRACE_CPUPROFILER_EVENT_SCOPE(UInterchangeSkeletalMeshFactory::CreateAsset)
 
 	FImportAssetResult ImportAssetResult;
 
@@ -1135,7 +1135,7 @@ UInterchangeFactoryBase::FImportAssetResult UInterchangeSkeletalMeshFactory::Imp
 	
 	for (int32 LodIndex = 0; LodIndex < LodCount; ++LodIndex)
 	{
-		TRACE_CPUPROFILER_EVENT_SCOPE("UInterchangeSkeletalMeshFactory::CreateAsset_LOD")
+		TRACE_CPUPROFILER_EVENT_SCOPE(UInterchangeSkeletalMeshFactory::CreateAsset_LOD)
 		ESkeletalMeshGeoImportVersions GeoImportVersion = ESkeletalMeshGeoImportVersions::LatestVersion;
 		ESkeletalMeshSkinningImportVersions SkinningImportVersion = ESkeletalMeshSkinningImportVersions::LatestVersion;
 		if (bIsReImport && SkeletalMesh->GetImportedModel() && SkeletalMesh->GetImportedModel()->LODModels.IsValidIndex(CurrentLodIndex))
@@ -1485,7 +1485,7 @@ UInterchangeFactoryBase::FImportAssetResult UInterchangeSkeletalMeshFactory::Imp
 
 UInterchangeFactoryBase::FImportAssetResult UInterchangeSkeletalMeshFactory::EndImportAsset_GameThread(const FImportAssetObjectParams& Arguments)
 {
-	TRACE_CPUPROFILER_EVENT_SCOPE("UInterchangeSkeletalMeshFactory::EndImportAssetObject_GameThread")
+	TRACE_CPUPROFILER_EVENT_SCOPE(UInterchangeSkeletalMeshFactory::EndImportAssetObject_GameThread)
 	
 	check(IsInGameThread());
 	FImportAssetResult ImportAssetResult;
@@ -1723,7 +1723,7 @@ void UInterchangeSkeletalMeshFactory::Cancel()
 /* This function is call in the completion task on the main thread, use it to call main thread post creation step for your assets*/
 void UInterchangeSkeletalMeshFactory::SetupObject_GameThread(const FSetupObjectParams& Arguments)
 {
-	TRACE_CPUPROFILER_EVENT_SCOPE("UInterchangeSkeletalMeshFactory::SetupObject_GameThread")
+	TRACE_CPUPROFILER_EVENT_SCOPE(UInterchangeSkeletalMeshFactory::SetupObject_GameThread)
 	check(IsInGameThread());
 	Super::SetupObject_GameThread(Arguments);
 
@@ -1817,7 +1817,7 @@ void UInterchangeSkeletalMeshFactory::SetupObject_GameThread(const FSetupObjectP
 
 bool UInterchangeSkeletalMeshFactory::GetSourceFilenames(const UObject* Object, TArray<FString>& OutSourceFilenames) const
 {
-	TRACE_CPUPROFILER_EVENT_SCOPE("UInterchangeSkeletalMeshFactory::GetSourceFilenames")
+	TRACE_CPUPROFILER_EVENT_SCOPE(UInterchangeSkeletalMeshFactory::GetSourceFilenames)
 #if WITH_EDITORONLY_DATA
 	if (const USkeletalMesh* SkeletalMesh = Cast<USkeletalMesh>(Object))
 	{
@@ -1830,7 +1830,7 @@ bool UInterchangeSkeletalMeshFactory::GetSourceFilenames(const UObject* Object,
 
 bool UInterchangeSkeletalMeshFactory::SetSourceFilename(const UObject* Object, const FString& SourceFilename, int32 SourceIndex) const
 {
-	TRACE_CPUPROFILER_EVENT_SCOPE("UInterchangeSkeletalMeshFactory::SetSourceFilename")
+	TRACE_CPUPROFILER_EVENT_SCOPE(UInterchangeSkeletalMeshFactory::SetSourceFilename)
 #if WITH_EDITORONLY_DATA
 	if (const USkeletalMesh* SkeletalMesh = Cast<USkeletalMesh>(Object))
 	{
@@ -1844,7 +1844,7 @@ bool UInterchangeSkeletalMeshFactory::SetSourceFilename(const UObject* Object, c
 
 bool UInterchangeSkeletalMeshFactory::SetReimportSourceIndex(const UObject* Object, int32 SourceIndex) const
 {
-	TRACE_CPUPROFILER_EVENT_SCOPE("UInterchangeSkeletalMeshFactory::SetReimportSourceIndex")
+	TRACE_CPUPROFILER_EVENT_SCOPE(UInterchangeSkeletalMeshFactory::SetReimportSourceIndex)
 #if WITH_EDITORONLY_DATA
 	if (const USkeletalMesh* SkeletalMesh = Cast<USkeletalMesh>(Object))
 	{
diff --git a/Engine/Source/Runtime/Interchange/Engine/Private/InterchangeImportCommon.cpp b/Engine/Source/Runtime/Interchange/Engine/Private/InterchangeImportCommon.cpp
--- a/Engine/Source/Runtime/Interchange/Engine/Private/InterchangeImportCommon.cpp
+++ b/Engine/Source/Runtime/Interchange/Engine/Private/InterchangeImportCommon.cpp
@@ -21,7 +21,7 @@ namespace UE::Interchange
 	{
 		UInterchangeAssetImportData* BeginSetupAssetData(FFactoryCommon::FUpdateImportAssetDataParameters& Parameters)
 		{
-			TRACE_CPUPROFILER_EVENT_SCOPE("UE::Interchange::Private::ImportCommon::BeginSetupAssetData")
+			TRACE_CPUPROFILER_EVENT_SCOPE(UE::Interchange::Private::ImportCommon::BeginSetupAssetData)
 			if (!ensure(IsInGameThread()))
 			{
 				return nullptr;
@@ -66,7 +66,7 @@ namespace UE::Interchange
 
 		void EndSetupAssetData(FFactoryCommon::FUpdateImportAssetDataParameters& Parameters, UInterchangeAssetImportData* AssetImportData)
 		{
-			TRACE_CPUPROFILER_EVENT_SCOPE("UE::Interchange::Private::ImportCommon::EndSetupAssetData")
+			TRACE_CPUPROFILER_EVENT_SCOPE(UE::Interchange::Private::ImportCommon::EndSetupAssetData)
 			UInterchangeBaseNodeContainer* FactoryNodeContainer = NewObject<UInterchangeBaseNodeContainer>(AssetImportData);
 			//We copy only the factory node dependencies, we use this only 
 			if (UInterchangeFactoryBaseNode* FactoryNode = Cast<UInterchangeFactoryBaseNode>(Parameters.NodeContainer->GetFactoryNode(Parameters.NodeUniqueID)))
@@ -130,7 +130,7 @@ namespace UE::Interchange
 
 	UAssetImportData* FFactoryCommon::UpdateImportAssetData(FUpdateImportAssetDataParameters& Parameters)
 	{
-		TRACE_CPUPROFILER_EVENT_SCOPE("UE::Interchange::FFactoryCommon::UpdateImportAssetData")
+		TRACE_CPUPROFILER_EVENT_SCOPE(UE::Interchange::FFactoryCommon::UpdateImportAssetData)
 		return UpdateImportAssetData(Parameters, [&Parameters](UInterchangeAssetImportData* AssetImportData)
 			{
 #if WITH_EDITORONLY_DATA
@@ -146,7 +146,7 @@ namespace UE::Interchange
 
 	UAssetImportData* FFactoryCommon::UpdateImportAssetData(FUpdateImportAssetDataParameters& Parameters, TFunctionRef<void(UInterchangeAssetImportData*)> CustomFileSourceUpdate)
 	{
-		TRACE_CPUPROFILER_EVENT_SCOPE("UE::Interchange::FFactoryCommon::UpdateImportAssetData2")
+		TRACE_CPUPROFILER_EVENT_SCOPE(UE::Interchange::FFactoryCommon::UpdateImportAssetData2)
 #if WITH_EDITORONLY_DATA
 		if (!ensure(IsInGameThread()))
 		{
@@ -200,7 +200,7 @@ namespace UE::Interchange
 
 	UAssetImportData* FFactoryCommon::SetImportAssetData(FSetImportAssetDataParameters& Parameters)
 	{
-		TRACE_CPUPROFILER_EVENT_SCOPE("UE::Interchange::FFactoryCommon::SetImportAssetData")
+		TRACE_CPUPROFILER_EVENT_SCOPE(UE::Interchange::FFactoryCommon::SetImportAssetData)
 		UInterchangeAssetImportData* AssetImportData = Private::ImportCommon::BeginSetupAssetData(Parameters);
 
 		// Update the source files
@@ -236,7 +236,7 @@ namespace UE::Interchange
 
 	bool FFactoryCommon::GetSourceFilenames(const UAssetImportData* AssetImportData, TArray<FString>& OutSourceFilenames)
 	{
-		TRACE_CPUPROFILER_EVENT_SCOPE("UE::Interchange::FFactoryCommon::GetSourceFilenames")
+		TRACE_CPUPROFILER_EVENT_SCOPE(UE::Interchange::FFactoryCommon::GetSourceFilenames)
 		if (Cast<UInterchangeAssetImportData>(AssetImportData) != nullptr)
 		{
 			AssetImportData->ExtractFilenames(OutSourceFilenames);
@@ -247,7 +247,7 @@ namespace UE::Interchange
 
 	bool FFactoryCommon::SetSourceFilename(UAssetImportData* AssetImportData, const FString& SourceFilename, int32 SourceIndex, const FString& SourceLabel)
 	{
-		TRACE_CPUPROFILER_EVENT_SCOPE("UE::Interchange::FFactoryCommon::SetSourceFilename")
+		TRACE_CPUPROFILER_EVENT_SCOPE(UE::Interchange::FFactoryCommon::SetSourceFilename)
 		if (AssetImportData)
 		{
 			const int32 SafeSourceIndex = SourceIndex == INDEX_NONE ? 0 : SourceIndex;				
@@ -268,7 +268,7 @@ namespace UE::Interchange
 		
 	bool FFactoryCommon::SetReimportSourceIndex(const UObject* Object, UAssetImportData* AssetImportData, int32 SourceIndex)
 	{
-		TRACE_CPUPROFILER_EVENT_SCOPE("UE::Interchange::FFactoryCommon::SetReimportSourceIndex")
+		TRACE_CPUPROFILER_EVENT_SCOPE(UE::Interchange::FFactoryCommon::SetReimportSourceIndex)
 		UInterchangeAssetImportData* InterchangeAssetImportData = Cast<UInterchangeAssetImportData>(AssetImportData);
 		if (!InterchangeAssetImportData)
 		{
@@ -302,7 +302,7 @@ namespace UE::Interchange
 										, const UInterchangeFactoryBaseNode* CurrentAssetNode
 										, UInterchangeFactoryBaseNode* PipelineAssetNode)
 	{
-		TRACE_CPUPROFILER_EVENT_SCOPE("UE::Interchange::FFactoryCommon::ApplyReimportStrategyToAsset")
+		TRACE_CPUPROFILER_EVENT_SCOPE(UE::Interchange::FFactoryCommon::ApplyReimportStrategyToAsset)
 		if (!ensure(PreviousAssetNode) || !ensure(PipelineAssetNode) || !ensure(CurrentAssetNode))
 		{
 			return;
diff --git a/Engine/Source/Runtime/Interchange/Engine/Private/InterchangeManager.cpp b/Engine/Source/Runtime/Interchange/Engine/Private/InterchangeManager.cpp
--- a/Engine/Source/Runtime/Interchange/Engine/Private/InterchangeManager.cpp
+++ b/Engine/Source/Runtime/Interchange/Engine/Private/InterchangeManager.cpp
@@ -1132,7 +1132,7 @@ bool UInterchangeManager::CanReimport(const UObject* Object, TArray<FString>& Ou
 
 void UInterchangeManager::StartQueuedTasks(bool bCancelAllTasks /*= false*/)
 {
-	TRACE_CPUPROFILER_EVENT_SCOPE_STR("UInterchangeManager::StartQueuedTasks")
+	TRACE_CPUPROFILER_EVENT_SCOPE(UInterchangeManager::StartQueuedTasks)
 	ensure(IsInterchangeImportEnabled());
 	if (!ensure(IsInGameThread()))
 	{
@@ -1312,7 +1312,7 @@ UInterchangeManager::ImportSceneAsync(const FString& ContentPath, const UInterch
 TTuple<UE::Interchange::FAssetImportResultRef, UE::Interchange::FSceneImportResultRef>
 UInterchangeManager::ImportInternal(const FString& ContentPath, const UInterchangeSourceData* SourceData, const FImportAssetParameters& ImportAssetParameters, const UE::Interchange::EImportType ImportType)
 {
-	TRACE_CPUPROFILER_EVENT_SCOPE("UInterchangeManager::ImportInternal")
+	TRACE_CPUPROFILER_EVENT_SCOPE(UInterchangeManager::ImportInternal)
 	ensure(IsInterchangeImportEnabled());
 	check(IsInGameThread());
 	static int32 GeneratedUniqueID = 0;
diff --git a/Engine/Source/Runtime/Interchange/Engine/Private/Tasks/InterchangeTaskCompletion.cpp b/Engine/Source/Runtime/Interchange/Engine/Private/Tasks/InterchangeTaskCompletion.cpp
--- a/Engine/Source/Runtime/Interchange/Engine/Private/Tasks/InterchangeTaskCompletion.cpp
+++ b/Engine/Source/Runtime/Interchange/Engine/Private/Tasks/InterchangeTaskCompletion.cpp
@@ -22,7 +22,7 @@
 
 void UE::Interchange::FTaskPreAsyncCompletion::DoTask(ENamedThreads::Type CurrentThread, const FGraphEventRef& MyCompletionGraphEvent)
 {
-	TRACE_CPUPROFILER_EVENT_SCOPE("UE::Interchange::FTaskPreAsyncCompletion::DoTask")
+	TRACE_CPUPROFILER_EVENT_SCOPE(UE::Interchange::FTaskPreAsyncCompletion::DoTask)
 #if INTERCHANGE_TRACE_ASYNCHRONOUS_TASK_ENABLED
 	INTERCHANGE_TRACE_ASYNCHRONOUS_TASK(PreAsyncCompletion)
 #endif
@@ -37,7 +37,7 @@ void UE::Interchange::FTaskPreAsyncCompletion::DoTask(ENamedThreads::Type Curren
 
 void UE::Interchange::FTaskPreCompletion::DoTask(ENamedThreads::Type CurrentThread, const FGraphEventRef& MyCompletionGraphEvent)
 {
-	TRACE_CPUPROFILER_EVENT_SCOPE("UE::Interchange::FTaskPreCompletion::DoTask")
+	TRACE_CPUPROFILER_EVENT_SCOPE(UE::Interchange::FTaskPreCompletion::DoTask)
 #if INTERCHANGE_TRACE_ASYNCHRONOUS_TASK_ENABLED
 	INTERCHANGE_TRACE_ASYNCHRONOUS_TASK(PreCompletion)
 #endif
@@ -202,7 +202,7 @@ void UE::Interchange::FTaskPreCompletion::DoTask(ENamedThreads::Type CurrentThre
 
 void UE::Interchange::FTaskCompletion::DoTask(ENamedThreads::Type CurrentThread, const FGraphEventRef& MyCompletionGraphEvent)
 {
-	TRACE_CPUPROFILER_EVENT_SCOPE("UE::Interchange::FTaskCompletion::DoTask")
+	TRACE_CPUPROFILER_EVENT_SCOPE(UE::Interchange::FTaskCompletion::DoTask)
 #if INTERCHANGE_TRACE_ASYNCHRONOUS_TASK_ENABLED
 	INTERCHANGE_TRACE_ASYNCHRONOUS_TASK(Completion)
 #endif
diff --git a/Engine/Source/Runtime/Interchange/Engine/Private/Tasks/InterchangeTaskCreateSceneObjects.cpp b/Engine/Source/Runtime/Interchange/Engine/Private/Tasks/InterchangeTaskCreateSceneObjects.cpp
--- a/Engine/Source/Runtime/Interchange/Engine/Private/Tasks/InterchangeTaskCreateSceneObjects.cpp
+++ b/Engine/Source/Runtime/Interchange/Engine/Private/Tasks/InterchangeTaskCreateSceneObjects.cpp
@@ -35,7 +35,7 @@ UE::Interchange::FTaskCreateSceneObjects::FTaskCreateSceneObjects(const FString&
 
 void UE::Interchange::FTaskCreateSceneObjects::DoTask(ENamedThreads::Type CurrentThread, const FGraphEventRef& MyCompletionGraphEvent)
 {
-	TRACE_CPUPROFILER_EVENT_SCOPE("UE::Interchange::FTaskCreateSceneObjects::DoTask")
+	TRACE_CPUPROFILER_EVENT_SCOPE(UE::Interchange::FTaskCreateSceneObjects::DoTask)
 #if INTERCHANGE_TRACE_ASYNCHRONOUS_TASK_ENABLED
 	INTERCHANGE_TRACE_ASYNCHRONOUS_TASK(SpawnActor)
 #endif
diff --git a/Engine/Source/Runtime/Interchange/Engine/Private/Tasks/InterchangeTaskImportObject.cpp b/Engine/Source/Runtime/Interchange/Engine/Private/Tasks/InterchangeTaskImportObject.cpp
--- a/Engine/Source/Runtime/Interchange/Engine/Private/Tasks/InterchangeTaskImportObject.cpp
+++ b/Engine/Source/Runtime/Interchange/Engine/Private/Tasks/InterchangeTaskImportObject.cpp
@@ -26,7 +26,7 @@ namespace UE::Interchange::Private
 {
 	void InternalGetPackageName(const UE::Interchange::FImportAsyncHelper& AsyncHelper, const int32 SourceIndex, const FString& PackageBasePath, const UInterchangeFactoryBaseNode* FactoryNode, FString& OutPackageName, FString& OutAssetName)
 	{
-		TRACE_CPUPROFILER_EVENT_SCOPE("UE::Interchange::Private::InternalGetPackageName")
+		TRACE_CPUPROFILER_EVENT_SCOPE(UE::Interchange::Private::InternalGetPackageName)
 		const UInterchangeSourceData* SourceData = AsyncHelper.SourceDatas[SourceIndex];
 		check(SourceData);
 		FString NodeDisplayName = FactoryNode->GetAssetName();
@@ -48,7 +48,7 @@ namespace UE::Interchange::Private
 	}
 	bool ShouldReimportFactoryNode(UInterchangeFactoryBaseNode* FactoryNode, const UInterchangeBaseNodeContainer* NodeContainer, UObject* ReimportObject)
 	{
-		TRACE_CPUPROFILER_EVENT_SCOPE("UE::Interchange::Private::ShouldReimportFactoryNode")
+		TRACE_CPUPROFILER_EVENT_SCOPE(UE::Interchange::Private::ShouldReimportFactoryNode)
 
 		if (!NodeContainer)
 		{
@@ -244,7 +244,7 @@ namespace UE::Interchange::Private
 
 void UE::Interchange::FTaskImportObject_GameThread::DoTask(ENamedThreads::Type CurrentThread, const FGraphEventRef& MyCompletionGraphEvent)
 {
-	TRACE_CPUPROFILER_EVENT_SCOPE("UE::Interchange::FTaskImportObject_GameThread::DoTask")
+	TRACE_CPUPROFILER_EVENT_SCOPE(UE::Interchange::FTaskImportObject_GameThread::DoTask)
 #if INTERCHANGE_TRACE_ASYNCHRONOUS_TASK_ENABLED
 	INTERCHANGE_TRACE_ASYNCHRONOUS_TASK(TaskImportObject_GameThread)
 #endif
diff --git a/Engine/Source/Runtime/Interchange/Engine/Private/Tasks/InterchangeTaskParsing.cpp b/Engine/Source/Runtime/Interchange/Engine/Private/Tasks/InterchangeTaskParsing.cpp
--- a/Engine/Source/Runtime/Interchange/Engine/Private/Tasks/InterchangeTaskParsing.cpp
+++ b/Engine/Source/Runtime/Interchange/Engine/Private/Tasks/InterchangeTaskParsing.cpp
@@ -35,7 +35,7 @@ struct FNodeDependencyCache
 {
 	const TSet<FString>& GetAccumulatedDependencies(const UInterchangeBaseNodeContainer* NodeContainer, const FString& NodeID)
 	{
-		TRACE_CPUPROFILER_EVENT_SCOPE("FNodeDependencyCache::GetAccumulatedDependencies")
+		TRACE_CPUPROFILER_EVENT_SCOPE(FNodeDependencyCache::GetAccumulatedDependencies)
 		TSet<FString> NodeStack;
 		return GetAccumulatedDependencies(NodeContainer, NodeID, NodeStack);
 	}
@@ -89,7 +89,7 @@ private:
 
 void UE::Interchange::FTaskParsing::DoTask(ENamedThreads::Type CurrentThread, const FGraphEventRef& MyCompletionGraphEvent)
 {
-	TRACE_CPUPROFILER_EVENT_SCOPE("UE::Interchange::FTaskParsing::DoTask")
+	TRACE_CPUPROFILER_EVENT_SCOPE(UE::Interchange::FTaskParsing::DoTask)
 #if INTERCHANGE_TRACE_ASYNCHRONOUS_TASK_ENABLED
 	INTERCHANGE_TRACE_ASYNCHRONOUS_TASK(ParsingGraph)
 #endif
diff --git a/Engine/Source/Runtime/Interchange/Engine/Private/Tasks/InterchangeTaskPipeline.cpp b/Engine/Source/Runtime/Interchange/Engine/Private/Tasks/InterchangeTaskPipeline.cpp
--- a/Engine/Source/Runtime/Interchange/Engine/Private/Tasks/InterchangeTaskPipeline.cpp
+++ b/Engine/Source/Runtime/Interchange/Engine/Private/Tasks/InterchangeTaskPipeline.cpp
@@ -22,7 +22,7 @@
 
 void UE::Interchange::FTaskPipeline::DoTask(ENamedThreads::Type CurrentThread, const FGraphEventRef& MyCompletionGraphEvent)
 {
-	TRACE_CPUPROFILER_EVENT_SCOPE("UE::Interchange::FTaskPipeline::DoTask")
+	TRACE_CPUPROFILER_EVENT_SCOPE(UE::Interchange::FTaskPipeline::DoTask)
 #if INTERCHANGE_TRACE_ASYNCHRONOUS_TASK_ENABLED
 	INTERCHANGE_TRACE_ASYNCHRONOUS_TASK(PipelinePreImport)
 #endif
@@ -58,7 +58,7 @@ void UE::Interchange::FTaskPipeline::DoTask(ENamedThreads::Type CurrentThread, c
 
 void UE::Interchange::FTaskWaitAssetCompilation::DoTask(ENamedThreads::Type CurrentThread, const FGraphEventRef& MyCompletionGraphEvent)
 {
-	TRACE_CPUPROFILER_EVENT_SCOPE("UE::Interchange::FTaskWaitAssetCompilation::DoTask")
+	TRACE_CPUPROFILER_EVENT_SCOPE(UE::Interchange::FTaskWaitAssetCompilation::DoTask)
 #if INTERCHANGE_TRACE_ASYNCHRONOUS_TASK_ENABLED
 		INTERCHANGE_TRACE_ASYNCHRONOUS_TASK(WaitAssetCompilation)
 #endif
@@ -98,7 +98,7 @@ void UE::Interchange::FTaskWaitAssetCompilation::DoTask(ENamedThreads::Type Curr
 			bool bCompilationFinish = false;
 			Async(EAsyncExecution::TaskGraphMainThread, [&bCompilationFinish, &ImportedObjects]()
 				{
-					TRACE_CPUPROFILER_EVENT_SCOPE("UE::Interchange::FTaskWaitAssetCompilation::DoTask::IsCompilingLambda_GameThread");
+					TRACE_CPUPROFILER_EVENT_SCOPE(UE::Interchange::FTaskWaitAssetCompilation::DoTask::IsCompilingLambda_GameThread);
 					//Make sure all asset compiling managers are up to date, In case the game thread is waiting for the import to finish (like automation test or synchronous import)
 					FAssetCompilingManager::Get().ProcessAsyncTasks();
 
@@ -132,7 +132,7 @@ void UE::Interchange::FTaskWaitAssetCompilation::DoTask(ENamedThreads::Type Curr
 
 void UE::Interchange::FTaskPipelinePostImport::DoTask(ENamedThreads::Type CurrentThread, const FGraphEventRef& MyCompletionGraphEvent)
 {
-	TRACE_CPUPROFILER_EVENT_SCOPE("UE::Interchange::FTaskPipelinePostImport::DoTask")
+	TRACE_CPUPROFILER_EVENT_SCOPE(UE::Interchange::FTaskPipelinePostImport::DoTask)
 #if INTERCHANGE_TRACE_ASYNCHRONOUS_TASK_ENABLED
 	INTERCHANGE_TRACE_ASYNCHRONOUS_TASK(PipelinePostImport)
 #endif
diff --git a/Engine/Source/Runtime/Interchange/Engine/Private/Tasks/InterchangeTaskTranslator.cpp b/Engine/Source/Runtime/Interchange/Engine/Private/Tasks/InterchangeTaskTranslator.cpp
--- a/Engine/Source/Runtime/Interchange/Engine/Private/Tasks/InterchangeTaskTranslator.cpp
+++ b/Engine/Source/Runtime/Interchange/Engine/Private/Tasks/InterchangeTaskTranslator.cpp
@@ -15,7 +15,7 @@
 
 void UE::Interchange::FTaskTranslator::DoTask(ENamedThreads::Type CurrentThread, const FGraphEventRef& MyCompletionGraphEvent)
 {
-	TRACE_CPUPROFILER_EVENT_SCOPE("UE::Interchange::FTaskTranslator::DoTask")
+	TRACE_CPUPROFILER_EVENT_SCOPE(UE::Interchange::FTaskTranslator::DoTask)
 #if INTERCHANGE_TRACE_ASYNCHRONOUS_TASK_ENABLED
 	INTERCHANGE_TRACE_ASYNCHRONOUS_TASK(Translator)
 #endif
