Fix some bad ensure

#rnx

diff --git a/Engine/Plugins/Interchange/Runtime/Source/Import/Private/Animation/InterchangeAnimSequenceFactory.cpp b/Engine/Plugins/Interchange/Runtime/Source/Import/Private/Animation/InterchangeAnimSequenceFactory.cpp
--- a/Engine/Plugins/Interchange/Runtime/Source/Import/Private/Animation/InterchangeAnimSequenceFactory.cpp
+++ b/Engine/Plugins/Interchange/Runtime/Source/Import/Private/Animation/InterchangeAnimSequenceFactory.cpp
@@ -1187,7 +1187,7 @@ void UInterchangeAnimSequenceFactory::SetupObject_GameThread(const FSetupObjectP
 bool UInterchangeAnimSequenceFactory::GetSourceFilenames(const UObject* Object, TArray<FString>& OutSourceFilenames) const
 {
 #if WITH_EDITORONLY_DATA
-	if (ensure(AnimSequence && (AnimSequence == Cast<UAnimSequence>(Object))))
+	if (AnimSequence && (AnimSequence == Cast<UAnimSequence>(Object)))
 	{
 		return UE::Interchange::FFactoryCommon::GetSourceFilenames(AnimSequence->AssetImportData.Get(), OutSourceFilenames);
 	}
@@ -1199,7 +1199,7 @@ bool UInterchangeAnimSequenceFactory::GetSourceFilenames(const UObject* Object,
 bool UInterchangeAnimSequenceFactory::SetSourceFilename(const UObject* Object, const FString& SourceFilename, int32 SourceIndex) const
 {
 #if WITH_EDITORONLY_DATA
-	if (ensure(AnimSequence && (AnimSequence == Cast<UAnimSequence>(Object))))
+	if (AnimSequence && (AnimSequence == Cast<UAnimSequence>(Object)))
 	{
 		return UE::Interchange::FFactoryCommon::SetSourceFilename(AnimSequence->AssetImportData.Get(), SourceFilename, SourceIndex);
 	}
