Adding GenerateHLSLExpression for MaterialX material expressions:
- Ramp4
- RampLeftRight
- RampTopBottom
- SplitLeftRight
- SplitTopBottom
- Swizzle
- Unpremult

Fixing swizzle Compile code, when one Channels character is not xyzw or rgba.

Exporting functions of the new HLSL Translator.

diff --git a/Engine/Plugins/Interchange/Runtime/Source/Import/Private/MaterialX/MaterialExpressions/MaterialExpressionRamp4.cpp b/Engine/Plugins/Interchange/Runtime/Source/Import/Private/MaterialX/MaterialExpressions/MaterialExpressionRamp4.cpp
--- a/Engine/Plugins/Interchange/Runtime/Source/Import/Private/MaterialX/MaterialExpressions/MaterialExpressionRamp4.cpp
+++ b/Engine/Plugins/Interchange/Runtime/Source/Import/Private/MaterialX/MaterialExpressions/MaterialExpressionRamp4.cpp
@@ -1,6 +1,8 @@
 // Copyright Epic Games, Inc. All Rights Reserved.
 #include "MaterialExpressionRamp4.h"
 #include "MaterialCompiler.h"
+#include "MaterialHLSLGenerator.h"
+#include "MaterialHLSLTree.h"
 
 #define LOCTEXT_NAMESPACE "MaterialExpressionMaterialXRamp4"
 
@@ -70,6 +72,34 @@ void UMaterialExpressionMaterialXRamp4::GetExpressionToolTip(TArray<FString>& Ou
 {
 	ConvertToMultilineToolTip(TEXT("A 4-corner bilinear value ramp."), 40, OutToolTip);
 }
+
+bool UMaterialExpressionMaterialXRamp4::GenerateHLSLExpression(FMaterialHLSLGenerator& Generator, UE::HLSLTree::FScope& Scope, int32 OutputIndex, UE::HLSLTree::FExpression const*& OutExpression) const
+{
+	using namespace UE::HLSLTree;
+
+	const FExpression* ExpressionA = A.AcquireHLSLExpression(Generator, Scope);
+	const FExpression* ExpressionB = B.AcquireHLSLExpression(Generator, Scope);
+	const FExpression* ExpressionC = C.AcquireHLSLExpression(Generator, Scope);
+	const FExpression* ExpressionD = D.AcquireHLSLExpression(Generator, Scope);
+	const FExpression* ExpressionCoordinates = Coordinates.AcquireHLSLExpressionOrExternalInput(Generator, Scope, Material::MakeInputTexCoord(ConstCoordinate));
+
+	if(!ExpressionA || !ExpressionB || !ExpressionC || !ExpressionD || !ExpressionCoordinates)
+	{
+		return false;
+	}
+
+	FTree& Tree = Generator.GetTree();
+
+	const FExpression* ExpressionCoordinatesClamped = Tree.NewSaturate(ExpressionCoordinates);
+	const FExpression* ExpressionCoordinatesU = Tree.NewSwizzle(FSwizzleParameters(0), ExpressionCoordinatesClamped);
+	const FExpression* ExpressionCoordinatesV = Tree.NewSwizzle(FSwizzleParameters(1), ExpressionCoordinatesClamped);
+
+	OutExpression = Tree.NewLerp(Tree.NewLerp(ExpressionC, ExpressionD, ExpressionCoordinatesU),
+								 Tree.NewLerp(ExpressionA, ExpressionB, ExpressionCoordinatesU),
+								 ExpressionCoordinatesV);
+
+	return true;
+}
 #endif
 
 #undef LOCTEXT_NAMESPACE 
\ No newline at end of file
diff --git a/Engine/Plugins/Interchange/Runtime/Source/Import/Private/MaterialX/MaterialExpressions/MaterialExpressionRamp4.h b/Engine/Plugins/Interchange/Runtime/Source/Import/Private/MaterialX/MaterialExpressions/MaterialExpressionRamp4.h
--- a/Engine/Plugins/Interchange/Runtime/Source/Import/Private/MaterialX/MaterialExpressions/MaterialExpressionRamp4.h
+++ b/Engine/Plugins/Interchange/Runtime/Source/Import/Private/MaterialX/MaterialExpressions/MaterialExpressionRamp4.h
@@ -42,6 +42,7 @@ class UMaterialExpressionMaterialXRamp4: public UMaterialExpression
 	virtual int32 Compile(class FMaterialCompiler* Compiler, int32 OutputIndex) override;
 	virtual void GetCaption(TArray<FString>& OutCaptions) const override;
 	virtual void GetExpressionToolTip(TArray<FString>& OutToolTip) override;
+	virtual bool GenerateHLSLExpression(FMaterialHLSLGenerator& Generator, UE::HLSLTree::FScope& Scope, int32 OutputIndex, UE::HLSLTree::FExpression const*& OutExpression) const override;
 #endif
 	//~ End UMaterialExpressionMaterialX Interface
 };
diff --git a/Engine/Plugins/Interchange/Runtime/Source/Import/Private/MaterialX/MaterialExpressions/MaterialExpressionRampLeftRight.cpp b/Engine/Plugins/Interchange/Runtime/Source/Import/Private/MaterialX/MaterialExpressions/MaterialExpressionRampLeftRight.cpp
--- a/Engine/Plugins/Interchange/Runtime/Source/Import/Private/MaterialX/MaterialExpressions/MaterialExpressionRampLeftRight.cpp
+++ b/Engine/Plugins/Interchange/Runtime/Source/Import/Private/MaterialX/MaterialExpressions/MaterialExpressionRampLeftRight.cpp
@@ -1,6 +1,8 @@
 // Copyright Epic Games, Inc. All Rights Reserved.
 #include "MaterialExpressionRampLeftRight.h"
 #include "MaterialCompiler.h"
+#include "MaterialHLSLGenerator.h"
+#include "MaterialHLSLTree.h"
 
 #include UE_INLINE_GENERATED_CPP_BY_NAME(MaterialExpressionRampLeftRight)
 
@@ -46,6 +48,28 @@ void UMaterialExpressionMaterialXRampLeftRight::GetCaption(TArray<FString>& OutC
 {
 	OutCaptions.Add(TEXT("MaterialX RampLR"));
 }
+
+bool UMaterialExpressionMaterialXRampLeftRight::GenerateHLSLExpression(FMaterialHLSLGenerator& Generator, UE::HLSLTree::FScope& Scope, int32 OutputIndex, UE::HLSLTree::FExpression const*& OutExpression) const
+{
+	using namespace UE::HLSLTree;
+
+	const FExpression* ExpressionA = A.AcquireHLSLExpression(Generator, Scope);
+	const FExpression* ExpressionB = B.AcquireHLSLExpression(Generator, Scope);
+	const FExpression* ExpressionCoordinates = Coordinates.AcquireHLSLExpressionOrExternalInput(Generator, Scope, Material::MakeInputTexCoord(ConstCoordinate));
+
+	if(!ExpressionA || !ExpressionB || !ExpressionCoordinates)
+	{
+		return false;
+	}
+
+	FTree& Tree = Generator.GetTree();
+
+	OutExpression = Tree.NewLerp(ExpressionA,
+								 ExpressionB,
+								 Tree.NewSaturate(Tree.NewSwizzle(FSwizzleParameters(0), ExpressionCoordinates)));
+
+	return false;
+}
 #endif
 
 #undef LOCTEXT_NAMESPACE 
\ No newline at end of file
diff --git a/Engine/Plugins/Interchange/Runtime/Source/Import/Private/MaterialX/MaterialExpressions/MaterialExpressionRampLeftRight.h b/Engine/Plugins/Interchange/Runtime/Source/Import/Private/MaterialX/MaterialExpressions/MaterialExpressionRampLeftRight.h
--- a/Engine/Plugins/Interchange/Runtime/Source/Import/Private/MaterialX/MaterialExpressions/MaterialExpressionRampLeftRight.h
+++ b/Engine/Plugins/Interchange/Runtime/Source/Import/Private/MaterialX/MaterialExpressions/MaterialExpressionRampLeftRight.h
@@ -31,6 +31,7 @@ class UMaterialExpressionMaterialXRampLeftRight : public UMaterialExpression
 #if WITH_EDITOR
 	virtual int32 Compile(class FMaterialCompiler* Compiler, int32 OutputIndex) override;
 	virtual void GetCaption(TArray<FString>& OutCaptions) const override;
+	virtual bool GenerateHLSLExpression(FMaterialHLSLGenerator& Generator, UE::HLSLTree::FScope& Scope, int32 OutputIndex, UE::HLSLTree::FExpression const*& OutExpression) const override;
 #endif
 	//~ End UMaterialExpressionMaterialX Interface
 };
diff --git a/Engine/Plugins/Interchange/Runtime/Source/Import/Private/MaterialX/MaterialExpressions/MaterialExpressionRampTopBottom.cpp b/Engine/Plugins/Interchange/Runtime/Source/Import/Private/MaterialX/MaterialExpressions/MaterialExpressionRampTopBottom.cpp
--- a/Engine/Plugins/Interchange/Runtime/Source/Import/Private/MaterialX/MaterialExpressions/MaterialExpressionRampTopBottom.cpp
+++ b/Engine/Plugins/Interchange/Runtime/Source/Import/Private/MaterialX/MaterialExpressions/MaterialExpressionRampTopBottom.cpp
@@ -1,6 +1,8 @@
 // Copyright Epic Games, Inc. All Rights Reserved.
 #include "MaterialExpressionRampTopBottom.h"
 #include "MaterialCompiler.h"
+#include "MaterialHLSLGenerator.h"
+#include "MaterialHLSLTree.h"
 
 #include UE_INLINE_GENERATED_CPP_BY_NAME(MaterialExpressionRampTopBottom)
 
@@ -46,6 +48,28 @@ void UMaterialExpressionMaterialXRampTopBottom::GetCaption(TArray<FString>& OutC
 {
 	OutCaptions.Add(TEXT("MaterialX RampTB"));
 }
+
+bool UMaterialExpressionMaterialXRampTopBottom::GenerateHLSLExpression(FMaterialHLSLGenerator& Generator, UE::HLSLTree::FScope& Scope, int32 OutputIndex, UE::HLSLTree::FExpression const*& OutExpression) const
+{
+	using namespace UE::HLSLTree;
+
+	const FExpression* ExpressionA = A.AcquireHLSLExpression(Generator, Scope);
+	const FExpression* ExpressionB = B.AcquireHLSLExpression(Generator, Scope);
+	const FExpression* ExpressionCoordinates = Coordinates.AcquireHLSLExpressionOrExternalInput(Generator, Scope, Material::MakeInputTexCoord(ConstCoordinate));
+
+	if(!ExpressionA || !ExpressionB || !ExpressionCoordinates)
+	{
+		return false;
+	}
+
+	FTree& Tree = Generator.GetTree();
+
+	OutExpression = Tree.NewLerp(ExpressionB,
+								 ExpressionA,
+								 Tree.NewSaturate(Tree.NewSwizzle(FSwizzleParameters(1), ExpressionCoordinates)));
+
+	return false;
+}
 #endif
 
 #undef LOCTEXT_NAMESPACE 
\ No newline at end of file
diff --git a/Engine/Plugins/Interchange/Runtime/Source/Import/Private/MaterialX/MaterialExpressions/MaterialExpressionRampTopBottom.h b/Engine/Plugins/Interchange/Runtime/Source/Import/Private/MaterialX/MaterialExpressions/MaterialExpressionRampTopBottom.h
--- a/Engine/Plugins/Interchange/Runtime/Source/Import/Private/MaterialX/MaterialExpressions/MaterialExpressionRampTopBottom.h
+++ b/Engine/Plugins/Interchange/Runtime/Source/Import/Private/MaterialX/MaterialExpressions/MaterialExpressionRampTopBottom.h
@@ -31,6 +31,7 @@ class UMaterialExpressionMaterialXRampTopBottom : public UMaterialExpression
 #if WITH_EDITOR
 	virtual int32 Compile(class FMaterialCompiler* Compiler, int32 OutputIndex) override;
 	virtual void GetCaption(TArray<FString>& OutCaptions) const override;
+	virtual bool GenerateHLSLExpression(FMaterialHLSLGenerator& Generator, UE::HLSLTree::FScope& Scope, int32 OutputIndex, UE::HLSLTree::FExpression const*& OutExpression) const override;
 #endif
 	//~ End UMaterialExpressionMaterialX Interface
 };
diff --git a/Engine/Plugins/Interchange/Runtime/Source/Import/Private/MaterialX/MaterialExpressions/MaterialExpressionSplitLeftRight.cpp b/Engine/Plugins/Interchange/Runtime/Source/Import/Private/MaterialX/MaterialExpressions/MaterialExpressionSplitLeftRight.cpp
--- a/Engine/Plugins/Interchange/Runtime/Source/Import/Private/MaterialX/MaterialExpressions/MaterialExpressionSplitLeftRight.cpp
+++ b/Engine/Plugins/Interchange/Runtime/Source/Import/Private/MaterialX/MaterialExpressions/MaterialExpressionSplitLeftRight.cpp
@@ -1,6 +1,10 @@
 // Copyright Epic Games, Inc. All Rights Reserved.
 #include "MaterialExpressionSplitLeftRight.h"
 #include "MaterialCompiler.h"
+#include "MaterialHLSLGenerator.h"
+#include "MaterialHLSLTree.h"
+#include "HLSLTree/HLSLTree.h"
+#include "HLSLTree/HLSLTreeCommon.h"
 
 #include UE_INLINE_GENERATED_CPP_BY_NAME(MaterialExpressionSplitLeftRight)
 
@@ -54,6 +58,42 @@ void UMaterialExpressionMaterialXSplitLeftRight::GetCaption(TArray<FString>& Out
 {
 	OutCaptions.Add(TEXT("MaterialX SplitLR"));
 }
+
+bool UMaterialExpressionMaterialXSplitLeftRight::GenerateHLSLExpression(FMaterialHLSLGenerator& Generator, UE::HLSLTree::FScope& Scope, int32 OutputIndex, UE::HLSLTree::FExpression const*& OutExpression) const
+{
+	using namespace UE::HLSLTree;
+
+	const FExpression* ExpressionA = A.AcquireHLSLExpression(Generator, Scope);
+	const FExpression* ExpressionB = B.AcquireHLSLExpression(Generator, Scope);
+	const FExpression* ExpressionCenter = Center.AcquireHLSLExpressionOrConstant(Generator, Scope, ConstCenter);
+	const FExpression* ExpressionCoordinates = Coordinates.AcquireHLSLExpressionOrExternalInput(Generator, Scope, Material::MakeInputTexCoord(ConstCoordinate));
+
+	if(!ExpressionA || !ExpressionB || !ExpressionCenter || !ExpressionCoordinates)
+	{
+		return false;
+	}
+
+	FTree& Tree = Generator.GetTree();
+
+	const FExpression* ExpressionTexCoordU = Tree.NewSwizzle(FSwizzleParameters(0), ExpressionCoordinates);
+
+	const FExpression* ExpressionAFWidth =
+		Tree.NewMul(
+			Tree.NewLength(
+				Tree.NewExpression<FExpressionAppend>(
+					Tree.NewExpression<FExpressionDerivative>(EDerivativeCoordinate::Ddx, ExpressionTexCoordU),
+					Tree.NewExpression<FExpressionDerivative>(EDerivativeCoordinate::Ddy, ExpressionTexCoordU))),
+			Tree.NewConstant(0.70710678118654757f));
+
+	OutExpression = Tree.NewLerp(ExpressionA,
+								 ExpressionB,
+								 Tree.NewSmoothStep(
+									 Tree.NewSub(ExpressionCenter, ExpressionAFWidth),
+									 Tree.NewAdd(ExpressionCenter, ExpressionAFWidth),
+									 ExpressionTexCoordU));
+
+	return true;
+}
 #endif
 
 #undef LOCTEXT_NAMESPACE 
\ No newline at end of file
diff --git a/Engine/Plugins/Interchange/Runtime/Source/Import/Private/MaterialX/MaterialExpressions/MaterialExpressionSplitLeftRight.h b/Engine/Plugins/Interchange/Runtime/Source/Import/Private/MaterialX/MaterialExpressions/MaterialExpressionSplitLeftRight.h
--- a/Engine/Plugins/Interchange/Runtime/Source/Import/Private/MaterialX/MaterialExpressions/MaterialExpressionSplitLeftRight.h
+++ b/Engine/Plugins/Interchange/Runtime/Source/Import/Private/MaterialX/MaterialExpressions/MaterialExpressionSplitLeftRight.h
@@ -38,6 +38,7 @@ class UMaterialExpressionMaterialXSplitLeftRight : public UMaterialExpression
 #if WITH_EDITOR
 	virtual int32 Compile(class FMaterialCompiler* Compiler, int32 OutputIndex) override;
 	virtual void GetCaption(TArray<FString>& OutCaptions) const override;
+	virtual bool GenerateHLSLExpression(FMaterialHLSLGenerator& Generator, UE::HLSLTree::FScope& Scope, int32 OutputIndex, UE::HLSLTree::FExpression const*& OutExpression) const override;
 #endif
 	//~ End UMaterialExpressionMaterialX Interface
 };
diff --git a/Engine/Plugins/Interchange/Runtime/Source/Import/Private/MaterialX/MaterialExpressions/MaterialExpressionSplitTopBottom.cpp b/Engine/Plugins/Interchange/Runtime/Source/Import/Private/MaterialX/MaterialExpressions/MaterialExpressionSplitTopBottom.cpp
--- a/Engine/Plugins/Interchange/Runtime/Source/Import/Private/MaterialX/MaterialExpressions/MaterialExpressionSplitTopBottom.cpp
+++ b/Engine/Plugins/Interchange/Runtime/Source/Import/Private/MaterialX/MaterialExpressions/MaterialExpressionSplitTopBottom.cpp
@@ -1,6 +1,10 @@
 // Copyright Epic Games, Inc. All Rights Reserved.
 #include "MaterialExpressionSplitTopBottom.h"
 #include "MaterialCompiler.h"
+#include "MaterialHLSLGenerator.h"
+#include "MaterialHLSLTree.h"
+#include "HLSLTree/HLSLTree.h"
+#include "HLSLTree/HLSLTreeCommon.h"
 
 #include UE_INLINE_GENERATED_CPP_BY_NAME(MaterialExpressionSplitTopBottom)
 
@@ -54,6 +58,42 @@ void UMaterialExpressionMaterialXSplitTopBottom::GetCaption(TArray<FString>& Out
 {
 	OutCaptions.Add(TEXT("MaterialX SplitTB"));
 }
+
+bool UMaterialExpressionMaterialXSplitTopBottom::GenerateHLSLExpression(FMaterialHLSLGenerator& Generator, UE::HLSLTree::FScope& Scope, int32 OutputIndex, UE::HLSLTree::FExpression const*& OutExpression) const
+{
+	using namespace UE::HLSLTree;
+
+	const FExpression* ExpressionA = A.AcquireHLSLExpression(Generator, Scope);
+	const FExpression* ExpressionB = B.AcquireHLSLExpression(Generator, Scope);
+	const FExpression* ExpressionCenter = Center.AcquireHLSLExpressionOrConstant(Generator, Scope, ConstCenter);
+	const FExpression* ExpressionCoordinates = Coordinates.AcquireHLSLExpressionOrExternalInput(Generator, Scope, Material::MakeInputTexCoord(ConstCoordinate));
+
+	if(!ExpressionA || !ExpressionB || !ExpressionCenter || !ExpressionCoordinates)
+	{
+		return false;
+	}
+
+	FTree& Tree = Generator.GetTree();
+
+	const FExpression* ExpressionTexCoordV = Tree.NewSwizzle(FSwizzleParameters(1), ExpressionCoordinates);
+
+	const FExpression* ExpressionAFWidth =
+		Tree.NewMul(
+			Tree.NewLength(
+				Tree.NewExpression<FExpressionAppend>(
+					Tree.NewExpression<FExpressionDerivative>(EDerivativeCoordinate::Ddx, ExpressionTexCoordV),
+					Tree.NewExpression<FExpressionDerivative>(EDerivativeCoordinate::Ddy, ExpressionTexCoordV))),
+			Tree.NewConstant(0.70710678118654757f));
+
+	OutExpression = Tree.NewLerp(ExpressionB,
+								 ExpressionA,
+								 Tree.NewSmoothStep(
+									 Tree.NewSub(ExpressionCenter, ExpressionAFWidth),
+									 Tree.NewAdd(ExpressionCenter, ExpressionAFWidth),
+									 ExpressionTexCoordV));
+
+	return true;
+}
 #endif
 
 #undef LOCTEXT_NAMESPACE 
\ No newline at end of file
diff --git a/Engine/Plugins/Interchange/Runtime/Source/Import/Private/MaterialX/MaterialExpressions/MaterialExpressionSplitTopBottom.h b/Engine/Plugins/Interchange/Runtime/Source/Import/Private/MaterialX/MaterialExpressions/MaterialExpressionSplitTopBottom.h
--- a/Engine/Plugins/Interchange/Runtime/Source/Import/Private/MaterialX/MaterialExpressions/MaterialExpressionSplitTopBottom.h
+++ b/Engine/Plugins/Interchange/Runtime/Source/Import/Private/MaterialX/MaterialExpressions/MaterialExpressionSplitTopBottom.h
@@ -38,6 +38,7 @@ class UMaterialExpressionMaterialXSplitTopBottom : public UMaterialExpression
 #if WITH_EDITOR
 	virtual int32 Compile(class FMaterialCompiler* Compiler, int32 OutputIndex) override;
 	virtual void GetCaption(TArray<FString>& OutCaptions) const override;
+	virtual bool GenerateHLSLExpression(FMaterialHLSLGenerator& Generator, UE::HLSLTree::FScope& Scope, int32 OutputIndex, UE::HLSLTree::FExpression const*& OutExpression) const override;
 #endif
 	//~ End UMaterialExpressionMaterialX Interface
 };
diff --git a/Engine/Plugins/Interchange/Runtime/Source/Import/Private/MaterialX/MaterialExpressions/MaterialExpressionSwizzle.cpp b/Engine/Plugins/Interchange/Runtime/Source/Import/Private/MaterialX/MaterialExpressions/MaterialExpressionSwizzle.cpp
--- a/Engine/Plugins/Interchange/Runtime/Source/Import/Private/MaterialX/MaterialExpressions/MaterialExpressionSwizzle.cpp
+++ b/Engine/Plugins/Interchange/Runtime/Source/Import/Private/MaterialX/MaterialExpressions/MaterialExpressionSwizzle.cpp
@@ -2,6 +2,8 @@
 
 #include "MaterialX/MaterialExpressions/MaterialExpressionSwizzle.h"
 #include "MaterialCompiler.h"
+#include "MaterialHLSLGenerator.h"
+#include "HLSLTree/HLSLTree.h"
 
 #include UE_INLINE_GENERATED_CPP_BY_NAME(MaterialExpressionSwizzle)
 
@@ -48,7 +50,6 @@ int32 UMaterialExpressionMaterialXSwizzle::Compile(FMaterialCompiler* Compiler,
 	//Don't allow mixing xyzw and rgba or invalid channels
 	bool bHasRGBA = false;
 	bool bHasXYZW = false;
-
 	bool SwizzleX[4]{ false, false, false, false };
 	bool SwizzleY[4]{ false, false, false, false };
 	bool SwizzleZ[4]{ false, false, false, false };
@@ -56,20 +57,23 @@ int32 UMaterialExpressionMaterialXSwizzle::Compile(FMaterialCompiler* Compiler,
 	
 	for(int32 i = 0; i < Channels.Len(); ++i)
 	{
+		bool bIsRGBAorXYZW = false;
 		const TCHAR Channel = Channels[i];
 		if(Channel == TEXT('r') || Channel == TEXT('g') || Channel == TEXT('b') || Channel == TEXT('a'))
 		{
 			bHasRGBA = true;
+			bIsRGBAorXYZW = true;
 		}
 		else if(Channel == TEXT('x') || Channel == TEXT('y') || Channel == TEXT('z') || Channel == TEXT('w'))
 		{
 			bHasXYZW = true;
+			bIsRGBAorXYZW = true;
 		}
 		if(bHasRGBA && bHasXYZW)
 		{
 			return Compiler->Errorf(TEXT("Cannot mix rgba and xyzw channels"));
 		}
-		if(!bHasRGBA && !bHasXYZW)
+		if(!bIsRGBAorXYZW)
 		{
 			return Compiler->Errorf(TEXT("%c is not recognized as a valid channel"), Channel);
 		}
@@ -141,6 +145,88 @@ void UMaterialExpressionMaterialXSwizzle::GetExpressionToolTip(TArray<FString>&
 								   "stream of the specified type. Individual channels may be replicated or omitted, and the output "
 								   "stream may have a different number of channels than the input."), 40, OutToolTip);
 }
+
+bool UMaterialExpressionMaterialXSwizzle::GenerateHLSLExpression(FMaterialHLSLGenerator& Generator, UE::HLSLTree::FScope& Scope, int32 OutputIndex, UE::HLSLTree::FExpression const*& OutExpression) const
+{
+	using namespace UE::HLSLTree;
+
+	const FExpression* ExpressionInput = Input.AcquireHLSLExpression(Generator, Scope);
+
+	if(!ExpressionInput)
+	{
+		return false;
+	}
+
+	if(Channels.IsEmpty())
+	{
+		OutExpression = ExpressionInput;
+		return true;
+	}
+
+	if(Channels.Len() > 4)
+	{
+		return Generator.Errorf(TEXT("Too many channels"));
+	}
+
+	//Don't allow mixing xyzw and rgba or invalid channels
+	bool bHasRGBA = false;
+	bool bHasXYZW = false;
+	FSwizzleParameters SwizzleParameters;
+	SwizzleParameters.NumComponents = Channels.Len();
+	SwizzleParameters.bHasSwizzle = true;
+	for(int32 i = 0; i < Channels.Len(); ++i)
+	{
+		bool bIsRGBAorXYZW = false;
+		const TCHAR Channel = Channels[i];
+		if(Channel == TEXT('r') || Channel == TEXT('g') || Channel == TEXT('b') || Channel == TEXT('a'))
+		{
+			bHasRGBA = true;
+			bIsRGBAorXYZW = true;
+		}
+		else if(Channel == TEXT('x') || Channel == TEXT('y') || Channel == TEXT('z') || Channel == TEXT('w'))
+		{
+			bHasXYZW = true;
+			bIsRGBAorXYZW = true;
+		}
+
+		if(bHasRGBA && bHasXYZW)
+		{
+			return Generator.Errorf(TEXT("Cannot mix rgba and xyzw channels"));
+		}
+
+		if(!bIsRGBAorXYZW)
+		{
+			return Generator.Errorf(TEXT("%c is not recognized as a valid channel"), Channel);
+		}
+
+		switch(Channel)
+		{
+		case TEXT('x'):
+		case TEXT('r'):
+			SwizzleParameters.SwizzleComponentIndex[i] = 0;
+			break;
+
+		case TEXT('y'):
+		case TEXT('g'):
+			SwizzleParameters.SwizzleComponentIndex[i] = 1;
+			break;
+
+		case TEXT('z'):
+		case TEXT('b'):
+			SwizzleParameters.SwizzleComponentIndex[i] = 2;
+			break;
+
+		case TEXT('w'):
+		case TEXT('a'):
+			SwizzleParameters.SwizzleComponentIndex[i] = 3;
+			break;
+		}
+	}
+
+	OutExpression = Generator.GetTree().NewSwizzle(SwizzleParameters, ExpressionInput);
+
+	return true;
+}
 #endif
 
 #undef LOCTEXT_NAMESPACE 
\ No newline at end of file
diff --git a/Engine/Plugins/Interchange/Runtime/Source/Import/Private/MaterialX/MaterialExpressions/MaterialExpressionUnpremult.cpp b/Engine/Plugins/Interchange/Runtime/Source/Import/Private/MaterialX/MaterialExpressions/MaterialExpressionUnpremult.cpp
--- a/Engine/Plugins/Interchange/Runtime/Source/Import/Private/MaterialX/MaterialExpressions/MaterialExpressionUnpremult.cpp
+++ b/Engine/Plugins/Interchange/Runtime/Source/Import/Private/MaterialX/MaterialExpressions/MaterialExpressionUnpremult.cpp
@@ -1,6 +1,9 @@
 // Copyright Epic Games, Inc. All Rights Reserved.
 #include "MaterialExpressionUnpremult.h"
 #include "MaterialCompiler.h"
+#include "MaterialHLSLGenerator.h"
+#include "HLSLTree/HLSLTree.h"
+#include "HLSLTree/HLSLTreeCommon.h"
 
 #include UE_INLINE_GENERATED_CPP_BY_NAME(MaterialExpressionUnpremult)
 
@@ -44,6 +47,32 @@ void UMaterialExpressionMaterialXUnpremult::GetCaption(TArray<FString>& OutCapti
 {
 	OutCaptions.Add(TEXT("MaterialX Unpremult"));
 }
+
+bool UMaterialExpressionMaterialXUnpremult::GenerateHLSLExpression(FMaterialHLSLGenerator& Generator, UE::HLSLTree::FScope& Scope, int32 OutputIndex, UE::HLSLTree::FExpression const*& OutExpression) const
+{
+	using namespace UE::HLSLTree;
+
+	const FExpression* ExpressionInput = Input.AcquireHLSLExpression(Generator, Scope);
+
+	if(!ExpressionInput)
+	{
+		return false;
+	}
+
+	FTree& Tree = Generator.GetTree();
+
+	const FExpression* ExpressionAlpha = Tree.NewSwizzle(FSwizzleParameters(3), ExpressionInput);
+
+	OutExpression = Generator.GenerateBranch(Scope,
+											 Tree.NewLess(Tree.NewAbs(ExpressionAlpha), Tree.NewConstant(UE_SMALL_NUMBER)),
+											 ExpressionInput,
+											 Tree.NewExpression<FExpressionAppend>(
+												 Tree.NewDiv(Tree.NewSwizzle(FSwizzleParameters(0, 1, 2), ExpressionInput),
+															 ExpressionAlpha),
+												 ExpressionAlpha));
+
+	return true;
+}
 #endif
 
 #undef LOCTEXT_NAMESPACE 
\ No newline at end of file
diff --git a/Engine/Plugins/Interchange/Runtime/Source/Import/Private/MaterialX/MaterialExpressions/MaterialExpressionUnpremult.h b/Engine/Plugins/Interchange/Runtime/Source/Import/Private/MaterialX/MaterialExpressions/MaterialExpressionUnpremult.h
--- a/Engine/Plugins/Interchange/Runtime/Source/Import/Private/MaterialX/MaterialExpressions/MaterialExpressionUnpremult.h
+++ b/Engine/Plugins/Interchange/Runtime/Source/Import/Private/MaterialX/MaterialExpressions/MaterialExpressionUnpremult.h
@@ -23,6 +23,7 @@ class UMaterialExpressionMaterialXUnpremult : public UMaterialExpression
 #if WITH_EDITOR
 	virtual int32 Compile(class FMaterialCompiler* Compiler, int32 OutputIndex) override;
 	virtual void GetCaption(TArray<FString>& OutCaptions) const override;
+	virtual bool GenerateHLSLExpression(FMaterialHLSLGenerator& Generator, UE::HLSLTree::FScope& Scope, int32 OutputIndex, UE::HLSLTree::FExpression const*& OutExpression) const override;
 #endif
 	//~ End UMaterialExpressionMaterialX Interface
 };
diff --git a/Engine/Plugins/Interchange/Runtime/Source/Import/Public/MaterialX/MaterialExpressions/MaterialExpressionSwizzle.h b/Engine/Plugins/Interchange/Runtime/Source/Import/Public/MaterialX/MaterialExpressions/MaterialExpressionSwizzle.h
--- a/Engine/Plugins/Interchange/Runtime/Source/Import/Public/MaterialX/MaterialExpressions/MaterialExpressionSwizzle.h
+++ b/Engine/Plugins/Interchange/Runtime/Source/Import/Public/MaterialX/MaterialExpressions/MaterialExpressionSwizzle.h
@@ -22,6 +22,7 @@ class UMaterialExpressionMaterialXSwizzle : public UMaterialExpression
 	virtual int32 Compile(class FMaterialCompiler* Compiler, int32 OutputIndex) override;
 	virtual void GetCaption(TArray<FString>& OutCaptions) const override;
 	void GetExpressionToolTip(TArray<FString>& OutToolTip) override;
+	virtual bool GenerateHLSLExpression(FMaterialHLSLGenerator& Generator, UE::HLSLTree::FScope& Scope, int32 OutputIndex, UE::HLSLTree::FExpression const*& OutExpression) const override;
 #endif
 	//~ End UMaterialExpressionMaterialX Interface
 };
diff --git a/Engine/Source/Runtime/Engine/Public/HLSLTree/HLSLTree.h b/Engine/Source/Runtime/Engine/Public/HLSLTree/HLSLTree.h
--- a/Engine/Source/Runtime/Engine/Public/HLSLTree/HLSLTree.h
+++ b/Engine/Source/Runtime/Engine/Public/HLSLTree/HLSLTree.h
@@ -636,7 +636,7 @@ public:
 	ENGINE_API const FExpression* NewSwizzle(const FSwizzleParameters& Params, const FExpression* Input);
 	ENGINE_API const FExpression* NewUnaryOp(EOperation Op, const FExpression* Input);
 	ENGINE_API const FExpression* NewBinaryOp(EOperation Op, const FExpression* Lhs, const FExpression* Rhs);
-	const FExpression* NewTernaryOp(EOperation Op, const FExpression* Input0, const FExpression* Input1, const FExpression* Input2);
+	ENGINE_API const FExpression* NewTernaryOp(EOperation Op, const FExpression* Input0, const FExpression* Input1, const FExpression* Input2);
 
 	const FExpression* NewAbs(const FExpression* Input) { return NewUnaryOp(EOperation::Abs, Input); }
 	const FExpression* NewNeg(const FExpression* Input) { return NewUnaryOp(EOperation::Neg, Input); }
diff --git a/Engine/Source/Runtime/Engine/Public/HLSLTree/HLSLTreeCommon.h b/Engine/Source/Runtime/Engine/Public/HLSLTree/HLSLTreeCommon.h
--- a/Engine/Source/Runtime/Engine/Public/HLSLTree/HLSLTreeCommon.h
+++ b/Engine/Source/Runtime/Engine/Public/HLSLTree/HLSLTreeCommon.h
@@ -162,7 +162,7 @@ public:
 	virtual bool EmitValueObject(FEmitContext& Context, FEmitScope& Scope, const FName& ObjectTypeName, void* OutObjectBase) const override;
 };
 
-class FExpressionDerivative : public FExpression
+class ENGINE_API FExpressionDerivative : public FExpression
 {
 public:
 	FExpressionDerivative(EDerivativeCoordinate InCoord, const FExpression* InInput) : Input(InInput), Coord(InCoord) {}
