Fix fbx animation issue

#rnx

diff --git a/Engine/Plugins/Interchange/Runtime/Source/Parsers/Fbx/Private/FbxScene.cpp b/Engine/Plugins/Interchange/Runtime/Source/Parsers/Fbx/Private/FbxScene.cpp
--- a/Engine/Plugins/Interchange/Runtime/Source/Parsers/Fbx/Private/FbxScene.cpp
+++ b/Engine/Plugins/Interchange/Runtime/Source/Parsers/Fbx/Private/FbxScene.cpp
@@ -85,22 +85,68 @@ namespace UE
 				CreateAssetNodeReference(Parser, UnrealSceneNode, NodeAttribute, NodeContainer, UInterchangeLightNode::StaticAssetTypeName());
 			}
 
-			bool DoesTheParentHierarchyContainJoints(FbxNode* Node)
+			bool DoesChildrenHierarchyContainJoints(FbxNode* Node, TFunction<bool(FbxNode*)> IsNodeAjoint)
 			{
 				if (!Node)
 				{
 					return false;
 				}
-				int32 AttributeCount = Node->GetNodeAttributeCount();
-				for (int32 AttributeIndex = 0; AttributeIndex < AttributeCount; ++AttributeIndex)
+				if (IsNodeAjoint(Node))
 				{
-					FbxNodeAttribute* NodeAttribute = Node->GetNodeAttributeByIndex(AttributeIndex);
-					if (NodeAttribute->GetAttributeType() == FbxNodeAttribute::eSkeleton)
+					return true;
+				}
+				int32 ChildCount = Node->GetChildCount();
+				for (int32 ChildIndex = 0; ChildIndex < ChildCount; ++ChildIndex)
+				{
+					FbxNode* ChildNode = Node->GetChild(ChildIndex);
+					if (DoesChildrenHierarchyContainJoints(ChildNode, IsNodeAjoint))
+					{
+						return true;
+					}
+				}
+				return false;
+			}
+
+			bool DoesTheParentOrChildrenHierarchyContainJoints(FbxNode* Node)
+			{
+				if (!Node)
+				{
+					return false;
+				}
+				auto IsNodeAjoint = [](FbxNode* NodeToTest)
+					{
+						int32 AttributeCount = NodeToTest->GetNodeAttributeCount();
+						for (int32 AttributeIndex = 0; AttributeIndex < AttributeCount; ++AttributeIndex)
+						{
+							FbxNodeAttribute* NodeAttribute = NodeToTest->GetNodeAttributeByIndex(AttributeIndex);
+							if (NodeAttribute->GetAttributeType() == FbxNodeAttribute::eSkeleton)
+							{
+								return true;
+							}
+						}
+						return false;
+					};
+				
+				if (IsNodeAjoint(Node))
+				{
+					return true;
+				}
+
+				FbxNode* ParentNode = Node->GetParent();
+				while (ParentNode)
+				{
+					if (IsNodeAjoint(ParentNode))
 					{
 						return true;
 					}
+					ParentNode = ParentNode->GetParent();
+				}
+
+				if (DoesChildrenHierarchyContainJoints(Node, IsNodeAjoint))
+				{
+					return true;
 				}
-				return DoesTheParentHierarchyContainJoints(Node->GetParent());
+				return false;
 			}
 
 			void FFbxScene::AddHierarchyRecursively(UInterchangeSceneNode* UnrealParentNode
@@ -113,7 +159,7 @@ namespace UE
 				constexpr bool bResetCache = false;
 				FString NodeName = Parser.GetFbxHelper()->GetFbxObjectName(Node);
 				FString NodeUniqueID = Parser.GetFbxHelper()->GetFbxNodeHierarchyName(Node);
-
+				const bool bIsRootNode = Node == SDKScene->GetRootNode();
 				UInterchangeSceneNode* UnrealNode = CreateTransformNode(NodeContainer, NodeName, NodeUniqueID);
 				check(UnrealNode);
 				if (UnrealParentNode)
@@ -237,12 +283,12 @@ namespace UE
 
 						case FbxNodeAttribute::eNull:
 						{
-							if (!DoesTheParentHierarchyContainJoints(Node->GetParent()))
+							if (!DoesTheParentOrChildrenHierarchyContainJoints(Node))
 							{
-								//eNull node not under any joint will be set has a transform
-								UnrealNode->AddSpecializedType(FSceneNodeStaticData::GetTransformSpecializeTypeString());
+								//eNull node not in a hierarchy containing any joint will not be set has joint
 								break;
 							}
+							UnrealNode->AddSpecializedType(FSceneNodeStaticData::GetTransformSpecializeTypeString());
 						}
 						//No break since the eNull act has a skeleton if possible
 						case FbxNodeAttribute::eSkeleton:
@@ -291,11 +337,21 @@ namespace UE
 					}
 				}
 
-				//If there is no attribute, make sure to threat the node like a joint if it's in the ForcejointNodes array
-				if (!bIsNodeContainJointAttribute && ForceJointNodes.Contains(Node))
+				if (!bIsNodeContainJointAttribute)
 				{
-					ApplySkeletonAttribute();
+					//Make sure to treat the node like a joint if it's in the ForcejointNodes array
+					if (ForceJointNodes.Contains(Node))
+					{
+						UnrealNode->AddSpecializedType(FSceneNodeStaticData::GetTransformSpecializeTypeString());
+						ApplySkeletonAttribute();
+					}
+					else if (!bIsRootNode && DoesTheParentOrChildrenHierarchyContainJoints(Node))
+					{
+						UnrealNode->AddSpecializedType(FSceneNodeStaticData::GetTransformSpecializeTypeString());
+						ApplySkeletonAttribute();
+					}
 				}
+				
 
 				//Add all custom Attributes for the node
 				FbxProperty Property = Node->GetFirstProperty();
@@ -618,6 +674,7 @@ namespace UE
 				, TArray<FbxNode*>& ForceJointNodes)
 			{
 				FString NodeUniqueID = Parser.GetFbxHelper()->GetFbxNodeHierarchyName(Node);
+				const bool bIsRootNode = Node == SDKScene->GetRootNode();
 				if (UInterchangeSceneNode* UnrealNode = const_cast<UInterchangeSceneNode*>(Cast< UInterchangeSceneNode >(NodeContainer.GetNode(NodeUniqueID))))
 				{
 					bool HasSkeletonAttribute = false;
@@ -655,7 +712,7 @@ namespace UE
 						switch (NodeAttribute->GetAttributeType())
 						{
 						case FbxNodeAttribute::eNull:
-							if (!DoesTheParentHierarchyContainJoints(Node->GetParent()))
+							if (!DoesTheParentOrChildrenHierarchyContainJoints(Node))
 							{
 								//eNull node not under any joint are not joint
 								break;
@@ -668,10 +725,19 @@ namespace UE
 							break;
 						}
 					}
-					//If there is no attribute, make sure to threat the node like a joint if it's in the ForcejointNodes array
-					if (!bIsNodeContainJointAttribute && ForceJointNodes.Contains(Node))
+
+					if (!bIsNodeContainJointAttribute)
 					{
-						ApplySkeletonAttribute();
+						//Make sure to threat the node like a joint if it's in the ForcejointNodes array
+						if (ForceJointNodes.Contains(Node))
+						{
+							ApplySkeletonAttribute();
+						}
+						else if (!bIsRootNode && DoesTheParentOrChildrenHierarchyContainJoints(Node))
+						{
+							ApplySkeletonAttribute();
+						}
+
 					}
 
 					if (!HasSkeletonAttribute)
