Additional fixes for usage of TRACE_CPUPROFILER_EVENT_SCOPE macros in Interchange plugin.

* TRACE_CPUPROFILER_EVENT_SCOPE(ScopeName) --> to be used with a plain text as scope name (not a string!). If the parameter is a string ("abc" or TEXT("abc")), the quotes (and the TEXT prefix) will be included in the actual timer name.
* TRACE_CPUPROFILER_EVENT_SCOPE_STR("scope name") --> to be used with a static string (const ANSICHAR* or const TCHAR*); ex: when the scope name needs to include a space char
* TRACE_CPUPROFILER_EVENT_SCOPE_TEXT(*MyScopeName) --> to be used with a dynamic string (const ANSICHAR* or const TCHAR*)
See also the comments for these macros in ProfilingDebugging\CpuProfilerTrace.h.

diff --git a/Engine/Plugins/Interchange/Runtime/Source/Import/Private/Animation/InterchangeAnimSequenceFactory.cpp b/Engine/Plugins/Interchange/Runtime/Source/Import/Private/Animation/InterchangeAnimSequenceFactory.cpp
--- a/Engine/Plugins/Interchange/Runtime/Source/Import/Private/Animation/InterchangeAnimSequenceFactory.cpp
+++ b/Engine/Plugins/Interchange/Runtime/Source/Import/Private/Animation/InterchangeAnimSequenceFactory.cpp
@@ -706,7 +706,7 @@ UClass* UInterchangeAnimSequenceFactory::GetFactoryClass() const
 
 UInterchangeFactoryBase::FImportAssetResult UInterchangeAnimSequenceFactory::BeginImportAsset_GameThread(const FImportAssetObjectParams& Arguments)
 {
-	TRACE_CPUPROFILER_EVENT_SCOPE("UInterchangeAnimSequenceFactory::BeginImportAsset_GameThread");
+	TRACE_CPUPROFILER_EVENT_SCOPE(UInterchangeAnimSequenceFactory::BeginImportAsset_GameThread);
 
 	UInterchangeFactoryBase::FImportAssetResult ImportAssetResult;
 #if !WITH_EDITOR || !WITH_EDITORONLY_DATA
@@ -833,7 +833,7 @@ UInterchangeFactoryBase::FImportAssetResult UInterchangeAnimSequenceFactory::Beg
 
 UInterchangeFactoryBase::FImportAssetResult UInterchangeAnimSequenceFactory::ImportAsset_Async(const FImportAssetObjectParams& Arguments)
 {
-	TRACE_CPUPROFILER_EVENT_SCOPE("UInterchangeAnimSequenceFactory::ImportAsset_Async");
+	TRACE_CPUPROFILER_EVENT_SCOPE(UInterchangeAnimSequenceFactory::ImportAsset_Async);
 
 	UInterchangeFactoryBase::FImportAssetResult ImportAssetResult;
 #if !WITH_EDITOR || !WITH_EDITORONLY_DATA
@@ -1009,7 +1009,7 @@ UInterchangeFactoryBase::FImportAssetResult UInterchangeAnimSequenceFactory::Imp
 
 UInterchangeFactoryBase::FImportAssetResult UInterchangeAnimSequenceFactory::EndImportAsset_GameThread(const FImportAssetObjectParams& Arguments)
 {
-	TRACE_CPUPROFILER_EVENT_SCOPE("UInterchangeAnimSequenceFactory::ImportAsset_Async");
+	TRACE_CPUPROFILER_EVENT_SCOPE(UInterchangeAnimSequenceFactory::ImportAsset_Async);
 
 	UInterchangeFactoryBase::FImportAssetResult ImportAssetResult;
 #if !WITH_EDITOR || !WITH_EDITORONLY_DATA
@@ -1164,7 +1164,7 @@ UInterchangeFactoryBase::FImportAssetResult UInterchangeAnimSequenceFactory::End
 /* This function is call in the completion task on the main thread, use it to call main thread post creation step for your assets*/
 void UInterchangeAnimSequenceFactory::SetupObject_GameThread(const FSetupObjectParams& Arguments)
 {
-	TRACE_CPUPROFILER_EVENT_SCOPE("UInterchangeAnimSequenceFactory::SetupObject_GameThread");
+	TRACE_CPUPROFILER_EVENT_SCOPE(UInterchangeAnimSequenceFactory::SetupObject_GameThread);
 
 	check(IsInGameThread());
 	Super::SetupObject_GameThread(Arguments);
diff --git a/Engine/Plugins/Interchange/Runtime/Source/Import/Private/Animation/InterchangeAnimationTrackSetFactory.cpp b/Engine/Plugins/Interchange/Runtime/Source/Import/Private/Animation/InterchangeAnimationTrackSetFactory.cpp
--- a/Engine/Plugins/Interchange/Runtime/Source/Import/Private/Animation/InterchangeAnimationTrackSetFactory.cpp
+++ b/Engine/Plugins/Interchange/Runtime/Source/Import/Private/Animation/InterchangeAnimationTrackSetFactory.cpp
@@ -662,7 +662,7 @@ UClass* UInterchangeAnimationTrackSetFactory::GetFactoryClass() const
 
 UInterchangeFactoryBase::FImportAssetResult UInterchangeAnimationTrackSetFactory::BeginImportAsset_GameThread(const FImportAssetObjectParams& Arguments)
 {
-	TRACE_CPUPROFILER_EVENT_SCOPE("UInterchangeAnimationTrackSetFactory::BeginImportAsset_GameThread");
+	TRACE_CPUPROFILER_EVENT_SCOPE(UInterchangeAnimationTrackSetFactory::BeginImportAsset_GameThread);
 
 	FImportAssetResult ImportAssetResult;
 #if !WITH_EDITOR || !WITH_EDITORONLY_DATA
@@ -747,7 +747,7 @@ UInterchangeFactoryBase::FImportAssetResult UInterchangeAnimationTrackSetFactory
 
 UObject* UInterchangeAnimationTrackSetFactory::ImportObjectSourceData(const FImportAssetObjectParams& Arguments)
 {
-	TRACE_CPUPROFILER_EVENT_SCOPE("UInterchangeAnimationTrackSetFactory::ImportObjectSourceData");
+	TRACE_CPUPROFILER_EVENT_SCOPE(UInterchangeAnimationTrackSetFactory::ImportObjectSourceData);
 
 #if !WITH_EDITOR || !WITH_EDITORONLY_DATA
 	// TODO: Can we import ULevelSequence at runtime
@@ -828,7 +828,7 @@ UObject* UInterchangeAnimationTrackSetFactory::ImportObjectSourceData(const FImp
 /* This function is call in the completion task on the main thread, use it to call main thread post creation step for your assets*/
 void UInterchangeAnimationTrackSetFactory::SetupObject_GameThread(const FSetupObjectParams& Arguments)
 {
-	TRACE_CPUPROFILER_EVENT_SCOPE("UInterchangeAnimationTrackSetFactory::SetupObject_GameThread");
+	TRACE_CPUPROFILER_EVENT_SCOPE(UInterchangeAnimationTrackSetFactory::SetupObject_GameThread);
 
 	check(IsInGameThread());
 	Super::SetupObject_GameThread(Arguments);
diff --git a/Engine/Plugins/Interchange/Runtime/Source/Import/Private/Material/InterchangeMaterialFactory.cpp b/Engine/Plugins/Interchange/Runtime/Source/Import/Private/Material/InterchangeMaterialFactory.cpp
--- a/Engine/Plugins/Interchange/Runtime/Source/Import/Private/Material/InterchangeMaterialFactory.cpp
+++ b/Engine/Plugins/Interchange/Runtime/Source/Import/Private/Material/InterchangeMaterialFactory.cpp
@@ -686,7 +686,7 @@ UClass* UInterchangeMaterialFactory::GetFactoryClass() const
 
 UInterchangeFactoryBase::FImportAssetResult UInterchangeMaterialFactory::BeginImportAsset_GameThread(const FImportAssetObjectParams& Arguments)
 {
-	TRACE_CPUPROFILER_EVENT_SCOPE("UInterchangeMaterialFactory::BeginImportAsset_GameThread");
+	TRACE_CPUPROFILER_EVENT_SCOPE(UInterchangeMaterialFactory::BeginImportAsset_GameThread);
 	FImportAssetResult ImportAssetResult;
 	UObject* Material = nullptr;
 
@@ -891,7 +891,7 @@ UInterchangeFactoryBase::FImportAssetResult UInterchangeMaterialFactory::BeginIm
 
 UInterchangeFactoryBase::FImportAssetResult UInterchangeMaterialFactory::ImportAsset_Async(const FImportAssetObjectParams& Arguments)
 {
-	TRACE_CPUPROFILER_EVENT_SCOPE("UInterchangeMaterialFactory::ImportAsset_Async");
+	TRACE_CPUPROFILER_EVENT_SCOPE(UInterchangeMaterialFactory::ImportAsset_Async);
 	FImportAssetResult ImportAssetResult;
 	ImportAssetResult.bIsFactorySkipAsset = bSkipImport;
 
@@ -971,7 +971,7 @@ UInterchangeFactoryBase::FImportAssetResult UInterchangeMaterialFactory::ImportA
 /* This function is call in the completion task on the main thread, use it to call main thread post creation step for your assets*/
 void UInterchangeMaterialFactory::SetupObject_GameThread(const FSetupObjectParams& Arguments)
 {
-	TRACE_CPUPROFILER_EVENT_SCOPE("UInterchangeMaterialFactory::SetupObject_GameThread");
+	TRACE_CPUPROFILER_EVENT_SCOPE(UInterchangeMaterialFactory::SetupObject_GameThread);
 	check(IsInGameThread());
 	Super::SetupObject_GameThread(Arguments);
 
@@ -1039,7 +1039,7 @@ bool UInterchangeMaterialFactory::SetSourceFilename(const UObject* Object, const
 #if WITH_EDITOR
 void UInterchangeMaterialFactory::SetupMaterial(UMaterial* Material, const FImportAssetObjectParams& Arguments, const UInterchangeBaseMaterialFactoryNode* BaseMaterialFactoryNode)
 {
-	TRACE_CPUPROFILER_EVENT_SCOPE("UInterchangeMaterialFactory::SetupMaterial");
+	TRACE_CPUPROFILER_EVENT_SCOPE(UInterchangeMaterialFactory::SetupMaterial);
 
 	using namespace UE::Interchange::Materials;
 	using namespace UE::Interchange::MaterialFactory::Internal;
@@ -1543,7 +1543,7 @@ void UInterchangeMaterialFactory::SetupMaterial(UMaterial* Material, const FImpo
 
 void UInterchangeMaterialFactory::SetupMaterialInstance(UMaterialInstance& MaterialInstance, const UInterchangeBaseNodeContainer& NodeContainer, const UInterchangeMaterialInstanceFactoryNode& FactoryNode, bool bResetInstance)
 {
-	TRACE_CPUPROFILER_EVENT_SCOPE("UInterchangeMaterialFactory::SetupMaterialInstance");
+	TRACE_CPUPROFILER_EVENT_SCOPE(UInterchangeMaterialFactory::SetupMaterialInstance);
 	using namespace UE::Interchange::MaterialFactory::Internal;
 
 	if (bResetInstance)
@@ -1735,7 +1735,7 @@ UClass* UInterchangeMaterialFunctionFactory::GetFactoryClass() const
 
 UInterchangeFactoryBase::FImportAssetResult UInterchangeMaterialFunctionFactory::BeginImportAsset_GameThread(const FImportAssetObjectParams& Arguments)
 {
-	TRACE_CPUPROFILER_EVENT_SCOPE("UInterchangeMaterialFunctionFactory::BeginImportAsset_GameThread");
+	TRACE_CPUPROFILER_EVENT_SCOPE(UInterchangeMaterialFunctionFactory::BeginImportAsset_GameThread);
 	FImportAssetResult ImportAssetResult;
 	UObject* Material = nullptr;
 
@@ -1821,7 +1821,7 @@ UInterchangeFactoryBase::FImportAssetResult UInterchangeMaterialFunctionFactory:
 
 UInterchangeFactoryBase::FImportAssetResult UInterchangeMaterialFunctionFactory::ImportAsset_Async(const FImportAssetObjectParams& Arguments)
 {
-	TRACE_CPUPROFILER_EVENT_SCOPE("UInterchangeMaterialFunctionFactory::ImportAsset_Async");
+	TRACE_CPUPROFILER_EVENT_SCOPE(UInterchangeMaterialFunctionFactory::ImportAsset_Async);
 	FImportAssetResult ImportAssetResult;
 	ImportAssetResult.bIsFactorySkipAsset = bSkipImport;
 
@@ -1877,7 +1877,7 @@ UInterchangeFactoryBase::FImportAssetResult UInterchangeMaterialFunctionFactory:
 /* This function is call in the completion task on the main thread, use it to call main thread post creation step for your assets*/
 void UInterchangeMaterialFunctionFactory::SetupObject_GameThread(const FSetupObjectParams& Arguments)
 {
-	TRACE_CPUPROFILER_EVENT_SCOPE("UInterchangeMaterialFunctionFactory::SetupObject_GameThread");
+	TRACE_CPUPROFILER_EVENT_SCOPE(UInterchangeMaterialFunctionFactory::SetupObject_GameThread);
 	Super::SetupObject_GameThread(Arguments);
 
 	if (bSkipImport)
@@ -1908,7 +1908,7 @@ void UInterchangeMaterialFunctionFactory::SetupObject_GameThread(const FSetupObj
 #if WITH_EDITOR
 void UInterchangeMaterialFunctionFactory::SetupMaterial(UMaterialFunction* MaterialFunction, const FImportAssetObjectParams& Arguments, const UInterchangeMaterialFunctionFactoryNode* MaterialFunctionFactoryNode)
 {
-	TRACE_CPUPROFILER_EVENT_SCOPE("UInterchangeMaterialFunctionFactory::SetupMaterial");
+	TRACE_CPUPROFILER_EVENT_SCOPE(UInterchangeMaterialFunctionFactory::SetupMaterial);
 	using namespace UE::Interchange::MaterialFactory::Internal;
 	using namespace UE::Interchange::Materials;
 
diff --git a/Engine/Plugins/Interchange/Runtime/Source/Import/Private/Mesh/InterchangePhysicsAssetFactory.cpp b/Engine/Plugins/Interchange/Runtime/Source/Import/Private/Mesh/InterchangePhysicsAssetFactory.cpp
--- a/Engine/Plugins/Interchange/Runtime/Source/Import/Private/Mesh/InterchangePhysicsAssetFactory.cpp
+++ b/Engine/Plugins/Interchange/Runtime/Source/Import/Private/Mesh/InterchangePhysicsAssetFactory.cpp
@@ -18,7 +18,7 @@ UClass* UInterchangePhysicsAssetFactory::GetFactoryClass() const
 
 UInterchangeFactoryBase::FImportAssetResult UInterchangePhysicsAssetFactory::BeginImportAsset_GameThread(const FImportAssetObjectParams& Arguments)
 {
-	TRACE_CPUPROFILER_EVENT_SCOPE("UInterchangePhysicsAssetFactory::BeginImportAsset_GameThread");
+	TRACE_CPUPROFILER_EVENT_SCOPE(UInterchangePhysicsAssetFactory::BeginImportAsset_GameThread);
 	FImportAssetResult ImportAssetResult;
 	UObject* PhysicsAsset = nullptr;
 
@@ -71,7 +71,7 @@ UInterchangeFactoryBase::FImportAssetResult UInterchangePhysicsAssetFactory::Beg
 
 UInterchangeFactoryBase::FImportAssetResult UInterchangePhysicsAssetFactory::ImportAsset_Async(const FImportAssetObjectParams& Arguments)
 {
-	TRACE_CPUPROFILER_EVENT_SCOPE("UInterchangePhysicsAssetFactory::ImportAsset_Async");
+	TRACE_CPUPROFILER_EVENT_SCOPE(UInterchangePhysicsAssetFactory::ImportAsset_Async);
 	FImportAssetResult ImportAssetResult;
 #if !WITH_EDITORONLY_DATA
 
diff --git a/Engine/Plugins/Interchange/Runtime/Source/Import/Private/Mesh/InterchangeSkeletonFactory.cpp b/Engine/Plugins/Interchange/Runtime/Source/Import/Private/Mesh/InterchangeSkeletonFactory.cpp
--- a/Engine/Plugins/Interchange/Runtime/Source/Import/Private/Mesh/InterchangeSkeletonFactory.cpp
+++ b/Engine/Plugins/Interchange/Runtime/Source/Import/Private/Mesh/InterchangeSkeletonFactory.cpp
@@ -19,7 +19,7 @@ UClass* UInterchangeSkeletonFactory::GetFactoryClass() const
 
 UInterchangeFactoryBase::FImportAssetResult UInterchangeSkeletonFactory::BeginImportAsset_GameThread(const FImportAssetObjectParams& Arguments)
 {
-	TRACE_CPUPROFILER_EVENT_SCOPE("UInterchangeSkeletonFactory::BeginImportAsset_GameThread");
+	TRACE_CPUPROFILER_EVENT_SCOPE(UInterchangeSkeletonFactory::BeginImportAsset_GameThread);
 
 	FImportAssetResult ImportAssetResult;
 
diff --git a/Engine/Plugins/Interchange/Runtime/Source/Import/Private/Mesh/InterchangeStaticMeshFactory.cpp b/Engine/Plugins/Interchange/Runtime/Source/Import/Private/Mesh/InterchangeStaticMeshFactory.cpp
--- a/Engine/Plugins/Interchange/Runtime/Source/Import/Private/Mesh/InterchangeStaticMeshFactory.cpp
+++ b/Engine/Plugins/Interchange/Runtime/Source/Import/Private/Mesh/InterchangeStaticMeshFactory.cpp
@@ -147,7 +147,7 @@ namespace UE::Interchange::Private::StaticMesh
 
 UInterchangeFactoryBase::FImportAssetResult UInterchangeStaticMeshFactory::BeginImportAsset_GameThread(const FImportAssetObjectParams& Arguments)
 {
-	TRACE_CPUPROFILER_EVENT_SCOPE("UInterchangeStaticMeshFactory::BeginImportAsset_GameThread");
+	TRACE_CPUPROFILER_EVENT_SCOPE(UInterchangeStaticMeshFactory::BeginImportAsset_GameThread);
 
 	FImportAssetResult ImportAssetResult;
 	UStaticMesh* StaticMesh = nullptr;
@@ -217,7 +217,7 @@ UInterchangeFactoryBase::FImportAssetResult UInterchangeStaticMeshFactory::Begin
 
 UInterchangeFactoryBase::FImportAssetResult UInterchangeStaticMeshFactory::ImportAsset_Async(const FImportAssetObjectParams& Arguments)
 {
-	TRACE_CPUPROFILER_EVENT_SCOPE("UInterchangeStaticMeshFactory::ImportAsset_Async");
+	TRACE_CPUPROFILER_EVENT_SCOPE(UInterchangeStaticMeshFactory::ImportAsset_Async);
 
 	FImportAssetResult ImportAssetResult;
 	if (!Arguments.AssetNode || !Arguments.AssetNode->GetObjectClass()->IsChildOf(GetFactoryClass()))
@@ -483,7 +483,7 @@ UInterchangeFactoryBase::FImportAssetResult UInterchangeStaticMeshFactory::Impor
 
 UInterchangeFactoryBase::FImportAssetResult UInterchangeStaticMeshFactory::EndImportAsset_GameThread(const FImportAssetObjectParams& Arguments)
 {
-	TRACE_CPUPROFILER_EVENT_SCOPE("UInterchangeStaticMeshFactory::EndImportAsset_GameThread");
+	TRACE_CPUPROFILER_EVENT_SCOPE(UInterchangeStaticMeshFactory::EndImportAsset_GameThread);
 
 	FImportAssetResult ImportAssetResult;
 	if (!Arguments.AssetNode || !Arguments.AssetNode->GetObjectClass()->IsChildOf(GetFactoryClass()))
@@ -755,7 +755,7 @@ void UInterchangeStaticMeshFactory::SetupSourceModelsSettings(UStaticMesh& Stati
 /* This function is call in the completion task on the main thread, use it to call main thread post creation step for your assets */
 void UInterchangeStaticMeshFactory::SetupObject_GameThread(const FSetupObjectParams& Arguments)
 {
-	TRACE_CPUPROFILER_EVENT_SCOPE("UInterchangeStaticMeshFactory::SetupObject_GameThread");
+	TRACE_CPUPROFILER_EVENT_SCOPE(UInterchangeStaticMeshFactory::SetupObject_GameThread);
 
 	check(IsInGameThread());
 	Super::SetupObject_GameThread(Arguments);
@@ -778,7 +778,7 @@ void UInterchangeStaticMeshFactory::SetupObject_GameThread(const FSetupObjectPar
 
 TArray<UInterchangeStaticMeshFactory::FMeshPayload> UInterchangeStaticMeshFactory::GetMeshPayloads(const FImportAssetObjectParams& Arguments, const TArray<FString>& MeshUids) const
 {
-	TRACE_CPUPROFILER_EVENT_SCOPE("UInterchangeStaticMeshFactory::GetMeshPayloads");
+	TRACE_CPUPROFILER_EVENT_SCOPE(UInterchangeStaticMeshFactory::GetMeshPayloads);
 
 	TArray<FMeshPayload> Payloads;
 	Payloads.Reserve(MeshUids.Num());
diff --git a/Engine/Plugins/Interchange/Runtime/Source/Import/Private/Scene/InterchangeSceneImportAssetFactory.cpp b/Engine/Plugins/Interchange/Runtime/Source/Import/Private/Scene/InterchangeSceneImportAssetFactory.cpp
--- a/Engine/Plugins/Interchange/Runtime/Source/Import/Private/Scene/InterchangeSceneImportAssetFactory.cpp
+++ b/Engine/Plugins/Interchange/Runtime/Source/Import/Private/Scene/InterchangeSceneImportAssetFactory.cpp
@@ -85,7 +85,7 @@ bool UInterchangeSceneImportAssetFactory::SetSourceFilename(const UObject* Objec
 
 UInterchangeFactoryBase::FImportAssetResult UInterchangeSceneImportAssetFactory::BeginImportAsset_GameThread(const FImportAssetObjectParams& Arguments)
 {
-	TRACE_CPUPROFILER_EVENT_SCOPE("UInterchangeSceneImportAssetFactory::BeginImportAsset_GameThread");
+	TRACE_CPUPROFILER_EVENT_SCOPE(UInterchangeSceneImportAssetFactory::BeginImportAsset_GameThread);
 	using namespace UE::Interchange::Private::InterchangeSceneImportAssetFactory;
 
 	UClass* TargetClass = GetFactoryClass();
@@ -115,7 +115,7 @@ UInterchangeFactoryBase::FImportAssetResult UInterchangeSceneImportAssetFactory:
 
 void UInterchangeSceneImportAssetFactory::SetupObject_GameThread(const FSetupObjectParams& Arguments)
 {
-	TRACE_CPUPROFILER_EVENT_SCOPE("UInterchangeSceneImportAssetFactory::SetupObject_GameThread");
+	TRACE_CPUPROFILER_EVENT_SCOPE(UInterchangeSceneImportAssetFactory::SetupObject_GameThread);
 	check(IsInGameThread());
 	Super::SetupObject_GameThread(Arguments);
 
@@ -137,7 +137,7 @@ void UInterchangeSceneImportAssetFactory::SetupObject_GameThread(const FSetupObj
 
 void UInterchangeSceneImportAssetFactory::FinalizeObject_GameThread(const FSetupObjectParams& Arguments)
 {
-	TRACE_CPUPROFILER_EVENT_SCOPE("UInterchangeSceneImportAssetFactory::FinalizeObject_GameThread");
+	TRACE_CPUPROFILER_EVENT_SCOPE(UInterchangeSceneImportAssetFactory::FinalizeObject_GameThread);
 	check(IsInGameThread());
 	Super::FinalizeObject_GameThread(Arguments);
 
diff --git a/Engine/Plugins/Interchange/Runtime/Source/Import/Private/Scene/InterchangeSceneVariantSetsFactory.cpp b/Engine/Plugins/Interchange/Runtime/Source/Import/Private/Scene/InterchangeSceneVariantSetsFactory.cpp
--- a/Engine/Plugins/Interchange/Runtime/Source/Import/Private/Scene/InterchangeSceneVariantSetsFactory.cpp
+++ b/Engine/Plugins/Interchange/Runtime/Source/Import/Private/Scene/InterchangeSceneVariantSetsFactory.cpp
@@ -296,7 +296,7 @@ UClass* UInterchangeSceneVariantSetsFactory::GetFactoryClass() const
 
 UInterchangeFactoryBase::FImportAssetResult UInterchangeSceneVariantSetsFactory::BeginImportAsset_GameThread(const FImportAssetObjectParams& Arguments)
 {
-	TRACE_CPUPROFILER_EVENT_SCOPE("UInterchangeSceneVariantSetsFactory::BeginImportAsset_GameThread");
+	TRACE_CPUPROFILER_EVENT_SCOPE(UInterchangeSceneVariantSetsFactory::BeginImportAsset_GameThread);
 	FImportAssetResult ImportAssetResult;
 #if !WITH_EDITOR || !WITH_EDITORONLY_DATA
 
@@ -374,7 +374,7 @@ UInterchangeFactoryBase::FImportAssetResult UInterchangeSceneVariantSetsFactory:
 
 UObject* UInterchangeSceneVariantSetsFactory::ImportObjectSourceData(const FImportAssetObjectParams& Arguments, ULevelVariantSets* LevelVariantSets)
 {
-	TRACE_CPUPROFILER_EVENT_SCOPE("UInterchangeSceneVariantSetsFactory::ImportObjectSourceData");
+	TRACE_CPUPROFILER_EVENT_SCOPE(UInterchangeSceneVariantSetsFactory::ImportObjectSourceData);
 #if !WITH_EDITOR || !WITH_EDITORONLY_DATA
 	// TODO: Can we import ULevelVariantSets at runtime
 	UE_LOG(LogInterchangeImport, Error, TEXT("Cannot import LevelVariantSets asset in runtime, this is an editor only feature."));
@@ -422,7 +422,7 @@ UObject* UInterchangeSceneVariantSetsFactory::ImportObjectSourceData(const FImpo
 
 void UInterchangeSceneVariantSetsFactory::SetupObject_GameThread(const FSetupObjectParams& Arguments)
 {
-	TRACE_CPUPROFILER_EVENT_SCOPE("UInterchangeSceneVariantSetsFactory::SetupObject_GameThread");
+	TRACE_CPUPROFILER_EVENT_SCOPE(UInterchangeSceneVariantSetsFactory::SetupObject_GameThread);
 	check(IsInGameThread());
 	Super::SetupObject_GameThread(Arguments);
 
diff --git a/Engine/Plugins/Interchange/Runtime/Source/Import/Private/Texture/InterchangeTextureFactory.cpp b/Engine/Plugins/Interchange/Runtime/Source/Import/Private/Texture/InterchangeTextureFactory.cpp
--- a/Engine/Plugins/Interchange/Runtime/Source/Import/Private/Texture/InterchangeTextureFactory.cpp
+++ b/Engine/Plugins/Interchange/Runtime/Source/Import/Private/Texture/InterchangeTextureFactory.cpp
@@ -1531,7 +1531,7 @@ UClass* UInterchangeTextureFactory::GetFactoryClass() const
 
 UInterchangeFactoryBase::FImportAssetResult UInterchangeTextureFactory::BeginImportAsset_GameThread(const FImportAssetObjectParams& Arguments)
 {
-	TRACE_CPUPROFILER_EVENT_SCOPE("UInterchangeTextureFactory::BeginImportAsset_GameThread");
+	TRACE_CPUPROFILER_EVENT_SCOPE(UInterchangeTextureFactory::BeginImportAsset_GameThread);
 	using namespace  UE::Interchange::Private::InterchangeTextureFactory;
 	FImportAssetResult ImportAssetResult;
 	UTexture* Texture = nullptr;
diff --git a/Engine/Plugins/Interchange/Runtime/Source/Pipelines/Private/InterchangeGenericSkeletalMeshPipeline.cpp b/Engine/Plugins/Interchange/Runtime/Source/Pipelines/Private/InterchangeGenericSkeletalMeshPipeline.cpp
--- a/Engine/Plugins/Interchange/Runtime/Source/Pipelines/Private/InterchangeGenericSkeletalMeshPipeline.cpp
+++ b/Engine/Plugins/Interchange/Runtime/Source/Pipelines/Private/InterchangeGenericSkeletalMeshPipeline.cpp
@@ -566,7 +566,7 @@ void UInterchangeGenericMeshPipeline::PostImportSkeletalMesh(UObject* CreatedAss
 
 void UInterchangeGenericMeshPipeline::PostImportPhysicsAssetImport(UObject* CreatedAsset, const UInterchangeFactoryBaseNode* FactoryNode)
 {
-	TRACE_CPUPROFILER_EVENT_SCOPE("UInterchangeGenericMeshPipeline::PostImportPhysicsAssetImport");
+	TRACE_CPUPROFILER_EVENT_SCOPE(UInterchangeGenericMeshPipeline::PostImportPhysicsAssetImport);
 #if WITH_EDITOR
 	if (!bCreatePhysicsAsset || !BaseNodeContainer)
 	{
@@ -593,7 +593,7 @@ void UInterchangeGenericMeshPipeline::PostImportPhysicsAssetImport(UObject* Crea
 					{
 						auto CreateFromSkeletalMeshLambda = [CreatedPhysicsAsset, SkeletalMesh]()
 						{
-							TRACE_CPUPROFILER_EVENT_SCOPE("UInterchangeGenericMeshPipeline::PostImportPhysicsAssetImport::CreateFromSkeletalMeshLambda");
+							TRACE_CPUPROFILER_EVENT_SCOPE(UInterchangeGenericMeshPipeline::PostImportPhysicsAssetImport::CreateFromSkeletalMeshLambda);
 							FPhysAssetCreateParams NewBodyData;
 							FText CreationErrorMessage;
 							constexpr bool bSetToMesh = true;
diff --git a/Engine/Source/Runtime/Interchange/Core/Public/InterchangeSourceData.h b/Engine/Source/Runtime/Interchange/Core/Public/InterchangeSourceData.h
--- a/Engine/Source/Runtime/Interchange/Core/Public/InterchangeSourceData.h
+++ b/Engine/Source/Runtime/Interchange/Core/Public/InterchangeSourceData.h
@@ -43,7 +43,7 @@ public:
 	/** Return the hash of the content point by Filename. */
 	FORCEINLINE TOptional<FMD5Hash> GetFileContentHash() const
 	{
-		TRACE_CPUPROFILER_EVENT_SCOPE("UInterchangeSourceData::GetFileContentHash");
+		TRACE_CPUPROFILER_EVENT_SCOPE(UInterchangeSourceData::GetFileContentHash);
 		if(!FileContentHashCache.IsSet())
 		{
 			ComputeFileContentHashCache();
