Rename InterchangeAnimationTrackSetFactory to InterchangeLevelSequenceFactory and InterchangeAnimationTrackSetFactoryNode to InterchangeLevelSequenceFactoryNode.

diff --git a/Engine/Plugins/Experimental/Enterprise/DatasmithInterchange/Source/DatasmithInterchange/Private/InterchangeDatasmithLevelPipeline.cpp b/Engine/Plugins/Experimental/Enterprise/DatasmithInterchange/Source/DatasmithInterchange/Private/InterchangeDatasmithLevelPipeline.cpp
--- a/Engine/Plugins/Experimental/Enterprise/DatasmithInterchange/Source/DatasmithInterchange/Private/InterchangeDatasmithLevelPipeline.cpp
+++ b/Engine/Plugins/Experimental/Enterprise/DatasmithInterchange/Source/DatasmithInterchange/Private/InterchangeDatasmithLevelPipeline.cpp
@@ -9,7 +9,7 @@
 #include "InterchangeDecalActorFactoryNode.h"
 #include "InterchangeDecalNode.h"
 
-#include "InterchangeAnimationTrackSetFactoryNode.h"
+#include "InterchangeLevelSequenceFactoryNode.h"
 #include "InterchangeAnimationTrackSetNode.h"
 #include "InterchangeMaterialFactoryNode.h"
 #include "InterchangeMeshActorFactoryNode.h"
@@ -104,7 +104,7 @@ void UInterchangeDatasmithLevelPipeline::ExecutePostImportPipeline(const UInterc
 	}
 
 	// If a ULevelSequence is referencing a ADatasmithAreaLightActor, update its mobility property
-	if (const UInterchangeAnimationTrackSetFactoryNode* FactoryNode = Cast<UInterchangeAnimationTrackSetFactoryNode>(NodeContainer->GetFactoryNode(FactoryNodeKey)))
+	if (const UInterchangeLevelSequenceFactoryNode* FactoryNode = Cast<UInterchangeLevelSequenceFactoryNode>(NodeContainer->GetFactoryNode(FactoryNodeKey)))
 	{
 		FSoftObjectPath FactoryNodeReferenceObject;
 		FactoryNode->GetCustomReferenceObject(FactoryNodeReferenceObject);
diff --git a/Engine/Plugins/Experimental/Enterprise/DatasmithInterchange/Source/DatasmithInterchange/Private/InterchangeDatasmithPipeline.cpp b/Engine/Plugins/Experimental/Enterprise/DatasmithInterchange/Source/DatasmithInterchange/Private/InterchangeDatasmithPipeline.cpp
--- a/Engine/Plugins/Experimental/Enterprise/DatasmithInterchange/Source/DatasmithInterchange/Private/InterchangeDatasmithPipeline.cpp
+++ b/Engine/Plugins/Experimental/Enterprise/DatasmithInterchange/Source/DatasmithInterchange/Private/InterchangeDatasmithPipeline.cpp
@@ -12,7 +12,7 @@
 #include "InterchangeDatasmithTexturePipeline.h"
 #include "InterchangeDatasmithUtils.h"
 
-#include "InterchangeAnimationTrackSetFactoryNode.h"
+#include "InterchangeLevelSequenceFactoryNode.h"
 #include "InterchangeGenericAnimationPipeline.h"
 #include "InterchangeGenericScenesPipeline.h"
 #include "InterchangeMaterialFactoryNode.h"
@@ -196,7 +196,7 @@ void UInterchangeDatasmithPipeline::ExecutePipeline(UInterchangeBaseNodeContaine
 	}
 
 	// LevelSequences
-	for (UInterchangeAnimationTrackSetFactoryNode* AnimationTrackSetFactoryNode : NodeUtils::GetNodes<UInterchangeAnimationTrackSetFactoryNode>(BaseNodeContainer))
+	for (UInterchangeLevelSequenceFactoryNode* AnimationTrackSetFactoryNode : NodeUtils::GetNodes<UInterchangeLevelSequenceFactoryNode>(BaseNodeContainer))
 	{
 		DependenciesUids.Add(AnimationTrackSetFactoryNode->GetUniqueID());
 		AnimationTrackSetFactoryNode->SetCustomSubPath(FPaths::Combine(PackageSubPath, "Animations"));
diff --git a/Engine/Plugins/Interchange/Runtime/Config/BaseInterchange.ini b/Engine/Plugins/Interchange/Runtime/Config/BaseInterchange.ini
--- a/Engine/Plugins/Interchange/Runtime/Config/BaseInterchange.ini
+++ b/Engine/Plugins/Interchange/Runtime/Config/BaseInterchange.ini
@@ -36,6 +36,10 @@
 +EnumRedirects=(OldName="ETextureSampleBlurFilter", NewName="EMaterialXTextureSampleBlurFilter")
 +EnumRedirects=(OldName="ELuminanceMode", NewName="EMaterialXLuminanceMode")
 
+; Rename InterchangeAnimationTrackSetFactory to InterchangeLevelSequenceFactory
++ClassRedirects=(OldName="/Script/InterchangeImport.InterchangeAnimationTrackSetFactory", NewName="/Script/InterchangeImport.InterchangeLevelSequenceFactory")
++ClassRedirects=(OldName="/Script/InterchangeFactoryNodes.InterchangeAnimationTrackSetFactoryNode", NewName="/Script/InterchangeFactoryNodes.InterchangeLevelSequenceFactoryNode")
+
 ; Add proper suffix to Interchange predefined materials
 +PackageRedirects=(OldName="/Interchange/Materials/PBRSurfaceMaterial.PBRSurfaceMaterial", NewName="/Interchange/Materials/PBRSurfaceMaterial_MR.PBRSurfaceMaterial_MR")
 +PackageRedirects=(OldName="/Interchange/Materials/ClearCoatMaterial.ClearCoatMaterial", NewName="/Interchange/Materials/ClearCoatMaterial_MR.ClearCoatMaterial_MR")
diff --git a/Engine/Plugins/Interchange/Runtime/Source/FactoryNodes/Private/InterchangeAnimationTrackSetFactoryNode.cpp b/Engine/Plugins/Interchange/Runtime/Source/FactoryNodes/Private/InterchangeAnimationTrackSetFactoryNode.cpp
deleted file mode 100644
diff --git a/Engine/Plugins/Interchange/Runtime/Source/FactoryNodes/Private/InterchangeLevelSequenceFactoryNode.cpp b/Engine/Plugins/Interchange/Runtime/Source/FactoryNodes/Private/InterchangeLevelSequenceFactoryNode.cpp
new file mode 100644
--- /dev/null
+++ b/Engine/Plugins/Interchange/Runtime/Source/FactoryNodes/Private/InterchangeLevelSequenceFactoryNode.cpp
@@ -0,0 +1,50 @@
+// Copyright Epic Games, Inc. All Rights Reserved.
+
+#include "InterchangeLevelSequenceFactoryNode.h"
+
+#include "LevelSequence.h"
+
+UInterchangeLevelSequenceFactoryNode::UInterchangeLevelSequenceFactoryNode()
+{
+	CustomAnimationTrackUids.Initialize(Attributes, TEXT("AnimationTrackUids"));
+}
+
+UClass* UInterchangeLevelSequenceFactoryNode::GetObjectClass() const
+{
+	return ULevelSequence::StaticClass();
+}
+
+int32 UInterchangeLevelSequenceFactoryNode::GetCustomAnimationTrackUidCount() const
+{
+	return CustomAnimationTrackUids.GetCount();
+}
+
+void UInterchangeLevelSequenceFactoryNode::GetCustomAnimationTrackUids(TArray<FString>& OutAnimationTrackUids) const
+{
+	CustomAnimationTrackUids.GetItems(OutAnimationTrackUids);
+}
+
+void UInterchangeLevelSequenceFactoryNode::GetCustomAnimationTrackUid(const int32 Index, FString& OutAnimationTrackUid) const
+{
+	CustomAnimationTrackUids.GetItem(Index, OutAnimationTrackUid);
+}
+
+bool UInterchangeLevelSequenceFactoryNode::AddCustomAnimationTrackUid(const FString& AnimationTrackUid)
+{
+	return CustomAnimationTrackUids.AddItem(AnimationTrackUid);
+}
+
+bool UInterchangeLevelSequenceFactoryNode::RemoveCustomAnimationTrackUid(const FString& AnimationTrackUid)
+{
+	return CustomAnimationTrackUids.RemoveItem(AnimationTrackUid);
+}
+
+bool UInterchangeLevelSequenceFactoryNode::GetCustomFrameRate(float& AttributeValue) const
+{
+	IMPLEMENT_NODE_ATTRIBUTE_GETTER(FrameRate, float);
+}
+
+bool UInterchangeLevelSequenceFactoryNode::SetCustomFrameRate(const float& AttributeValue)
+{
+	IMPLEMENT_NODE_ATTRIBUTE_SETTER_NODELEGATE(FrameRate, float);
+}
diff --git a/Engine/Plugins/Interchange/Runtime/Source/FactoryNodes/Public/InterchangeAnimationTrackSetFactoryNode.h b/Engine/Plugins/Interchange/Runtime/Source/FactoryNodes/Public/InterchangeAnimationTrackSetFactoryNode.h
deleted file mode 100644
diff --git a/Engine/Plugins/Interchange/Runtime/Source/FactoryNodes/Public/InterchangeLevelSequenceFactoryNode.h b/Engine/Plugins/Interchange/Runtime/Source/FactoryNodes/Public/InterchangeLevelSequenceFactoryNode.h
new file mode 100644
--- /dev/null
+++ b/Engine/Plugins/Interchange/Runtime/Source/FactoryNodes/Public/InterchangeLevelSequenceFactoryNode.h
@@ -0,0 +1,82 @@
+// Copyright Epic Games, Inc. All Rights Reserved.
+
+#pragma once
+
+#include "CoreMinimal.h"
+#include "Nodes/InterchangeFactoryBaseNode.h"
+
+#include "InterchangeLevelSequenceFactoryNode.generated.h"
+
+UCLASS(BlueprintType)
+class INTERCHANGEFACTORYNODES_API UInterchangeLevelSequenceFactoryNode : public UInterchangeFactoryBaseNode
+{
+	GENERATED_BODY()
+
+	UInterchangeLevelSequenceFactoryNode();
+
+public:
+	static FStringView StaticAssetTypeName()
+	{
+		return TEXT("LevelSequenceFactory");
+	}
+
+	/**
+	 * Return the node type name of the class, we use this when reporting errors
+	 */
+	virtual FString GetTypeName() const override
+	{
+		const FString TypeName = TEXT("LevelSequenceFactoryNode");
+		return TypeName;
+	}
+
+	/** Get the class this node want to create */
+	UFUNCTION(BlueprintCallable, Category = "Interchange | Node | LevelSequenceFactory")
+	virtual class UClass* GetObjectClass() const override;
+
+	/**
+	 * This function allow to retrieve the number of track dependencies for this object.
+	 */
+	UFUNCTION(BlueprintCallable, Category = "Interchange | Node | LevelSequenceFactory")
+	int32 GetCustomAnimationTrackUidCount() const;
+
+	/**
+	 * This function allow to retrieve the track dependency for this object.
+	 */
+	UFUNCTION(BlueprintCallable, Category = "Interchange | Node | LevelSequenceFactory")
+	void GetCustomAnimationTrackUids(TArray<FString>& OutAnimationTrackUids) const;
+
+	/**
+	 * This function allow to retrieve one track dependency for this object.
+	 */
+	UFUNCTION(BlueprintCallable, Category = "Interchange | Node | LevelSequenceFactory")
+	void GetCustomAnimationTrackUid(const int32 Index, FString& OutAnimationTrackUid) const;
+
+	/**
+	 * Add one track dependency to this object.
+	 */
+	UFUNCTION(BlueprintCallable, Category = "Interchange | Node | LevelSequenceFactory")
+	bool AddCustomAnimationTrackUid(const FString& AnimationTrackUid);
+
+	/**
+	 * Remove one track dependency from this object.
+	 */
+	UFUNCTION(BlueprintCallable, Category = "Interchange | Node | LevelSequenceFactory")
+	bool RemoveCustomAnimationTrackUid(const FString& AnimationTrackUid);
+
+	/**
+	 * Set the frame rate for the animations in the level sequence.
+	 */
+	UFUNCTION(BlueprintCallable, Category = "Interchange | Node | LevelSequenceFactory")
+	bool SetCustomFrameRate(const float& AttributeValue);
+
+	/**
+	 * Get the frame rate for the animations in the level sequence.
+	 */
+	UFUNCTION(BlueprintCallable, Category = "Interchange | Node | LevelSequenceFactory")
+	bool GetCustomFrameRate(float& AttributeValue) const;
+
+private:
+	const UE::Interchange::FAttributeKey Macro_CustomFrameRateKey = UE::Interchange::FAttributeKey(TEXT("FrameRate"));
+
+	UE::Interchange::TArrayAttributeHelper<FString> CustomAnimationTrackUids;
+};
diff --git a/Engine/Plugins/Interchange/Runtime/Source/Import/Private/Animation/InterchangeAnimationTrackSetFactory.cpp b/Engine/Plugins/Interchange/Runtime/Source/Import/Private/Animation/InterchangeAnimationTrackSetFactory.cpp
deleted file mode 100644
diff --git a/Engine/Plugins/Interchange/Runtime/Source/Import/Private/Animation/InterchangeLevelSequenceFactory.cpp b/Engine/Plugins/Interchange/Runtime/Source/Import/Private/Animation/InterchangeLevelSequenceFactory.cpp
new file mode 100644
--- /dev/null
+++ b/Engine/Plugins/Interchange/Runtime/Source/Import/Private/Animation/InterchangeLevelSequenceFactory.cpp
@@ -0,0 +1,839 @@
+// Copyright Epic Games, Inc. All Rights Reserved. 
+
+#include "Animation/InterchangeLevelSequenceFactory.h"
+
+#include "Animation/InterchangeAnimationPayloadInterface.h"
+#include "InterchangeLevelSequenceFactoryNode.h"
+#include "InterchangeAnimationTrackSetNode.h"
+#include "InterchangeAnimSequenceFactoryNode.h"
+#include "InterchangeImportCommon.h"
+#include "InterchangeImportLog.h"
+#include "InterchangeResult.h"
+#include "InterchangeSourceData.h"
+#include "InterchangeTranslatorBase.h"
+#include "Nodes/InterchangeBaseNode.h"
+#include "Nodes/InterchangeBaseNodeContainer.h"
+
+//#include "Algo/Transform.h"
+#include "Channels/MovieSceneBoolChannel.h"
+#include "Channels/MovieSceneChannelProxy.h"
+#include "Channels/MovieSceneDoubleChannel.h"
+#include "LevelSequence.h"
+#include "MovieScene.h"
+#include "MovieSceneSection.h"
+#include "MovieSceneSequence.h"
+#include "Sections/MovieScene3DTransformSection.h"
+#include "Sections/MovieSceneBoolSection.h"
+#include "Sections/MovieSceneSubSection.h"
+#include "Tracks/MovieScene3DTransformTrack.h"
+#include "Tracks/MovieSceneSubTrack.h"
+#include "Tracks/MovieSceneVisibilityTrack.h"
+
+#if WITH_EDITORONLY_DATA
+
+#include "EditorFramework/AssetImportData.h"
+
+#endif //WITH_EDITORONLY_DATA
+
+#define LOCTEXT_NAMESPACE "InterchangeLevelSequenceFactory"
+
+#if WITH_EDITOR
+namespace UE::Interchange::Private
+{
+	AActor* GetActor(const UInterchangeBaseNodeContainer* NodeContainer, const UInterchangeAnimationTrackNode* TrackNode)
+	{
+		AActor* Actor = nullptr;
+
+		FString ActorNodeUid;
+		if (TrackNode->GetCustomActorDependencyUid(ActorNodeUid))
+		{
+			const FString ActorFactoryNodeUid = UInterchangeFactoryBaseNode::BuildFactoryNodeUid(ActorNodeUid);
+			const UInterchangeFactoryBaseNode* ActorFactoryNode = Cast<UInterchangeFactoryBaseNode>(NodeContainer->GetNode(ActorFactoryNodeUid));
+
+			if (ActorFactoryNode)
+			{
+				FSoftObjectPath ReferenceObject;
+				ActorFactoryNode->GetCustomReferenceObject(ReferenceObject);
+				if (ReferenceObject.IsValid())
+				{
+					Actor = Cast<AActor>(ReferenceObject.TryLoad());
+				}
+			}
+		}
+
+		return Actor;
+	}
+
+	bool HasActorToUse(const UInterchangeBaseNodeContainer* NodeContainer, const UInterchangeLevelSequenceFactoryNode* FactoryNode)
+	{
+
+		TArray<FString> AnimationTrackUids;
+		FactoryNode->GetCustomAnimationTrackUids(AnimationTrackUids);
+
+		for (const FString& AnimationTrackUid : AnimationTrackUids)
+		{
+			if (const UInterchangeBaseNode* TranslatedNode = NodeContainer->GetNode(AnimationTrackUid))
+			{
+				if (const UInterchangeTransformAnimationTrackNode* TransformTrackNode = Cast<UInterchangeTransformAnimationTrackNode>(TranslatedNode))
+				{
+					AActor* Actor = GetActor(NodeContainer, TransformTrackNode);
+					if (Actor)
+					{
+						return true;
+					}
+				}
+				else if (const UInterchangeAnimationTrackSetInstanceNode* InstanceTrackNode = Cast<UInterchangeAnimationTrackSetInstanceNode>(TranslatedNode))
+				{
+					FString TrackSetNodeUid;
+					if (!InstanceTrackNode->GetCustomTrackSetDependencyUid(TrackSetNodeUid))
+					{
+						continue;
+					}
+
+					const FString TrackSetFactoryNodeUid = UInterchangeFactoryBaseNode::BuildFactoryNodeUid(TrackSetNodeUid);
+					const UInterchangeLevelSequenceFactoryNode* InstanceFactoryNode = Cast<UInterchangeLevelSequenceFactoryNode>(NodeContainer->GetNode(TrackSetFactoryNodeUid));
+
+					if (!InstanceFactoryNode)
+					{
+						continue;
+					}
+					FSoftObjectPath ReferenceObject;
+					InstanceFactoryNode->GetCustomReferenceObject(ReferenceObject);
+					if (!ReferenceObject.IsValid())
+					{
+						continue;
+					}
+
+					return true;
+				}
+				else if (const UInterchangeAnimationTrackNode* TrackNode = Cast<UInterchangeAnimationTrackNode>(TranslatedNode))
+				{
+					int32 TargetedProperty;
+					if (!TrackNode->GetCustomTargetedProperty(TargetedProperty))
+					{
+						continue;
+					}
+
+					// Only visibility is supported for the time being
+					if (TargetedProperty != (int32)EInterchangeAnimatedProperty::Visibility)
+					{
+						continue;
+					}
+
+					// Get targeted actor exists
+					AActor* Actor = GetActor(NodeContainer, TrackNode);
+					if (Actor)
+					{
+						return true;
+					}
+				}
+			}
+		}
+
+		return false;
+	}
+
+	class FLevelSequenceHelper
+	{
+	public:
+		FLevelSequenceHelper(ULevelSequence& InLevelSequence, UInterchangeLevelSequenceFactoryNode& InFactoryNode, const UInterchangeBaseNodeContainer& InNodeContainer, const IInterchangeAnimationPayloadInterface& InPayloadInterface)
+			: LevelSequence(InLevelSequence)
+			, MovieScene(InLevelSequence.MovieScene)
+			, FactoryNode(InFactoryNode)
+			, NodeContainer(InNodeContainer)
+			, PayloadInterface(InPayloadInterface)
+		{
+		}
+
+		void PopulateLevelSequence();
+
+	private:
+		void PopulateTransformTrack(const UInterchangeTransformAnimationTrackNode& TransformTrackNode, int32 TrackIndex);
+		void PopulateSubsequenceTrack(const UInterchangeAnimationTrackSetInstanceNode& InstanceNode);
+		void PopulateAnimationTrack(const UInterchangeAnimationTrackNode& AnimationTrackNode, int32 TrackIndex);
+
+		AActor* GetActor(const UInterchangeAnimationTrackNode& TrackNode);
+		FGuid BindActorToLevelSequence(AActor& Actor, int32 TrackIndex);
+		void UpdateTransformChannels(TArrayView<FMovieSceneDoubleChannel*>& Channels, int32 IndexOffset, const TArray<FRichCurve>& Curves);
+		
+		template<typename ChannelType, typename ValueType> 
+		void UpdateStepChannel(ChannelType& Channel, const TArray<float>& KeyTimes, const TArray<ValueType>& Values)
+		{
+			const FFrameRate& FrameRate = MovieScene->GetTickResolution();
+
+			TMovieSceneChannelData<ValueType> Data = Channel.GetData();
+
+			Data.Reset();
+
+			for (int32 KeyIndex = 0; KeyIndex < KeyTimes.Num(); ++KeyIndex)
+			{
+				FFrameNumber FrameNumber = FrameRate.AsFrameNumber(KeyTimes[KeyIndex]);
+
+				if (FrameNumber < this->MinFrameNumber)
+				{
+					this->MinFrameNumber = FrameNumber;
+				}
+
+				if (FrameNumber > this->MaxFrameNumber)
+				{
+					this->MaxFrameNumber = FrameNumber;
+				}
+
+				Data.AddKey(FrameNumber, Values[KeyIndex]);
+			}
+		}
+
+	private:
+		ULevelSequence& LevelSequence;
+		UMovieScene* MovieScene = nullptr;
+		UInterchangeLevelSequenceFactoryNode& FactoryNode;
+		const UInterchangeBaseNodeContainer& NodeContainer;
+		const IInterchangeAnimationPayloadInterface& PayloadInterface;
+
+		FFrameNumber MinFrameNumber;
+		FFrameNumber MaxFrameNumber;
+		bool ClearSubsequenceTrack = true;
+	};
+
+	void FLevelSequenceHelper::PopulateLevelSequence()
+	{
+		if (!MovieScene || FactoryNode.GetCustomAnimationTrackUidCount() == 0)
+		{
+			return;
+		}
+
+		float FrameRate;
+		if (FactoryNode.GetCustomFrameRate(FrameRate))
+		{
+			if (FrameRate <= 0.f)
+			{
+				FrameRate = 30.f;
+			}
+
+			MovieScene->SetDisplayRate(UE::Interchange::Animation::ConvertSampleRatetoFrameRate(FrameRate));
+		}
+		// Set 30 FPS as the default frame rate
+		else
+		{
+			MovieScene->SetDisplayRate(UE::Interchange::Animation::ConvertSampleRatetoFrameRate(30.f));
+		}
+
+		TArray<FString> AnimationTrackUids;
+		FactoryNode.GetCustomAnimationTrackUids(AnimationTrackUids);
+
+		int32 TrackIndex = 0;
+		for (const FString& AnimationTrackUid : AnimationTrackUids)
+		{
+			if (const UInterchangeBaseNode* TranslatedNode = NodeContainer.GetNode(AnimationTrackUid))
+			{
+				if (const UInterchangeTransformAnimationTrackNode* TransformTrackNode = Cast<UInterchangeTransformAnimationTrackNode>(TranslatedNode))
+				{
+					PopulateTransformTrack(*TransformTrackNode, TrackIndex);
+				}
+				else if (const UInterchangeAnimationTrackSetInstanceNode* InstanceTrackNode = Cast<UInterchangeAnimationTrackSetInstanceNode>(TranslatedNode))
+				{
+					PopulateSubsequenceTrack(*InstanceTrackNode);
+				}
+				else if (const UInterchangeAnimationTrackNode* TrackNode = Cast<UInterchangeAnimationTrackNode>(TranslatedNode))
+				{
+					PopulateAnimationTrack(*TrackNode, TrackIndex);
+				}
+
+				++TrackIndex;
+			}
+		}
+
+		LevelSequence.MovieScene->SetPlaybackRange(TRange<FFrameNumber>::Inclusive(MinFrameNumber, MaxFrameNumber));
+		LevelSequence.MovieScene->SetEvaluationType(EMovieSceneEvaluationType::FrameLocked);
+	}
+
+	void FLevelSequenceHelper::PopulateTransformTrack(const UInterchangeTransformAnimationTrackNode& TransformTrackNode, int32 TrackIndex)
+	{
+		// Get targeted actor exists
+		AActor* Actor = GetActor(TransformTrackNode);
+		if (!Actor)
+		{
+			UE_LOG(LogInterchangeImport, Warning, TEXT("Cannot find actor for animation track %s"), *TransformTrackNode.GetDisplayLabel());
+			return;
+		}
+
+		// Get payload
+		FInterchangeAnimationPayLoadKey PayloadKey;
+		if (!TransformTrackNode.GetCustomAnimationPayloadKey(PayloadKey))
+		{
+			UE_LOG(LogInterchangeImport, Warning, TEXT("No payload key for animation track %s on actor %s"), *TransformTrackNode.GetDisplayLabel(), *Actor->GetActorLabel());
+			return;
+		}
+
+		TFuture<TOptional<UE::Interchange::FAnimationPayloadData>> Result = PayloadInterface.GetAnimationPayloadData(PayloadKey);
+		const TOptional<UE::Interchange::FAnimationPayloadData>& PayloadData = Result.Get();
+		if (!PayloadData.IsSet() || PayloadData->Curves.Num() != 9)
+		{
+			UE_LOG(LogInterchangeImport, Warning, TEXT("No payload for animation track %s on actor %s"), *TransformTrackNode.GetDisplayLabel(), *Actor->GetActorLabel());
+			return;
+		}
+
+		FGuid ObjectBinding = BindActorToLevelSequence(*Actor, TrackIndex);
+
+		UMovieScene3DTransformTrack* TransformTrack = MovieScene->FindTrack<UMovieScene3DTransformTrack>(ObjectBinding);
+		if (!TransformTrack)
+		{
+			TransformTrack = MovieScene->AddTrack<UMovieScene3DTransformTrack>(ObjectBinding);
+		}
+		else
+		{
+			TransformTrack->RemoveAllAnimationData();
+		}
+
+		if (!TransformTrack)
+		{
+			return;
+		}
+
+		bool bSectionAdded = false;
+		UMovieScene3DTransformSection* TransformSection = Cast<UMovieScene3DTransformSection>(TransformTrack->FindOrAddSection(0, bSectionAdded));
+		if (!TransformSection)
+		{
+			return;
+		}
+
+		if (bSectionAdded)
+		{
+			int32 CompletionMode;
+			if (TransformTrackNode.GetCustomCompletionMode(CompletionMode))
+			{
+				// Make sure EMovieSceneCompletionMode enum value are still between 0 and 2
+				static_assert((uint8)EInterchangeAimationCompletionMode::KeepState == (uint8)EMovieSceneCompletionMode::KeepState, "ENUM_VALUE_HAS_CHANGED");
+				static_assert((uint8)EInterchangeAimationCompletionMode::RestoreState == (uint8)EMovieSceneCompletionMode::RestoreState, "ENUM_VALUE_HAS_CHANGED");
+				static_assert((uint8)EInterchangeAimationCompletionMode::ProjectDefault == (uint8)EMovieSceneCompletionMode::ProjectDefault, "ENUM_VALUE_HAS_CHANGED");
+
+				TransformSection->EvalOptions.CompletionMode = (EMovieSceneCompletionMode)CompletionMode;
+			}
+			// By default the completion mode is EMovieSceneCompletionMode::ProjectDefault
+			else
+			{
+				TransformSection->EvalOptions.CompletionMode = EMovieSceneCompletionMode::ProjectDefault;
+			}
+
+			TransformSection->SetRange(TRange<FFrameNumber>::All());
+		}
+
+		const FFrameRate FrameRate = MovieScene->GetDisplayRate();
+		const FFrameRate TickResolution = MovieScene->GetTickResolution();
+
+		TArrayView<FMovieSceneDoubleChannel*> Channels = TransformSection->GetChannelProxy().GetChannels<FMovieSceneDoubleChannel>();
+
+		UpdateTransformChannels(Channels, 0, PayloadData->Curves); // Translation
+		UpdateTransformChannels(Channels, 3, PayloadData->Curves); // Rotation
+		UpdateTransformChannels(Channels, 6, PayloadData->Curves); // Scaling
+
+		// Remove unnecessary keys
+		FKeyDataOptimizationParams OptimParams;
+		OptimParams.DisplayRate = MovieScene->GetDisplayRate();
+		for (FMovieSceneDoubleChannel* Channel : Channels)
+		{
+			Channel->Optimize(OptimParams);
+		}
+
+
+		int32 EnabledTransformChannels;
+		if (TransformTrackNode.GetCustomUsedChannels(EnabledTransformChannels))
+		{
+			TransformSection->SetMask((EMovieSceneTransformChannel)EnabledTransformChannels);
+		}
+		// By default all channels are enabled
+		else
+		{
+			TransformSection->SetMask(EMovieSceneTransformChannel::AllTransform);
+		}
+
+		if (USceneComponent* SceneComp = Actor->GetRootComponent())
+		{
+			SceneComp->SetMobility(EComponentMobility::Movable);
+		}
+	}
+
+	void FLevelSequenceHelper::PopulateSubsequenceTrack(const UInterchangeAnimationTrackSetInstanceNode& InstanceNode)
+	{
+		FString TrackSetNodeUid;
+		if (!InstanceNode.GetCustomTrackSetDependencyUid(TrackSetNodeUid))
+		{
+			UE_LOG(LogInterchangeImport, Error, TEXT("No unique id specified for the animation track set referenced by animation track %s."), *InstanceNode.GetDisplayLabel());
+			return;
+		}
+
+		const FString TrackSetFactoryNodeUid = UInterchangeFactoryBaseNode::BuildFactoryNodeUid(TrackSetNodeUid);
+		const UInterchangeLevelSequenceFactoryNode* InstanceFactoryNode = Cast<UInterchangeLevelSequenceFactoryNode>(NodeContainer.GetNode(TrackSetFactoryNodeUid));
+		
+		FString InstanceNodeDisplayLabel = InstanceNode.GetDisplayLabel();
+		auto LogMissingTrackError = [&InstanceNodeDisplayLabel]()
+		{
+			UE_LOG(LogInterchangeImport, Error, TEXT("Cannot find factory of animation track set referenced by animation track %s."), *InstanceNodeDisplayLabel);
+		};
+
+		if (!InstanceFactoryNode)
+		{
+			LogMissingTrackError();
+			return;
+		}
+		FSoftObjectPath ReferenceObject;
+		InstanceFactoryNode->GetCustomReferenceObject(ReferenceObject);
+		if (!ReferenceObject.IsValid())
+		{
+			LogMissingTrackError();
+			return;
+		}
+		UMovieSceneSequence* TargetMovieSceneSequence = CastChecked<UMovieSceneSequence>(ReferenceObject.TryLoad());
+
+		// Create SubTrack
+		UMovieSceneSubTrack* SubTrack = MovieScene->FindTrack<UMovieSceneSubTrack>();
+		if (!SubTrack)
+		{
+			SubTrack = MovieScene->AddTrack<UMovieSceneSubTrack>();
+		}
+		else if (SubTrack && ClearSubsequenceTrack)
+		{
+			SubTrack->RemoveAllAnimationData();
+		}
+
+		if (!SubTrack)
+		{
+			return;
+		}
+
+		FFrameRate DisplayRate = MovieScene->GetDisplayRate();
+
+		float SrcLowerBound = 0.f;
+		float SrcUpperBound = 0.f;
+		int32 StartFrame = 0;
+		if (InstanceNode.GetCustomStartFrame(StartFrame))
+		{
+			SrcLowerBound = DisplayRate.AsSeconds(StartFrame);
+		}
+
+		int32 Duration = 0;
+		if (InstanceNode.GetCustomDuration(Duration))
+		{
+			SrcUpperBound = DisplayRate.AsSeconds(StartFrame + Duration);
+		}
+
+		const FFrameRate TickResolution = MovieScene->GetTickResolution();
+		const FFrameNumber DstLowerBound = TickResolution.AsFrameNumber(SrcLowerBound);
+		const FFrameNumber DstUpperBound = TickResolution.AsFrameNumber(SrcUpperBound);
+
+		// Internally AddSequenceOnRow will automatically bump overlapping subsequences, so we can just add where it's ideal for us
+		UMovieSceneSubSection* NewSection = SubTrack->AddSequenceOnRow(TargetMovieSceneSequence, DstLowerBound, DstUpperBound.Value - DstLowerBound.Value, INDEX_NONE);
+		
+		NewSection->Parameters.TimeScale = 1.f;
+		InstanceNode.GetCustomTimeScale(NewSection->Parameters.TimeScale);
+		
+		int32 CompletionMode;
+		if (InstanceNode.GetCustomCompletionMode(CompletionMode))
+		{
+			NewSection->EvalOptions.CompletionMode = (EMovieSceneCompletionMode)CompletionMode;
+		}
+		// By default the completion mode is EMovieSceneCompletionMode::ProjectDefault
+		else
+		{
+			NewSection->EvalOptions.CompletionMode = EMovieSceneCompletionMode::ProjectDefault;
+		}
+
+		MinFrameNumber = FMath::Min(MinFrameNumber, NewSection->GetRange().GetLowerBoundValue());
+		MaxFrameNumber = FMath::Max(MaxFrameNumber, NewSection->GetRange().GetUpperBoundValue());
+
+		ClearSubsequenceTrack = false;
+	}
+
+	void FLevelSequenceHelper::PopulateAnimationTrack(const UInterchangeAnimationTrackNode& AnimationTrackNode, int32 TrackIndex)
+	{
+		int32 TargetedProperty;
+		if (!AnimationTrackNode.GetCustomTargetedProperty(TargetedProperty))
+		{
+			return;
+		}
+
+		// Only visibility is supported for the time being
+		if (TargetedProperty != (int32)EInterchangeAnimatedProperty::Visibility)
+		{
+			return;
+		}
+
+		// Get targeted actor exists
+		AActor* Actor = GetActor(AnimationTrackNode);
+		if (!Actor)
+		{
+			UE_LOG(LogInterchangeImport, Warning, TEXT("Cannot find actor for animation track %s"), *AnimationTrackNode.GetDisplayLabel());
+			return;
+		}
+
+		// Get payload
+		FInterchangeAnimationPayLoadKey PayloadKey;
+		if (!AnimationTrackNode.GetCustomAnimationPayloadKey(PayloadKey))
+		{
+			UE_LOG(LogInterchangeImport, Warning, TEXT("No payload key for animation track %s on actor %s"), *AnimationTrackNode.GetDisplayLabel(), *Actor->GetActorLabel());
+			return;
+		}
+
+		TFuture<TOptional<UE::Interchange::FAnimationPayloadData>> Result = PayloadInterface.GetAnimationPayloadData(PayloadKey);
+		const TOptional<UE::Interchange::FAnimationPayloadData>& PayloadData = Result.Get();
+		if (!PayloadData.IsSet() || PayloadData->StepCurves.Num() != 1)
+		{
+			UE_LOG(LogInterchangeImport, Warning, TEXT("No payload for animation track %s on actor %s"), *AnimationTrackNode.GetDisplayLabel(), *Actor->GetActorLabel());
+			return;
+		}
+
+		FGuid ObjectBinding = BindActorToLevelSequence(*Actor, TrackIndex);
+
+		UMovieSceneVisibilityTrack* VisibilityTrack = MovieScene->FindTrack<UMovieSceneVisibilityTrack>(ObjectBinding);
+		if (!VisibilityTrack)
+		{
+			VisibilityTrack = MovieScene->AddTrack<UMovieSceneVisibilityTrack>(ObjectBinding);
+		}
+		else
+		{
+			VisibilityTrack->RemoveAllAnimationData();
+		}
+
+		if (!VisibilityTrack)
+		{
+			return;
+		}
+
+		bool bSectionAdded = false;
+		UMovieSceneBoolSection* BoolSection = Cast<UMovieSceneBoolSection>(VisibilityTrack->FindOrAddSection(0, bSectionAdded));
+		if (!BoolSection)
+		{
+			return;
+		}
+
+		if (bSectionAdded)
+		{
+			int32 CompletionMode;
+			if (AnimationTrackNode.GetCustomCompletionMode(CompletionMode))
+			{
+				// Make sure EMovieSceneCompletionMode enum value are still between 0 and 2
+				static_assert(0 == (uint32)EMovieSceneCompletionMode::KeepState, "ENUM_VALUE_HAS_CHANGED");
+				static_assert(1 == (uint32)EMovieSceneCompletionMode::RestoreState, "ENUM_VALUE_HAS_CHANGED");
+				static_assert(2 == (uint32)EMovieSceneCompletionMode::ProjectDefault, "ENUM_VALUE_HAS_CHANGED");
+
+				BoolSection->EvalOptions.CompletionMode = (EMovieSceneCompletionMode)CompletionMode;
+			}
+			// By default the completion mode is EMovieSceneCompletionMode::ProjectDefault
+			else
+			{
+				BoolSection->EvalOptions.CompletionMode = EMovieSceneCompletionMode::ProjectDefault;
+			}
+
+			BoolSection->SetRange(TRange<FFrameNumber>::All());
+		}
+
+		TArrayView<FMovieSceneBoolChannel*> Channels = BoolSection->GetChannelProxy().GetChannels<FMovieSceneBoolChannel>();
+		if (ensure(Channels[0]))
+		{
+			const FInterchangeStepCurve& Curve = PayloadData->StepCurves[0];
+			const TArray<float>& KeyTimes = Curve.KeyTimes;
+			const TArray<bool>& Values = Curve.BooleanKeyValues.GetValue();
+
+			UpdateStepChannel(*(Channels[0]), KeyTimes, Values);
+		}
+	}
+
+	AActor* FLevelSequenceHelper::GetActor(const UInterchangeAnimationTrackNode& TrackNode)
+	{
+		AActor* Actor = nullptr;
+
+		FString ActorNodeUid;
+		if (TrackNode.GetCustomActorDependencyUid(ActorNodeUid))
+		{
+			const FString ActorFactoryNodeUid = UInterchangeFactoryBaseNode::BuildFactoryNodeUid(ActorNodeUid);
+			const UInterchangeFactoryBaseNode* ActorFactoryNode = Cast<UInterchangeFactoryBaseNode>(NodeContainer.GetNode(ActorFactoryNodeUid));
+
+			if (ActorFactoryNode)
+			{
+				FSoftObjectPath ReferenceObject;
+				ActorFactoryNode->GetCustomReferenceObject(ReferenceObject);
+				if (ReferenceObject.IsValid())
+				{
+					Actor = Cast<AActor>(ReferenceObject.TryLoad());
+				}
+			}
+		}
+
+		return Actor;
+	}
+
+	FGuid FLevelSequenceHelper::BindActorToLevelSequence(AActor& Actor, int32 TrackIndex)
+	{
+		// Bind the actor to the level sequence
+		// But first, check if there's already a possessable at the current index
+		FMovieScenePossessable* OldPossessable = nullptr;
+		if (TrackIndex < MovieScene->GetPossessableCount())
+		{
+			OldPossessable = &MovieScene->GetPossessable(TrackIndex);
+		}
+
+		FString ActorLabel = Actor.GetActorLabel();
+		FGuid ObjectBinding;
+		if (OldPossessable)
+		{
+			// If so, remove any track associated with it and unbind it from the level sequence
+			FGuid OldGuid = OldPossessable->GetGuid();
+			UMovieScene3DTransformTrack* TransformTrack = MovieScene->FindTrack<UMovieScene3DTransformTrack>(OldGuid);
+			if (TransformTrack)
+			{
+				MovieScene->RemoveTrack(*TransformTrack);
+			}
+			LevelSequence.UnbindPossessableObjects(OldGuid);
+
+			// Replace the old possessable with the new one in the MovieScene
+			FMovieScenePossessable NewPossessable(ActorLabel, Actor.GetClass());
+			MovieScene->ReplacePossessable(OldGuid, NewPossessable);
+			ObjectBinding = NewPossessable.GetGuid();
+		}
+		else
+		{
+			// No previous possessable, so simply add it to the MovieScene
+			ObjectBinding = MovieScene->AddPossessable(ActorLabel, Actor.GetClass());
+		}
+
+		// Finally, associate the guid with the actor
+		LevelSequence.BindPossessableObject(ObjectBinding, Actor, Actor.GetWorld());
+
+		return ObjectBinding;
+	}
+
+	void FLevelSequenceHelper::UpdateTransformChannels(TArrayView<FMovieSceneDoubleChannel*>& Channels, int32 IndexOffset, const TArray<FRichCurve>& Curves)
+	{
+		auto CopyToChannel = [this](FMovieSceneDoubleChannel* Channel, const FRichCurve& Curve)
+		{
+			const FFrameRate& FrameRate = this->MovieScene->GetTickResolution();
+
+			const TArray<FRichCurveKey>& CurveKeys = Curve.GetConstRefOfKeys();
+			
+			TArray<FFrameNumber> FrameNumbers;
+			FrameNumbers.Reserve(CurveKeys.Num());
+			
+			TArray<FMovieSceneDoubleValue> MovieSceneDoubleValues;
+			MovieSceneDoubleValues.Reserve(CurveKeys.Num());
+
+			for (int32 KeyIndex = 0; KeyIndex < CurveKeys.Num(); ++KeyIndex)
+			{
+				const FRichCurveKey& CurveKey = CurveKeys[KeyIndex];
+
+				FFrameNumber& FrameNumber = FrameNumbers.Add_GetRef(FrameRate.AsFrameNumber(CurveKey.Time));
+
+				if (FrameNumber < this->MinFrameNumber)
+				{
+					this->MinFrameNumber = FrameNumber;
+				}
+
+				if (FrameNumber > this->MaxFrameNumber)
+				{
+					this->MaxFrameNumber = FrameNumber;
+				}
+
+				FMovieSceneDoubleValue& SceneValue = MovieSceneDoubleValues.AddDefaulted_GetRef();
+
+				SceneValue.InterpMode = CurveKey.InterpMode;
+				SceneValue.Value = CurveKey.Value;
+			}
+
+			if (!MovieSceneDoubleValues.IsEmpty())
+			{
+				Channel->Set(FrameNumbers, MovieSceneDoubleValues);
+			}
+			else
+			{
+				Channel->RemoveDefault();
+			}
+		};
+
+		CopyToChannel(Channels[IndexOffset + 0], Curves[IndexOffset + 0]);
+		CopyToChannel(Channels[IndexOffset + 1], Curves[IndexOffset + 1]);
+		CopyToChannel(Channels[IndexOffset + 2], Curves[IndexOffset + 2]);
+	}
+} //namespace UE::Interchange::Private
+#endif //WITH_EDITOR
+
+UClass* UInterchangeLevelSequenceFactory::GetFactoryClass() const
+{
+	return ULevelSequence::StaticClass();
+}
+
+UInterchangeFactoryBase::FImportAssetResult UInterchangeLevelSequenceFactory::BeginImportAsset_GameThread(const FImportAssetObjectParams& Arguments)
+{
+	TRACE_CPUPROFILER_EVENT_SCOPE(UInterchangeLevelSequenceFactory::BeginImportAsset_GameThread);
+
+	FImportAssetResult ImportAssetResult;
+#if !WITH_EDITOR || !WITH_EDITORONLY_DATA
+
+	UE_LOG(LogInterchangeImport, Error, TEXT("Cannot import levelsequence asset in runtime, this is an editor only feature."));
+	return ImportAssetResult;
+#else
+
+	auto CannotReimportMessage = [&Arguments, this]()
+	{
+		UInterchangeResultError_Generic* Message = AddMessage<UInterchangeResultError_Generic>();
+		Message->SourceAssetName = Arguments.SourceData->GetFilename();
+		Message->DestinationAssetName = Arguments.AssetName;
+		Message->AssetType = ULevelSequence::StaticClass();
+		Message->Text = LOCTEXT("CreateEmptyAssetUnsupportedReimport", "Re-import of ULevelSequence not supported yet.");
+		Arguments.AssetNode->SetSkipNodeImport();
+	};
+
+	if (Arguments.ReimportObject)
+	{
+		CannotReimportMessage();
+		return ImportAssetResult;
+	}
+
+	ULevelSequence* LevelSequence = nullptr;
+	if (!Arguments.AssetNode || !Arguments.AssetNode->GetObjectClass()->IsChildOf(GetFactoryClass()))
+	{
+		return ImportAssetResult;
+	}
+
+	UInterchangeLevelSequenceFactoryNode* FactoryNode = Cast<UInterchangeLevelSequenceFactoryNode>(Arguments.AssetNode);
+	if (FactoryNode == nullptr)
+	{
+		return ImportAssetResult;
+	}
+
+	if (!UE::Interchange::Private::HasActorToUse(Arguments.NodeContainer, FactoryNode))
+	{
+		UE_LOG(LogInterchangeImport, Warning, TEXT("Level sequence asset, %s, not imported, because all referenced actors are missing."), *FactoryNode->GetDisplayLabel());
+		return ImportAssetResult;
+	}
+
+	UObject* ExistingAsset = Arguments.ReimportObject;
+	if (!ExistingAsset)
+	{
+		FSoftObjectPath ReferenceObject;
+		if (FactoryNode->GetCustomReferenceObject(ReferenceObject))
+		{
+			ExistingAsset = ReferenceObject.TryLoad();
+		}
+	}
+
+	// create a new material or overwrite existing asset, if possible
+	if (!ExistingAsset)
+	{
+		LevelSequence = NewObject<ULevelSequence>(Arguments.Parent, *Arguments.AssetName, RF_Public | RF_Standalone);
+	}
+	else
+	{
+		// This is a reimport or an override, we are just re-updating the source data
+
+		//TODO: put back the Cast when the LevelSequence will support re-import
+		//LevelSequence = Cast<ULevelSequence>(ExistingAsset);
+		CannotReimportMessage();
+		return ImportAssetResult;
+	}
+
+	if (!LevelSequence)
+	{
+		UE_LOG(LogInterchangeImport, Warning, TEXT("Could not create LevelSequence asset %s"), *Arguments.AssetName);
+		return ImportAssetResult;
+	}
+
+	FactoryNode->SetCustomReferenceObject(LevelSequence);
+
+	LevelSequence->PreEditChange(nullptr);
+
+	ImportAssetResult.ImportedObject = ImportObjectSourceData(Arguments);
+	return ImportAssetResult;
+#endif //else !WITH_EDITOR || !WITH_EDITORONLY_DATA
+}
+
+UObject* UInterchangeLevelSequenceFactory::ImportObjectSourceData(const FImportAssetObjectParams& Arguments)
+{
+	TRACE_CPUPROFILER_EVENT_SCOPE(UInterchangeLevelSequenceFactory::ImportObjectSourceData);
+
+#if !WITH_EDITOR || !WITH_EDITORONLY_DATA
+	// TODO: Can we import ULevelSequence at runtime
+	UE_LOG(LogInterchangeImport, Error, TEXT("Cannot import LevelSequence asset in runtime, this is an editor only feature."));
+	return nullptr;
+
+#else
+	using namespace UE::Interchange;
+
+	// Re-import is not supported yet
+	// Need to talk to the Sequencer team about adding an AssetImportData to ULevelSequence
+	if (Arguments.ReimportObject)
+	{
+		UInterchangeResultError_Generic* Message = AddMessage<UInterchangeResultError_Generic>();
+		Message->SourceAssetName = Arguments.SourceData->GetFilename();
+		Message->DestinationAssetName = Arguments.AssetName;
+		Message->AssetType = ULevelSequence::StaticClass();
+		Message->Text = LOCTEXT("CreateAssetUnsupportedReimport", "Re-import of ULevelSequence not supported yet.");
+
+		return nullptr;
+	}
+
+	if (!Arguments.NodeContainer || !Arguments.AssetNode || !Arguments.AssetNode->GetObjectClass()->IsChildOf(GetFactoryClass()))
+	{
+		return nullptr;
+	}
+
+	UInterchangeLevelSequenceFactoryNode* FactoryNode = Cast<UInterchangeLevelSequenceFactoryNode>(Arguments.AssetNode);
+	if (!FactoryNode)
+	{
+		return nullptr;
+	}
+
+	Translator = Arguments.Translator;
+	const IInterchangeAnimationPayloadInterface* AnimSequenceTranslatorPayloadInterface = Cast<IInterchangeAnimationPayloadInterface>(Arguments.Translator);
+	if (!AnimSequenceTranslatorPayloadInterface)
+	{
+		UE_LOG(LogInterchangeImport, Error, TEXT("Cannot import LevelSequence, the translator do not implement the IInterchangeAnimationPayloadInterface."));
+		return nullptr;
+	}
+
+	UObject* ExistingAsset = UE::Interchange::FFactoryCommon::AsyncFindObject(FactoryNode, GetFactoryClass(), Arguments.Parent, Arguments.AssetName);
+
+	if (!ExistingAsset)
+	{
+		UE_LOG(LogInterchangeImport, Error, TEXT("Could not import the LevelSequence asset %s, because the asset do not exist."), *Arguments.AssetName);
+		return nullptr;
+	}
+
+	ULevelSequence* LevelSequence = Cast<ULevelSequence>(ExistingAsset);
+
+	if (!ensure(LevelSequence))
+	{
+		UInterchangeResultError_Generic* Message = AddMessage<UInterchangeResultError_Generic>();
+		Message->SourceAssetName = Arguments.SourceData->GetFilename();
+		Message->DestinationAssetName = Arguments.AssetName;
+		Message->AssetType = ULevelSequence::StaticClass();
+		Message->Text = FText::Format(LOCTEXT("CreateAssetFailed", "Could not create nor find LevelSequence asset {0}."), FText::FromString(Arguments.AssetName));
+		return nullptr;
+	}
+
+	LevelSequence->Initialize();
+
+	Private::FLevelSequenceHelper Helper(*LevelSequence, *FactoryNode, *Arguments.NodeContainer, *AnimSequenceTranslatorPayloadInterface);
+	Helper.PopulateLevelSequence();
+
+	/** Apply all FactoryNode custom attributes to the level sequence asset */
+	FactoryNode->ApplyAllCustomAttributeToObject(LevelSequence);
+
+	//Getting the file Hash will cache it into the source data
+	Arguments.SourceData->GetFileContentHash();
+
+	//The interchange completion task (call in the GameThread after the factories pass), will call PostEditChange which will trig another asynchronous system that will build all material in parallel
+	return LevelSequence;
+#endif //else !WITH_EDITOR || !WITH_EDITORONLY_DATA
+}
+
+/* This function is call in the completion task on the main thread, use it to call main thread post creation step for your assets*/
+void UInterchangeLevelSequenceFactory::SetupObject_GameThread(const FSetupObjectParams& Arguments)
+{
+	TRACE_CPUPROFILER_EVENT_SCOPE(UInterchangeLevelSequenceFactory::SetupObject_GameThread);
+
+	check(IsInGameThread());
+	Super::SetupObject_GameThread(Arguments);
+
+	// TODO: Talk with sequence team about adding AssetImportData to ULevelSequence for re-import 
+}
+
+#undef LOCTEXT_NAMESPACE
\ No newline at end of file
diff --git a/Engine/Plugins/Interchange/Runtime/Source/Import/Private/InterchangeImportModule.cpp b/Engine/Plugins/Interchange/Runtime/Source/Import/Private/InterchangeImportModule.cpp
--- a/Engine/Plugins/Interchange/Runtime/Source/Import/Private/InterchangeImportModule.cpp
+++ b/Engine/Plugins/Interchange/Runtime/Source/Import/Private/InterchangeImportModule.cpp
@@ -3,7 +3,7 @@
 #include "InterchangeImportModule.h"
 
 #include "Animation/InterchangeAnimSequenceFactory.h"
-#include "Animation/InterchangeAnimationTrackSetFactory.h"
+#include "Animation/InterchangeLevelSequenceFactory.h"
 #include "CoreMinimal.h"
 #include "Engine/Engine.h"
 #include "Engine/RendererSettings.h"
@@ -102,7 +102,7 @@ void FInterchangeImportModule::StartupModule()
 		InterchangeManager.RegisterFactory(UInterchangeStaticMeshFactory::StaticClass());
 		InterchangeManager.RegisterFactory(UInterchangePhysicsAssetFactory::StaticClass());
 		InterchangeManager.RegisterFactory(UInterchangeActorFactory::StaticClass());
-		InterchangeManager.RegisterFactory(UInterchangeAnimationTrackSetFactory::StaticClass());
+		InterchangeManager.RegisterFactory(UInterchangeLevelSequenceFactory::StaticClass());
 		InterchangeManager.RegisterFactory(UInterchangeAnimSequenceFactory::StaticClass());
 		InterchangeManager.RegisterFactory(UInterchangeCineCameraActorFactory::StaticClass());
 		InterchangeManager.RegisterFactory(UInterchangeCameraActorFactory::StaticClass());
diff --git a/Engine/Plugins/Interchange/Runtime/Source/Import/Public/Animation/InterchangeAnimationTrackSetFactory.h b/Engine/Plugins/Interchange/Runtime/Source/Import/Public/Animation/InterchangeAnimationTrackSetFactory.h
deleted file mode 100644
diff --git a/Engine/Plugins/Interchange/Runtime/Source/Import/Public/Animation/InterchangeLevelSequenceFactory.h b/Engine/Plugins/Interchange/Runtime/Source/Import/Public/Animation/InterchangeLevelSequenceFactory.h
new file mode 100644
--- /dev/null
+++ b/Engine/Plugins/Interchange/Runtime/Source/Import/Public/Animation/InterchangeLevelSequenceFactory.h
@@ -0,0 +1,35 @@
+// Copyright Epic Games, Inc. All Rights Reserved.
+
+#pragma once
+
+#include "CoreMinimal.h"
+#include "InterchangeFactoryBase.h"
+#include "UObject/Object.h"
+#include "UObject/ObjectMacros.h"
+
+#include "InterchangeLevelSequenceFactory.generated.h"
+
+UCLASS(BlueprintType)
+class INTERCHANGEIMPORT_API UInterchangeLevelSequenceFactory : public UInterchangeFactoryBase
+{
+	GENERATED_BODY()
+public:
+
+	//////////////////////////////////////////////////////////////////////////
+	// Interchange factory base interface begin
+
+	virtual UClass* GetFactoryClass() const override;
+	virtual EInterchangeFactoryAssetType GetFactoryAssetType() override { return EInterchangeFactoryAssetType::Animations; }
+	virtual FImportAssetResult BeginImportAsset_GameThread(const FImportAssetObjectParams& Arguments) override;
+	virtual void SetupObject_GameThread(const FSetupObjectParams& Arguments) override;
+
+	// Interchange factory base interface end
+	//////////////////////////////////////////////////////////////////////////
+
+private:
+	const UInterchangeTranslatorBase* Translator = nullptr;
+
+	UObject* ImportObjectSourceData(const FImportAssetObjectParams& Arguments);
+};
+
+
diff --git a/Engine/Plugins/Interchange/Runtime/Source/Pipelines/Private/InterchangeGenericAnimationPipeline.cpp b/Engine/Plugins/Interchange/Runtime/Source/Pipelines/Private/InterchangeGenericAnimationPipeline.cpp
--- a/Engine/Plugins/Interchange/Runtime/Source/Pipelines/Private/InterchangeGenericAnimationPipeline.cpp
+++ b/Engine/Plugins/Interchange/Runtime/Source/Pipelines/Private/InterchangeGenericAnimationPipeline.cpp
@@ -3,7 +3,7 @@
 
 #include "Animation/AnimationSettings.h"
 #include "CoreMinimal.h"
-#include "InterchangeAnimationTrackSetFactoryNode.h"
+#include "InterchangeLevelSequenceFactoryNode.h"
 #include "InterchangeAnimationTrackSetNode.h"
 #include "InterchangeAnimSequenceFactoryNode.h"
 #include "InterchangeMeshNode.h"
@@ -281,7 +281,7 @@ void UInterchangeGenericAnimationPipeline::ExecutePipeline(UInterchangeBaseNodeC
 		{
 			if (TrackSetNode)
 			{
-				CreateAnimationTrackSetFactoryNode(*TrackSetNode);
+				CreateLevelSequenceFactoryNode(*TrackSetNode);
 			}
 		}
 	}
@@ -318,11 +318,11 @@ void UInterchangeGenericAnimationPipeline::ExecutePipeline(UInterchangeBaseNodeC
 	}
 }
 
-void UInterchangeGenericAnimationPipeline::CreateAnimationTrackSetFactoryNode(UInterchangeAnimationTrackSetNode& TranslatedNode)
+void UInterchangeGenericAnimationPipeline::CreateLevelSequenceFactoryNode(UInterchangeAnimationTrackSetNode& TranslatedNode)
 {
 	const FString FactoryNodeUid = UInterchangeFactoryBaseNode::BuildFactoryNodeUid(TranslatedNode.GetUniqueID());
 
-	UInterchangeAnimationTrackSetFactoryNode* FactoryNode = NewObject<UInterchangeAnimationTrackSetFactoryNode>(BaseNodeContainer, NAME_None);
+	UInterchangeLevelSequenceFactoryNode* FactoryNode = NewObject<UInterchangeLevelSequenceFactoryNode>(BaseNodeContainer, NAME_None);
 
 	FactoryNode->InitializeNode(FactoryNodeUid, TranslatedNode.GetDisplayLabel(), EInterchangeNodeContainerType::FactoryData);
 	FactoryNode->SetEnabled(true);
diff --git a/Engine/Plugins/Interchange/Runtime/Source/Pipelines/Public/InterchangeGenericAnimationPipeline.h b/Engine/Plugins/Interchange/Runtime/Source/Pipelines/Public/InterchangeGenericAnimationPipeline.h
--- a/Engine/Plugins/Interchange/Runtime/Source/Pipelines/Public/InterchangeGenericAnimationPipeline.h
+++ b/Engine/Plugins/Interchange/Runtime/Source/Pipelines/Public/InterchangeGenericAnimationPipeline.h
@@ -127,7 +127,7 @@ protected:
 
 private:
 
-	void CreateAnimationTrackSetFactoryNode(UInterchangeAnimationTrackSetNode& Node);
+	void CreateLevelSequenceFactoryNode(UInterchangeAnimationTrackSetNode& Node);
 
 	void CreateAnimSequenceFactoryNode(UInterchangeSkeletalAnimationTrackNode& Node);
 
