[Interchange] - Elements animation problem caused by 0 Scale

#rnx

diff --git a/Engine/Plugins/Interchange/Runtime/Source/FactoryNodes/Private/InterchangeActorFactoryNode.cpp b/Engine/Plugins/Interchange/Runtime/Source/FactoryNodes/Private/InterchangeActorFactoryNode.cpp
--- a/Engine/Plugins/Interchange/Runtime/Source/FactoryNodes/Private/InterchangeActorFactoryNode.cpp
+++ b/Engine/Plugins/Interchange/Runtime/Source/FactoryNodes/Private/InterchangeActorFactoryNode.cpp
@@ -38,8 +38,28 @@ bool UInterchangeActorFactoryNode::SetCustomGlobalTransform(const FTransform& At
 	IMPLEMENT_NODE_ATTRIBUTE_SETTER_WITH_CUSTOM_DELEGATE_WITH_CUSTOM_CLASS(UInterchangeActorFactoryNode, GlobalTransform, FTransform, USceneComponent);
 }
 
+bool UInterchangeActorFactoryNode::GetCustomLocalTransform(FTransform& AttributeValue) const
+{
+	IMPLEMENT_NODE_ATTRIBUTE_GETTER(GlobalTransform, FTransform);
+}
+
+bool UInterchangeActorFactoryNode::SetCustomLocalTransform(const FTransform& AttributeValue, bool bAddApplyDelegate)
+{
+	IMPLEMENT_NODE_ATTRIBUTE_SETTER_WITH_CUSTOM_DELEGATE_WITH_CUSTOM_CLASS(UInterchangeActorFactoryNode, GlobalTransform, FTransform, USceneComponent);
+}
+
 bool UInterchangeActorFactoryNode::ApplyCustomGlobalTransformToAsset(UObject* Asset) const
 {
+	FTransform LocalTransform;
+	if (GetCustomLocalTransform(LocalTransform))
+	{
+		if (USceneComponent* Component = Cast<USceneComponent>(Asset))
+		{
+			Component->SetRelativeTransform(LocalTransform);
+			return true;
+		}
+	}
+
 	FTransform GlobalTransform;
 	if (GetCustomGlobalTransform(GlobalTransform))
 	{
@@ -57,8 +77,10 @@ bool UInterchangeActorFactoryNode::FillCustomGlobalTransformFromAsset(UObject* A
 {
 	if (const USceneComponent* Component = Cast<USceneComponent>(Asset))
 	{
-		FTransform Transform = Component->GetComponentToWorld();
-		return this->SetCustomGlobalTransform(Transform, false);
+		FTransform LocalTransform = Component->GetRelativeTransform();
+		FTransform GlobalTransform = Component->GetComponentToWorld();
+
+		return this->SetCustomLocalTransform(LocalTransform, false) || this->SetCustomGlobalTransform(GlobalTransform, false);
 	}
 
 	return false;
diff --git a/Engine/Plugins/Interchange/Runtime/Source/FactoryNodes/Public/InterchangeActorFactoryNode.h b/Engine/Plugins/Interchange/Runtime/Source/FactoryNodes/Public/InterchangeActorFactoryNode.h
--- a/Engine/Plugins/Interchange/Runtime/Source/FactoryNodes/Public/InterchangeActorFactoryNode.h
+++ b/Engine/Plugins/Interchange/Runtime/Source/FactoryNodes/Public/InterchangeActorFactoryNode.h
@@ -20,6 +20,12 @@ public:
 	UFUNCTION(BlueprintCallable, Category = "Interchange | Node | ActorFactory")
 	bool SetCustomGlobalTransform(const FTransform& AttributeValue, bool bAddApplyDelegate = true);
 
+	UFUNCTION(BlueprintCallable, Category = "Interchange | Node | ActorFactory")
+	bool GetCustomLocalTransform(FTransform& AttributeValue) const;
+
+	UFUNCTION(BlueprintCallable, Category = "Interchange | Node | ActorFactory")
+	bool SetCustomLocalTransform(const FTransform& AttributeValue, bool bAddApplyDelegate = true);
+
 	UFUNCTION(BlueprintCallable, Category = "Interchange | Node | ActorFactory")
 	bool GetCustomActorClassName(FString& AttributeValue) const
 	{
@@ -52,6 +58,7 @@ private:
 	bool FillCustomGlobalTransformFromAsset(UObject* Asset);
 
 	const UE::Interchange::FAttributeKey Macro_CustomGlobalTransformKey = UE::Interchange::FAttributeKey(TEXT("GlobalTransform"));
+	const UE::Interchange::FAttributeKey Macro_CustomLocalTransformKey = UE::Interchange::FAttributeKey(TEXT("LocalTransform"));
 	const UE::Interchange::FAttributeKey Macro_CustomActorClassNameKey = UE::Interchange::FAttributeKey(TEXT("ActorClassName"));
 	const UE::Interchange::FAttributeKey Macro_CustomMobilityKey = UE::Interchange::FAttributeKey(TEXT("Mobility"));
 };
\ No newline at end of file
diff --git a/Engine/Plugins/Interchange/Runtime/Source/Pipelines/Private/InterchangeGenericScenesPipeline.cpp b/Engine/Plugins/Interchange/Runtime/Source/Pipelines/Private/InterchangeGenericScenesPipeline.cpp
--- a/Engine/Plugins/Interchange/Runtime/Source/Pipelines/Private/InterchangeGenericScenesPipeline.cpp
+++ b/Engine/Plugins/Interchange/Runtime/Source/Pipelines/Private/InterchangeGenericScenesPipeline.cpp
@@ -466,23 +466,52 @@ void UInterchangeGenericLevelPipeline::ExecuteSceneNodePreImport(const FTransfor
 	}
 
 	//TODO move this code to the factory, a stack over pipeline can change the global offset transform which will affect this value.
-	FTransform GlobalTransform;
-	if (SceneNode->GetCustomGlobalTransform(BaseNodeContainer, GlobalOffsetTransform, GlobalTransform))
+	//We prioritize Local (Relative) Transforms due to issues introduced by 0 scales with Global Transforms.
+	//In case the LocalTransform is not available we fallback onto GlobalTransforms
+	FTransform LocalTransform;
+	if (SceneNode->GetCustomLocalTransform(LocalTransform))
 	{
 		if (bRootJointNode)
 		{
-			GlobalTransform = FTransform::Identity;
+			LocalTransform = FTransform::Identity;
 			//LocalTransform of RootjointNode is already baked into the Skeletal and animation.
 			//due to that we acquire the Parent SceneNode and get its GlobalTransform:
 			if (!SceneNode->GetParentUid().IsEmpty())
 			{
 				if (const UInterchangeSceneNode* ParentSceneNode = Cast<UInterchangeSceneNode>(BaseNodeContainer->GetNode(SceneNode->GetParentUid())))
 				{
-					ParentSceneNode->GetCustomGlobalTransform(BaseNodeContainer, GlobalOffsetTransform, GlobalTransform);
+					ParentSceneNode->GetCustomLocalTransform(LocalTransform);
 				}
 			}
 		}
-		ActorFactoryNode->SetCustomGlobalTransform(GlobalTransform);
+
+		if (SceneNode->GetParentUid().IsEmpty())
+		{
+			LocalTransform = LocalTransform * GlobalOffsetTransform;
+		}
+
+		ActorFactoryNode->SetCustomLocalTransform(LocalTransform);
+	}
+	else
+	{
+		FTransform GlobalTransform;
+		if (SceneNode->GetCustomGlobalTransform(BaseNodeContainer, GlobalOffsetTransform, GlobalTransform))
+		{
+			if (bRootJointNode)
+			{
+				GlobalTransform = FTransform::Identity;
+				//LocalTransform of RootjointNode is already baked into the Skeletal and animation.
+				//due to that we acquire the Parent SceneNode and get its GlobalTransform:
+				if (!SceneNode->GetParentUid().IsEmpty())
+				{
+					if (const UInterchangeSceneNode* ParentSceneNode = Cast<UInterchangeSceneNode>(BaseNodeContainer->GetNode(SceneNode->GetParentUid())))
+					{
+						ParentSceneNode->GetCustomGlobalTransform(BaseNodeContainer, GlobalOffsetTransform, GlobalTransform);
+					}
+				}
+			}
+			ActorFactoryNode->SetCustomGlobalTransform(GlobalTransform);
+		}
 	}
 
 	ActorFactoryNode->SetCustomMobility(EComponentMobility::Static);
