Interchange: Fix memory leak with node and container when the async helper is clean up.

#rnx

diff --git a/Engine/Plugins/Interchange/Runtime/Source/Import/Private/Gltf/InterchangeGltfTranslator.cpp b/Engine/Plugins/Interchange/Runtime/Source/Import/Private/Gltf/InterchangeGltfTranslator.cpp
--- a/Engine/Plugins/Interchange/Runtime/Source/Import/Private/Gltf/InterchangeGltfTranslator.cpp
+++ b/Engine/Plugins/Interchange/Runtime/Source/Import/Private/Gltf/InterchangeGltfTranslator.cpp
@@ -785,6 +785,7 @@ TOptional< UE::Interchange::FImportImage > UInterchangeGLTFTranslator::GetTextur
 		UInterchangeImageWrapperTranslator* ImageWrapperTranslator = NewObject<UInterchangeImageWrapperTranslator>(GetTransientPackage(), NAME_None);
 		ImageWrapperTranslator->SetResultsContainer(Results);
 		TexturePayloadData = ImageWrapperTranslator->GetTexturePayloadDataFromBuffer(ImageData);
+		ImageWrapperTranslator->ClearInternalFlags(EInternalObjectFlags::Async);
 	}
 	else
 	{
diff --git a/Engine/Source/Runtime/Interchange/Engine/Private/InterchangeManager.cpp b/Engine/Source/Runtime/Interchange/Engine/Private/InterchangeManager.cpp
--- a/Engine/Source/Runtime/Interchange/Engine/Private/InterchangeManager.cpp
+++ b/Engine/Source/Runtime/Interchange/Engine/Private/InterchangeManager.cpp
@@ -548,6 +548,15 @@ void UE::Interchange::FImportAsyncHelper::InitCancel()
 void UE::Interchange::FImportAsyncHelper::CleanUp()
 {
 	//Release the graph
+	for (TStrongObjectPtr<UInterchangeBaseNodeContainer>& Container : BaseNodeContainers)
+	{
+		Container->IterateNodes([](const FString&, UInterchangeBaseNode* Node)
+			{
+				Node->ClearInternalFlags(EInternalObjectFlags::Async);
+			}
+		);
+		Container->ClearInternalFlags(EInternalObjectFlags::Async);
+	}
 	BaseNodeContainers.Empty();
 
 	for (UInterchangeSourceData* SourceData : SourceDatas)
