Fix async animation import

#rnx

diff --git a/Engine/Plugins/Interchange/Runtime/Source/Import/Private/Mesh/InterchangeSkeletalMeshFactory.cpp b/Engine/Plugins/Interchange/Runtime/Source/Import/Private/Mesh/InterchangeSkeletalMeshFactory.cpp
--- a/Engine/Plugins/Interchange/Runtime/Source/Import/Private/Mesh/InterchangeSkeletalMeshFactory.cpp
+++ b/Engine/Plugins/Interchange/Runtime/Source/Import/Private/Mesh/InterchangeSkeletalMeshFactory.cpp
@@ -895,6 +895,27 @@ namespace UE
 				}
 			}
 
+			struct FContentInfo
+			{
+				bool bApplyGeometry = false;
+				bool bApplySkinning = false;
+				bool bApplyPartialContent = false;
+				bool bApplyGeometryOnly = false;
+				bool bApplySkinningOnly = false;
+			};
+
+			FContentInfo GetContentInfo(const UInterchangeSkeletalMeshFactoryNode* SkeletalMeshFactoryNode, const bool bIsReImport)
+			{
+				FContentInfo ContentInfo;
+				EInterchangeSkeletalMeshContentType ImportContent = EInterchangeSkeletalMeshContentType::All;
+				SkeletalMeshFactoryNode->GetCustomImportContentType(ImportContent);
+				ContentInfo.bApplyGeometry = !bIsReImport || (ImportContent == EInterchangeSkeletalMeshContentType::All || ImportContent == EInterchangeSkeletalMeshContentType::Geometry);
+				ContentInfo.bApplySkinning = !bIsReImport || (ImportContent == EInterchangeSkeletalMeshContentType::All || ImportContent == EInterchangeSkeletalMeshContentType::SkinningWeights);
+				ContentInfo.bApplyPartialContent = bIsReImport && ImportContent != EInterchangeSkeletalMeshContentType::All;
+				ContentInfo.bApplyGeometryOnly = ContentInfo.bApplyPartialContent && ContentInfo.bApplyGeometry;
+				ContentInfo.bApplySkinningOnly = ContentInfo.bApplyPartialContent && ContentInfo.bApplySkinning;
+				return ContentInfo;
+			}
 		} //Namespace Private
 	} //namespace Interchange
 } //namespace UE
@@ -925,7 +946,7 @@ UInterchangeFactoryBase::FImportAssetResult UInterchangeSkeletalMeshFactory::Beg
 		return ImportAssetResult;
 	}
 
-	const UInterchangeSkeletalMeshFactoryNode* SkeletalMeshFactoryNode = Cast<UInterchangeSkeletalMeshFactoryNode>(Arguments.AssetNode);
+	UInterchangeSkeletalMeshFactoryNode* SkeletalMeshFactoryNode = Cast<UInterchangeSkeletalMeshFactoryNode>(Arguments.AssetNode);
 	if (SkeletalMeshFactoryNode == nullptr)
 	{
 		return ImportAssetResult;
@@ -974,6 +995,122 @@ UInterchangeFactoryBase::FImportAssetResult UInterchangeSkeletalMeshFactory::Beg
 	}
 
 	ImportAssetResult.ImportedObject = SkeletalMesh;
+
+	//Make sure we can modify the skeletalmesh properties
+	FSkinnedAssetAsyncBuildScope AsyncBuildScope(SkeletalMesh);
+
+	//This is consider has a re-import if we have a reimport object or if the object exist and have some valid LOD
+	const bool bIsReImport = (Arguments.ReimportObject != nullptr) || (SkeletalMesh->GetLODNum() > 0);
+
+	//Dirty the DDC Key for any imported Skeletal Mesh
+	SkeletalMesh->InvalidateDeriveDataCacheGUID();
+	USkeleton* SkeletonReference = nullptr;
+
+
+	int32 LodCount = SkeletalMeshFactoryNode->GetLodDataCount();
+	TArray<FString> LodDataUniqueIds;
+	SkeletalMeshFactoryNode->GetLodDataUniqueIds(LodDataUniqueIds);
+	ensure(LodDataUniqueIds.Num() == LodCount);
+	int32 CurrentLodIndex = 0;
+
+	UE::Interchange::Private::FContentInfo ContentInfo = UE::Interchange::Private::GetContentInfo(SkeletalMeshFactoryNode, bIsReImport);
+
+	ImportAssetObjectData.bIsReImport = bIsReImport;
+	ImportAssetObjectData.bApplyGeometryOnly = ContentInfo.bApplyGeometryOnly;
+	for (int32 LodIndex = 0; LodIndex < LodCount; ++LodIndex)
+	{
+		FString LodUniqueId = LodDataUniqueIds[LodIndex];
+		const UInterchangeSkeletalMeshLodDataNode* LodDataNode = Cast<UInterchangeSkeletalMeshLodDataNode>(Arguments.NodeContainer->GetNode(LodUniqueId));
+		if (!LodDataNode)
+		{
+			UE_LOG(LogInterchangeImport, Warning, TEXT("Invalid LOD when importing SkeletalMesh asset %s"), *Arguments.AssetName);
+			continue;
+		}
+
+		FString SkeletonNodeUid;
+		if (!LodDataNode->GetCustomSkeletonUid(SkeletonNodeUid))
+		{
+			UE_LOG(LogInterchangeImport, Warning, TEXT("Invalid Skeleton LOD when importing SkeletalMesh asset %s"), *Arguments.AssetName);
+			continue;
+		}
+		const UInterchangeSkeletonFactoryNode* SkeletonNode = Cast<UInterchangeSkeletonFactoryNode>(Arguments.NodeContainer->GetNode(SkeletonNodeUid));
+		if (!SkeletonNode)
+		{
+			UE_LOG(LogInterchangeImport, Warning, TEXT("Invalid Skeleton LOD when importing SkeletalMesh asset %s"), *Arguments.AssetName);
+			continue;
+		}
+		FSoftObjectPath SkeletonNodeReferenceObject;
+		SkeletonNode->GetCustomReferenceObject(SkeletonNodeReferenceObject);
+
+		FSoftObjectPath SpecifiedSkeleton;
+		SkeletalMeshFactoryNode->GetCustomSkeletonSoftObjectPath(SpecifiedSkeleton);
+		bool bSpecifiedSkeleton = SpecifiedSkeleton.IsValid();
+		if (SkeletonReference == nullptr)
+		{
+			UObject* SkeletonObject = nullptr;
+
+			if (SpecifiedSkeleton.IsValid())
+			{
+				SkeletonObject = SpecifiedSkeleton.TryLoad();
+			}
+			else if (SkeletonNodeReferenceObject.IsValid())
+			{
+				SkeletonObject = SkeletonNodeReferenceObject.TryLoad();
+			}
+
+			if (SkeletonObject)
+			{
+				SkeletonReference = Cast<USkeleton>(SkeletonObject);
+
+			}
+
+			if (!ensure(SkeletonReference))
+			{
+				UE_LOG(LogInterchangeImport, Warning, TEXT("Invalid Skeleton LOD when importing SkeletalMesh asset %s"), *Arguments.AssetName);
+				break;
+			}
+			ImportAssetObjectData.SkeletonReference = SkeletonReference;
+		}
+
+		FString RootJointNodeId;
+		if (!SkeletonNode->GetCustomRootJointUid(RootJointNodeId))
+		{
+			UE_LOG(LogInterchangeImport, Warning, TEXT("Invalid Skeleton LOD Root Joint when importing SkeletalMesh asset %s"), *Arguments.AssetName);
+			continue;
+		}
+
+		const UInterchangeSceneNode* RootJointNode = Cast<UInterchangeSceneNode>(Arguments.NodeContainer->GetNode(RootJointNodeId));
+		if (!RootJointNode)
+		{
+			UE_LOG(LogInterchangeImport, Warning, TEXT("Invalid Skeleton RootJointNode."));
+			continue;
+		}
+
+		FImportAssetObjectLODData& ImportAssetObjectLODData = ImportAssetObjectData.LodDatas.AddDefaulted_GetRef();
+		ImportAssetObjectLODData.LodIndex = CurrentLodIndex;
+
+		int32 SkeletonDepth = 0;
+		SkeletonNode->GetCustomUseTimeZeroForBindPose(ImportAssetObjectLODData.bUseTimeZeroAsBindPose);
+		
+		//Do not alter the skeletal mesh reference skeleton when importing geometry only
+		FReferenceSkeleton RefSkeleton;
+		UE::Interchange::Private::FSkeletonHelper::ProcessImportMeshSkeleton(SkeletonReference
+			, ContentInfo.bApplyGeometryOnly ? RefSkeleton : SkeletalMesh->GetRefSkeleton()
+			, SkeletonDepth
+			, Arguments.NodeContainer
+			, RootJointNodeId
+			, ImportAssetObjectLODData.RefBonesBinary
+			, ImportAssetObjectLODData.bUseTimeZeroAsBindPose
+			, ImportAssetObjectLODData.bDiffPose);
+
+		if (bSpecifiedSkeleton && !SkeletonReference->IsCompatibleMesh(SkeletalMesh))
+		{
+			UE_LOG(LogInterchangeImport, Warning, TEXT("The skeleton %s is incompatible with the imported skeletalmesh asset %s"), *SkeletonReference->GetName(), *Arguments.AssetName);
+		}
+
+		CurrentLodIndex++;
+	}
+
 	return ImportAssetResult;
 
 #endif //else !WITH_EDITOR || !WITH_EDITORONLY_DATA
@@ -1027,9 +1164,6 @@ UInterchangeFactoryBase::FImportAssetResult UInterchangeSkeletalMeshFactory::Imp
 	//Make sure we can modify the skeletalmesh properties
 	FSkinnedAssetAsyncBuildScope AsyncBuildScope(SkeletalMesh);
 
-	//This is consider has a re-import if we have a reimport object or if the object exist and have some valid LOD
-	const bool bIsReImport = (Arguments.ReimportObject != nullptr) || (SkeletalMesh->GetLODNum() > 0);
-
 	FTransform GlobalOffsetTransform = FTransform::Identity;
 	bool bBakeMeshes = false;
 	if (UInterchangeCommonPipelineDataFactoryNode* CommonPipelineDataFactoryNode = UInterchangeCommonPipelineDataFactoryNode::GetUniqueInstance(Arguments.NodeContainer))
@@ -1038,12 +1172,10 @@ UInterchangeFactoryBase::FImportAssetResult UInterchangeSkeletalMeshFactory::Imp
 		CommonPipelineDataFactoryNode->GetBakeMeshes(bBakeMeshes);
 	}
 
-	//Dirty the DDC Key for any imported Skeletal Mesh
-	SkeletalMesh->InvalidateDeriveDataCacheGUID();
-	USkeleton* SkeletonReference = nullptr;
+	USkeleton* SkeletonReference = ImportAssetObjectData.SkeletonReference;
 		
 	FSkeletalMeshModel* ImportedResource = SkeletalMesh->GetImportedModel();
-	if (!bIsReImport)
+	if (!ImportAssetObjectData.bIsReImport)
 	{
 		if (!ensure(ImportedResource->LODModels.Num() == 0))
 		{
@@ -1063,15 +1195,9 @@ UInterchangeFactoryBase::FImportAssetResult UInterchangeSkeletalMeshFactory::Imp
 	ensure(LodDataUniqueIds.Num() == LodCount);
 	int32 CurrentLodIndex = 0;
 
-	EInterchangeSkeletalMeshContentType ImportContent = EInterchangeSkeletalMeshContentType::All;
-	SkeletalMeshFactoryNode->GetCustomImportContentType(ImportContent);
-	const bool bApplyGeometry = !bIsReImport || (ImportContent == EInterchangeSkeletalMeshContentType::All || ImportContent == EInterchangeSkeletalMeshContentType::Geometry);
-	const bool bApplySkinning = !bIsReImport || (ImportContent == EInterchangeSkeletalMeshContentType::All || ImportContent == EInterchangeSkeletalMeshContentType::SkinningWeights);
-	const bool bApplyPartialContent = bIsReImport && ImportContent != EInterchangeSkeletalMeshContentType::All;
-	const bool bApplyGeometryOnly = bApplyPartialContent && bApplyGeometry;
-	const bool bApplySkinningOnly = bApplyPartialContent && bApplySkinning;
+	UE::Interchange::Private::FContentInfo ContentInfo = UE::Interchange::Private::GetContentInfo(SkeletalMeshFactoryNode, ImportAssetObjectData.bIsReImport);
 	
-	if (bApplySkinningOnly)
+	if (ContentInfo.bApplySkinningOnly)
 	{
 		//Ignore vertex color when we import only the skinning
 		constexpr bool bForceIgnoreVertexColor = true;
@@ -1082,7 +1208,7 @@ UInterchangeFactoryBase::FImportAssetResult UInterchangeSkeletalMeshFactory::Imp
 
 	// Update skeletal materials
 	TArray<FSkeletalMaterial>& Materials = SkeletalMesh->GetMaterials();
-	auto UpdateOrAddSkeletalMaterial = [&Materials, bIsReImport](const FName& MaterialSlotName, UMaterialInterface* MaterialInterface)
+	auto UpdateOrAddSkeletalMaterial = [&Materials, bIsReImport = ImportAssetObjectData.bIsReImport](const FName& MaterialSlotName, UMaterialInterface* MaterialInterface)
 	{
 		UMaterialInterface* NewMaterial = MaterialInterface ? MaterialInterface : UMaterial::GetDefaultMaterial(MD_Surface);
 
@@ -1130,15 +1256,12 @@ UInterchangeFactoryBase::FImportAssetResult UInterchangeSkeletalMeshFactory::Imp
 		UpdateOrAddSkeletalMaterial(MaterialSlotName, MaterialInterface ? MaterialInterface : nullptr);
 	}
 
-	ImportAssetObjectData.bIsReImport = bIsReImport;
-	ImportAssetObjectData.bApplyGeometryOnly = bApplyGeometryOnly;
-	
 	for (int32 LodIndex = 0; LodIndex < LodCount; ++LodIndex)
 	{
 		TRACE_CPUPROFILER_EVENT_SCOPE(UInterchangeSkeletalMeshFactory::CreateAsset_LOD)
 		ESkeletalMeshGeoImportVersions GeoImportVersion = ESkeletalMeshGeoImportVersions::LatestVersion;
 		ESkeletalMeshSkinningImportVersions SkinningImportVersion = ESkeletalMeshSkinningImportVersions::LatestVersion;
-		if (bIsReImport && SkeletalMesh->GetImportedModel() && SkeletalMesh->GetImportedModel()->LODModels.IsValidIndex(CurrentLodIndex))
+		if (ImportAssetObjectData.bIsReImport && SkeletalMesh->GetImportedModel() && SkeletalMesh->GetImportedModel()->LODModels.IsValidIndex(CurrentLodIndex))
 		{
 			SkeletalMesh->GetLODImportedDataVersions(CurrentLodIndex, GeoImportVersion, SkinningImportVersion);
 		}
@@ -1163,37 +1286,13 @@ UInterchangeFactoryBase::FImportAssetResult UInterchangeSkeletalMeshFactory::Imp
 			UE_LOG(LogInterchangeImport, Warning, TEXT("Invalid Skeleton LOD when importing SkeletalMesh asset %s"), *Arguments.AssetName);
 			continue;
 		}
-		FSoftObjectPath SkeletonNodeReferenceObject;
-		SkeletonNode->GetCustomReferenceObject(SkeletonNodeReferenceObject);
 
 		FSoftObjectPath SpecifiedSkeleton;
 		SkeletalMeshFactoryNode->GetCustomSkeletonSoftObjectPath(SpecifiedSkeleton);
 		bool bSpecifiedSkeleton = SpecifiedSkeleton.IsValid();
 		if (SkeletonReference == nullptr)
 		{
-			UObject* SkeletonObject = nullptr;
-
-			if (SpecifiedSkeleton.IsValid())
-			{
-				SkeletonObject = SpecifiedSkeleton.TryLoad();
-			}
-			else if (SkeletonNodeReferenceObject.IsValid())
-			{
-				SkeletonObject = SkeletonNodeReferenceObject.TryLoad();
-			}
-
-			if (SkeletonObject)
-			{
-				SkeletonReference = Cast<USkeleton>(SkeletonObject);
-
-			}
-				
-			if (!ensure(SkeletonReference))
-			{
-				UE_LOG(LogInterchangeImport, Warning, TEXT("Invalid Skeleton LOD when importing SkeletalMesh asset %s"), *Arguments.AssetName);
-				break;
-			}
-			ImportAssetObjectData.SkeletonReference = SkeletonReference;
+			break;
 		}
 
 		FString RootJointNodeId;
@@ -1209,32 +1308,15 @@ UInterchangeFactoryBase::FImportAssetResult UInterchangeSkeletalMeshFactory::Imp
 			UE_LOG(LogInterchangeImport, Warning, TEXT("Invalid Skeleton RootJointNode."));
 			continue;
 		}
+		FImportAssetObjectLODData& ImportAssetObjectLODData = ImportAssetObjectData.LodDatas[CurrentLodIndex];
+		ensure(ImportAssetObjectLODData.LodIndex == CurrentLodIndex);
+
 		FTransform RootJointNodeGlobalTransform;
 		ensure(RootJointNode->GetCustomGlobalTransform(Arguments.NodeContainer, GlobalOffsetTransform, RootJointNodeGlobalTransform));
 		FTransform RootJointNodeLocalTransform;
 		ensure(RootJointNode->GetCustomLocalTransform(RootJointNodeLocalTransform));
 		FTransform BakeToRootJointTransfromModifier = RootJointNodeGlobalTransform.Inverse() * RootJointNodeLocalTransform;
 
-		int32 SkeletonDepth = 0;
-		TArray<SkeletalMeshImportData::FBone> RefBonesBinary;
-		bool bUseTimeZeroAsBindPose = false;
-		SkeletonNode->GetCustomUseTimeZeroForBindPose(bUseTimeZeroAsBindPose);
-		bool bDiffPose = false;
-		if(bApplyGeometryOnly)
-		{
-			//Do not alter the skeletal mesh reference skeleton when importing geometry only
-			FReferenceSkeleton RefSkeleton;
-			UE::Interchange::Private::FSkeletonHelper::ProcessImportMeshSkeleton(SkeletonReference, RefSkeleton, SkeletonDepth, Arguments.NodeContainer, RootJointNodeId, RefBonesBinary, bUseTimeZeroAsBindPose, bDiffPose);
-		}
-		else
-		{
-			UE::Interchange::Private::FSkeletonHelper::ProcessImportMeshSkeleton(SkeletonReference, SkeletalMesh->GetRefSkeleton(), SkeletonDepth, Arguments.NodeContainer, RootJointNodeId, RefBonesBinary, bUseTimeZeroAsBindPose, bDiffPose);
-		}
-		if (bSpecifiedSkeleton && !SkeletonReference->IsCompatibleMesh(SkeletalMesh))
-		{
-			UE_LOG(LogInterchangeImport, Warning, TEXT("The skeleton %s is incompatible with the imported skeletalmesh asset %s"), *SkeletonReference->GetName(), *Arguments.AssetName);
-		}
-				
 		TArray<UE::Interchange::Private::FMeshNodeContext> MeshReferences;
 		//Scope to query the mesh node
 		{
@@ -1259,7 +1341,7 @@ UInterchangeFactoryBase::FImportAssetResult UInterchangeSkeletalMeshFactory::Imp
 					MeshReference.MeshNode = Cast<UInterchangeMeshNode>(Arguments.NodeContainer->GetNode(MeshDependencyUid));
 					//Cache the scene node global matrix, we will use this matrix to bake the vertices, add the node geometric mesh offset to this matrix to bake it properly
 					FTransform SceneNodeTransform;
-					if (!bUseTimeZeroAsBindPose || !MeshReference.SceneNode->GetCustomTimeZeroGlobalTransform(Arguments.NodeContainer, GlobalOffsetTransform, SceneNodeTransform))
+					if (!ImportAssetObjectLODData.bUseTimeZeroAsBindPose || !MeshReference.SceneNode->GetCustomTimeZeroGlobalTransform(Arguments.NodeContainer, GlobalOffsetTransform, SceneNodeTransform))
 					{
 						ensure(MeshReference.SceneNode->GetCustomBindPoseGlobalTransform(Arguments.NodeContainer, GlobalOffsetTransform, SceneNodeTransform));
 						if (!bBakeMeshes)
@@ -1301,12 +1383,12 @@ UInterchangeFactoryBase::FImportAssetResult UInterchangeSkeletalMeshFactory::Imp
 
 		//Add the lod mesh data to the skeletalmesh
 		FSkeletalMeshImportData SkeletalMeshImportData;
-		const bool bSkinControlPointToTimeZero = bUseTimeZeroAsBindPose && bDiffPose;
+		const bool bSkinControlPointToTimeZero = ImportAssetObjectLODData.bUseTimeZeroAsBindPose && ImportAssetObjectLODData.bDiffPose;
 		//Get all meshes and morph targets payload and fill the SkeletalMeshImportData structure
 		UE::Interchange::Private::RetrieveAllSkeletalMeshPayloadsAndFillImportData(SkeletalMeshFactoryNode
 																					, SkeletalMeshImportData
 																					, MeshReferences
-																					, RefBonesBinary
+																					, ImportAssetObjectLODData.RefBonesBinary
 																					, Arguments
 																					, MeshTranslatorPayloadInterface
 																					, bSkinControlPointToTimeZero
@@ -1326,7 +1408,7 @@ UInterchangeFactoryBase::FImportAssetResult UInterchangeSkeletalMeshFactory::Imp
 				SkeletalMeshFactoryNode->GetCustomVertexColorIgnore(bIgnoreVertexColor);
 				if (bIgnoreVertexColor)
 				{
-					if (bIsReImport)
+					if (ImportAssetObjectData.bIsReImport)
 					{
 						//Get the vertex color we have in the current asset, 
 						UE::Interchange::Private::RemapSkeletalMeshVertexColorToImportData(SkeletalMesh, LodIndex, &SkeletalMeshImportData);
@@ -1355,7 +1437,7 @@ UInterchangeFactoryBase::FImportAssetResult UInterchangeSkeletalMeshFactory::Imp
 				}
 			}
 
-			if (bApplyGeometry)
+			if (ContentInfo.bApplyGeometry)
 			{
 				// Store whether or not this mesh has vertex colors
 				SkeletalMesh->SetHasVertexColors(SkeletalMeshImportData.bHasVertexColors);
@@ -1363,7 +1445,7 @@ UInterchangeFactoryBase::FImportAssetResult UInterchangeSkeletalMeshFactory::Imp
 			}
 		}
 
-		if (bIsReImport)
+		if (ImportAssetObjectData.bIsReImport)
 		{
 			while (ImportedResource->LODModels.Num() <= CurrentLodIndex)
 			{
@@ -1379,18 +1461,18 @@ UInterchangeFactoryBase::FImportAssetResult UInterchangeSkeletalMeshFactory::Imp
 
 		UE::Interchange::Private::ProcessImportMeshInfluences(SkeletalMeshImportData.Wedges.Num(), SkeletalMeshImportData.Influences);
 
-		if (bApplyGeometryOnly)
+		if (ContentInfo.bApplyGeometryOnly)
 		{
 			FSkeletalMeshImportData::ReplaceSkeletalMeshRigImportData(SkeletalMesh, &SkeletalMeshImportData, CurrentLodIndex);
 		}
-		else if(bApplySkinningOnly)
+		else if(ContentInfo.bApplySkinningOnly)
 		{
 			FSkeletalMeshImportData::ReplaceSkeletalMeshGeometryImportData(SkeletalMesh, &SkeletalMeshImportData, CurrentLodIndex);
 		}
 
 		//Store the existing material import data before updating it so we can remap properly the material  and section data
 		TArray<FName> ExistingOriginalPerSectionMaterialImportName;
-		if (bIsReImport)
+		if (ImportAssetObjectData.bIsReImport)
 		{
 			if (CurrentLodIndex != 0)
 			{
@@ -1416,11 +1498,11 @@ UInterchangeFactoryBase::FImportAssetResult UInterchangeSkeletalMeshFactory::Imp
 		//Store the original fbx import data the SkelMeshImportDataPtr should not be modified after this
 		SkeletalMesh->SaveLODImportedData(CurrentLodIndex, SkeletalMeshImportData);
 
-		if (bApplySkinningOnly)
+		if (ContentInfo.bApplySkinningOnly)
 		{
 			SkeletalMesh->SetLODImportedDataVersions(CurrentLodIndex, GeoImportVersion, ESkeletalMeshSkinningImportVersions::LatestVersion);
 		}
-		else if (bApplyGeometryOnly)
+		else if (ContentInfo.bApplyGeometryOnly)
 		{
 			SkeletalMesh->SetLODImportedDataVersions(CurrentLodIndex, ESkeletalMeshGeoImportVersions::LatestVersion, SkinningImportVersion);
 		}
@@ -1440,7 +1522,7 @@ UInterchangeFactoryBase::FImportAssetResult UInterchangeSkeletalMeshFactory::Imp
 			NewLODInfo.bImportWithBaseMesh = true;
 		};
 
-		if (bIsReImport)
+		if (ImportAssetObjectData.bIsReImport)
 		{
 			while (SkeletalMesh->GetLODNum() <= CurrentLodIndex)
 			{
@@ -1466,8 +1548,6 @@ UInterchangeFactoryBase::FImportAssetResult UInterchangeSkeletalMeshFactory::Imp
 			SkeletalMesh->SetImportedBounds(FBoxSphereBounds((FBox)BoundingBox));
 		}
 		//Copy the data into the game thread structure so we can finish the import in the game thread callback
-		FImportAssetObjectLODData& ImportAssetObjectLODData = ImportAssetObjectData.LodDatas.AddDefaulted_GetRef();
-		ImportAssetObjectLODData.LodIndex = CurrentLodIndex;
 		ImportAssetObjectLODData.ExistingOriginalPerSectionMaterialImportName = ExistingOriginalPerSectionMaterialImportName;
 		ImportAssetObjectLODData.ImportedMaterials = SkeletalMeshImportData.Materials;
 		
diff --git a/Engine/Plugins/Interchange/Runtime/Source/Import/Public/Mesh/InterchangeSkeletalMeshFactory.h b/Engine/Plugins/Interchange/Runtime/Source/Import/Public/Mesh/InterchangeSkeletalMeshFactory.h
--- a/Engine/Plugins/Interchange/Runtime/Source/Import/Public/Mesh/InterchangeSkeletalMeshFactory.h
+++ b/Engine/Plugins/Interchange/Runtime/Source/Import/Public/Mesh/InterchangeSkeletalMeshFactory.h
@@ -25,14 +25,19 @@ public:
 		TArray<FName> ExistingOriginalPerSectionMaterialImportName;
 #if WITH_EDITOR
 		TArray<SkeletalMeshImportData::FMaterial> ImportedMaterials;
+		TArray<SkeletalMeshImportData::FBone> RefBonesBinary;
 #endif
+		bool bUseTimeZeroAsBindPose = false;
+		bool bDiffPose = false;
 	};
+
 	struct FImportAssetObjectData
 	{
 		bool bIsReImport = false;
 		USkeleton* SkeletonReference = nullptr;
 		bool bApplyGeometryOnly = false;
 		TArray<FImportAssetObjectLODData> LodDatas;
+
 		bool IsValid() const;
 	};
 	//////////////////////////////////////////////////////////////////////////
