[Interchange - glTF] - Level import - Crash when importing GLTF skeletal animation without .bin file.
Added more defensive codeing to GLTFCore.

#rnx

diff --git a/Engine/Plugins/Interchange/Runtime/Source/Parsers/GLTFCore/Private/GLTF/GLTFReader.cpp b/Engine/Plugins/Interchange/Runtime/Source/Parsers/GLTFCore/Private/GLTF/GLTFReader.cpp
--- a/Engine/Plugins/Interchange/Runtime/Source/Parsers/GLTFCore/Private/GLTF/GLTFReader.cpp
+++ b/Engine/Plugins/Interchange/Runtime/Source/Parsers/GLTFCore/Private/GLTF/GLTFReader.cpp
@@ -1041,23 +1041,21 @@ namespace GLTF
 			}
 		}
 
-		SetupObjects(BufferCount, TEXT("buffers"), [this, InResourcesPath](const FJsonObject& Object) { SetupBuffer(Object, InResourcesPath); });
-		SetupObjects(BufferViewCount, TEXT("bufferViews"), [this](const FJsonObject& Object) { SetupBufferView(Object); });
-		SetupObjects(AccessorCount, TEXT("accessors"), [this](const FJsonObject& Object) { SetupAccessor(Object); });
+		if (!SetupObjects(BufferCount, TEXT("buffers"), [this, InResourcesPath](const FJsonObject& Object) { SetupBuffer(Object, InResourcesPath); })) { return; }
+		if (!SetupObjects(BufferViewCount, TEXT("bufferViews"), [this](const FJsonObject& Object) { SetupBufferView(Object); })) { return; }
+		if (!SetupObjects(AccessorCount, TEXT("accessors"), [this](const FJsonObject& Object) { SetupAccessor(Object); })) { return; }
 
-		SetupObjects(MeshCount, TEXT("meshes"), [this, &bMeshQuantized](const FJsonObject& Object) { SetupMesh(Object, bMeshQuantized); });
-		SetupObjects(NodeCount, TEXT("nodes"), [this](const FJsonObject& Object) { SetupNode(Object); });
-		SetupObjects(SceneCount, TEXT("scenes"), [this](const FJsonObject& Object) { SetupScene(Object); });
-		SetupObjects(CameraCount, TEXT("cameras"), [this](const FJsonObject& Object) { SetupCamera(Object); });
-		SetupObjects(SkinCount, TEXT("skins"), [this](const FJsonObject& Object) { SetupSkin(Object); });
-		SetupObjects(AnimationsCount, TEXT("animations"), [this](const FJsonObject& Object) { SetupAnimation(Object); });
+		if (!SetupObjects(MeshCount, TEXT("meshes"), [this, &bMeshQuantized](const FJsonObject& Object) { SetupMesh(Object, bMeshQuantized); })) { return; }
+		if (!SetupObjects(NodeCount, TEXT("nodes"), [this](const FJsonObject& Object) { SetupNode(Object); })) { return; }
+		if (!SetupObjects(SceneCount, TEXT("scenes"), [this](const FJsonObject& Object) { SetupScene(Object); })) { return; }
+		if (!SetupObjects(CameraCount, TEXT("cameras"), [this](const FJsonObject& Object) { SetupCamera(Object); })) { return; }
+		if (!SetupObjects(SkinCount, TEXT("skins"), [this](const FJsonObject& Object) { SetupSkin(Object); })) { return; }
+		if (!SetupObjects(AnimationsCount, TEXT("animations"), [this](const FJsonObject& Object) { SetupAnimation(Object); })) { return; }
 
-		SetupObjects(ImageCount, TEXT("images"), [this, InResourcesPath, bInLoadImageData](const FJsonObject& Object) {
-			SetupImage(Object, InResourcesPath, bInLoadImageData);
-		});
-		SetupObjects(SamplerCount, TEXT("samplers"), [this](const FJsonObject& Object) { SetupSampler(Object); });
-		SetupObjects(TextureCount, TEXT("textures"), [this](const FJsonObject& Object) { SetupTexture(Object); });
-		SetupObjects(MaterialCount, TEXT("materials"), [this](const FJsonObject& Object) { SetupMaterial(Object); });
+		if (!SetupObjects(ImageCount, TEXT("images"), [this, InResourcesPath, bInLoadImageData](const FJsonObject& Object) { SetupImage(Object, InResourcesPath, bInLoadImageData); })) { return; }
+		if (!SetupObjects(SamplerCount, TEXT("samplers"), [this](const FJsonObject& Object) { SetupSampler(Object); })) { return; }
+		if (!SetupObjects(TextureCount, TEXT("textures"), [this](const FJsonObject& Object) { SetupTexture(Object); })) { return; }
+		if (!SetupObjects(MaterialCount, TEXT("materials"), [this](const FJsonObject& Object) { SetupMaterial(Object); })) { return; }
 
 
 		const TArray<TSharedPtr<FJsonValue>>* ExtensionsUsed;
@@ -1080,9 +1078,24 @@ namespace GLTF
 		BuildRootJoints();
 	}
 
+	bool FFileReader::CheckForErrors(int32 StartIndex) const
+	{
+		for (size_t Index = StartIndex; Index < Messages.Num(); Index++)
+		{
+			if (Messages[Index].Key == EMessageSeverity::Error)
+			{
+				return true;
+			}
+		}
+
+		return false;
+	}
+
 	template <typename SetupFunc>
-	void FFileReader::SetupObjects(uint32 ObjectCount, const TCHAR* FieldName, SetupFunc Func) const
+	bool FFileReader::SetupObjects(uint32 ObjectCount, const TCHAR* FieldName, SetupFunc Func) const
 	{
+		int32 StartIndex = Messages.Num();
+		
 		if (ObjectCount > 0)
 		{
 			for (const TSharedPtr<FJsonValue>& Value : JsonRoot->GetArrayField(FieldName))
@@ -1091,6 +1104,14 @@ namespace GLTF
 				Func(Object);
 			}
 		}
+
+		if (CheckForErrors(StartIndex))
+		{
+			//Any error found should automatically halt the import.
+			return false;
+		}
+
+		return true;
 	}
 
 	void FFileReader::SetupNodesType() const
@@ -1137,7 +1158,8 @@ namespace GLTF
 		for (size_t SkinIndex = 0; SkinIndex < Asset->Skins.Num(); SkinIndex++)
 		{
 			const FSkinInfo& Skin = Asset->Skins[SkinIndex];
-			if (Skin.InverseBindMatrices.Count == Skin.Joints.Num())
+			if (Skin.InverseBindMatrices.Count == Skin.Joints.Num() &&
+				Skin.InverseBindMatrices.IsValid())
 			{
 				for (size_t JointCounter = 0; JointCounter < Skin.Joints.Num(); JointCounter++)
 				{
@@ -1176,7 +1198,8 @@ namespace GLTF
 		for (size_t SkinIndex = 0; SkinIndex < Asset->Skins.Num(); SkinIndex++)
 		{
 			const FSkinInfo& Skin = Asset->Skins[SkinIndex];
-			if (Skin.InverseBindMatrices.Count == Skin.Joints.Num())
+			if (Skin.InverseBindMatrices.Count == Skin.Joints.Num() &&
+				Skin.InverseBindMatrices.IsValid())
 			{
 				for (size_t JointCounter = 0; JointCounter < Skin.Joints.Num(); JointCounter++)
 				{
@@ -1231,7 +1254,8 @@ namespace GLTF
 		for (size_t SkinIndex = 0; SkinIndex < Asset->Skins.Num(); SkinIndex++)
 		{
 			const FSkinInfo& Skin = Asset->Skins[SkinIndex];
-			if (Skin.InverseBindMatrices.Count == Skin.Joints.Num())
+			if (Skin.InverseBindMatrices.Count == Skin.Joints.Num() &&
+				Skin.InverseBindMatrices.IsValid())
 			{
 				for (size_t JointCounter = 0; JointCounter < Skin.Joints.Num(); JointCounter++)
 				{
diff --git a/Engine/Plugins/Interchange/Runtime/Source/Parsers/GLTFCore/Public/GLTFReader.h b/Engine/Plugins/Interchange/Runtime/Source/Parsers/GLTFCore/Public/GLTFReader.h
--- a/Engine/Plugins/Interchange/Runtime/Source/Parsers/GLTFCore/Public/GLTFReader.h
+++ b/Engine/Plugins/Interchange/Runtime/Source/Parsers/GLTFCore/Public/GLTFReader.h
@@ -59,7 +59,7 @@ namespace GLTF
 		void SetupMaterial(const FJsonObject& Object) const;
 
 		template <typename SetupFunc>
-		void SetupObjects(uint32 ObjectCount, const TCHAR* FieldName, SetupFunc Func) const;
+		bool SetupObjects(uint32 ObjectCount, const TCHAR* FieldName, SetupFunc Func) const;
 		void SetupNodesType() const;
 
 		void GenerateInverseBindPosesPerSkinIndices() const; //Per node
@@ -71,6 +71,8 @@ namespace GLTF
 		void BuildRootJoints() const;
 
 	private:
+		bool CheckForErrors(int32 StartIndex = 0) const;
+
 		uint32 BufferCount;
 		uint32 BufferViewCount;
 		uint32 ImageCount;
