Adding GenerateHLSLExpression for UMaterialExpressionMaterialXTextureSampleParameterBlur

diff --git a/Engine/Plugins/Interchange/Runtime/Source/Import/Private/MaterialX/MaterialExpressions/MaterialExpressionTextureSampleParameterBlur.cpp b/Engine/Plugins/Interchange/Runtime/Source/Import/Private/MaterialX/MaterialExpressions/MaterialExpressionTextureSampleParameterBlur.cpp
--- a/Engine/Plugins/Interchange/Runtime/Source/Import/Private/MaterialX/MaterialExpressions/MaterialExpressionTextureSampleParameterBlur.cpp
+++ b/Engine/Plugins/Interchange/Runtime/Source/Import/Private/MaterialX/MaterialExpressions/MaterialExpressionTextureSampleParameterBlur.cpp
@@ -7,6 +7,10 @@
 #include "Materials/MaterialExpressionCustom.h"
 #include "Materials/MaterialExpressionMultiply.h"
 #include "Materials/MaterialExpressionTextureCoordinate.h"
+#include "MaterialHLSLGenerator.h"
+#include "MaterialHLSLTree.h"
+#include "HLSLTree/HLSLTree.h"
+#include "HLSLTree/HLSLTreeCommon.h"
 
 #include UE_INLINE_GENERATED_CPP_BY_NAME(MaterialExpressionTextureSampleParameterBlur)
 
@@ -78,6 +82,23 @@ namespace Box
 	}
 }
 
+namespace
+{
+	constexpr const TCHAR* ComputeSampleSizeUV = TEXT(
+		R"(float2 derivUVx = ddx(uv) * 0.5f;
+   float2 derivUVy = ddy(uv) * 0.5f;
+   float derivX = abs(derivUVx.x) + abs(derivUVy.x);
+   float derivY = abs(derivUVx.y) + abs(derivUVy.y);
+   float sampleSizeU = 2.0f * filterSize * derivX + filterOffset;
+   if (sampleSizeU < 1.0E-05f)
+       sampleSizeU = 1.0E-05f;
+   float sampleSizeV = 2.0f * filterSize * derivY + filterOffset;
+   if (sampleSizeV < 1.0E-05f)
+       sampleSizeV = 1.0E-05f;
+   return float2(sampleSizeU, sampleSizeV);)"
+	);
+}
+
 UMaterialExpressionMaterialXTextureSampleParameterBlur::UMaterialExpressionMaterialXTextureSampleParameterBlur(const FObjectInitializer& ObjectInitializer)
 	: Super(ObjectInitializer)
 {
@@ -132,40 +153,10 @@ int32 UMaterialExpressionMaterialXTextureSampleParameterBlur::Compile(FMaterialC
 		ComputeSampleSizeInput(FilterOffset, TEXT("filterOffset"));
 
 		ComputeSampleSize->OutputType = ECustomMaterialOutputType::CMOT_Float2;
-		ComputeSampleSize->Code = TEXT(
-			R"(float2 derivUVx = ddx(uv) * 0.5f;
-   float2 derivUVy = ddy(uv) * 0.5f;
-   float derivX = abs(derivUVx.x) + abs(derivUVy.x);
-   float derivY = abs(derivUVx.y) + abs(derivUVy.y);
-   float sampleSizeU = 2.0f * filterSize * derivX + filterOffset;
-   if (sampleSizeU < 1.0E-05f)
-       sampleSizeU = 1.0E-05f;
-   float sampleSizeV = 2.0f * filterSize * derivY + filterOffset;
-   if (sampleSizeV < 1.0E-05f)
-       sampleSizeV = 1.0E-05f;
-   return float2(sampleSizeU, sampleSizeV);)"
-		);
-
-		const float* Kernel;
+		ComputeSampleSize->Code = ComputeSampleSizeUV;
+				
 		int32 FilterWidth;
-
-		switch(KernelSize)
-		{
-		case EMAterialXTextureSampleBlurKernel::Kernel5:
-			FilterWidth = 5 / 2;
-			Kernel = Filter == EMaterialXTextureSampleBlurFilter::Box ? Box::KernelWeights5x5 : Gauss::KernelWeights5x5;
-			break;
-
-		case EMAterialXTextureSampleBlurKernel::Kernel7:
-			FilterWidth = 7 / 2;
-			Kernel = Filter == EMaterialXTextureSampleBlurFilter::Box ? Box::KernelWeights7x7 : Gauss::KernelWeights7x7;
-			break;
-
-		default:
-			FilterWidth = 3 / 2;
-			Kernel = Filter == EMaterialXTextureSampleBlurFilter::Box ? Box::KernelWeights3x3 : Gauss::KernelWeights3x3;
-			break;
-		}
+		const float* Kernel = GetKernel(FilterWidth);
 
 		UMaterialExpressionCustom* Convolution = NewObject<UMaterialExpressionCustom>();
 		Convolution->OutputType = ECustomMaterialOutputType::CMOT_Float4;
@@ -173,7 +164,7 @@ int32 UMaterialExpressionMaterialXTextureSampleParameterBlur::Compile(FMaterialC
 		Convolution->Inputs.Reserve(25);
 		Convolution->Code = TEXT("float4 result = float4(0,0,0,0);\n");
 
-		int32 i = 0;
+		int32 Index = 0;
 
 		for(int Row = -FilterWidth; Row <= FilterWidth; ++Row)
 		{
@@ -198,11 +189,11 @@ int32 UMaterialExpressionMaterialXTextureSampleParameterBlur::Compile(FMaterialC
 				TextureSample->SamplerType = SamplerType;
 
 				Convolution->Inputs.Add({});
-				FString InputName{ TEXT("S") + FString::FromInt(i) };
-				Convolution->Inputs[i].InputName = FName(InputName);// FName(TEXT("S") + Index);
-				Convolution->Inputs[i].Input.Expression = TextureSample;
-				Convolution->Code += (TEXT("result += ") + InputName + TEXT("*") + FString::SanitizeFloat(Kernel[i]) + TEXT(";\n"));
-				++i;
+				FString InputName{ TEXT("S") + FString::FromInt(Index) };
+				Convolution->Inputs[Index].InputName = FName(InputName);// FName(TEXT("S") + Index);
+				Convolution->Inputs[Index].Input.Expression = TextureSample;
+				Convolution->Code += (TEXT("result += ") + InputName + TEXT("*") + FString::SanitizeFloat(Kernel[Index]) + TEXT(";\n"));
+				++Index;
 			}
 		}
 		Convolution->Code += TEXT("return result;");
@@ -220,6 +211,141 @@ void UMaterialExpressionMaterialXTextureSampleParameterBlur::GetCaption(TArray<F
 	OutCaptions.Add(TEXT("MaterialX ParamBlur"));
 	OutCaptions.Add(FString::Printf(TEXT("'%s'"), *ParameterName.ToString()));
 }
+
+bool UMaterialExpressionMaterialXTextureSampleParameterBlur::GenerateHLSLExpression(FMaterialHLSLGenerator& Generator, UE::HLSLTree::FScope& Scope, int32 OutputIndex, UE::HLSLTree::FExpression const*& OutExpression) const
+{
+	using namespace UE::HLSLTree;
+	using namespace UE::Shader;
+
+	if(KernelSize == EMAterialXTextureSampleBlurKernel::Kernel1)
+	{
+		return Super::GenerateHLSLExpression(Generator, Scope, OutputIndex, OutExpression);
+	}
+
+	FTree& Tree = Generator.GetTree();
+
+	const FExpression* ExpressionCoord = Coordinates.AcquireHLSLExpressionOrExternalInput(Generator, Scope, Material::MakeInputTexCoord(ConstCoordinate));
+
+	const FExpression* ExpressionDerivUVx = Tree.NewMul(Tree.NewExpression<FExpressionDerivative>(EDerivativeCoordinate::Ddx, ExpressionCoord), Tree.NewConstant(0.5f));
+	const FExpression* ExpressionDerivUVy = Tree.NewMul(Tree.NewExpression<FExpressionDerivative>(EDerivativeCoordinate::Ddy, ExpressionCoord), Tree.NewConstant(0.5f));
+
+	const FExpression* ExpressionDerivX = Tree.NewAdd(Tree.NewAbs(Tree.NewSwizzle(FSwizzleParameters{ 0 }, ExpressionDerivUVx)), Tree.NewAbs(Tree.NewSwizzle(FSwizzleParameters{ 0 }, ExpressionDerivUVy)));
+	const FExpression* ExpressionDerivY = Tree.NewAdd(Tree.NewAbs(Tree.NewSwizzle(FSwizzleParameters{ 1 }, ExpressionDerivUVx)), Tree.NewAbs(Tree.NewSwizzle(FSwizzleParameters{ 1 }, ExpressionDerivUVy)));
+
+	const FExpression* ExpressionSampleSizeU = Tree.NewAdd(Tree.NewMul(Tree.NewMul(Tree.NewConstant(2.f), Tree.NewConstant(FilterSize)), ExpressionDerivX), Tree.NewConstant(FilterOffset));
+	const FExpression* ExpressionSampleSizeV = Tree.NewAdd(Tree.NewMul(Tree.NewMul(Tree.NewConstant(2.f), Tree.NewConstant(FilterSize)), ExpressionDerivY), Tree.NewConstant(FilterOffset));
+
+	const FExpression* ExpressionSmallNumber = Tree.NewConstant(UE_SMALL_NUMBER);
+	const FExpression* ExpressionComputeSampleSizeUV = Tree.NewAppend(
+		Tree.NewMax(ExpressionSampleSizeU, ExpressionSmallNumber),
+		Tree.NewMax(ExpressionSampleSizeV, ExpressionSmallNumber));
+
+	int32 FilterWidth, Index = 0;
+	const float* Kernel = GetKernel(FilterWidth);
+
+	OutExpression = Tree.NewConstant({ 0.f, 0.f, 0.f, 0.f });
+	for(int32 Row = -FilterWidth; Row <= FilterWidth; ++Row)
+	{
+		for(int32 Col = -FilterWidth; Col <= FilterWidth; ++Col)
+		{
+			const FExpression* ExpressionConstantColRow = Tree.NewAppend(Tree.NewConstant(Col), Tree.NewConstant(Row));
+
+			const FExpression* ExpressionCoordinates = Tree.NewAdd(Tree.NewMul(ExpressionComputeSampleSizeUV, ExpressionConstantColRow), ExpressionCoord);
+
+			const FExpression* ExpressionTextureSample = GenerateHLSLExpressionBase(Generator, Scope, ExpressionCoordinates);
+
+			OutExpression = Tree.NewAdd(OutExpression, Tree.NewMul(ExpressionTextureSample, Tree.NewConstant(Kernel[Index++])));
+		}
+	}
+
+	return true;
+}
+
+const float* UMaterialExpressionMaterialXTextureSampleParameterBlur::GetKernel(int32& FilterWidth) const
+{
+	const float* Kernel;
+
+	switch(KernelSize)
+	{
+	case EMAterialXTextureSampleBlurKernel::Kernel5:
+		FilterWidth = 5 / 2;
+		Kernel = Filter == EMaterialXTextureSampleBlurFilter::Box ? Box::KernelWeights5x5 : Gauss::KernelWeights5x5;
+		break;
+
+	case EMAterialXTextureSampleBlurKernel::Kernel7:
+		FilterWidth = 7 / 2;
+		Kernel = Filter == EMaterialXTextureSampleBlurFilter::Box ? Box::KernelWeights7x7 : Gauss::KernelWeights7x7;
+		break;
+
+	default:
+		FilterWidth = 3 / 2;
+		Kernel = Filter == EMaterialXTextureSampleBlurFilter::Box ? Box::KernelWeights3x3 : Gauss::KernelWeights3x3;
+		break;
+	}
+
+	return Kernel;
+}
+
+const UE::HLSLTree::FExpression* UMaterialExpressionMaterialXTextureSampleParameterBlur::GenerateHLSLExpressionBase(FMaterialHLSLGenerator& Generator, UE::HLSLTree::FScope& Scope, const UE::HLSLTree::FExpression* TexCoordExpression) const
+{
+	using namespace UE::HLSLTree;
+	using namespace UE::Shader;
+	const FExpression* TextureExpression = nullptr;
+	if(ParameterName.IsNone() && TextureObject.GetTracedInput().Expression)
+	{
+		TextureExpression = TextureObject.AcquireHLSLExpression(Generator, Scope);
+	}
+	else if(Texture)
+	{
+		FMaterialParameterMetadata ParameterMeta;
+		if(!GetParameterValue(ParameterMeta))
+		{
+			Generator.Error(TEXT("Failed to get parameter value"));
+			return nullptr;
+		}
+		TextureExpression = Generator.GenerateMaterialParameter(ParameterName, ParameterMeta, SamplerType);
+	}
+
+	if(!TextureExpression)
+	{
+		Generator.Error(TEXT("Missing input texture"));
+		return nullptr;
+	}
+
+	const FExpression* MipLevelExpression = nullptr;
+	FExpressionDerivatives TexCoordDerivatives;
+	switch(MipValueMode)
+	{
+	case TMVM_None:
+		TexCoordDerivatives = Generator.GetTree().GetAnalyticDerivatives(TexCoordExpression);
+		break;
+	case TMVM_MipBias:
+	{
+		MipLevelExpression = MipValue.AcquireHLSLExpressionOrConstant(Generator, Scope, ConstMipValue);
+		TexCoordDerivatives = Generator.GetTree().GetAnalyticDerivatives(TexCoordExpression);
+		if(TexCoordDerivatives.IsValid())
+		{
+			const FExpression* DerivativeScale = Generator.GetTree().NewExp2(MipLevelExpression);
+			TexCoordDerivatives.ExpressionDdx = Generator.GetTree().NewMul(TexCoordDerivatives.ExpressionDdx, DerivativeScale);
+			TexCoordDerivatives.ExpressionDdy = Generator.GetTree().NewMul(TexCoordDerivatives.ExpressionDdy, DerivativeScale);
+		}
+		break;
+	}
+	case TMVM_MipLevel:
+		MipLevelExpression = MipValue.AcquireHLSLExpressionOrConstant(Generator, Scope, ConstMipValue);
+		break;
+	case TMVM_Derivative:
+		TexCoordDerivatives.ExpressionDdx = CoordinatesDX.AcquireHLSLExpression(Generator, Scope);
+		TexCoordDerivatives.ExpressionDdy = CoordinatesDY.AcquireHLSLExpression(Generator, Scope);
+		break;
+	default:
+		checkNoEntry();
+		break;
+	}
+
+	const FExpression* AutomaticMipBiasExpression = AutomaticViewMipBiasValue.AcquireHLSLExpressionOrConstant(Generator, Scope, (bool)AutomaticViewMipBias);
+	return Generator.GetTree().NewExpression<Material::FExpressionTextureSample>(TextureExpression, TexCoordExpression, MipLevelExpression, AutomaticMipBiasExpression, TexCoordDerivatives, SamplerSource, MipValueMode);
+}
 #endif // WITH_EDITOR
 
 #undef LOCTEXT_NAMESPACE
\ No newline at end of file
diff --git a/Engine/Plugins/Interchange/Runtime/Source/Import/Public/MaterialX/MaterialExpressions/MaterialExpressionTextureSampleParameterBlur.h b/Engine/Plugins/Interchange/Runtime/Source/Import/Public/MaterialX/MaterialExpressions/MaterialExpressionTextureSampleParameterBlur.h
--- a/Engine/Plugins/Interchange/Runtime/Source/Import/Public/MaterialX/MaterialExpressions/MaterialExpressionTextureSampleParameterBlur.h
+++ b/Engine/Plugins/Interchange/Runtime/Source/Import/Public/MaterialX/MaterialExpressions/MaterialExpressionTextureSampleParameterBlur.h
@@ -45,6 +45,14 @@ class UMaterialExpressionMaterialXTextureSampleParameterBlur: public UMaterialEx
 #if WITH_EDITOR
 	virtual int32 Compile(class FMaterialCompiler* Compiler, int32 OutputIndex) override;
 	virtual void GetCaption(TArray<FString>& OutCaptions) const override;
+	virtual bool GenerateHLSLExpression(FMaterialHLSLGenerator& Generator, UE::HLSLTree::FScope& Scope, int32 OutputIndex, UE::HLSLTree::FExpression const*& OutExpression) const override;
 #endif
 	//~ End UMaterialExpressionMaterialX Interface
+
+private:
+
+#if WITH_EDITOR
+	const float* GetKernel(int32& FilterWidth) const;
+	const UE::HLSLTree::FExpression* GenerateHLSLExpressionBase(FMaterialHLSLGenerator& Generator, UE::HLSLTree::FScope& Scope, const UE::HLSLTree::FExpression* TexCoordExpression) const;
+#endif
 };
\ No newline at end of file
diff --git a/Engine/Source/Runtime/Engine/Public/HLSLTree/HLSLTree.h b/Engine/Source/Runtime/Engine/Public/HLSLTree/HLSLTree.h
--- a/Engine/Source/Runtime/Engine/Public/HLSLTree/HLSLTree.h
+++ b/Engine/Source/Runtime/Engine/Public/HLSLTree/HLSLTree.h
@@ -624,7 +624,7 @@ public:
 
 	const FExpression* NewFunctionCall(FScope& Scope, FFunction* Function, int32 OutputIndex);
 
-	FExpressionDerivatives GetAnalyticDerivatives(const FExpression* InExpression);
+	ENGINE_API FExpressionDerivatives GetAnalyticDerivatives(const FExpression* InExpression);
 	const FExpression* GetPreviousFrame(const FExpression* InExpression, const FRequestedType& RequestedType);
 
 	FScope* NewScope(FScope& Scope);
@@ -678,7 +678,7 @@ public:
 
 	const FExpression* NewTruncateLWC(const FExpression* Input) { return NewUnaryOp(EOperation::TruncateLWC, Input); }
 
-	const FExpression* NewAppend(const FExpression* Lhs, const FExpression* Rhs);
+	ENGINE_API const FExpression* NewAppend(const FExpression* Lhs, const FExpression* Rhs);
 
 	template <typename... ArgTypes>
 	const FExpression* NewAppend(const FExpression* Lhs, ArgTypes&&... Rhs)
diff --git a/Engine/Source/Runtime/Engine/Public/MaterialHLSLGenerator.h b/Engine/Source/Runtime/Engine/Public/MaterialHLSLGenerator.h
--- a/Engine/Source/Runtime/Engine/Public/MaterialHLSLGenerator.h
+++ b/Engine/Source/Runtime/Engine/Public/MaterialHLSLGenerator.h
@@ -145,7 +145,7 @@ public:
 
 	bool GenerateStatements(UE::HLSLTree::FScope& Scope, UMaterialExpression* MaterialExpression);
 
-	const UE::HLSLTree::FExpression* GenerateMaterialParameter(FName InParameterName,
+	ENGINE_API const UE::HLSLTree::FExpression* GenerateMaterialParameter(FName InParameterName,
 		const FMaterialParameterMetadata& InParameterMeta,
 		EMaterialSamplerType InSamplerType = SAMPLERTYPE_Color,
 		const FGuid& InExternalTextureGuid = FGuid());
diff --git a/Engine/Source/Runtime/Engine/Public/MaterialHLSLTree.h b/Engine/Source/Runtime/Engine/Public/MaterialHLSLTree.h
--- a/Engine/Source/Runtime/Engine/Public/MaterialHLSLTree.h
+++ b/Engine/Source/Runtime/Engine/Public/MaterialHLSLTree.h
@@ -554,7 +554,7 @@ public:
 	virtual void EmitValueShader(FEmitContext& Context, FEmitScope& Scope, const FRequestedType& RequestedType, FEmitValueShaderResult& OutResult) const override;
 };
 
-class FExpressionTextureSample : public FExpression
+class ENGINE_API FExpressionTextureSample : public FExpression
 {
 public:
 	FExpressionTextureSample(const FExpression* InTextureExpression,
