TVariant trait compile time, PDB size and correctness improvements.

diff --git a/Engine/Plugins/Experimental/Web/WebAPI/Source/WebAPIEditor/Public/WebAPIJsonUtilities.inl b/Engine/Plugins/Experimental/Web/WebAPI/Source/WebAPIEditor/Public/WebAPIJsonUtilities.inl
--- a/Engine/Plugins/Experimental/Web/WebAPI/Source/WebAPIEditor/Public/WebAPIJsonUtilities.inl
+++ b/Engine/Plugins/Experimental/Web/WebAPI/Source/WebAPIEditor/Public/WebAPIJsonUtilities.inl
@@ -562,7 +562,7 @@ namespace UE::Json
 	bool TryParse(const TSharedPtr<FJsonValue>& InJsonValue, TVariant<ValueTypes...>& OutValue)
 	{
 		using VariantType = TVariant<ValueTypes...>;
-		if constexpr (TypeIndex < TVariantSize<VariantType>::Value)
+		if constexpr (TypeIndex < TVariantSize_V<VariantType>)
 		{
 			using ValueType = typename TNthTypeFromParameterPack<TypeIndex, ValueTypes...>::Type;
 			ValueType Value;
diff --git a/Engine/Plugins/Interchange/Runtime/Source/Import/Private/Texture/InterchangeTextureFactory.cpp b/Engine/Plugins/Interchange/Runtime/Source/Import/Private/Texture/InterchangeTextureFactory.cpp
--- a/Engine/Plugins/Interchange/Runtime/Source/Import/Private/Texture/InterchangeTextureFactory.cpp
+++ b/Engine/Plugins/Interchange/Runtime/Source/Import/Private/Texture/InterchangeTextureFactory.cpp
@@ -206,7 +206,7 @@ namespace UE::Interchange::Private::InterchangeTextureFactory
 
 	FTextureFactoryNodeVariant GetAsTextureFactoryNodeVariant(UInterchangeFactoryBaseNode* AssetNode, UClass* SupportedFactoryNodeClass)
 	{
-		static_assert(TVariantSize<FTextureFactoryNodeVariant>::Value == 7, "Please update the code below and this assert to reflect the change to the variant type.");
+		static_assert(TVariantSize_V<FTextureFactoryNodeVariant> == 7, "Please update the code below and this assert to reflect the change to the variant type.");
 
 		if (AssetNode)
 		{
@@ -270,7 +270,7 @@ namespace UE::Interchange::Private::InterchangeTextureFactory
 
 	UInterchangeTextureFactoryNode* GetTextureFactoryNodeFromVariant(const FTextureFactoryNodeVariant& FactoryNodeVariant)
 	{
-		static_assert(TVariantSize<FTextureFactoryNodeVariant>::Value == 7, "Please update the code below and this assert to reflect the change to the variant type.");
+		static_assert(TVariantSize_V<FTextureFactoryNodeVariant> == 7, "Please update the code below and this assert to reflect the change to the variant type.");
 
 		if (UInterchangeTexture2DFactoryNode* const* TextureFactoryNode = FactoryNodeVariant.TryGet<UInterchangeTexture2DFactoryNode*>())
 		{
@@ -315,9 +315,9 @@ namespace UE::Interchange::Private::InterchangeTextureFactory
 
 	FTextureNodeVariant GetTextureNodeVariantFromFactoryVariant(const FTextureFactoryNodeVariant& FactoryVariant, const UInterchangeBaseNodeContainer* NodeContainer)
 	{
-		static_assert(TVariantSize<FTextureFactoryNodeVariant>::Value == 7, "Please update the code below and this assert to reflect the change to the variant type.");
+		static_assert(TVariantSize_V<FTextureFactoryNodeVariant> == 7, "Please update the code below and this assert to reflect the change to the variant type.");
 
-		static_assert(TVariantSize<FTextureNodeVariant>::Value == 7, "Please update the code below and this assert to reflect the change to the variant type.");
+		static_assert(TVariantSize_V<FTextureNodeVariant> == 7, "Please update the code below and this assert to reflect the change to the variant type.");
 
 		FString TextureNodeUniqueID;
 
@@ -384,7 +384,7 @@ namespace UE::Interchange::Private::InterchangeTextureFactory
 
 	bool HasPayloadKey(const FTextureNodeVariant& TextureNodeVariant)
 	{
-		static_assert(TVariantSize<FTextureNodeVariant>::Value == 7, "Please update the code below and this assert to reflect the change to the variant type.");
+		static_assert(TVariantSize_V<FTextureNodeVariant> == 7, "Please update the code below and this assert to reflect the change to the variant type.");
 
 		if (const UInterchangeTexture2DNode* const* TextureNode =  TextureNodeVariant.TryGet<const UInterchangeTexture2DNode*>())
 		{
@@ -421,7 +421,7 @@ namespace UE::Interchange::Private::InterchangeTextureFactory
 
 	TOptional<FString> GetPayloadKey(const FTextureNodeVariant& TextureNodeVariant)
 	{
-		static_assert(TVariantSize<FTextureNodeVariant>::Value == 7, "Please update the code below and this assert to reflect the change to the variant type.");
+		static_assert(TVariantSize_V<FTextureNodeVariant> == 7, "Please update the code below and this assert to reflect the change to the variant type.");
 
 		if (const UInterchangeTexture2DNode* const* TextureNode =  TextureNodeVariant.TryGet<const UInterchangeTexture2DNode*>())
 		{
@@ -653,11 +653,11 @@ namespace UE::Interchange::Private::InterchangeTextureFactory
 
 	FTexturePayloadVariant GetTexturePayload(const UInterchangeSourceData* SourceData, const FString& PayloadKey, const FTextureNodeVariant& TextureNodeVariant, const FTextureFactoryNodeVariant& FactoryNodeVariant, const UInterchangeTranslatorBase* Translator, TOptional<FString>& AlternateTexturePath)
 	{
-		static_assert(TVariantSize<FTextureFactoryNodeVariant>::Value == 7, "Please update the code below and this assert to reflect the change to the variant type.");
+		static_assert(TVariantSize_V<FTextureFactoryNodeVariant> == 7, "Please update the code below and this assert to reflect the change to the variant type.");
 
-		static_assert(TVariantSize<FTexturePayloadVariant>::Value == 5, "Please update the code below and this assert to reflect the change to the variant type.");
+		static_assert(TVariantSize_V<FTexturePayloadVariant> == 5, "Please update the code below and this assert to reflect the change to the variant type.");
 
-		static_assert(TVariantSize<FTextureNodeVariant>::Value == 7, "Please update the code below and this assert to reflect the change to the variant type.");
+		static_assert(TVariantSize_V<FTextureNodeVariant> == 7, "Please update the code below and this assert to reflect the change to the variant type.");
 
 		// Standard texture 2D payload
 		if (const UInterchangeTexture2DNode* const* TextureNode =  TextureNodeVariant.TryGet<const UInterchangeTexture2DNode*>())
@@ -830,7 +830,7 @@ namespace UE::Interchange::Private::InterchangeTextureFactory
 
 	bool CanSetupTextureCubeSourceData(FTexturePayloadVariant& TexturePayload)
 	{
-		static_assert(TVariantSize<FTexturePayloadVariant>::Value == 5, "Please update the code below and this assert to reflect the change to the variant type.");
+		static_assert(TVariantSize_V<FTexturePayloadVariant> == 5, "Please update the code below and this assert to reflect the change to the variant type.");
 
 		if (TOptional<FImportSlicedImage>* SlicedImage = TexturePayload.TryGet<TOptional<FImportSlicedImage>>())
 		{
@@ -897,7 +897,7 @@ namespace UE::Interchange::Private::InterchangeTextureFactory
 
 	bool CanSetupTextureCubeArraySourceData(FTexturePayloadVariant& TexturePayload)
 	{
-		static_assert(TVariantSize<FTexturePayloadVariant>::Value == 5, "Please update the code below and this assert to reflect the change to the variant type.");
+		static_assert(TVariantSize_V<FTexturePayloadVariant> == 5, "Please update the code below and this assert to reflect the change to the variant type.");
 
 		if (TOptional<FImportSlicedImage>* SlicedImage = TexturePayload.TryGet<TOptional<FImportSlicedImage>>())
 		{
@@ -963,7 +963,7 @@ namespace UE::Interchange::Private::InterchangeTextureFactory
 
 	bool CanSetupTexture2DArraySourceData(FTexturePayloadVariant& TexturePayload)
 	{
-		static_assert(TVariantSize<FTexturePayloadVariant>::Value == 5, "Please update the code below and this assert to reflect the change to the variant type.");
+		static_assert(TVariantSize_V<FTexturePayloadVariant> == 5, "Please update the code below and this assert to reflect the change to the variant type.");
 
 		if (TOptional<FImportSlicedImage>* SlicedImage = TexturePayload.TryGet<TOptional<FImportSlicedImage>>())
 		{
@@ -1026,7 +1026,7 @@ namespace UE::Interchange::Private::InterchangeTextureFactory
 
 	bool CanSetupVolumeTextureSourceData(FTexturePayloadVariant& TexturePayload)
 	{
-		static_assert(TVariantSize<FTexturePayloadVariant>::Value == 5, "Please update the code below and this assert to reflect the change to the variant type.");
+		static_assert(TVariantSize_V<FTexturePayloadVariant> == 5, "Please update the code below and this assert to reflect the change to the variant type.");
 
 		if (TOptional<FImportSlicedImage>* SlicedImage = TexturePayload.TryGet<TOptional<FImportSlicedImage>>())
 		{
@@ -1253,7 +1253,7 @@ namespace UE::Interchange::Private::InterchangeTextureFactory
 
 	FSharedBuffer MoveRawDataToSharedBuffer(FTexturePayloadVariant& TexturePayload)
 	{
-		static_assert(TVariantSize<FTexturePayloadVariant>::Value == 5, "Please update the code below and this assert to reflect the change to the variant type.");
+		static_assert(TVariantSize_V<FTexturePayloadVariant> == 5, "Please update the code below and this assert to reflect the change to the variant type.");
 
 		if (TOptional<FImportBlockedImage>* BlockedImage = TexturePayload.TryGet<TOptional<FImportBlockedImage>>())
 		{
@@ -1335,7 +1335,7 @@ namespace UE::Interchange::Private::InterchangeTextureFactory
 
 	bool CanSetupTexture2DSourceData(FTexturePayloadVariant& TexturePayload)
 	{
-		static_assert(TVariantSize<FTexturePayloadVariant>::Value == 5, "Please update the code below and this assert to reflect the change to the variant type.");
+		static_assert(TVariantSize_V<FTexturePayloadVariant> == 5, "Please update the code below and this assert to reflect the change to the variant type.");
 
 #if WITH_EDITOR
 		if (TOptional<FImportBlockedImage>* BlockedImage = TexturePayload.TryGet<TOptional<FImportBlockedImage>>())
@@ -1899,7 +1899,7 @@ void UInterchangeTextureFactory::CheckForInvalidResolutions(UE::Interchange::Pri
 	bool bAllowNonPowerOfTwo = false;
 	TextureFactoryNode->GetCustomAllowNonPowerOfTwo(bAllowNonPowerOfTwo);
 
-	static_assert(TVariantSize<Private::InterchangeTextureFactory::FTexturePayloadVariant>::Value == 5, "Please update the code below and this assert to reflect the change to the variant type.");
+	static_assert(TVariantSize_V<Private::InterchangeTextureFactory::FTexturePayloadVariant> == 5, "Please update the code below and this assert to reflect the change to the variant type.");
 
 	FText ErrorMessage;
 	if (TOptional<FImportBlockedImage>* BlockedImagePtr = InPayloadVariant.TryGet<TOptional<FImportBlockedImage>>())
diff --git a/Engine/Plugins/Runtime/Metasound/Source/MetasoundGraphCore/Private/MetasoundLiteral.cpp b/Engine/Plugins/Runtime/Metasound/Source/MetasoundGraphCore/Private/MetasoundLiteral.cpp
--- a/Engine/Plugins/Runtime/Metasound/Source/MetasoundGraphCore/Private/MetasoundLiteral.cpp
+++ b/Engine/Plugins/Runtime/Metasound/Source/MetasoundGraphCore/Private/MetasoundLiteral.cpp
@@ -118,7 +118,7 @@ namespace Metasound
 
 			case 12:
 			default:
-				static_assert(TVariantSize<FVariantType>::Value == 13, "Possible missing FVariantType case coverage");
+				static_assert(TVariantSize_V<FVariantType> == 13, "Possible missing FVariantType case coverage");
 				return ELiteralType::Invalid;
 		}
 	}
diff --git a/Engine/Source/Runtime/Core/Public/Misc/TVariant.h b/Engine/Source/Runtime/Core/Public/Misc/TVariant.h
--- a/Engine/Source/Runtime/Core/Public/Misc/TVariant.h
+++ b/Engine/Source/Runtime/Core/Public/Misc/TVariant.h
@@ -3,9 +3,7 @@
 #pragma once
 
 #include "TVariantMeta.h"
-#include "Templates/EnableIf.h"
-#include "Templates/IsConstructible.h"
-#include "Templates/Decay.h"
+#include <type_traits>
 
 /**
  * A special tag used to indicate that in-place construction of a variant should take place.
@@ -42,7 +40,7 @@ public:
 	/** Default initialize the TVariant to the first type in the parameter pack */
 	TVariant()
 	{
-		static_assert(TIsConstructible<T>::Value, "To default-initialize a TVariant, the first type in the parameter pack must be default constructible. Use FEmptyVariantState as the first type if none of the other types can be listed first.");
+		static_assert(std::is_constructible_v<T>, "To default-initialize a TVariant, the first type in the parameter pack must be default constructible. Use FEmptyVariantState as the first type if none of the other types can be listed first.");
 		new(&UE::Core::Private::CastToStorage(*this).Storage) T();
 		TypeIndex = 0;
 	}
@@ -194,48 +192,18 @@ private:
 	uint8 TypeIndex;
 };
 
-/** Determine if a type is a variant */
-template <typename T>
-struct TIsVariant
-{
-	static constexpr inline bool Value = false;
-};
-
-template <typename... Ts>
-struct TIsVariant<TVariant<Ts...>>
-{
-	static constexpr inline bool Value = true;
-};
-
-template <typename T> struct TIsVariant<T&> : public TIsVariant<T> {};
-template <typename T> struct TIsVariant<T&&> : public TIsVariant<T> {};
-template <typename T> struct TIsVariant<const T> : public TIsVariant<T> {};
-
-/** Determine the number of types in a TVariant */
-template <typename> struct TVariantSize;
-
-template <typename... Ts>
-struct TVariantSize<TVariant<Ts...>>
-{
-	static constexpr inline SIZE_T Value = sizeof...(Ts);
-};
-
-template <typename T> struct TVariantSize<T&> : public TVariantSize<T> {};
-template <typename T> struct TVariantSize<T&&> : public TVariantSize<T> {};
-template <typename T> struct TVariantSize<const T> : public TVariantSize<T> {};
-
 /** Apply a visitor function to the list of variants */
 template <
 	typename Func,
-	typename... Variants,
-	typename = typename TEnableIf<UE::Core::Private::TIsAllVariant<typename TDecay<Variants>::Type...>::Value>::Type
+	typename... Variants
+	UE_REQUIRES((TIsVariant_V<std::decay_t<Variants>> && ...))
 >
 decltype(auto) Visit(Func&& Callable, Variants&&... Args)
 {
 #if PLATFORM_COMPILER_HAS_FOLD_EXPRESSIONS
-	constexpr SIZE_T NumPermutations = (1 * ... * (TVariantSize<Variants>::Value));
+	constexpr SIZE_T NumPermutations = (1 * ... * (TVariantSize_V<std::decay_t<Variants>>));
 #else
-	constexpr SIZE_T VariantSizes[] = { TVariantSize<Variants>::Value... };
+	constexpr SIZE_T VariantSizes[] = { TVariantSize_V<std::decay_t<Variants>>... };
 	constexpr SIZE_T NumPermutations = UE::Core::Private::Multiply(VariantSizes, sizeof...(Variants));
 #endif
 
diff --git a/Engine/Source/Runtime/Core/Public/Misc/TVariantMeta.h b/Engine/Source/Runtime/Core/Public/Misc/TVariantMeta.h
--- a/Engine/Source/Runtime/Core/Public/Misc/TVariantMeta.h
+++ b/Engine/Source/Runtime/Core/Public/Misc/TVariantMeta.h
@@ -7,7 +7,6 @@
 #include "Templates/UnrealTemplate.h"
 #include "Templates/UnrealTypeTraits.h"
 #include "Delegates/IntegerSequence.h"
-#include "Templates/AndOrNot.h"
 #include "Concepts/Insertable.h"
 
 #include "Misc/AssertionMacros.h"
@@ -15,11 +14,42 @@
 template <typename T, typename... Ts>
 class TVariant;
 
+/** Determine if a type is a variant */
 template <typename T>
-struct TIsVariant;
+constexpr bool TIsVariant_V = false;
+
+template <typename... Ts>
+constexpr bool TIsVariant_V<TVariant<Ts...>> = true;
+
+template <typename T> constexpr bool TIsVariant_V<const          T> = TIsVariant_V<T>;
+template <typename T> constexpr bool TIsVariant_V<      volatile T> = TIsVariant_V<T>;
+template <typename T> constexpr bool TIsVariant_V<const volatile T> = TIsVariant_V<T>;
+
+/** Determine the number of types in a TVariant */
+template <typename>
+constexpr SIZE_T TVariantSize_V = 0;
+
+template <typename... Ts>
+constexpr SIZE_T TVariantSize_V<TVariant<Ts...>> = sizeof...(Ts);
+
+template <typename T> constexpr SIZE_T TVariantSize_V<const          T> = TVariantSize_V<T>;
+template <typename T> constexpr SIZE_T TVariantSize_V<      volatile T> = TVariantSize_V<T>;
+template <typename T> constexpr SIZE_T TVariantSize_V<const volatile T> = TVariantSize_V<T>;
+
+template <typename T>
+struct UE_DEPRECATED(5.4, "TIsVariant<T> has been deprecated, please use TIsVariant_V<std::remove_reference_t<T>> instead") TIsVariant
+{
+	static constexpr inline bool Value = TIsVariant_V<T>;
+};
+
+/** Determine the number of types in a TVariant */
+template <typename> struct TVariantSize;
 
 template <typename T>
-struct TVariantSize;
+struct UE_DEPRECATED(5.4, "TVariantSize<T> has been deprecated, please use TVariantSize_V<std::remove_reference_t<T>> instead") TVariantSize
+{
+	static constexpr inline SIZE_T Value = TVariantSize_V<std::remove_reference_t<T>>;
+};
 
 namespace UE
 {
@@ -66,7 +96,7 @@ namespace Private
 	template <typename... Ts>
 	struct TContainsReferenceType
 	{
-		static constexpr inline bool Value = TOr<TIsReferenceType<Ts>...>::Value;
+		static constexpr inline bool Value = (std::is_reference_v<Ts> || ...);
 	};
 
 	/** Determine the max alignof and sizeof of all types in a template parameter pack and provide a type that is compatible with those sizes */
@@ -245,7 +275,7 @@ namespace Private
 	struct TVariantLoadFromArchiveLookup
 	{
 		using VariantType = TVariant<Ts...>;
-		static_assert((std::is_default_constructible<Ts>::value && ...), "Each type in TVariant template parameter pack must be default constructible in order to use FArchive serialization");
+		static_assert((std::is_default_constructible_v<Ts> && ...), "Each type in TVariant template parameter pack must be default constructible in order to use FArchive serialization");
 		static_assert((TModels_V<CInsertable<FArchive&>, Ts> && ...), "Each type in TVariant template parameter pack must be able to use operator<< with an FArchive");
 
 		/** Load the type at the specified index from the FArchive and emplace it into the TVariant */
@@ -270,10 +300,6 @@ namespace Private
 		}
 	};
 
-	/** Determine if all the types are TVariant<...> */
-	template <typename... Ts>
-	using TIsAllVariant = TAnd<TIsVariant<Ts>...>;
-
 	/** Encode the stored index of a bunch of variants into a single value used to lookup a Visit invocation function */
 	template <typename T>
 	inline SIZE_T EncodeIndices(const T& Variant)
@@ -284,7 +310,7 @@ namespace Private
 	template <typename Variant0, typename... Variants>
 	inline SIZE_T EncodeIndices(const Variant0& First, const Variants&... Rest)
 	{
-		return First.GetIndex() + TVariantSize<Variant0>::Value * EncodeIndices(Rest...);
+		return First.GetIndex() + TVariantSize_V<Variant0> * EncodeIndices(Rest...);
 	}
 
 	/** Inverse operation of EncodeIndices. Decodes an encoded index into the individual index for the specified variant index. */
@@ -336,7 +362,7 @@ namespace Private
 	template <SIZE_T EncodedIndex, SIZE_T... VariantIndices, typename Func, typename... Variants>
 	inline decltype(auto) VisitApplyEncoded(Func&& Callable, Variants&&... Args)
 	{
-		constexpr SIZE_T VariantSizes[] = { TVariantSize<Variants>::Value... };
+		constexpr SIZE_T VariantSizes[] = { TVariantSize_V<std::decay_t<Variants>>... };
 		return Callable(CastToStorage(Forward<Variants>(Args)).template GetValueAsIndexedType<DecodeIndex(EncodedIndex, VariantIndices, VariantSizes)>()...);
 	}
 
diff --git a/Engine/Source/Runtime/Core/Tests/Misc/TVariantTest.cpp b/Engine/Source/Runtime/Core/Tests/Misc/TVariantTest.cpp
--- a/Engine/Source/Runtime/Core/Tests/Misc/TVariantTest.cpp
+++ b/Engine/Source/Runtime/Core/Tests/Misc/TVariantTest.cpp
@@ -12,29 +12,25 @@
 
 namespace UE::Core::Private
 {
-	// TIsAllVariant
-	static_assert(TIsAllVariant<TVariant<int, float>>::Value, "Expected TIsAllVariant<TVariant> to be true");
-	static_assert(TIsAllVariant<TVariant<int, float>, TVariant<float, int>>::Value, "Expected TIsAllVariant<TVariant, TVariant> to be true");
-	static_assert(TIsAllVariant<TVariant<int, float>, const TVariant<float, int>>::Value, "Expected TIsAllVariant<TVariant, const TVariant> to be true");
-	static_assert(TIsAllVariant<>::Value, "Expected TIsAllVariant<> for base case to be true");
-	static_assert(TIsAllVariant<bool>::Value == false, "Expected TIsAllVariant<non-TVariant> to be false");
-	static_assert(TIsAllVariant<TVariant<int, float>, bool>::Value == false, "Expected TIsAllVariant<TVariant, non-TVariant> to be false");
-	static_assert(TIsAllVariant<TVariant<int, float>, bool, TVariant<float, int>>::Value == false, "Expected TIsAllVariant<TVariant, non-TVariant, TVariant> to be false");
-
-	// TIsVariant
-	static_assert(TIsVariant<TVariant<int, float>>::Value, "Expected TIsVariant<TVariant> to be true");
-	static_assert(TIsVariant<const TVariant<int, float>>::Value, "Expected TIsVariant<const TVariant> to be true");
-	static_assert(TIsVariant<TVariant<int, float>&>::Value, "Expected TIsVariant<TVariant&> to be true");
-	static_assert(TIsVariant<const TVariant<int, float>&>::Value, "Expected TIsVariant<const TVariant&> to be true");
-	static_assert(TIsVariant<TVariant<int, float>&&>::Value, "Expected TIsVariant<TVariant&&> to be true");
-	static_assert(TIsVariant<bool>::Value == false, "Expected TIsVariant<non-TVariant> to be false");
-
-	// TVariantSize
-	static_assert(TVariantSize<TVariant<int, float>>::Value == 2, "Expected TVariantSize<TVariant> to be 2");
-	static_assert(TVariantSize<const TVariant<int, float>>::Value == 2, "Expected TVariantSize<const TVariant> to be 2");
-	static_assert(TVariantSize<TVariant<int, float>&>::Value == 2, "Expected TVariantSize<TVariant&> to be 2");
-	static_assert(TVariantSize<const TVariant<int, float>&>::Value == 2, "Expected TVariantSize<const TVariant&> to be 2");
-	static_assert(TVariantSize<TVariant<int, float>&&>::Value == 2, "Expected TVariantSize<TVariant&&> to be 2");
+	// TIsVariant_V of variants
+	static_assert(TIsVariant_V<      TVariant<int, float>>, "Expected TIsVariant_V<TVariant> to be true");
+	static_assert(TIsVariant_V<const TVariant<int, float>>, "Expected TIsVariant_V<const TVariant> to be true");
+
+	// TIsVariant_V of non-variants (references are not variants)
+	static_assert(TIsVariant_V<      TVariant<int, float>&>  == false, "Expected TIsVariant_V<TVariant&> to be false");
+	static_assert(TIsVariant_V<const TVariant<int, float>&>  == false, "Expected TIsVariant_V<const TVariant&> to be false");
+	static_assert(TIsVariant_V<      TVariant<int, float>&&> == false, "Expected TIsVariant_V<TVariant&&> to be false");
+	static_assert(TIsVariant_V<      bool>                   == false, "Expected TIsVariant_V<non-TVariant> to be false");
+
+	// TVariantSize_V of variants (references are not variants)
+	static_assert(TVariantSize_V<      TVariant<int, float>> == 2, "Expected TVariantSize_V<TVariant> to be 2");
+	static_assert(TVariantSize_V<const TVariant<int, float>> == 2, "Expected TVariantSize_V<const TVariant> to be 2");
+
+	// TVariantSize_V of non-variants
+	static_assert(TVariantSize_V<      TVariant<int, float>&>  == 0, "Expected TVariantSize_V<TVariant&> to be 0");
+	static_assert(TVariantSize_V<const TVariant<int, float>&>  == 0, "Expected TVariantSize_V<const TVariant&> to be 0");
+	static_assert(TVariantSize_V<      TVariant<int, float>&&> == 0, "Expected TVariantSize_V<TVariant&&> to be 0");
+	static_assert(TVariantSize_V<      bool>                   == 0, "Expected TVariantSize_V<non-TVariant> to be 0");
 } // namespace UE::Core::Private
 
 struct FNonDefaultConstructible
diff --git a/Engine/Source/Runtime/Datasmith/DatasmithCore/Private/DatasmithCloth.cpp b/Engine/Source/Runtime/Datasmith/DatasmithCore/Private/DatasmithCloth.cpp
--- a/Engine/Source/Runtime/Datasmith/DatasmithCore/Private/DatasmithCloth.cpp
+++ b/Engine/Source/Runtime/Datasmith/DatasmithCore/Private/DatasmithCloth.cpp
@@ -35,7 +35,7 @@ FArchive& operator<<(FArchive& Ar, FParameterData& ParameterData)
 {
 	Ar.UsingCustomVersion(FDatasmithClothSerializationVersion::GUID);
 
-	static_assert(TVariantSize<decltype(ParameterData.Data)>::Value == 2, "Serialization code not synced with structure");
+	static_assert(TVariantSize_V<decltype(ParameterData.Data)> == 2, "Serialization code not synced with structure");
 	static_assert(decltype(ParameterData.Data)::IndexOfType<TArray<float>>() == 0, "Serialization relies on this specific order");
 	static_assert(decltype(ParameterData.Data)::IndexOfType<TArray<double>>() == 1, "Serialization relies on this specific order");
 
