Interchange: Remove bad ensure when filtering options for generic asset pipeline.

#rnx

diff --git a/Engine/Plugins/Interchange/Runtime/Source/Pipelines/Private/InterchangeGenericAssetsPipeline.cpp b/Engine/Plugins/Interchange/Runtime/Source/Pipelines/Private/InterchangeGenericAssetsPipeline.cpp
--- a/Engine/Plugins/Interchange/Runtime/Source/Pipelines/Private/InterchangeGenericAssetsPipeline.cpp
+++ b/Engine/Plugins/Interchange/Runtime/Source/Pipelines/Private/InterchangeGenericAssetsPipeline.cpp
@@ -230,11 +230,6 @@ void UInterchangeGenericAssetsPipeline::FilterPropertiesFromTranslatedData(UInte
 		{
 			bHideStaticMeshes = true;
 		}
-		else
-		{
-			constexpr bool bFilterCodeLogicError = false;
-			ensure(bFilterCodeLogicError);
-		}
 
 		if (SkeletalMeshes.Num() > 0)
 		{
