fix default texture SRGB setting at import time should be set primarily by TC and only secondarily by the import image SRGB
fixes U16 images coming in with SRGB off, also mixed format udims
SRGB flag on the texture primarily controls the output encoding, which we want to be SRGB for Color

#rnx

diff --git a/Engine/Plugins/Interchange/Runtime/Source/Import/Private/Texture/InterchangeTextureFactory.cpp b/Engine/Plugins/Interchange/Runtime/Source/Import/Private/Texture/InterchangeTextureFactory.cpp
--- a/Engine/Plugins/Interchange/Runtime/Source/Import/Private/Texture/InterchangeTextureFactory.cpp
+++ b/Engine/Plugins/Interchange/Runtime/Source/Import/Private/Texture/InterchangeTextureFactory.cpp
@@ -577,7 +577,7 @@ namespace UE::Interchange::Private::InterchangeTextureFactory
 
 			if (bMismatchedGammaSpace || bMismatchedFormats)
 			{
-				UE_LOG(LogInterchangeImport, Warning, TEXT("Mismatched UDIM image %s, converting all to %s/%s ..."), bMismatchedGammaSpace ? TEXT("gamma spaces") : TEXT("pixel formats"),
+				UE_LOG(LogInterchangeImport, Display, TEXT("Mismatched UDIM image %s, converting all to %s/%s ..."), bMismatchedGammaSpace ? TEXT("gamma spaces") : TEXT("pixel formats"),
 					ERawImageFormat::GetName(FImageCoreUtils::ConvertToRawImageFormat(BlockedImage.Format)), BlockedImage.bSRGB ? TEXT("sRGB") : TEXT("Linear"));
 
 				for (UE::Interchange::FImportImage& Image : Images)
@@ -749,7 +749,8 @@ namespace UE::Interchange::Private::InterchangeTextureFactory
 		if (!bIsReimport)
 		{
 			Texture->CompressionSettings = Image.CompressionSettings;
-			Texture->SRGB = Image.bSRGB;
+			
+			Texture->SRGB = UE::TextureUtilitiesCommon::GetDefaultSRGB(Image.CompressionSettings,Image.Format,Image.bSRGB);
 
 			//If the MipGenSettings was set by the translator, we must apply it before the build
 			if (Image.MipGenSettings.IsSet())
@@ -775,7 +776,9 @@ namespace UE::Interchange::Private::InterchangeTextureFactory
 			if (!bIsReimport)
 			{
 				Texture2D->CompressionSettings = BlockedImage.CompressionSettings;
-				Texture2D->SRGB = BlockedImage.bSRGB;
+				
+				Texture2D->SRGB = UE::TextureUtilitiesCommon::GetDefaultSRGB(BlockedImage.CompressionSettings,BlockedImage.Format,BlockedImage.bSRGB);
+
 				Texture2D->VirtualTextureStreaming = true;
 
 				if (BlockedImage.MipGenSettings.IsSet())
@@ -818,7 +821,8 @@ namespace UE::Interchange::Private::InterchangeTextureFactory
 		if (!bIsReimport)
 		{
 			Texture->CompressionSettings = SlicedImage.CompressionSettings;
-			Texture->SRGB = SlicedImage.bSRGB;
+			
+			Texture->SRGB = UE::TextureUtilitiesCommon::GetDefaultSRGB(SlicedImage.CompressionSettings,SlicedImage.Format,SlicedImage.bSRGB);
 
 			if (SlicedImage.MipGenSettings.IsSet())
 			{
@@ -1191,10 +1195,11 @@ namespace UE::Interchange::Private::InterchangeTextureFactory
 
 		return TasksToDo;
 	}
-#else // WITH_EDITOR
+#else // !WITH_EDITOR
 	void SetupTexture2DSourceDataFromBulkData_Runtime(UTexture2D* Texture2D, const FImportImage& ImportImage, UE::Serialization::FEditorBulkData::FSharedBufferWithID&& BufferAndId)
 	{
-		Texture2D->SRGB = ImportImage.bSRGB;
+		Texture2D->SRGB = UE::TextureUtilitiesCommon::GetDefaultSRGB(Texture2D->CompressionSettings,ImportImage.Format,ImportImage.bSRGB);
+		
 		ERawImageFormat::Type PixelFormatRawFormat;
 		const EPixelFormat PixelFormat = FImageCoreUtils::GetPixelFormatForRawImageFormat(FImageCoreUtils::ConvertToRawImageFormat(ImportImage.Format), &PixelFormatRawFormat);
 
@@ -1202,9 +1207,9 @@ namespace UE::Interchange::Private::InterchangeTextureFactory
 		BulkData.UpdatePayload(MoveTemp(BufferAndId));
 		TFuture<FSharedBuffer> Payload = BulkData.GetPayload();
 
-		const EGammaSpace GammaSpace = ImportImage.bSRGB ? EGammaSpace::sRGB : EGammaSpace::Linear;
+		const EGammaSpace SourceGammaSpace = ImportImage.bSRGB ? EGammaSpace::sRGB : EGammaSpace::Linear;
 		constexpr int32 NumSlices = 1;
-		FImageView SourceImageView(const_cast<void*>(Payload.Get().GetData()), ImportImage.SizeX, ImportImage.SizeY, NumSlices, PixelFormatRawFormat, GammaSpace);
+		FImageView SourceImageView(const_cast<void*>(Payload.Get().GetData()), ImportImage.SizeX, ImportImage.SizeY, NumSlices, PixelFormatRawFormat, SourceGammaSpace);
 
 		FImage DecompressedSourceImage;
 		if (ImportImage.RawDataCompressionFormat != TSCF_None)
@@ -1235,11 +1240,21 @@ namespace UE::Interchange::Private::InterchangeTextureFactory
 
 			Texture2D->bNotOfflineProcessed = true;
 
+			EGammaSpace MipGammaSpace;
+			if ( Texture2D->SRGB && ERawImageFormat::GetFormatNeedsGammaSpace(PixelFormatRawFormat) )
+			{
+				MipGammaSpace = EGammaSpace::sRGB;
+			}
+			else
+			{
+				MipGammaSpace = EGammaSpace::Linear;
+			}
+
 			uint8* MipData = static_cast<uint8*>(Texture2D->GetPlatformData()->Mips[0].BulkData.Lock(LOCK_READ_WRITE));
 			check(MipData != nullptr);
 			int64 MipDataSize = Texture2D->GetPlatformData()->Mips[0].BulkData.GetBulkDataSize();
 
-			FImageView MipImageView(MipData, ImportImage.SizeX, ImportImage.SizeY, NumSlices, PixelFormatRawFormat, GammaSpace);
+			FImageView MipImageView(MipData, ImportImage.SizeX, ImportImage.SizeY, NumSlices, PixelFormatRawFormat, MipGammaSpace);
 
 			// copy into texture and convert if necessary :
 			FImageCore::CopyImage(SourceImageView, MipImageView);
diff --git a/Engine/Source/Editor/UnrealEd/Private/Factories/EditorFactories.cpp b/Engine/Source/Editor/UnrealEd/Private/Factories/EditorFactories.cpp
--- a/Engine/Source/Editor/UnrealEd/Private/Factories/EditorFactories.cpp
+++ b/Engine/Source/Editor/UnrealEd/Private/Factories/EditorFactories.cpp
@@ -3049,6 +3049,9 @@ bool UTextureFactory::ImportImage(const uint8* Buffer, int64 Length, FFeedbackCo
 				if ( OutImage.CompressionSettings == TC_Grayscale && TGA->ImageTypeCode == 3)
 				{
 					// default grayscales to linear as they wont get compression otherwise and are commonly used as masks
+					// -> this is wrong way to do this
+					//	Image.SRGB should be the setting of the import image data
+					//	not contain information about how we want the texture platform data to be set up
 					OutImage.SRGB = false;
 				}
 			}
@@ -3343,7 +3346,7 @@ UTexture* UTextureFactory::ImportTextureUDIM(UClass* Class, UObject* InParent, F
 
 	if (ColorSpaceMode == ETextureSourceColorSpace::Auto)
 	{
-		Texture->SRGB = bSRGB;
+		Texture->SRGB = UE::TextureUtilitiesCommon::GetDefaultSRGB(TCSettings,Format,bSRGB);
 	}
 	else if (ColorSpaceMode == ETextureSourceColorSpace::Linear)
 	{
@@ -3541,7 +3544,7 @@ UTexture * UTextureFactory::ImportDDS(const uint8* Buffer,int64 Length,UObject*
 		check( bSRGB == false );
 	}
 
-	Texture->SRGB = bSRGB;
+	Texture->SRGB = UE::TextureUtilitiesCommon::GetDefaultSRGB(Texture->CompressionSettings,TSFormat,bSRGB);
 	
 	if ( MipCount > 1)
 	{
@@ -3764,30 +3767,22 @@ UTexture* UTextureFactory::ImportTexture(UClass* Class, UObject* InParent, FName
 
 			Texture->CompressionSettings = Image.CompressionSettings;
 
-			// check Texture Format before setting SRGB
-			if ( ERawImageFormat::GetFormatNeedsGammaSpace( FImageCoreUtils::ConvertToRawImageFormat(Image.Format) ) )
+			if (ColorSpaceMode == ETextureSourceColorSpace::Auto)
 			{
-				if (ColorSpaceMode == ETextureSourceColorSpace::Auto)
-				{
-					Texture->SRGB = Image.SRGB;
-				}
-				else if (ColorSpaceMode == ETextureSourceColorSpace::Linear)
-				{
-					Texture->SRGB = false;
-				}
-				else if (ColorSpaceMode == ETextureSourceColorSpace::SRGB)
-				{
-					Texture->SRGB = true;
-				}
-				else
-				{
-					check(0);
-				}
+				Texture->SRGB = UE::TextureUtilitiesCommon::GetDefaultSRGB(Image.CompressionSettings,Image.Format,Image.SRGB);
 			}
-			else
+			else if (ColorSpaceMode == ETextureSourceColorSpace::Linear)
 			{
 				Texture->SRGB = false;
 			}
+			else if (ColorSpaceMode == ETextureSourceColorSpace::SRGB)
+			{
+				Texture->SRGB = true;
+			}
+			else
+			{
+				check(0);
+			}
 		}
 		return Texture;
 	}
diff --git a/Engine/Source/Runtime/Engine/Private/Texture.cpp b/Engine/Source/Runtime/Engine/Private/Texture.cpp
--- a/Engine/Source/Runtime/Engine/Private/Texture.cpp
+++ b/Engine/Source/Runtime/Engine/Private/Texture.cpp
@@ -2695,9 +2695,15 @@ ETextureClass FTextureSource::GetTextureClass() const
 	}
 }
 
+// Texture->GetGammaSpace is the desired gamma space of the Platform texture we make
+// TextureSource->GetGammaSpace is the way the source image pixels should be interpretted
 EGammaSpace FTextureSource::GetGammaSpace(int LayerIndex) const
 {
-	// Texture->GetGammaSpace does not validate against format, but I do? a bit weird, fix that
+	// note: does not respect ETextureSourceEncoding EncodingOverride (but should)
+	//	in most cases it is not possible to map EncodingOverride into our FImage GammaSpace
+	//	even when EncodingOverride is Linear or sRGB it's often not possible, eg. if format is F32 we don't allow that to be SRGB
+	//	in some cases it is possible; so we could return an EGammaSpace from EncodingOverride in those cases
+
 	if ( ! ERawImageFormat::GetFormatNeedsGammaSpace( FImageCoreUtils::ConvertToRawImageFormat(GetFormat(LayerIndex)) ) )
 	{
 		return EGammaSpace::Linear;
diff --git a/Engine/Source/Runtime/TextureUtilitiesCommon/Private/TextureImportSettings.cpp b/Engine/Source/Runtime/TextureUtilitiesCommon/Private/TextureImportSettings.cpp
--- a/Engine/Source/Runtime/TextureUtilitiesCommon/Private/TextureImportSettings.cpp
+++ b/Engine/Source/Runtime/TextureUtilitiesCommon/Private/TextureImportSettings.cpp
@@ -5,6 +5,7 @@
 #include "UObject/UnrealType.h"
 #include "UObject/PropertyPortFlags.h"
 #include "Misc/ConfigCacheIni.h"
+#include "ImageCoreUtils.h"
 
 #if WITH_EDITOR
 #include "Engine/Texture.h"
@@ -64,9 +65,9 @@ void UTextureImportSettings::PostEditChangeProperty(FPropertyChangedEvent& Prope
 
 
 
-#if WITH_EDITOR
 namespace UE::TextureUtilitiesCommon
 {
+	#if WITH_EDITOR
 	/* Set default properties on Texture for newly imported textures, or reimports.
 	*  Should be called after all texture properties are set, before PostEditChange() 
 	*/
@@ -151,6 +152,49 @@ namespace UE::TextureUtilitiesCommon
 			}
 		}
 	}
+	#endif // WITH_EDITOR
 
-}
-#endif
\ No newline at end of file
+	/* Get the default value for Texture->SRGB
+	* ImportImageSRGB is the SRGB setting of the imported image
+	*/
+	TEXTUREUTILITIESCOMMON_API bool GetDefaultSRGB(TextureCompressionSettings TC, ETextureSourceFormat ImportImageFormat, bool ImportImageSRGB)
+	{
+		// Texture->SRGB sets the gamma correction of the platform texture we make
+		//	so this is not just = ImportImageSRGB
+		
+		if ( TC == TC_Default || TC == TC_EditorIcon )
+		{
+			// DXT1,DXT3,R8G8B8 encodings
+			//	we typically want SRGB on for these (for the platform encoding)
+			// only exception is if the source is U8 linear
+			//	 in that case, staying U8 linear probably preserves bits better
+			//	 and ALSO we don't have the choice of turning on SRGB even if we want to
+			//	 because of the way it is overloaded to mean both the source encoding and the platform encoding
+
+			if ( ERawImageFormat::GetFormatNeedsGammaSpace( FImageCoreUtils::ConvertToRawImageFormat(ImportImageFormat) ) )
+			{
+				// if the imported image supported gamma (eg. U8)
+				//	then we will set SRGB=false if it was U8-linear (can only happen with DDS U8 linear import, very rare)
+				// note that texture SRGB flag in this case affects both the source interpretation and the platform encoding
+				return ImportImageSRGB;
+			}
+			else
+			{
+				// counter-intuitively, U16 and F32 always want SRGB *on*
+				//	the source will be treated as linear no matter what we set SRGB to
+				//  SRGB will only affect the Platform encoding, so we prefer that to be sRGB color space
+				return true;
+			}
+		}
+		else
+		{
+			// TC_HDR, NormalMap, etc. we want SRGB off
+
+			// TC_Grayscale we would prefer to have SRGB on, but default to off because
+			//	G8 + SRGB is not supported well
+
+			return false;
+		}
+	}
+
+} // TextureUtilitiesCommon
diff --git a/Engine/Source/Runtime/TextureUtilitiesCommon/Public/TextureImportSettings.h b/Engine/Source/Runtime/TextureUtilitiesCommon/Public/TextureImportSettings.h
--- a/Engine/Source/Runtime/TextureUtilitiesCommon/Public/TextureImportSettings.h
+++ b/Engine/Source/Runtime/TextureUtilitiesCommon/Public/TextureImportSettings.h
@@ -6,6 +6,7 @@
 #include "UObject/ObjectMacros.h"
 #include "Engine/EngineTypes.h"
 #include "Engine/DeveloperSettings.h"
+#include "Engine/TextureDefines.h"
 
 #include "TextureImportSettings.generated.h"
 
@@ -93,15 +94,19 @@ public:
 };
 
 
-#if WITH_EDITOR
 class UTexture;
 
 namespace UE::TextureUtilitiesCommon
 {
+#if WITH_EDITOR
 	/* Set default properties on Texture for newly imported textures, or reimports.
 	*  Should be called after all texture properties are set, before PostEditChange() 
 	*/
 	TEXTUREUTILITIESCOMMON_API void ApplyDefaultsForNewlyImportedTextures(UTexture * Texture, bool bIsReimport);
+#endif
 
+	/* Get the default value for Texture->SRGB
+	* ImportImageSRGB is the SRGB setting of the imported image
+	*/
+	TEXTUREUTILITIESCOMMON_API bool GetDefaultSRGB(TextureCompressionSettings TC, ETextureSourceFormat ImportImageFormat, bool ImportImageSRGB); 
 }
-#endif
