Interchange: fix the anim sequence re-import handler which was not answering true when ask CanImport.

#rnx

diff --git a/Engine/Plugins/Interchange/Runtime/Source/Import/Private/Animation/InterchangeAnimSequenceFactory.cpp b/Engine/Plugins/Interchange/Runtime/Source/Import/Private/Animation/InterchangeAnimSequenceFactory.cpp
--- a/Engine/Plugins/Interchange/Runtime/Source/Import/Private/Animation/InterchangeAnimSequenceFactory.cpp
+++ b/Engine/Plugins/Interchange/Runtime/Source/Import/Private/Animation/InterchangeAnimSequenceFactory.cpp
@@ -1187,9 +1187,9 @@ void UInterchangeAnimSequenceFactory::SetupObject_GameThread(const FSetupObjectP
 bool UInterchangeAnimSequenceFactory::GetSourceFilenames(const UObject* Object, TArray<FString>& OutSourceFilenames) const
 {
 #if WITH_EDITORONLY_DATA
-	if (AnimSequence && (AnimSequence == Cast<UAnimSequence>(Object)))
+	if (const UAnimSequence* TempAnimSequence = Cast<UAnimSequence>(Object))
 	{
-		return UE::Interchange::FFactoryCommon::GetSourceFilenames(AnimSequence->AssetImportData.Get(), OutSourceFilenames);
+		return UE::Interchange::FFactoryCommon::GetSourceFilenames(TempAnimSequence->AssetImportData.Get(), OutSourceFilenames);
 	}
 #endif
 
@@ -1199,9 +1199,9 @@ bool UInterchangeAnimSequenceFactory::GetSourceFilenames(const UObject* Object,
 bool UInterchangeAnimSequenceFactory::SetSourceFilename(const UObject* Object, const FString& SourceFilename, int32 SourceIndex) const
 {
 #if WITH_EDITORONLY_DATA
-	if (AnimSequence && (AnimSequence == Cast<UAnimSequence>(Object)))
+	if (const UAnimSequence* TempAnimSequence = Cast<UAnimSequence>(Object))
 	{
-		return UE::Interchange::FFactoryCommon::SetSourceFilename(AnimSequence->AssetImportData.Get(), SourceFilename, SourceIndex);
+		return UE::Interchange::FFactoryCommon::SetSourceFilename(TempAnimSequence->AssetImportData.Get(), SourceFilename, SourceIndex);
 	}
 #endif
 
