Interchange: Avoid creating texture asset if the source filename cannot be open.

#rnx

diff --git a/Engine/Plugins/Interchange/Runtime/Source/Import/Private/Texture/TextureTranslatorUtilities.cpp b/Engine/Plugins/Interchange/Runtime/Source/Import/Private/Texture/TextureTranslatorUtilities.cpp
--- a/Engine/Plugins/Interchange/Runtime/Source/Import/Private/Texture/TextureTranslatorUtilities.cpp
+++ b/Engine/Plugins/Interchange/Runtime/Source/Import/Private/Texture/TextureTranslatorUtilities.cpp
@@ -112,7 +112,7 @@ bool UE::Interchange::FTextureTranslatorUtilities::IsTranslatorValid(const UInte
 
 	if (!SourceData)
 	{
-		LogError(TextureTranslator, FText::Format(LOCTEXT("TextureImportFailure_BadData", "Failed to import {0}, bad source data."), FText::FromString(Format)));
+		LogError(TextureTranslator, FText::Format(LOCTEXT("TextureImportFailure_BadData", "{0}: Failed to import, bad source data."), FText::FromString(Format)));
 		return false;
 	}
 
@@ -120,7 +120,7 @@ bool UE::Interchange::FTextureTranslatorUtilities::IsTranslatorValid(const UInte
 
 	if (!FPaths::FileExists(Filename))
 	{
-		LogError(TextureTranslator, FText::Format(LOCTEXT("TextureImportFailure_OpenFile", "Failed to import {0}, cannot open file."), FText::FromString(Format)));
+		LogError(TextureTranslator, FText::Format(LOCTEXT("TextureImportFailure_OpenFile", "{0}: Failed to import texture asset, cannot open file [{1}]."), FText::FromString(Format), FText::FromString(Filename)));
 		return false;
 	}
 
diff --git a/Engine/Plugins/Interchange/Runtime/Source/Parsers/Fbx/Private/FbxMaterial.cpp b/Engine/Plugins/Interchange/Runtime/Source/Parsers/Fbx/Private/FbxMaterial.cpp
--- a/Engine/Plugins/Interchange/Runtime/Source/Parsers/Fbx/Private/FbxMaterial.cpp
+++ b/Engine/Plugins/Interchange/Runtime/Source/Parsers/Fbx/Private/FbxMaterial.cpp
@@ -41,6 +41,12 @@ namespace UE
 				{
 					return nullptr;
 				}
+				FString NormalizeFilePath = TextureFilePath;
+				FPaths::NormalizeFilename(NormalizeFilePath);
+				if (!FPaths::FileExists(NormalizeFilePath))
+				{
+					return nullptr;
+				}
 				const FString TextureName = FPaths::GetBaseFilename(TextureFilePath);
 				const FString TextureNodeID = UInterchangeTextureNode::MakeNodeUid(TextureName);
 
@@ -52,8 +58,7 @@ namespace UE
 				UInterchangeTexture2DNode* NewTextureNode = UInterchangeTexture2DNode::Create(&NodeContainer, TextureName);
 
 				//All texture translator expect a file as the payload key
-				FString NormalizeFilePath = TextureFilePath;
-				FPaths::NormalizeFilename(NormalizeFilePath);
+				
 				NewTextureNode->SetPayLoadKey(NormalizeFilePath);
 
 				return NewTextureNode;
@@ -352,23 +357,25 @@ namespace UE
 							if (FbxFileTexture* FbxTexture = MaterialProperty.GetSrcObject<FbxFileTexture>(0))
 							{
 								const FString TexturePath = FbxTexture->GetFileName();
-
-								if (const UInterchangeTexture2DNode* TextureNode = CreateTexture2DNode(NodeContainer, TexturePath))
+								if (TexturePath.IsEmpty() || !FPaths::FileExists(TexturePath))
 								{
-									const FString TextureName = FPaths::GetBaseFilename(TexturePath);
+									if (const UInterchangeTexture2DNode* TextureNode = CreateTexture2DNode(NodeContainer, TexturePath))
+									{
+										const FString TextureName = FPaths::GetBaseFilename(TexturePath);
 
-									// NormalFromHeightmap needs TextureObject(not just a sample as it takes multiple samples from it)
-									UInterchangeShaderNode* TextureObjectNode = UInterchangeShaderNode::Create(&NodeContainer, TextureName, ShaderGraphNode->GetUniqueID());
-									TextureObjectNode->SetCustomShaderType(TextureObject::Name.ToString());
-									TextureObjectNode->AddStringAttribute(UInterchangeShaderPortsAPI::MakeInputValueKey(TextureObject::Inputs::Texture.ToString()), TextureNode->GetUniqueID());
+										// NormalFromHeightmap needs TextureObject(not just a sample as it takes multiple samples from it)
+										UInterchangeShaderNode* TextureObjectNode = UInterchangeShaderNode::Create(&NodeContainer, TextureName, ShaderGraphNode->GetUniqueID());
+										TextureObjectNode->SetCustomShaderType(TextureObject::Name.ToString());
+										TextureObjectNode->AddStringAttribute(UInterchangeShaderPortsAPI::MakeInputValueKey(TextureObject::Inputs::Texture.ToString()), TextureNode->GetUniqueID());
 
-									UInterchangeShaderNode* HeightMapNode = UInterchangeShaderNode::Create(&NodeContainer, NormalFromHeightMap::Name.ToString(), ShaderGraphNode->GetUniqueID());
-									HeightMapNode->SetCustomShaderType(NormalFromHeightMap::Name.ToString());
+										UInterchangeShaderNode* HeightMapNode = UInterchangeShaderNode::Create(&NodeContainer, NormalFromHeightMap::Name.ToString(), ShaderGraphNode->GetUniqueID());
+										HeightMapNode->SetCustomShaderType(NormalFromHeightMap::Name.ToString());
 
-									UInterchangeShaderPortsAPI::ConnectDefaultOuputToInput(HeightMapNode, NormalFromHeightMap::Inputs::HeightMap.ToString(), TextureObjectNode->GetUniqueID());
-									HeightMapNode->AddFloatAttribute(UInterchangeShaderPortsAPI::MakeInputValueKey(NormalFromHeightMap::Inputs::Intensity.ToString()), Factor);
+										UInterchangeShaderPortsAPI::ConnectDefaultOuputToInput(HeightMapNode, NormalFromHeightMap::Inputs::HeightMap.ToString(), TextureObjectNode->GetUniqueID());
+										HeightMapNode->AddFloatAttribute(UInterchangeShaderPortsAPI::MakeInputValueKey(NormalFromHeightMap::Inputs::Intensity.ToString()), Factor);
 
-									UInterchangeShaderPortsAPI::ConnectDefaultOuputToInput(ShaderGraphNode, Materials::Common::Parameters::Normal.ToString(), HeightMapNode->GetUniqueID());
+										UInterchangeShaderPortsAPI::ConnectDefaultOuputToInput(ShaderGraphNode, Materials::Common::Parameters::Normal.ToString(), HeightMapNode->GetUniqueID());
+									}
 								}
 							}
 						}
