[Interchange - glTF] T0 as ref does not work for glTF

#rnx

diff --git a/Engine/Plugins/Interchange/Runtime/Source/Import/Private/Gltf/InterchangeGltfAnimation.cpp b/Engine/Plugins/Interchange/Runtime/Source/Import/Private/Gltf/InterchangeGltfAnimation.cpp
--- a/Engine/Plugins/Interchange/Runtime/Source/Import/Private/Gltf/InterchangeGltfAnimation.cpp
+++ b/Engine/Plugins/Interchange/Runtime/Source/Import/Private/Gltf/InterchangeGltfAnimation.cpp
@@ -456,6 +456,85 @@ namespace UE::Interchange::Gltf::Private
 		return true;
 	}
 
+	void GetT0Transform(const GLTF::FAnimation& GltfAnimation, const GLTF::FNode& AnimatedNode, const TArray<int32>& ChannelIndices, FTransform& OutTransform)
+	{
+		int32 LowerIndex = 0;
+		int32 HigherIndex = 0;
+		double RequestedTime = 0.;
+
+		FVector3f TranslationData = FVector3f(OutTransform.GetTranslation());
+		FQuat4f RotationData = FQuat4f(OutTransform.GetRotation());
+		FVector3f ScaleData = FVector3f(OutTransform.GetScale3D());
+
+		for (const int32& ChannelIndex : ChannelIndices)
+		{
+			const GLTF::FAnimation::FChannel& Channel = GltfAnimation.Channels[ChannelIndex];
+			const GLTF::FAnimation::FSampler& Sampler = GltfAnimation.Samplers[Channel.Sampler];
+
+			TArray<float> Seconds;
+			Sampler.Input.GetFloatArray(Seconds);
+
+			if (Seconds.Num() == 0)
+			{
+				continue;
+			}
+
+			switch (Channel.Target.Path)
+			{
+				case GLTF::FAnimation::EPath::Translation:
+				{
+					TArray<FVector3f> TranslationBuffer;
+					TranslationBuffer.SetNumUninitialized(Sampler.Output.Count);
+					Sampler.Output.GetCoordArray(TranslationBuffer.GetData());
+					const float Scale = GltfUnitConversionMultiplier; // Convert m to cm
+					for (FVector3f& Position : TranslationBuffer)
+					{
+						Position *= Scale;
+					}
+
+					TranslationData = InterpolateValue(Seconds, TranslationBuffer, LowerIndex, HigherIndex, RequestedTime, Sampler.Interpolation, false);
+					break;
+				}
+				case GLTF::FAnimation::EPath::Rotation:
+				{
+					TArray<FQuat4f> RotationBuffer;
+					TArray<FVector4f> RotationBufferTemp;
+					RotationBufferTemp.SetNumUninitialized(Sampler.Output.Count);
+					RotationBuffer.SetNumUninitialized(Sampler.Output.Count);
+
+					Sampler.Output.GetQuatArray(RotationBufferTemp.GetData());
+
+					for (size_t Index = 0; Index < Sampler.Output.Count; Index++)
+					{
+						RotationBuffer[Index] = FQuat4f(RotationBufferTemp[Index][0], RotationBufferTemp[Index][1], RotationBufferTemp[Index][2], RotationBufferTemp[Index][3]);
+					}
+
+					RotationData = InterpolateValue(Seconds, RotationBuffer, LowerIndex, HigherIndex, RequestedTime, Sampler.Interpolation, true);
+					RotationData.Normalize();
+					break;
+				}
+				case GLTF::FAnimation::EPath::Scale:
+				{
+					TArray<FVector3f> ScaleBuffer;
+					ScaleBuffer.SetNumUninitialized(Sampler.Output.Count);
+					Sampler.Output.GetCoordArray(ScaleBuffer.GetData());
+
+					ScaleData = InterpolateValue(Seconds, ScaleBuffer, LowerIndex, HigherIndex, RequestedTime, Sampler.Interpolation, false);
+
+					break;
+				}
+				default:
+					UE_LOG(LogInterchangeImport, Warning, TEXT("Animation type not supported"));
+					break;
+			}
+		}
+
+		OutTransform = FTransform(
+			FQuat4d(RotationData),
+			FVector3d(TranslationData),
+			FVector3d(ScaleData));
+	}
+
 	bool GetBakedAnimationTransformPayloadData(const FString& PayLoadKey, const GLTF::FAsset& GltfAsset, FAnimationPayloadData& PayloadData)
 	{
 		const double BakeInterval = 1.0 / PayloadData.BakeFrequency;
diff --git a/Engine/Plugins/Interchange/Runtime/Source/Import/Private/Gltf/InterchangeGltfPrivate.h b/Engine/Plugins/Interchange/Runtime/Source/Import/Private/Gltf/InterchangeGltfPrivate.h
--- a/Engine/Plugins/Interchange/Runtime/Source/Import/Private/Gltf/InterchangeGltfPrivate.h
+++ b/Engine/Plugins/Interchange/Runtime/Source/Import/Private/Gltf/InterchangeGltfPrivate.h
@@ -25,6 +25,7 @@ namespace UE::Interchange
 		// Animation related functions
 		bool GetTransformAnimationPayloadData(const FString& PayLoadKey, const GLTF::FAsset& GltfAsset, FAnimationPayloadData& OutPayloadData);
 		bool GetMorphTargetAnimationPayloadData(const FString& PayLoadKey, const GLTF::FAsset& GltfAsset, FAnimationPayloadData& OutPayloadData);
+		void GetT0Transform(const GLTF::FAnimation& GltfAnimation, const GLTF::FNode& AnimatedNode, const TArray<int32>& ChannelIndices, FTransform& OutTransform);
 		bool GetBakedAnimationTransformPayloadData(const FString& PayLoadKey, const GLTF::FAsset& GltfAsset, FAnimationPayloadData& PayloadData);
 		// 
 
diff --git a/Engine/Plugins/Interchange/Runtime/Source/Import/Private/Gltf/InterchangeGltfTranslator.cpp b/Engine/Plugins/Interchange/Runtime/Source/Import/Private/Gltf/InterchangeGltfTranslator.cpp
--- a/Engine/Plugins/Interchange/Runtime/Source/Import/Private/Gltf/InterchangeGltfTranslator.cpp
+++ b/Engine/Plugins/Interchange/Runtime/Source/Import/Private/Gltf/InterchangeGltfTranslator.cpp
@@ -212,7 +212,8 @@ namespace UE::Interchange::Gltf::Private
 }
 
 void UInterchangeGLTFTranslator::HandleGltfNode( UInterchangeBaseNodeContainer& NodeContainer, const GLTF::FNode& GltfNode, const FString& ParentNodeUid, const int32 NodeIndex, 
-	bool &bHasVariants, TArray<int32>& SkinnedMeshNodes, TSet<int>& UnusedMeshIndices ) const
+	bool &bHasVariants, TArray<int32>& SkinnedMeshNodes, TSet<int>& UnusedMeshIndices,
+	const TMap<int32, FTransform>& T0Transforms ) const
 {
 	using namespace UE::Interchange::Gltf::Private;
 
@@ -269,6 +270,16 @@ void UInterchangeGLTFTranslator::HandleGltfNode( UInterchangeBaseNodeContainer&
 			{
 				InterchangeSceneNode->SetCustomBindPoseLocalTransform(&NodeContainer, GltfNode.LocalBindPose);
 			}
+			
+			if (GltfAsset.Animations.Num() == 0 || !T0Transforms.Contains(GltfNode.Index))
+			{
+				//If no animations present, use Local Transform for T0
+				InterchangeSceneNode->SetCustomTimeZeroLocalTransform(&NodeContainer, GltfNode.Transform);
+			}
+			else
+			{
+				InterchangeSceneNode->SetCustomTimeZeroLocalTransform(&NodeContainer, T0Transforms[GltfNode.Index]);
+			}
 			break;
 		}
 
@@ -349,7 +360,7 @@ void UInterchangeGLTFTranslator::HandleGltfNode( UInterchangeBaseNodeContainer&
 	{
 		if ( GltfAsset.Nodes.IsValidIndex( ChildIndex ) )
 		{
-			HandleGltfNode( NodeContainer, GltfAsset.Nodes[ ChildIndex ], NodeUid, ChildIndex, bHasVariants, SkinnedMeshNodes, UnusedMeshIndices );
+			HandleGltfNode( NodeContainer, GltfAsset.Nodes[ ChildIndex ], NodeUid, ChildIndex, bHasVariants, SkinnedMeshNodes, UnusedMeshIndices, T0Transforms);
 		}
 	}
 }
@@ -663,6 +674,33 @@ bool UInterchangeGLTFTranslator::Translate( UInterchangeBaseNodeContainer& NodeC
 
 	// Scenes
 	{
+		//Generate T0 Transforms
+		TMap<int32, FTransform> T0Transforms;
+		if (GltfAsset.Animations.Num() > 0)
+		{
+			const GLTF::FAnimation& Animation = GltfAsset.Animations[0];
+
+			//Only Skeletal Animations (no Morph Animations as those do not produce FTransforms)
+			TMap<int32, TArray<int32>> AnimatedNodesIndexToChannelIndices;//(AnimatedNodeIndex, [Channels])
+			for (int32 ChannelIndex = 0; ChannelIndex < Animation.Channels.Num(); ++ChannelIndex)
+			{
+				const GLTF::FAnimation::FChannel& Channel = Animation.Channels[ChannelIndex];
+
+				if (Channel.Target.Node.Type == GLTF::FNode::EType::Joint)
+				{
+					TArray<int32>& ChannelIndices = AnimatedNodesIndexToChannelIndices.FindOrAdd(Channel.Target.Node.Index);
+					ChannelIndices.Add(ChannelIndex);
+				}
+			}
+
+			for (const TPair<int32, TArray<int32>>& AnimatedNodeIndexToChannelIndices : AnimatedNodesIndexToChannelIndices)
+			{
+				FTransform T0Transform;
+				UE::Interchange::Gltf::Private::GetT0Transform(Animation, GltfAsset.Nodes[AnimatedNodeIndexToChannelIndices.Key], AnimatedNodeIndexToChannelIndices.Value, T0Transform);
+				T0Transforms.Emplace(AnimatedNodeIndexToChannelIndices.Key, T0Transform);
+			}
+		}
+
 		int32 SceneIndex = 0;
 		for ( const GLTF::FScene& GltfScene : GltfAsset.Scenes )
 		{
@@ -682,7 +720,7 @@ bool UInterchangeGLTFTranslator::Translate( UInterchangeBaseNodeContainer& NodeC
 			{
 				if ( GltfAsset.Nodes.IsValidIndex( NodeIndex ) )
 				{
-					HandleGltfNode( NodeContainer, GltfAsset.Nodes[ NodeIndex ], SceneNodeUid, NodeIndex, bHasVariants, SkinnedMeshNodes, UnusedGltfMeshIndices );
+					HandleGltfNode( NodeContainer, GltfAsset.Nodes[ NodeIndex ], SceneNodeUid, NodeIndex, bHasVariants, SkinnedMeshNodes, UnusedGltfMeshIndices, T0Transforms);
 				}
 			}
 
diff --git a/Engine/Plugins/Interchange/Runtime/Source/Import/Public/Gltf/InterchangeGltfTranslator.h b/Engine/Plugins/Interchange/Runtime/Source/Import/Public/Gltf/InterchangeGltfTranslator.h
--- a/Engine/Plugins/Interchange/Runtime/Source/Import/Public/Gltf/InterchangeGltfTranslator.h
+++ b/Engine/Plugins/Interchange/Runtime/Source/Import/Public/Gltf/InterchangeGltfTranslator.h
@@ -65,7 +65,7 @@ protected:
 
 	void HandleGltfSkeletons( UInterchangeBaseNodeContainer& NodeContainer, const FString& SceneNodeUid, const TArray<int32>& SkinnedMeshNodes, TSet<int>& UnusedMeshIndices ) const;
 	void HandleGltfNode( UInterchangeBaseNodeContainer& NodeContainer, const GLTF::FNode& GltfNode, const FString& ParentNodeUid, const int32 NodeIndex, 
-		bool &bHasVariants, TArray<int32>& SkinnedMeshNodes, TSet<int>& UnusedMeshIndices ) const;
+		bool &bHasVariants, TArray<int32>& SkinnedMeshNodes, TSet<int>& UnusedMeshIndices, const TMap<int32, FTransform>& T0Transforms) const;
 	void HandleGltfVariants(UInterchangeBaseNodeContainer& NodeContainer, const FString& FileName) const;
 	UInterchangeMeshNode* HandleGltfMesh(UInterchangeBaseNodeContainer& NodeContainer, const GLTF::FMesh& GltfMesh, int MeshIndex,
 		TSet<int>& UnusedMeshIndices, const FString& SkeletalName = "" /*If set it creates the mesh even if it was already created (for Skeletals)*/, const FString& SkeletalId = "") const;
