Create an option to run fbx translator async not in parallel without the InterchangeWorker.

#rnx

diff --git a/Engine/Plugins/Interchange/Runtime/Source/Dispatcher/Private/InterchangeDispatcher.cpp b/Engine/Plugins/Interchange/Runtime/Source/Dispatcher/Private/InterchangeDispatcher.cpp
--- a/Engine/Plugins/Interchange/Runtime/Source/Dispatcher/Private/InterchangeDispatcher.cpp
+++ b/Engine/Plugins/Interchange/Runtime/Source/Dispatcher/Private/InterchangeDispatcher.cpp
@@ -207,6 +207,16 @@ namespace UE
 			return CompletedTaskCount == TaskPool.Num();
 		}
 
+		bool FInterchangeDispatcher::IsInterchangeWorkerRunning()
+		{
+			if (IsHandlerAlive())
+			{
+				WaitAllTaskToCompleteExecution();
+				return GetInterchangeWorkerFatalError().IsEmpty();
+			}
+			return false;
+		}
+
 		//this is a static function
 		bool FInterchangeDispatcher::IsInterchangeWorkerAvailable()
 		{
diff --git a/Engine/Plugins/Interchange/Runtime/Source/Dispatcher/Public/InterchangeDispatcher.h b/Engine/Plugins/Interchange/Runtime/Source/Dispatcher/Public/InterchangeDispatcher.h
--- a/Engine/Plugins/Interchange/Runtime/Source/Dispatcher/Public/InterchangeDispatcher.h
+++ b/Engine/Plugins/Interchange/Runtime/Source/Dispatcher/Public/InterchangeDispatcher.h
@@ -42,6 +42,7 @@ namespace UE
 				return InterchangeWorkerFatalError;
 			}
 
+			bool IsInterchangeWorkerRunning();
 			static bool IsInterchangeWorkerAvailable();
 		private:
 			void SpawnHandler();
diff --git a/Engine/Plugins/Interchange/Runtime/Source/Import/InterchangeImport.Build.cs b/Engine/Plugins/Interchange/Runtime/Source/Import/InterchangeImport.Build.cs
--- a/Engine/Plugins/Interchange/Runtime/Source/Import/InterchangeImport.Build.cs
+++ b/Engine/Plugins/Interchange/Runtime/Source/Import/InterchangeImport.Build.cs
@@ -50,6 +50,7 @@ namespace UnrealBuildTool.Rules
 					new string[]
 					{
 						"BSPUtils",
+						"InterchangeFbxParser",
 						"MaterialEditor",
 						"SkeletalMeshUtilitiesCommon",
 						"UnrealEd",
diff --git a/Engine/Plugins/Interchange/Runtime/Source/Import/Private/Fbx/InterchangeFbxTranslator.cpp b/Engine/Plugins/Interchange/Runtime/Source/Import/Private/Fbx/InterchangeFbxTranslator.cpp
--- a/Engine/Plugins/Interchange/Runtime/Source/Import/Private/Fbx/InterchangeFbxTranslator.cpp
+++ b/Engine/Plugins/Interchange/Runtime/Source/Import/Private/Fbx/InterchangeFbxTranslator.cpp
@@ -30,7 +30,14 @@ static bool GInterchangeEnableFBXImport = false;
 static FAutoConsoleVariableRef CCvarInterchangeEnableFBXImport(
 	TEXT("Interchange.FeatureFlags.Import.FBX"),
 	GInterchangeEnableFBXImport,
-	TEXT("[Experimental] Whether FBX support is enabled."),
+	TEXT("Whether FBX support is enabled."),
+	ECVF_Default);
+
+static bool GInterchangeFBXTranslatorUseWorker = false;
+static FAutoConsoleVariableRef CCvarInterchangeFBXTranslatorUseWorker(
+	TEXT("Interchange.FeatureFlags.Translator.UseWorker.FBX"),
+	GInterchangeFBXTranslatorUseWorker,
+	TEXT("Whether FBX translator can be execute in parallel using the InterchangeWorker process."),
 	ECVF_Default);
 
 namespace UE::Interchange::Private
@@ -56,6 +63,52 @@ namespace UE::Interchange::Private
 UInterchangeFbxTranslator::UInterchangeFbxTranslator()
 {
 	Dispatcher = nullptr;
+	bUseWorkerImport = false;
+	if (GInterchangeFBXTranslatorUseWorker)
+	{
+		//Run the import in parallel only if we can start interchange worker
+		FGuid RandomGuid;
+		FPlatformMisc::CreateGuid(RandomGuid);
+		IPlatformFile& PlatformFile = FPlatformFileManager::Get().GetPlatformFile();
+		const FString ProjectSavedDir = FPaths::ConvertRelativePathToFull(FPaths::ProjectSavedDir());
+		const FString RamdomGuidDir = RandomGuid.ToString(EGuidFormats::Base36Encoded);
+		if (!PlatformFile.DirectoryExists(*ProjectSavedDir))
+		{
+			PlatformFile.CreateDirectory(*ProjectSavedDir);
+		}
+		const FString InterchangeDir = FPaths::Combine(ProjectSavedDir, TEXT("Interchange"));
+		if (!PlatformFile.DirectoryExists(*InterchangeDir))
+		{
+			PlatformFile.CreateDirectory(*InterchangeDir);
+		}
+		ResultFolder = FPaths::Combine(InterchangeDir, RamdomGuidDir);
+		if (!PlatformFile.DirectoryExists(*ResultFolder))
+		{
+			PlatformFile.CreateDirectory(*ResultFolder);
+		}
+
+		//Create the dispatcher
+		Dispatcher = MakeUnique<UE::Interchange::FInterchangeDispatcher>(ResultFolder);
+
+		if (ensure(Dispatcher.IsValid()))
+		{
+			Dispatcher->StartProcess();
+		}
+
+		if (Dispatcher->IsInterchangeWorkerRunning())
+		{
+			bUseWorkerImport = true;
+		}
+		else
+		{
+			Dispatcher.Reset();
+		}
+	}
+}
+
+bool UInterchangeFbxTranslator::IsThreadSafe() const
+{
+	return bUseWorkerImport;
 }
 
 EInterchangeTranslatorType UInterchangeFbxTranslator::GetTranslatorType() const
@@ -71,15 +124,14 @@ EInterchangeTranslatorAssetType UInterchangeFbxTranslator::GetSupportedAssetType
 
 TArray<FString> UInterchangeFbxTranslator::GetSupportedFormats() const
 {
+#if WITH_EDITOR
 	if (GInterchangeEnableFBXImport)
 	{
 		TArray<FString> Formats{ TEXT("fbx;Filmbox") };
 		return Formats;
 	}
-	else
-	{
-		return TArray<FString>{};
-	}
+#endif
+	return TArray<FString>{};
 }
 
 bool UInterchangeFbxTranslator::Translate(UInterchangeBaseNodeContainer& BaseNodeContainer) const
@@ -93,80 +145,54 @@ bool UInterchangeFbxTranslator::Translate(UInterchangeBaseNodeContainer& BaseNod
 		return false;
 	}
 
-	if(!Dispatcher.IsValid())
+	if (bUseWorkerImport)
 	{
-		//Dispatch an Interchange worker by using the InterchangeDispatcher
-		//Build Result folder
-		FGuid RandomGuid;
-		FPlatformMisc::CreateGuid(RandomGuid);
-		IPlatformFile& PlatformFile = FPlatformFileManager::Get().GetPlatformFile();
-		const FString ProjectSavedDir = FPaths::ConvertRelativePathToFull(FPaths::ProjectSavedDir());
-		const FString RamdomGuidDir = RandomGuid.ToString(EGuidFormats::Base36Encoded);
-		if (!PlatformFile.DirectoryExists(*ProjectSavedDir))
-		{
-			PlatformFile.CreateDirectory(*ProjectSavedDir);
-		}
-		const FString InterchangeDir = FPaths::Combine(ProjectSavedDir, TEXT("Interchange"));
-		if (!PlatformFile.DirectoryExists(*InterchangeDir))
-		{
-			PlatformFile.CreateDirectory(*InterchangeDir);
-		}
-		const FString ResultFolder = FPaths::Combine(InterchangeDir, RamdomGuidDir);
-		if (!PlatformFile.DirectoryExists(*ResultFolder))
+		if (!Dispatcher.IsValid())
 		{
-			PlatformFile.CreateDirectory(*ResultFolder);
+			return false;
 		}
+		//Create a json command to read the fbx file
+		FString JsonCommand = CreateLoadFbxFileCommand(Filename);
+		int32 TaskIndex = Dispatcher->AddTask(JsonCommand);
 
-		//Create the dispatcher
-		Dispatcher = MakeUnique<UE::Interchange::FInterchangeDispatcher>(ResultFolder);
+		//Blocking call until all tasks are executed
+		Dispatcher->WaitAllTaskToCompleteExecution();
 
-		if(ensure(Dispatcher.IsValid()))
+		FString WorkerFatalError = Dispatcher->GetInterchangeWorkerFatalError();
+		if (!WorkerFatalError.IsEmpty())
 		{
-			Dispatcher->StartProcess();
+			AddMessage(UInterchangeResult::FromJson(WorkerFatalError));
 		}
-	}
-
-	if(!Dispatcher.IsValid())
-	{
-		return false;
-	}
 
-	//Create a json command to read the fbx file
-	FString JsonCommand = CreateLoadFbxFileCommand(Filename);
-	int32 TaskIndex = Dispatcher->AddTask(JsonCommand);
+		UE::Interchange::ETaskState TaskState;
+		FString JsonResult;
+		TArray<FString> JsonMessages;
+		Dispatcher->GetTaskState(TaskIndex, TaskState, JsonResult, JsonMessages);
 
-	//Blocking call until all tasks are executed
-	Dispatcher->WaitAllTaskToCompleteExecution();
-	
-	FString WorkerFatalError = Dispatcher->GetInterchangeWorkerFatalError();
-	if (!WorkerFatalError.IsEmpty())
-	{
-		AddMessage(UInterchangeResult::FromJson(WorkerFatalError));
-	}
+		// Parse the Json messages into UInterchangeResults
+		for (const FString& JsonMessage : JsonMessages)
+		{
+			UE::Interchange::Private::ApplyTranslatorMessage(this, JsonMessage);
+		}
 
-	UE::Interchange::ETaskState TaskState;
-	FString JsonResult;
-	TArray<FString> JsonMessages;
-	Dispatcher->GetTaskState(TaskIndex, TaskState, JsonResult, JsonMessages);
+		if (TaskState != UE::Interchange::ETaskState::ProcessOk)
+		{
+			return false;
+		}
+		//Grab the result file and fill the BaseNodeContainer
+		UE::Interchange::FJsonLoadSourceCmd::JsonResultParser ResultParser;
+		ResultParser.FromJson(JsonResult);
+		FString BaseNodeContainerFilename = ResultParser.GetResultFilename();
 
-	// Parse the Json messages into UInterchangeResults
-	for (const FString& JsonMessage : JsonMessages)
-	{
-		UE::Interchange::Private::ApplyTranslatorMessage(this, JsonMessage);
+		//Parse the filename to fill the container
+		BaseNodeContainer.LoadFromFile(BaseNodeContainerFilename);
 	}
-
-	if(TaskState != UE::Interchange::ETaskState::ProcessOk)
+	else
 	{
-		return false;
+#if WITH_EDITOR
+		FbxParser.LoadFbxFile(Filename, BaseNodeContainer);
+#endif
 	}
-	//Grab the result file and fill the BaseNodeContainer
-	UE::Interchange::FJsonLoadSourceCmd::JsonResultParser ResultParser;
-	ResultParser.FromJson(JsonResult);
-	FString BaseNodeContainerFilename = ResultParser.GetResultFilename();
-
-	//Parse the filename to fill the container
-	BaseNodeContainer.LoadFromFile(BaseNodeContainerFilename);
-
 	return true;
 }
 
@@ -177,6 +203,10 @@ void UInterchangeFbxTranslator::ReleaseSource()
 		//Do not block the main thread
 		Dispatcher->StopProcess(!IsInGameThread());
 	}
+#if WITH_EDITOR
+	FbxParser.ReleaseResources();
+#endif
+	ResultFolder.Empty();
 }
 
 void UInterchangeFbxTranslator::ImportFinish()
@@ -208,67 +238,36 @@ TFuture<TOptional<UE::Interchange::FMeshPayloadData>> UInterchangeFbxTranslator:
 {
 	TSharedPtr<TPromise<TOptional<UE::Interchange::FMeshPayloadData>>> Promise = MakeShared<TPromise<TOptional<UE::Interchange::FMeshPayloadData>>>();
 
-	if (!Dispatcher.IsValid())
-	{
-		Promise->SetValue(TOptional<UE::Interchange::FMeshPayloadData>());
-		return Promise->GetFuture();
-	}
-
-	// Create a json command to read the fbx file
-	FString JsonCommand = CreateFetchMeshPayloadFbxCommand(PayLoadKey.UniqueId, MeshGlobalTransform);
-	const int32 CreatedTaskIndex = Dispatcher->AddTask(JsonCommand, FInterchangeDispatcherTaskCompleted::CreateLambda([this, Promise, PayLoadKey](const int32 TaskIndex)
-	{
-		TRACE_CPUPROFILER_EVENT_SCOPE(UInterchangeFbxTranslator::GetStaticMeshPayloadData::Dispatcher->AddTaskDone)
-		UE::Interchange::ETaskState TaskState;
-		FString JsonResult;
-		TArray<FString> JsonMessages;
-		Dispatcher->GetTaskState(TaskIndex, TaskState, JsonResult, JsonMessages);
-
-		// Parse the Json messages into UInterchangeResults
-		for (const FString& JsonMessage : JsonMessages)
-		{
-			UE::Interchange::Private::ApplyTranslatorMessage(this, JsonMessage);
-		}
-
-		if (TaskState != UE::Interchange::ETaskState::ProcessOk)
-		{
-			Promise->SetValue(TOptional<UE::Interchange::FMeshPayloadData>());
-			return;
-		}
-
-		// Grab the result file and fill the BaseNodeContainer
-		UE::Interchange::FJsonFetchMeshPayloadCmd::JsonResultParser ResultParser;
-		ResultParser.FromJson(JsonResult);
-		FString StaticMeshPayloadFilename = ResultParser.GetResultFilename();
-
-		//Mesh payload file generation can fail due to invalid Mesh (for ep No Polygons/Only Degenerate Polygons)
-		if (!FPaths::FileExists(StaticMeshPayloadFilename))
+	auto OnPayloadReady = [this, Promise, PayLoadKey](const FString& MeshPayloadFilename)
 		{
-			UE_LOG(LogInterchangeImport, Warning, TEXT("Expected mesh payload file does not exist for PayloadKey: %s"), *PayLoadKey.UniqueId);
+			//Mesh payload file generation can fail due to invalid Mesh (for ep No Polygons/Only Degenerate Polygons)
+			if (!FPaths::FileExists(MeshPayloadFilename))
+			{
+				UE_LOG(LogInterchangeImport, Warning, TEXT("Expected mesh payload file does not exist for PayloadKey: %s"), *PayLoadKey.UniqueId);
 
-			Promise->SetValue(TOptional<UE::Interchange::FMeshPayloadData>());
-			return;
-		}
+				Promise->SetValue(TOptional<UE::Interchange::FMeshPayloadData>());
+				return;
+			}
 
-		UE::Interchange::FMeshPayloadData MeshPayloadData;
-		MeshPayloadData.MeshDescription.Empty();
+			UE::Interchange::FMeshPayloadData MeshPayloadData;
+			MeshPayloadData.MeshDescription.Empty();
 
-		// All sub object should be gone with the reset
-		TArray64<uint8> Buffer;
-		FFileHelper::LoadFileToArray(Buffer, *StaticMeshPayloadFilename);
-		uint8* FileData = Buffer.GetData();
-		int64 FileDataSize = Buffer.Num();
-		if (FileDataSize < 1)
-		{
-			// Nothing to load from this file
-			Promise->SetValue(TOptional<UE::Interchange::FMeshPayloadData>());
-			return;
-		}
+			// All sub object should be gone with the reset
+			TArray64<uint8> Buffer;
+			FFileHelper::LoadFileToArray(Buffer, *MeshPayloadFilename);
+			uint8* FileData = Buffer.GetData();
+			int64 FileDataSize = Buffer.Num();
+			if (FileDataSize < 1)
+			{
+				// Nothing to load from this file
+				Promise->SetValue(TOptional<UE::Interchange::FMeshPayloadData>());
+				return;
+			}
 
-		switch (PayLoadKey.Type)
-		{
-		case EInterchangeMeshPayLoadType::STATIC:
-		case EInterchangeMeshPayLoadType::SKELETAL:
+			switch (PayLoadKey.Type)
+			{
+			case EInterchangeMeshPayLoadType::STATIC:
+			case EInterchangeMeshPayLoadType::SKELETAL:
 			{
 				// Buffer keeps the ownership of the data, the large memory reader is use to serialize the TMap
 				FLargeMemoryReader Ar(FileData, FileDataSize);
@@ -283,34 +282,80 @@ TFuture<TOptional<UE::Interchange::FMeshPayloadData>> UInterchangeFbxTranslator:
 				}
 			}
 			break;
-		case EInterchangeMeshPayLoadType::MORPHTARGET:
+			case EInterchangeMeshPayLoadType::MORPHTARGET:
 			{
 				//Buffer keep the ownership of the data, the large memory reader is use to serialize the TMap
 				FLargeMemoryReader Ar(FileData, FileDataSize);
 				MeshPayloadData.MeshDescription.Serialize(Ar);
 			}
 			break;
-		case EInterchangeMeshPayLoadType::NONE:
-		default:
-			break;
-		}
+			case EInterchangeMeshPayLoadType::NONE:
+			default:
+				break;
+			}
 
-		if (!FStaticMeshOperations::ValidateAndFixData(MeshPayloadData.MeshDescription, PayLoadKey.UniqueId))
-		{
-			UInterchangeResultError_Generic* ErrorResult = AddMessage<UInterchangeResultError_Generic>();
-			ErrorResult->SourceAssetName = SourceData ? SourceData->GetFilename() : FString();
-			ErrorResult->Text = NSLOCTEXT("UInterchangeFbxTranslator", "GetMeshPayloadData_ValidateMeshDescriptionFail", "Invalid mesh data (NAN) was found and fix to zero. Mesh render can be bad.");
-		}
+			if (!FStaticMeshOperations::ValidateAndFixData(MeshPayloadData.MeshDescription, PayLoadKey.UniqueId))
+			{
+				UInterchangeResultError_Generic* ErrorResult = AddMessage<UInterchangeResultError_Generic>();
+				ErrorResult->SourceAssetName = SourceData ? SourceData->GetFilename() : FString();
+				ErrorResult->Text = NSLOCTEXT("UInterchangeFbxTranslator", "GetMeshPayloadData_ValidateMeshDescriptionFail", "Invalid mesh data (NAN) was found and fix to zero. Mesh render can be bad.");
+			}
 
-		Promise->SetValue(MoveTemp(MeshPayloadData));
-	}));
+			Promise->SetValue(MoveTemp(MeshPayloadData));
+		};
 
-	// The task was not added to the dispatcher
-	if (CreatedTaskIndex == INDEX_NONE)
+	if (!bUseWorkerImport)
 	{
-		Promise->SetValue(TOptional<UE::Interchange::FMeshPayloadData>{});
+#if WITH_EDITOR
+		FString ResultMeshPayloadsUniqueId = FbxParser.FetchMeshPayload(PayLoadKey.UniqueId, MeshGlobalTransform, ResultFolder);
+		FString MeshPayloadFilename = FbxParser.GetResultPayloadFilepath(ResultMeshPayloadsUniqueId);
+		OnPayloadReady(MeshPayloadFilename);
+#endif
 	}
+	else
+	{
+		if (!Dispatcher.IsValid())
+		{
+			Promise->SetValue(TOptional<UE::Interchange::FMeshPayloadData>());
+			return Promise->GetFuture();
+		}
 
+		// Create a json command to read the fbx file
+		FString JsonCommand = CreateFetchMeshPayloadFbxCommand(PayLoadKey.UniqueId, MeshGlobalTransform);
+		const int32 CreatedTaskIndex = Dispatcher->AddTask(JsonCommand, FInterchangeDispatcherTaskCompleted::CreateLambda([this, Promise, PayLoadKey, OnPayloadReadyClosure = MoveTemp(OnPayloadReady)](const int32 TaskIndex)
+			{
+				TRACE_CPUPROFILER_EVENT_SCOPE(UInterchangeFbxTranslator::GetStaticMeshPayloadData::Dispatcher->AddTaskDone)
+					UE::Interchange::ETaskState TaskState;
+				FString JsonResult;
+				TArray<FString> JsonMessages;
+				Dispatcher->GetTaskState(TaskIndex, TaskState, JsonResult, JsonMessages);
+
+				// Parse the Json messages into UInterchangeResults
+				for (const FString& JsonMessage : JsonMessages)
+				{
+					UE::Interchange::Private::ApplyTranslatorMessage(this, JsonMessage);
+				}
+
+				if (TaskState != UE::Interchange::ETaskState::ProcessOk)
+				{
+					Promise->SetValue(TOptional<UE::Interchange::FMeshPayloadData>());
+					return;
+				}
+
+				// Grab the result file and fill the BaseNodeContainer
+				UE::Interchange::FJsonFetchMeshPayloadCmd::JsonResultParser ResultParser;
+				ResultParser.FromJson(JsonResult);
+				FString MeshPayloadFilename = ResultParser.GetResultFilename();
+
+				OnPayloadReadyClosure(MeshPayloadFilename);
+			}));
+
+		// The task was not added to the dispatcher
+		if (CreatedTaskIndex == INDEX_NONE)
+		{
+			Promise->SetValue(TOptional<UE::Interchange::FMeshPayloadData>{});
+		}
+	}
 	return Promise->GetFuture();
 }
 
@@ -318,44 +363,9 @@ TFuture<TOptional<UE::Interchange::FAnimationPayloadData>> UInterchangeFbxTransl
 {
 	TSharedPtr<TPromise<TOptional<UE::Interchange::FAnimationPayloadData>>> Promise = MakeShared<TPromise<TOptional<UE::Interchange::FAnimationPayloadData>>>();
 
-	if (!Dispatcher.IsValid())
-	{
-		Promise->SetValue(TOptional<UE::Interchange::FAnimationPayloadData>());
-		return Promise->GetFuture();
-	}
-
-	// Create a json command to read the fbx file
-	FString JsonCommand = 
-		(PayLoadKey.Type == EInterchangeAnimationPayLoadType::BAKED)
-		? CreateFetchAnimationBakeTransformPayloadFbxCommand(PayLoadKey.UniqueId, BakeFrequency, RangeStartSecond, RangeStopSecond)
-		: CreateFetchPayloadFbxCommand(PayLoadKey.UniqueId);
-
-	const int32 CreatedTaskIndex = Dispatcher->AddTask(JsonCommand, FInterchangeDispatcherTaskCompleted::CreateLambda([this, Promise, PayLoadKey](const int32 TaskIndex)
+	auto OnPayloadReady = [Promise, PayLoadKey](const FString& AnimationPayloadFilename)
 		{
-			TRACE_CPUPROFILER_EVENT_SCOPE(UInterchangeFbxTranslator::GetAnimationCurvePayloadData::Dispatcher->AddTaskDone)
-			UE::Interchange::ETaskState TaskState;
-			FString JsonResult;
-			TArray<FString> JsonMessages;
-			Dispatcher->GetTaskState(TaskIndex, TaskState, JsonResult, JsonMessages);
-
-			// Parse the Json messages into UInterchangeResults
-			for (const FString& JsonMessage : JsonMessages)
-			{
-				UE::Interchange::Private::ApplyTranslatorMessage(this, JsonMessage);
-			}
-
-			if (TaskState != UE::Interchange::ETaskState::ProcessOk)
-			{
-				Promise->SetValue(TOptional<UE::Interchange::FAnimationPayloadData>());
-				return;
-			}
-
-			// Grab the result file and fill the BaseNodeContainer
-			UE::Interchange::FJsonFetchPayloadCmd::JsonResultParser ResultParser;
-			ResultParser.FromJson(JsonResult);
-			FString AnimationTransformPayloadFilename = ResultParser.GetResultFilename();
-
-			if (!ensure(FPaths::FileExists(AnimationTransformPayloadFilename)))
+			if (!ensure(FPaths::FileExists(AnimationPayloadFilename)))
 			{
 				// TODO log an error saying the payload file does not exist even if the get payload command succeeded
 				Promise->SetValue(TOptional<UE::Interchange::FAnimationPayloadData>());
@@ -366,7 +376,7 @@ TFuture<TOptional<UE::Interchange::FAnimationPayloadData>> UInterchangeFbxTransl
 
 			// All sub object should be gone with the reset
 			TArray64<uint8> Buffer;
-			FFileHelper::LoadFileToArray(Buffer, *AnimationTransformPayloadFilename);
+			FFileHelper::LoadFileToArray(Buffer, *AnimationPayloadFilename);
 			uint8* FileData = Buffer.GetData();
 			int64 FileDataSize = Buffer.Num();
 			if (FileDataSize < 1)
@@ -408,12 +418,75 @@ TFuture<TOptional<UE::Interchange::FAnimationPayloadData>> UInterchangeFbxTransl
 			}
 
 			Promise->SetValue(MoveTemp(AnimationTransformPayload));
-		}));
+		};
 
-	// The task was not added to the dispatcher
-	if (CreatedTaskIndex == INDEX_NONE)
+	if (!bUseWorkerImport)
+	{
+#if WITH_EDITOR
+		FString AnimationPayloadFile;
+		if (PayLoadKey.Type == EInterchangeAnimationPayLoadType::BAKED)
+		{
+			FString ResultPayloadsUniqueId = FbxParser.FetchAnimationBakeTransformPayload(PayLoadKey.UniqueId
+				, BakeFrequency
+				, RangeStartSecond
+				, RangeStopSecond
+				, ResultFolder);
+			AnimationPayloadFile = FbxParser.GetResultPayloadFilepath(ResultPayloadsUniqueId);
+		}
+		else
+		{
+			FbxParser.FetchPayload(PayLoadKey.UniqueId, ResultFolder);
+			AnimationPayloadFile = FbxParser.GetResultPayloadFilepath(PayLoadKey.UniqueId);
+		}
+		OnPayloadReady(AnimationPayloadFile);
+#endif //WITH_EDITOR
+	}
+	else
 	{
-		Promise->SetValue(TOptional<UE::Interchange::FAnimationPayloadData>{});
+		if (!Dispatcher.IsValid())
+		{
+			Promise->SetValue(TOptional<UE::Interchange::FAnimationPayloadData>());
+			return Promise->GetFuture();
+		}
+
+		// Create a json command to read the fbx file
+		FString JsonCommand = 
+			(PayLoadKey.Type == EInterchangeAnimationPayLoadType::BAKED)
+			? CreateFetchAnimationBakeTransformPayloadFbxCommand(PayLoadKey.UniqueId, BakeFrequency, RangeStartSecond, RangeStopSecond)
+			: CreateFetchPayloadFbxCommand(PayLoadKey.UniqueId);
+
+		const int32 CreatedTaskIndex = Dispatcher->AddTask(JsonCommand, FInterchangeDispatcherTaskCompleted::CreateLambda([this, Promise, PayLoadKey, OnPayloadReadyClosure = MoveTemp(OnPayloadReady)](const int32 TaskIndex)
+			{
+				TRACE_CPUPROFILER_EVENT_SCOPE(UInterchangeFbxTranslator::GetAnimationCurvePayloadData::Dispatcher->AddTaskDone)
+				UE::Interchange::ETaskState TaskState;
+				FString JsonResult;
+				TArray<FString> JsonMessages;
+				Dispatcher->GetTaskState(TaskIndex, TaskState, JsonResult, JsonMessages);
+
+				// Parse the Json messages into UInterchangeResults
+				for (const FString& JsonMessage : JsonMessages)
+				{
+					UE::Interchange::Private::ApplyTranslatorMessage(this, JsonMessage);
+				}
+
+				if (TaskState != UE::Interchange::ETaskState::ProcessOk)
+				{
+					Promise->SetValue(TOptional<UE::Interchange::FAnimationPayloadData>());
+					return;
+				}
+
+				// Grab the result file and fill the BaseNodeContainer
+				UE::Interchange::FJsonFetchPayloadCmd::JsonResultParser ResultParser;
+				ResultParser.FromJson(JsonResult);
+				FString AnimationTransformPayloadFilename = ResultParser.GetResultFilename();
+				OnPayloadReadyClosure(AnimationTransformPayloadFilename);
+			}));
+
+		// The task was not added to the dispatcher
+		if (CreatedTaskIndex == INDEX_NONE)
+		{
+			Promise->SetValue(TOptional<UE::Interchange::FAnimationPayloadData>{});
+		}
 	}
 
 	return Promise->GetFuture();
diff --git a/Engine/Plugins/Interchange/Runtime/Source/Import/Private/InterchangeImportModule.cpp b/Engine/Plugins/Interchange/Runtime/Source/Import/Private/InterchangeImportModule.cpp
--- a/Engine/Plugins/Interchange/Runtime/Source/Import/Private/InterchangeImportModule.cpp
+++ b/Engine/Plugins/Interchange/Runtime/Source/Import/Private/InterchangeImportModule.cpp
@@ -77,7 +77,9 @@ void FInterchangeImportModule::StartupModule()
 
 		//Register the translators
 		//Scenes
+#if WITH_EDITOR
 		InterchangeManager.RegisterTranslator(UInterchangeFbxTranslator::StaticClass());
+#endif
 		InterchangeManager.RegisterTranslator(UInterchangeGLTFTranslator::StaticClass());
 		InterchangeManager.RegisterTranslator(UInterchangeOBJTranslator::StaticClass());
 
diff --git a/Engine/Plugins/Interchange/Runtime/Source/Import/Public/Fbx/InterchangeFbxTranslator.h b/Engine/Plugins/Interchange/Runtime/Source/Import/Public/Fbx/InterchangeFbxTranslator.h
--- a/Engine/Plugins/Interchange/Runtime/Source/Import/Public/Fbx/InterchangeFbxTranslator.h
+++ b/Engine/Plugins/Interchange/Runtime/Source/Import/Public/Fbx/InterchangeFbxTranslator.h
@@ -12,6 +12,10 @@
 #include "UObject/Object.h"
 #include "UObject/ObjectMacros.h"
 
+#if WITH_EDITOR
+#include "InterchangeFbxParser.h"
+#endif //WITH_EDITOR
+
 #include "InterchangeFbxTranslator.generated.h"
 
 /* Fbx translator class support import of texture, material, static mesh, skeletal mesh, */
@@ -27,6 +31,7 @@ public:
 	UInterchangeFbxTranslator();
 
 	/** Begin UInterchangeTranslatorBase API*/
+	virtual bool IsThreadSafe() const override;
 	virtual EInterchangeTranslatorType GetTranslatorType() const override;
 	virtual EInterchangeTranslatorAssetType GetSupportedAssetTypes() const override;
 	virtual TArray<FString> GetSupportedFormats() const override;
@@ -85,6 +90,14 @@ private:
 	//in the constructor because Archetype, CDO and registered translators will
 	//never translate a source.
 	mutable TUniquePtr<UE::Interchange::FInterchangeDispatcher> Dispatcher;
+
+	//If true this translator will use the dispatcher (InterchangeWorker program) to translate and return payloads.
+	//If false, this translator will not use the dispatcher
+	bool bUseWorkerImport = false;
+#if WITH_EDITOR
+	mutable UE::Interchange::FInterchangeFbxParser FbxParser;
+#endif //WITH_EDITOR
+	FString ResultFolder;
 };
 
 
diff --git a/Engine/Plugins/Interchange/Runtime/Source/Parsers/Fbx/Private/FbxAPI.cpp b/Engine/Plugins/Interchange/Runtime/Source/Parsers/Fbx/Private/FbxAPI.cpp
--- a/Engine/Plugins/Interchange/Runtime/Source/Parsers/Fbx/Private/FbxAPI.cpp
+++ b/Engine/Plugins/Interchange/Runtime/Source/Parsers/Fbx/Private/FbxAPI.cpp
@@ -17,6 +17,13 @@
 
 #define LOCTEXT_NAMESPACE "InterchangeFbxParser"
 
+#define DESTROY_FBX_OBJECT(Object) \
+if(Object) \
+{ \
+	Object->Destroy(); \
+	Object = nullptr; \
+}
+
 namespace UE
 {
 	namespace Interchange
@@ -27,6 +34,16 @@ namespace UE
 			{
 				PayloadContexts.Empty();
 				FbxHelper = nullptr;
+
+				DESTROY_FBX_OBJECT(SDKImporter);
+				DESTROY_FBX_OBJECT(SDKScene);
+				if (SDKGeometryConverter)
+				{
+					delete SDKGeometryConverter;
+					SDKGeometryConverter = nullptr;
+				}
+				DESTROY_FBX_OBJECT(SDKIoSettings);
+				DESTROY_FBX_OBJECT(SDKManager);
 			}
 			
 			const TSharedPtr<FFbxHelper> FFbxParser::GetFbxHelper()
@@ -54,18 +71,18 @@ namespace UE
 				}
 
 				//Create an IOSettings object. This object holds all import/export settings.
-				FbxIOSettings* ios = FbxIOSettings::Create(SDKManager, IOSROOT);
-				ios->SetBoolProp(IMP_FBX_MATERIAL, true);
-				ios->SetBoolProp(IMP_FBX_TEXTURE, true);
-				ios->SetBoolProp(IMP_FBX_LINK, true);
-				ios->SetBoolProp(IMP_FBX_SHAPE, true);
-				ios->SetBoolProp(IMP_FBX_GOBO, true);
-				ios->SetBoolProp(IMP_FBX_ANIMATION, true);
-				ios->SetBoolProp(IMP_SKINS, true);
-				ios->SetBoolProp(IMP_DEFORMATION, true);
-				ios->SetBoolProp(IMP_FBX_GLOBAL_SETTINGS, true);
-				ios->SetBoolProp(IMP_TAKE, true);
-				SDKManager->SetIOSettings(ios);
+				SDKIoSettings = FbxIOSettings::Create(SDKManager, IOSROOT);
+				SDKIoSettings->SetBoolProp(IMP_FBX_MATERIAL, true);
+				SDKIoSettings->SetBoolProp(IMP_FBX_TEXTURE, true);
+				SDKIoSettings->SetBoolProp(IMP_FBX_LINK, true);
+				SDKIoSettings->SetBoolProp(IMP_FBX_SHAPE, true);
+				SDKIoSettings->SetBoolProp(IMP_FBX_GOBO, true);
+				SDKIoSettings->SetBoolProp(IMP_FBX_ANIMATION, true);
+				SDKIoSettings->SetBoolProp(IMP_SKINS, true);
+				SDKIoSettings->SetBoolProp(IMP_DEFORMATION, true);
+				SDKIoSettings->SetBoolProp(IMP_FBX_GLOBAL_SETTINGS, true);
+				SDKIoSettings->SetBoolProp(IMP_TAKE, true);
+				SDKManager->SetIOSettings(SDKIoSettings);
 
 				SDKGeometryConverter = new FbxGeometryConverter(SDKManager);
 
diff --git a/Engine/Plugins/Interchange/Runtime/Source/Parsers/Fbx/Private/FbxAPI.h b/Engine/Plugins/Interchange/Runtime/Source/Parsers/Fbx/Private/FbxAPI.h
--- a/Engine/Plugins/Interchange/Runtime/Source/Parsers/Fbx/Private/FbxAPI.h
+++ b/Engine/Plugins/Interchange/Runtime/Source/Parsers/Fbx/Private/FbxAPI.h
@@ -39,6 +39,11 @@ namespace UE
 
 				~FFbxParser();
 
+				void SetResultContainer(UInterchangeResultsContainer* Result)
+				{
+					ResultsContainer = Result;
+				}
+
 				//return the fbx helper for this parser
 				const TSharedPtr<FFbxHelper> GetFbxHelper();
 
@@ -95,6 +100,7 @@ namespace UE
 				FbxScene* SDKScene = nullptr;
 				FbxImporter* SDKImporter = nullptr;
 				FbxGeometryConverter* SDKGeometryConverter = nullptr;
+				FbxIOSettings* SDKIoSettings = nullptr;
 				FString SourceFilename;
 				TMap<FString, TSharedPtr<FPayloadContextBase>> PayloadContexts;
 				TSharedPtr<FFbxHelper> FbxHelper;
diff --git a/Engine/Plugins/Interchange/Runtime/Source/Parsers/Fbx/Private/InterchangeFbxParser.cpp b/Engine/Plugins/Interchange/Runtime/Source/Parsers/Fbx/Private/InterchangeFbxParser.cpp
--- a/Engine/Plugins/Interchange/Runtime/Source/Parsers/Fbx/Private/InterchangeFbxParser.cpp
+++ b/Engine/Plugins/Interchange/Runtime/Source/Parsers/Fbx/Private/InterchangeFbxParser.cpp
@@ -46,9 +46,21 @@ namespace UE::Interchange
 	}
 	FInterchangeFbxParser::~FInterchangeFbxParser()
 	{
+		ReleaseResources();
+	}
+
+	void FInterchangeFbxParser::ReleaseResources()
+	{
+		ResultsContainer = nullptr;
 		FbxParserPrivate = nullptr;
 	}
 
+	void FInterchangeFbxParser::SetResultContainer(UInterchangeResultsContainer* Result)
+	{
+		InternalResultsContainer = Result;
+		FbxParserPrivate->SetResultContainer(Result);
+	}
+
 	void FInterchangeFbxParser::LoadFbxFile(const FString& Filename, const FString& ResultFolder)
 	{
 		check(FbxParserPrivate.IsValid());
@@ -57,9 +69,11 @@ namespace UE::Interchange
 
 		if (!FbxParserPrivate->LoadFbxFile(Filename))
 		{
-			UInterchangeResultError_Generic* Error = AddMessage<UInterchangeResultError_Generic>();
-			Error->SourceAssetName = SourceFilename;
-			Error->Text = LOCTEXT("CantLoadFbxFile", "Cannot load the FBX file.");
+			if (UInterchangeResultError_Generic* Error = AddMessage<UInterchangeResultError_Generic>())
+			{
+				Error->SourceAssetName = SourceFilename;
+				Error->Text = LOCTEXT("CantLoadFbxFile", "Cannot load the FBX file.");
+			}
 			return;
 		}
 
@@ -68,9 +82,11 @@ namespace UE::Interchange
 		UInterchangeBaseNodeContainer* Container = NewObject<UInterchangeBaseNodeContainer>(GetTransientPackage(), NAME_None);
 		if (!ensure(Container != nullptr))
 		{
-			UInterchangeResultError_Generic* Error = AddMessage<UInterchangeResultError_Generic>();
-			Error->SourceAssetName = SourceFilename;
-			Error->Text = LOCTEXT("CantAllocate", "Cannot allocate base node container to add FBX scene data.");
+			if (UInterchangeResultError_Generic* Error = AddMessage<UInterchangeResultError_Generic>())
+			{
+				Error->SourceAssetName = SourceFilename;
+				Error->Text = LOCTEXT("CantAllocate", "Cannot allocate base node container to add FBX scene data.");
+			}
 			return;
 		}
 
@@ -80,10 +96,44 @@ namespace UE::Interchange
 		Container->RemoveFromRoot();
 	}
 
+	void FInterchangeFbxParser::LoadFbxFile(const FString& Filename, UInterchangeBaseNodeContainer& BaseNodecontainer)
+	{
+		SourceFilename = Filename;
+		if (!ensure(FbxParserPrivate.IsValid()))
+		{
+			if (UInterchangeResultError_Generic* Error = AddMessage<UInterchangeResultError_Generic>())
+			{
+				Error->SourceAssetName = SourceFilename;
+				Error->Text = LOCTEXT("CantLoadFbxFile_ParserInvalid", "FInterchangeFbxParser::LoadFbxFile: Cannot load the FBX file. The internal fbx parser is invalid.");
+			}
+			return;
+		}
+		
+		if (!FbxParserPrivate->LoadFbxFile(Filename))
+		{
+			if (UInterchangeResultError_Generic* Error = AddMessage<UInterchangeResultError_Generic>())
+			{
+				Error->SourceAssetName = SourceFilename;
+				Error->Text = LOCTEXT("CantLoadFbxFile_ParserError", "FInterchangeFbxParser::LoadFbxFile: Cannot load the FBX file. There was an error when parsing the file.");
+			}
+			return;
+		}
+		FbxParserPrivate->FillContainerWithFbxScene(BaseNodecontainer);
+	}
+
 	void FInterchangeFbxParser::FetchPayload(const FString& PayloadKey, const FString& ResultFolder)
 	{
-		check(FbxParserPrivate.IsValid());
 		ResultsContainer->Empty();
+		if (!ensure(FbxParserPrivate.IsValid()))
+		{
+			if (UInterchangeResultError_Generic* Error = AddMessage<UInterchangeResultError_Generic>())
+			{
+				Error->SourceAssetName = PayloadKey;
+				Error->Text = LOCTEXT("CantFetchPayload_ParserInvalid", "FInterchangeFbxParser::FetchPayload: Cannot fetch the payload. The internal fbx parser is invalid.");
+			}
+			return;
+		}
+		
 		FString PayloadFilepathCopy;
 		{
 			FScopeLock Lock(&ResultPayloadsCriticalSection);
@@ -97,19 +147,30 @@ namespace UE::Interchange
 		}
 		if (!FbxParserPrivate->FetchPayloadData(PayloadKey, PayloadFilepathCopy))
 		{
-			UInterchangeResultError_Generic* Error = AddMessage<UInterchangeResultError_Generic>();
-			Error->SourceAssetName = SourceFilename;
-			Error->Text = LOCTEXT("CantFetchPayload", "Cannot fetch FBX payload data.");
+			if (UInterchangeResultError_Generic* Error = AddMessage<UInterchangeResultError_Generic>())
+			{
+				Error->SourceAssetName = SourceFilename;
+				Error->Text = LOCTEXT("CantFetchPayload", "Cannot fetch FBX payload data.");
+			}
 			return;
 		}
 	}
 
 	FString FInterchangeFbxParser::FetchMeshPayload(const FString& PayloadKey, const FTransform& MeshGlobalTransform, const FString& ResultFolder)
 	{
-		check(FbxParserPrivate.IsValid());
 		ResultsContainer->Empty();
 		FString PayloadFilepathCopy;
 		FString ResultPayloadUniqueId = PayloadKey + MeshGlobalTransform.ToString();
+		if (!ensure(FbxParserPrivate.IsValid()))
+		{
+			if (UInterchangeResultError_Generic* Error = AddMessage<UInterchangeResultError_Generic>())
+			{
+				Error->SourceAssetName = PayloadKey;
+				Error->Text = LOCTEXT("CantFetchMeshPayload_ParserInvalid", "FInterchangeFbxParser::FetchMeshPayload: Cannot fetch the mesh payload. The internal fbx parser is invalid.");
+			}
+			return ResultPayloadUniqueId;
+		}
+		
 		//If we already have extract this mesh, no need to extract again
 		if (ResultPayloads.Contains(ResultPayloadUniqueId))
 		{
@@ -128,19 +189,31 @@ namespace UE::Interchange
 		}
 		if (!FbxParserPrivate->FetchMeshPayloadData(PayloadKey, MeshGlobalTransform, PayloadFilepathCopy))
 		{
-			UInterchangeResultError_Generic* Error = AddMessage<UInterchangeResultError_Generic>();
-			Error->SourceAssetName = SourceFilename;
-			Error->Text = LOCTEXT("CantFetchPayload", "Cannot fetch FBX payload data.");
+			if (UInterchangeResultError_Generic* Error = AddMessage<UInterchangeResultError_Generic>())
+			{
+				Error->SourceAssetName = SourceFilename;
+				Error->Text = LOCTEXT("CantFetchPayload", "Cannot fetch FBX payload data.");
+			}
 		}
 		return ResultPayloadUniqueId;
 	}
 
 	FString FInterchangeFbxParser::FetchAnimationBakeTransformPayload(const FString& PayloadKey, const double BakeFrequency, const double RangeStartTime, const double RangeEndTime, const FString& ResultFolder)
 	{
-		check(FbxParserPrivate.IsValid());
 		ResultsContainer->Empty();
 		FString PayloadFilepathCopy;
-		FString ResultPayloadUniqueId = PayloadKey + FString::FromInt(static_cast<int32>(BakeFrequency*1000.0)) + FString::FromInt(static_cast<int32>(RangeStartTime * 1000.0)) + FString::FromInt(static_cast<int32>(RangeEndTime * 1000.0));
+		FString ResultPayloadUniqueId = PayloadKey + FString::FromInt(static_cast<int32>(BakeFrequency * 1000.0)) + FString::FromInt(static_cast<int32>(RangeStartTime * 1000.0)) + FString::FromInt(static_cast<int32>(RangeEndTime * 1000.0));
+
+		if (!ensure(FbxParserPrivate.IsValid()))
+		{
+			if (UInterchangeResultError_Generic* Error = AddMessage<UInterchangeResultError_Generic>())
+			{
+				Error->SourceAssetName = PayloadKey;
+				Error->Text = LOCTEXT("CantFetchAnimationBakeTransformPayload_ParserInvalid", "FInterchangeFbxParser::FetchAnimationBakeTransformPayload: Cannot fetch the animation bake transform payload. The internal fbx parser is invalid.");
+			}
+			return ResultPayloadUniqueId;
+		}
+		
 		//If we already have extract this mesh, no need to extract again
 		if (ResultPayloads.Contains(ResultPayloadUniqueId))
 		{
@@ -159,9 +232,11 @@ namespace UE::Interchange
 		}
 		if (!FbxParserPrivate->FetchAnimationBakeTransformPayload(PayloadKey, BakeFrequency, RangeStartTime, RangeEndTime, PayloadFilepathCopy))
 		{
-			UInterchangeResultError_Generic* Error = AddMessage<UInterchangeResultError_Generic>();
-			Error->SourceAssetName = SourceFilename;
-			Error->Text = LOCTEXT("CantFetchPayload", "Cannot fetch FBX payload data.");
+			if (UInterchangeResultError_Generic* Error = AddMessage<UInterchangeResultError_Generic>())
+			{
+				Error->SourceAssetName = SourceFilename;
+				Error->Text = LOCTEXT("CantFetchPayload", "Cannot fetch FBX payload data.");
+			}
 		}
 		return ResultPayloadUniqueId;
 	}
@@ -169,7 +244,7 @@ namespace UE::Interchange
 	TArray<FString> FInterchangeFbxParser::GetJsonLoadMessages() const
 	{
 		TArray<FString> JsonResults;
-		for (UInterchangeResult* Result : ResultsContainer->GetResults())
+		for (UInterchangeResult* Result : GetResultContainer()->GetResults())
 		{
 			JsonResults.Add(Result->ToJson());
 		}
@@ -177,6 +252,16 @@ namespace UE::Interchange
 		return JsonResults;
 	}
 
+	UInterchangeResultsContainer* FInterchangeFbxParser::GetResultContainer() const
+	{
+		if (InternalResultsContainer)
+		{
+			return InternalResultsContainer;
+		}
+		ensure(ResultsContainer);
+		return ResultsContainer.Get();
+	}
+
 }//ns UE::Interchange
 
 #undef LOCTEXT_NAMESPACE
diff --git a/Engine/Plugins/Interchange/Runtime/Source/Parsers/Fbx/Public/InterchangeFbxParser.h b/Engine/Plugins/Interchange/Runtime/Source/Parsers/Fbx/Public/InterchangeFbxParser.h
--- a/Engine/Plugins/Interchange/Runtime/Source/Parsers/Fbx/Public/InterchangeFbxParser.h
+++ b/Engine/Plugins/Interchange/Runtime/Source/Parsers/Fbx/Public/InterchangeFbxParser.h
@@ -23,7 +23,10 @@ namespace UE
 		public:
 			FInterchangeFbxParser();
 			~FInterchangeFbxParser();
+			
+			void ReleaseResources();
 
+			void SetResultContainer(UInterchangeResultsContainer* Result);
 			/**
 			 * Parse a file support by the fbx sdk. It just extract all the fbx node and create a FBaseNodeContainer and dump it in a json file inside the ResultFolder
 			 * @param - Filename is the file that the fbx sdk will read (.fbx or .obj)
@@ -31,6 +34,13 @@ namespace UE
 			 */
 			void LoadFbxFile(const FString& Filename, const FString& ResultFolder);
 
+			/**
+			 * Parse a file support by the fbx sdk. It just extract all the fbx node and create a FBaseNodeContainer and dump it in a json file inside the ResultFolder
+			 * @param - Filename is the file that the fbx sdk will read (.fbx or .obj)
+			 * @param - BaseNodecontainer is the container of the scene graph
+			 */
+			void LoadFbxFile(const FString& Filename, UInterchangeBaseNodeContainer& BaseNodecontainer);
+
 			/**
 			 * Extract payload data from the fbx, the key tell the translator what payload the client ask
 			 * @param - PayloadKey is the key that describe the payload data to extract from the fbx file
@@ -77,11 +87,17 @@ namespace UE
 			template <typename T>
 			T* AddMessage()
 			{
-				return ResultsContainer->Add<T>();
+				if (UInterchangeResultsContainer* ResultContainer = GetResultContainer())
+				{
+					return ResultContainer->Add<T>();
+				}
+				return nullptr;
 			}
 
-
 		private:
+			UInterchangeResultsContainer* GetResultContainer() const;
+
+			TObjectPtr<UInterchangeResultsContainer> InternalResultsContainer = nullptr;
 			TStrongObjectPtr<UInterchangeResultsContainer> ResultsContainer = nullptr;
 			FString SourceFilename;
 			FString ResultFilepath;
diff --git a/Engine/Plugins/Tests/InterchangeTests/Source/InterchangeTests/Private/InterchangeImportTest.cpp b/Engine/Plugins/Tests/InterchangeTests/Source/InterchangeTests/Private/InterchangeImportTest.cpp
--- a/Engine/Plugins/Tests/InterchangeTests/Source/InterchangeTests/Private/InterchangeImportTest.cpp
+++ b/Engine/Plugins/Tests/InterchangeTests/Source/InterchangeTests/Private/InterchangeImportTest.cpp
@@ -7,7 +7,6 @@
 #include "Editor/Transactor.h"
 #include "HAL/FileManager.h"
 #include "Misc/AutomationTest.h"
-#include "InterchangeDispatcher.h"
 #include "InterchangeHelper.h"
 #include "InterchangeImportTestData.h"
 #include "InterchangeImportTestPlan.h"
@@ -15,39 +14,7 @@
 #include "Modules/ModuleManager.h"
 #include "ObjectTools.h"
 
-
-class FInterchangeImportTestBase : public FAutomationTestBase
-{
-public:
-	FInterchangeImportTestBase(const FString& InName, const bool bInComplexTask)
-		: FAutomationTestBase(InName, bInComplexTask)
-	{
-	}
-
-	virtual bool CanRunInEnvironment(const FString& TestParams, FString* OutReason, bool* OutWarn) const override;
-};
-
-bool FInterchangeImportTestBase::CanRunInEnvironment(const FString& TestParams, FString* OutReason, bool* OutWarn) const
-{
-	//Make sure interchange worker is available, skip the test execution if unavailable
-	if (!UE::Interchange::FInterchangeDispatcher::IsInterchangeWorkerAvailable())
-	{
-		if (nullptr != OutWarn)
-		{
-			*OutWarn = false;
-		}
-
-		if (nullptr != OutReason)
-		{
-			*OutReason = "Interchange worker is not available";
-		}
-
-		return false;
-	}
-	return true;
-}
-
-IMPLEMENT_CUSTOM_COMPLEX_AUTOMATION_TEST(FInterchangeImportTest, FInterchangeImportTestBase, "Editor.Interchange", EAutomationTestFlags::EditorContext | EAutomationTestFlags::EngineFilter)
+IMPLEMENT_CUSTOM_COMPLEX_AUTOMATION_TEST(FInterchangeImportTest, FAutomationTestBase, "Editor.Interchange", EAutomationTestFlags::EditorContext | EAutomationTestFlags::EngineFilter)
 
 const FString InterchangeTestsRootGameFolder = TEXT("/Game/Tests/Interchange/");
 
diff --git a/Engine/Source/Runtime/Interchange/Core/Public/InterchangeTranslatorBase.h b/Engine/Source/Runtime/Interchange/Core/Public/InterchangeTranslatorBase.h
--- a/Engine/Source/Runtime/Interchange/Core/Public/InterchangeTranslatorBase.h
+++ b/Engine/Source/Runtime/Interchange/Core/Public/InterchangeTranslatorBase.h
@@ -49,6 +49,13 @@ public:
 	/** return true if the translator can translate the given source data. */
 	INTERCHANGECORE_API virtual bool CanImportSourceData(const UInterchangeSourceData* InSourceData) const;
 
+	/**
+	 * Return true if the translator can be instantiate more then one time.
+	 * Some translator are not thread safe and cannot translate two files in same time.
+	 * A non thread safe translator could use the InterchangeWorker program to run in a separate process and return true.
+	 */
+	INTERCHANGECORE_API virtual bool IsThreadSafe() const { return true; }
+
 	/** Specifies the capabilities of a translator. */
 	virtual EInterchangeTranslatorType GetTranslatorType() const { return EInterchangeTranslatorType::Assets; }
 
diff --git a/Engine/Source/Runtime/Interchange/Engine/Private/InterchangeManager.cpp b/Engine/Source/Runtime/Interchange/Engine/Private/InterchangeManager.cpp
--- a/Engine/Source/Runtime/Interchange/Engine/Private/InterchangeManager.cpp
+++ b/Engine/Source/Runtime/Interchange/Engine/Private/InterchangeManager.cpp
@@ -521,10 +521,6 @@ FGraphEventArray UE::Interchange::FImportAsyncHelper::GetCompletionTaskGraphEven
 	}
 	TasksToComplete.Append(PipelinePostImportTasks);
 
-	if (PreAsyncCompletionTask.GetReference())
-	{
-		TasksToComplete.Add(PreAsyncCompletionTask);
-	}
 	if (PreCompletionTask.GetReference())
 	{
 		TasksToComplete.Add(PreCompletionTask);
@@ -1259,6 +1255,26 @@ void UInterchangeManager::StartQueuedTasks(bool bCancelAllTasks /*= false*/)
 	//Each import can use 2 tasks in same time if the build of the asset ddc use the same task pool (i.e. staticmesh, skeletalmesh, texture...)
 	const int32 PoolWorkerThreadCount = FTaskGraphInterface::Get().GetNumWorkerThreads() / 2;
 	const int32 MaxNumWorker = FMath::Max(PoolWorkerThreadCount, 1);
+
+	for (TPair<UClass*, TArray<FQueuedTaskData>>& ClassAndTasks : NonParallelTranslatorQueueTasks)
+	{
+		if(ClassAndTasks.Value.IsEmpty())
+		{
+			continue;
+		}
+		bool& TranslatorLock = NonParallelTranslatorLocks.FindChecked(ClassAndTasks.Key);
+		if (!TranslatorLock)
+		{
+			FQueuedTaskData QueuedTaskData = ClassAndTasks.Value[0];
+			QueuedTasks.Enqueue(QueuedTaskData);
+			TranslatorLock = true;
+			constexpr bool bAllowShrinking = false;
+			ClassAndTasks.Value.RemoveAt(0, 1, bAllowShrinking);
+			//No need to process an another the lock is set
+			continue;
+		}
+	}
+
 	while (!QueuedTasks.IsEmpty() && (ImportTasks.Num() < MaxNumWorker || bCancelAllTasks))
 	{
 		FQueuedTaskData QueuedTaskData;
@@ -1821,7 +1837,28 @@ UInterchangeManager::ImportInternal(const FString& ContentPath, const UInterchan
 	FQueuedTaskData QueuedTaskData;
 	QueuedTaskData.AsyncHelper = AsyncHelper;
 	QueuedTaskData.PackageBasePath = PackageBasePath;
-	QueuedTasks.Enqueue(QueuedTaskData);
+	QueuedTaskData.TranslatorClass = AsyncTranslator->GetClass();
+
+	//If we cancel or abort the task we want to avoid putting it in the NonParallelTranslatorQueueTasks (the locks will not be release if the task doesn't start)
+	bool bTranslatorIsThreadSafe = AsyncTranslator->IsThreadSafe() || (bImportCanceled || bImportAborted);
+	if (bTranslatorIsThreadSafe)
+	{
+		QueuedTasks.Enqueue(QueuedTaskData);
+	}
+	else
+	{
+		//Add a NonParallelTranslatorLocks for this translator class
+		if (!NonParallelTranslatorLocks.Contains(AsyncTranslator->GetClass()))
+		{
+			//Create a boolean lock and initialize it to false
+			bool& TranslatorLock = NonParallelTranslatorLocks.FindOrAdd(AsyncTranslator->GetClass());
+			TranslatorLock = false;
+		}
+		//Add an entry in NonParallelTranslatorQueueTasks
+		TArray<FQueuedTaskData>& NonParallelQueuedTasks = NonParallelTranslatorQueueTasks.FindOrAdd(AsyncTranslator->GetClass());
+		NonParallelQueuedTasks.Add(QueuedTaskData);
+	}
+
 	QueueTaskCount = FMath::Clamp(QueueTaskCount + 1, 0, MAX_int32);
 
 	StartQueuedTasks();
@@ -1930,6 +1967,15 @@ void UInterchangeManager::ReleaseAsyncHelper(TWeakPtr<UE::Interchange::FImportAs
 	bool bSucceeded = false;
 	{
 		TSharedPtr<FImportAsyncHelper> AsyncHelperPtr = AsyncHelper.Pin();
+
+		//Free the lock to allow the next import to happen
+		if (AsyncHelperPtr->Translators.IsValidIndex(0))
+		{
+			if (bool* bTranslatorLock = NonParallelTranslatorLocks.Find(AsyncHelperPtr->Translators[0]->GetClass()))
+			{
+				*bTranslatorLock = false;
+			}
+		}
 		
 		auto ForEachResult = [&bSucceeded, bLogWarningsAndErrors](TArray<UInterchangeResult*>&& Results)
 		{
diff --git a/Engine/Source/Runtime/Interchange/Engine/Private/Tasks/InterchangeTaskCompletion.cpp b/Engine/Source/Runtime/Interchange/Engine/Private/Tasks/InterchangeTaskCompletion.cpp
--- a/Engine/Source/Runtime/Interchange/Engine/Private/Tasks/InterchangeTaskCompletion.cpp
+++ b/Engine/Source/Runtime/Interchange/Engine/Private/Tasks/InterchangeTaskCompletion.cpp
@@ -20,21 +20,6 @@
 #include "UObject/ObjectMacros.h"
 #include "UObject/WeakObjectPtrTemplates.h"
 
-void UE::Interchange::FTaskPreAsyncCompletion::DoTask(ENamedThreads::Type CurrentThread, const FGraphEventRef& MyCompletionGraphEvent)
-{
-	TRACE_CPUPROFILER_EVENT_SCOPE(UE::Interchange::FTaskPreAsyncCompletion::DoTask)
-#if INTERCHANGE_TRACE_ASYNCHRONOUS_TASK_ENABLED
-	INTERCHANGE_TRACE_ASYNCHRONOUS_TASK(PreAsyncCompletion)
-#endif
-	FGCScopeGuard GCScopeGuard;
-
-	TSharedPtr<FImportAsyncHelper, ESPMode::ThreadSafe> AsyncHelper = WeakAsyncHelper.Pin();
-	check(AsyncHelper.IsValid());
-
-	//No need anymore of the translators sources
-	AsyncHelper->ReleaseTranslatorsSource();
-}
-
 void UE::Interchange::FTaskPreCompletion::DoTask(ENamedThreads::Type CurrentThread, const FGraphEventRef& MyCompletionGraphEvent)
 {
 	TRACE_CPUPROFILER_EVENT_SCOPE(UE::Interchange::FTaskPreCompletion::DoTask)
diff --git a/Engine/Source/Runtime/Interchange/Engine/Private/Tasks/InterchangeTaskCompletion.h b/Engine/Source/Runtime/Interchange/Engine/Private/Tasks/InterchangeTaskCompletion.h
--- a/Engine/Source/Runtime/Interchange/Engine/Private/Tasks/InterchangeTaskCompletion.h
+++ b/Engine/Source/Runtime/Interchange/Engine/Private/Tasks/InterchangeTaskCompletion.h
@@ -12,35 +12,6 @@ namespace UE
 {
 	namespace Interchange
 	{
-		class FTaskPreAsyncCompletion
-		{
-		private:
-			UInterchangeManager* InterchangeManager;
-			TWeakPtr<FImportAsyncHelper, ESPMode::ThreadSafe> WeakAsyncHelper;
-		public:
-			FTaskPreAsyncCompletion(UInterchangeManager* InInterchangeManager, TWeakPtr<FImportAsyncHelper, ESPMode::ThreadSafe> InAsyncHelper)
-				: InterchangeManager(InInterchangeManager)
-				, WeakAsyncHelper(InAsyncHelper)
-			{
-			}
-
-			static FORCEINLINE ENamedThreads::Type GetDesiredThread()
-			{
-				return ENamedThreads::AnyBackgroundThreadNormalTask;
-			}
-			static FORCEINLINE ESubsequentsMode::Type GetSubsequentsMode()
-			{
-				return ESubsequentsMode::TrackSubsequents;
-			}
-
-			FORCEINLINE TStatId GetStatId() const
-			{
-				RETURN_QUICK_DECLARE_CYCLE_STAT(FTaskPreAsyncCompletion, STATGROUP_TaskGraphTasks);
-			}
-
-			void DoTask(ENamedThreads::Type CurrentThread, const FGraphEventRef& MyCompletionGraphEvent);
-		};
-
 		class FTaskPreCompletion
 		{
 		private:
diff --git a/Engine/Source/Runtime/Interchange/Engine/Private/Tasks/InterchangeTaskParsing.cpp b/Engine/Source/Runtime/Interchange/Engine/Private/Tasks/InterchangeTaskParsing.cpp
--- a/Engine/Source/Runtime/Interchange/Engine/Private/Tasks/InterchangeTaskParsing.cpp
+++ b/Engine/Source/Runtime/Interchange/Engine/Private/Tasks/InterchangeTaskParsing.cpp
@@ -432,9 +432,5 @@ void UE::Interchange::FTaskParsing::DoTask(ENamedThreads::Type CurrentThread, co
 		}
 	}
 
-	FGraphEventArray PreAsyncCompletionPrerequistes;
-	AsyncHelper->PreAsyncCompletionTask = TGraphTask<FTaskPreAsyncCompletion>::CreateTask(&AssetCompilationPrerequistes).ConstructAndDispatchWhenReady(InterchangeManager, WeakAsyncHelper);
-	PreAsyncCompletionPrerequistes.Add(AsyncHelper->PreAsyncCompletionTask);
-
-	AsyncHelper->CompletionTask = TGraphTask<FTaskCompletion>::CreateTask(&PreAsyncCompletionPrerequistes).ConstructAndDispatchWhenReady(InterchangeManager, WeakAsyncHelper);
+	AsyncHelper->CompletionTask = TGraphTask<FTaskCompletion>::CreateTask(&AssetCompilationPrerequistes).ConstructAndDispatchWhenReady(InterchangeManager, WeakAsyncHelper);
 }
diff --git a/Engine/Source/Runtime/Interchange/Engine/Public/InterchangeManager.h b/Engine/Source/Runtime/Interchange/Engine/Public/InterchangeManager.h
--- a/Engine/Source/Runtime/Interchange/Engine/Public/InterchangeManager.h
+++ b/Engine/Source/Runtime/Interchange/Engine/Public/InterchangeManager.h
@@ -225,7 +225,6 @@ namespace UE
 			TArray<FGraphEventRef> FinalizeImportObjectTasks;
 			TArray<FGraphEventRef> SceneTasks;
 
-			FGraphEventRef PreAsyncCompletionTask;
 			FGraphEventRef PreCompletionTask;
 			FGraphEventRef CompletionTask;
 
@@ -676,7 +675,11 @@ private:
 	{
 		FString PackageBasePath;
 		TSharedPtr<UE::Interchange::FImportAsyncHelper, ESPMode::ThreadSafe> AsyncHelper;
+		UClass* TranslatorClass = nullptr;
 	};
+
+	TMap<UClass*, bool> NonParallelTranslatorLocks;
+	TMap<UClass*, TArray<FQueuedTaskData>> NonParallelTranslatorQueueTasks;
 	
 	//Queue all incomming tasks if there is more started task then we have cores
 	TQueue<FQueuedTaskData> QueuedTasks;
