Interchange: Make sure we compute the Interchange temporary folder when using the fbx in process translator.

#rnx

diff --git a/Engine/Plugins/Interchange/Runtime/Source/Import/Private/Fbx/InterchangeFbxTranslator.cpp b/Engine/Plugins/Interchange/Runtime/Source/Import/Private/Fbx/InterchangeFbxTranslator.cpp
--- a/Engine/Plugins/Interchange/Runtime/Source/Import/Private/Fbx/InterchangeFbxTranslator.cpp
+++ b/Engine/Plugins/Interchange/Runtime/Source/Import/Private/Fbx/InterchangeFbxTranslator.cpp
@@ -64,29 +64,30 @@ UInterchangeFbxTranslator::UInterchangeFbxTranslator()
 {
 	Dispatcher = nullptr;
 	bUseWorkerImport = false;
-	if (GInterchangeFBXTranslatorUseWorker)
+	
+	FGuid RandomGuid;
+	FPlatformMisc::CreateGuid(RandomGuid);
+	IPlatformFile& PlatformFile = FPlatformFileManager::Get().GetPlatformFile();
+	const FString ProjectSavedDir = FPaths::ConvertRelativePathToFull(FPaths::ProjectSavedDir());
+	const FString RamdomGuidDir = RandomGuid.ToString(EGuidFormats::Base36Encoded);
+	if (!PlatformFile.DirectoryExists(*ProjectSavedDir))
 	{
-		//Run the import in parallel only if we can start interchange worker
-		FGuid RandomGuid;
-		FPlatformMisc::CreateGuid(RandomGuid);
-		IPlatformFile& PlatformFile = FPlatformFileManager::Get().GetPlatformFile();
-		const FString ProjectSavedDir = FPaths::ConvertRelativePathToFull(FPaths::ProjectSavedDir());
-		const FString RamdomGuidDir = RandomGuid.ToString(EGuidFormats::Base36Encoded);
-		if (!PlatformFile.DirectoryExists(*ProjectSavedDir))
-		{
-			PlatformFile.CreateDirectory(*ProjectSavedDir);
-		}
-		const FString InterchangeDir = FPaths::Combine(ProjectSavedDir, TEXT("Interchange"));
-		if (!PlatformFile.DirectoryExists(*InterchangeDir))
-		{
-			PlatformFile.CreateDirectory(*InterchangeDir);
-		}
-		ResultFolder = FPaths::Combine(InterchangeDir, RamdomGuidDir);
-		if (!PlatformFile.DirectoryExists(*ResultFolder))
-		{
-			PlatformFile.CreateDirectory(*ResultFolder);
-		}
+		PlatformFile.CreateDirectory(*ProjectSavedDir);
+	}
+	const FString InterchangeDir = FPaths::Combine(ProjectSavedDir, TEXT("Interchange"));
+	if (!PlatformFile.DirectoryExists(*InterchangeDir))
+	{
+		PlatformFile.CreateDirectory(*InterchangeDir);
+	}
+	ResultFolder = FPaths::Combine(InterchangeDir, RamdomGuidDir);
+	if (!PlatformFile.DirectoryExists(*ResultFolder))
+	{
+		PlatformFile.CreateDirectory(*ResultFolder);
+	}
 
+	//Run the import in parallel only if we can start interchange worker
+	if (GInterchangeFBXTranslatorUseWorker)
+	{
 		//Create the dispatcher
 		Dispatcher = MakeUnique<UE::Interchange::FInterchangeDispatcher>(ResultFolder);
 
@@ -206,6 +207,15 @@ void UInterchangeFbxTranslator::ReleaseSource()
 #if WITH_EDITOR
 	FbxParser.ReleaseResources();
 #endif
+	//Delete the result folder if we are not running with the worker, in the other case the dispatcher will delete the folder on TerminateProcess.
+	if (!bUseWorkerImport)
+	{
+		IPlatformFile& PlatformFile = FPlatformFileManager::Get().GetPlatformFile();
+		constexpr bool RequireExists = false;
+		//Delete recursively folder's content
+		constexpr bool Tree = true;
+		IFileManager::Get().DeleteDirectory(*ResultFolder, RequireExists, Tree);
+	}
 	ResultFolder.Empty();
 }
 
