Remove unecessary line.

#rnx

diff --git a/Engine/Plugins/Interchange/Runtime/Source/Import/Private/Fbx/InterchangeFbxTranslator.cpp b/Engine/Plugins/Interchange/Runtime/Source/Import/Private/Fbx/InterchangeFbxTranslator.cpp
--- a/Engine/Plugins/Interchange/Runtime/Source/Import/Private/Fbx/InterchangeFbxTranslator.cpp
+++ b/Engine/Plugins/Interchange/Runtime/Source/Import/Private/Fbx/InterchangeFbxTranslator.cpp
@@ -210,7 +210,6 @@ void UInterchangeFbxTranslator::ReleaseSource()
 	//Delete the result folder if we are not running with the worker, in the other case the dispatcher will delete the folder on TerminateProcess.
 	if (!bUseWorkerImport)
 	{
-		IPlatformFile& PlatformFile = FPlatformFileManager::Get().GetPlatformFile();
 		constexpr bool RequireExists = false;
 		//Delete recursively folder's content
 		constexpr bool Tree = true;
