Interchange: Prevent pipeline filtering to change options

#rnx

diff --git a/Engine/Plugins/Interchange/Runtime/Source/Pipelines/Private/InterchangeGenericMaterialPipeline.cpp b/Engine/Plugins/Interchange/Runtime/Source/Pipelines/Private/InterchangeGenericMaterialPipeline.cpp
--- a/Engine/Plugins/Interchange/Runtime/Source/Pipelines/Private/InterchangeGenericMaterialPipeline.cpp
+++ b/Engine/Plugins/Interchange/Runtime/Source/Pipelines/Private/InterchangeGenericMaterialPipeline.cpp
@@ -529,7 +529,6 @@ void UInterchangeGenericMaterialPipeline::FilterPropertiesFromTranslatedData(UIn
 	MaterialCount += TmpMaterialNodes.Num();
 	if(MaterialCount == 0)
 	{
-		bImportMaterials = false;
 		TArray<FString> HideCategories;
 		//Filter out all material properties
 		HideCategories.Add(TEXT("Materials"));
diff --git a/Engine/Plugins/Interchange/Runtime/Source/Pipelines/Private/InterchangeGenericTexturePipeline.cpp b/Engine/Plugins/Interchange/Runtime/Source/Pipelines/Private/InterchangeGenericTexturePipeline.cpp
--- a/Engine/Plugins/Interchange/Runtime/Source/Pipelines/Private/InterchangeGenericTexturePipeline.cpp
+++ b/Engine/Plugins/Interchange/Runtime/Source/Pipelines/Private/InterchangeGenericTexturePipeline.cpp
@@ -146,7 +146,6 @@ void UInterchangeGenericTexturePipeline::FilterPropertiesFromTranslatedData(UInt
 	InBaseNodeContainer->GetNodes(UInterchangeTextureNode::StaticClass(), TmpTextureNodes);
 	if (TmpTextureNodes.Num() == 0)
 	{
-		bImportTextures = false;
 		//Filter out all Textures properties
 		if (UInterchangePipelineBase* OuterMostPipeline = GetMostPipelineOuter())
 		{
