[Interchange] Adding VirtualTextureStreaming conversion based on Base Materials (+Opacity chain and render settings)

#rnx

diff --git a/Engine/Plugins/Interchange/Runtime/Source/Import/Private/Material/InterchangeMaterialFactory.cpp b/Engine/Plugins/Interchange/Runtime/Source/Import/Private/Material/InterchangeMaterialFactory.cpp
--- a/Engine/Plugins/Interchange/Runtime/Source/Import/Private/Material/InterchangeMaterialFactory.cpp
+++ b/Engine/Plugins/Interchange/Runtime/Source/Import/Private/Material/InterchangeMaterialFactory.cpp
@@ -31,6 +31,7 @@
 #include "Nodes/InterchangeUserDefinedAttribute.h"
 #include "UObject/ObjectRedirector.h"
 #include "Misc/PackageName.h"
+#include "Engine/RendererSettings.h"
 
 #include UE_INLINE_GENERATED_CPP_BY_NAME(InterchangeMaterialFactory)
 
@@ -609,6 +610,94 @@ namespace UE::Interchange::MaterialFactory::Internal
 			}
 		}
 	}
+
+	UTexture* GetVirtualTextureStreamingMatchedTexture(uint8 DesiredVirtualTextureStreaming, UTexture* UsedTexture)
+	{
+		if (DesiredVirtualTextureStreaming == UsedTexture->VirtualTextureStreaming)
+		{
+			return nullptr;
+		}
+
+		auto GetExistingPackage = [](const FString& PackageName)
+			{
+				//Try to find the package in memory
+				UPackage* Pkg = FindPackage(nullptr, *PackageName);
+				if (!Pkg)
+				{
+					//Try to load the package from disk
+					Pkg = LoadPackage(nullptr, *PackageName, LOAD_NoWarn | LOAD_Quiet);
+				}
+
+				if (Pkg)
+				{
+					Pkg->FullyLoad();
+				}
+
+				return Pkg;
+			};
+
+		//Duplicate and Converts texture to what the DefaultTexture's setting requires (if needed)
+		//in case No VirtualTextureStreaming  => make the Texture's duplicate into a nonVT (and the name should be suffixed with "_nonVT")
+		//in case VirtualTextureStreaming     => make the Texture's duplicate into a VT (and the name should be suffixed with "_VT")
+
+		FString Suffix = DesiredVirtualTextureStreaming ? TEXT("_VT") : TEXT("_nonVT");
+
+		FString OriginalPackageName = UsedTexture->GetPackage()->GetName();
+		FString ConvertedPackageName = OriginalPackageName + Suffix;
+		FString ConvertedAssetName = FPaths::GetCleanFilename(ConvertedPackageName);
+
+		UPackage* ToBeConvertedPackage = GetExistingPackage(ConvertedPackageName);
+		UObject* ToBeConvertedObject = ToBeConvertedPackage ? StaticFindObject(nullptr, ToBeConvertedPackage, *ConvertedAssetName) : nullptr;
+
+		{
+			//object exists with the desired ConvertedName but its not a texture:
+			uint32 Counter = 1;
+			while (ToBeConvertedObject && !Cast<UTexture>(ToBeConvertedObject))
+			{
+				ConvertedPackageName += FString::FromInt(Counter);
+				Counter++;
+
+				ToBeConvertedPackage = GetExistingPackage(ConvertedPackageName);
+				ToBeConvertedObject = ToBeConvertedPackage ? StaticFindObject(nullptr, ToBeConvertedPackage, *ConvertedAssetName) : nullptr;
+			}
+		}
+
+		if (!ToBeConvertedObject)
+		{
+			UPackage* NewPackage = CreatePackage(*ConvertedPackageName);
+			NewPackage->SetPackageFlags(PKG_NewlyCreated);
+
+			ConvertedAssetName = FPaths::GetCleanFilename(ConvertedPackageName);
+			ToBeConvertedObject = StaticDuplicateObject(UsedTexture, NewPackage, *ConvertedAssetName);
+		}
+
+		if (!ToBeConvertedObject)
+		{
+			ensure(false);
+			return nullptr;
+		}
+
+		UTexture* ToBeConvertedTexture = Cast<UTexture>(ToBeConvertedObject);
+		if (!ToBeConvertedTexture)
+		{
+			ensure(false);
+			return nullptr;
+		}
+
+		if (ToBeConvertedTexture->VirtualTextureStreaming != DesiredVirtualTextureStreaming)
+		{
+#if WITH_EDITOR
+			FPropertyChangedEvent PropertyChangeEvent(UTexture::StaticClass()->FindPropertyByName(GET_MEMBER_NAME_CHECKED(UTexture, VirtualTextureStreaming)));
+			ToBeConvertedTexture->Modify();
+			ToBeConvertedTexture->VirtualTextureStreaming = DesiredVirtualTextureStreaming;
+			ToBeConvertedTexture->PostEditChangeProperty(PropertyChangeEvent);
+#else
+			ToBeConvertedTexture->VirtualTextureStreaming = DesiredVirtualTextureStreaming;
+#endif
+		}
+
+		return ToBeConvertedTexture;
+	}
 }
 		
 namespace UE::Interchange::Materials::HashUtils
@@ -989,16 +1078,54 @@ void UInterchangeMaterialFactory::SetupObject_GameThread(const FSetupObjectParam
 		//Update the samplers type in case the textures were changed during their SetupObject_GameThread
 		if (UMaterial* ImportedMaterial = Cast<UMaterial>(ImportedMaterialInterface))
 		{
+			TMap<UTexture*, UTexture*> OriginalToConvertedTextureMaps; //for opacity when bEnableVirtualTextureOpacityMask is false
+			if (!GetDefault<URendererSettings>()->bEnableVirtualTextureOpacityMask)
+			{
+				//Virtual textures are not supported in the OpacityMask slot, convert any textures back to a regular texture.
+				TArray<UTexture*> OutOpacityMaskTextures;
+				if (ImportedMaterialInterface->GetTexturesInPropertyChain(MP_OpacityMask, OutOpacityMaskTextures, nullptr, nullptr))
+				{
+					for (UTexture* CurrentTexture : OutOpacityMaskTextures)
+					{
+						if (UTexture* ConvertedTexture = UE::Interchange::MaterialFactory::Internal::GetVirtualTextureStreamingMatchedTexture(0 /*false*/, CurrentTexture))
+						{
+							OriginalToConvertedTextureMaps.Emplace(CurrentTexture, ConvertedTexture);
+						}
+					}
+				}
+			}
+
 			for (UMaterialExpression* Expression : ImportedMaterial->GetExpressions())
 			{
 				if (UMaterialExpressionTextureBase* TextureSample = Cast<UMaterialExpressionTextureBase>(Expression))
 				{
+					if (TextureSample->Texture && OriginalToConvertedTextureMaps.Contains(TextureSample->Texture))
+					{
+						TextureSample->Texture = OriginalToConvertedTextureMaps[TextureSample->Texture];
+					}
 					TextureSample->AutoSetSampleType();
 				}
 			}
 		}
 #endif // WITH_EDITOR
 
+		if (UMaterialInstance* ImportedMaterial = Cast<UMaterialInstance>(ImportedMaterialInterface))
+		{
+			for (struct FTextureParameterValue& TextureParameterValue : ImportedMaterial->TextureParameterValues)
+			{
+				if (TextureParameterValue.ParameterValue)
+				{
+					UTexture* DefaultTexture;
+					ImportedMaterial->GetTextureParameterDefaultValue(TextureParameterValue.ParameterInfo.Name, DefaultTexture);
+
+					if (UTexture* ConvertedTexture = UE::Interchange::MaterialFactory::Internal::GetVirtualTextureStreamingMatchedTexture(DefaultTexture->VirtualTextureStreaming, TextureParameterValue.ParameterValue))
+					{
+						TextureParameterValue.ParameterValue = ConvertedTexture;
+					}
+				}
+			}
+		}
+
 #if WITH_EDITORONLY_DATA
 		UE::Interchange::FFactoryCommon::FUpdateImportAssetDataParameters UpdateImportAssetDataParameters(ImportedMaterialInterface
 																										  , ImportedMaterialInterface->AssetImportData
