[Interchange - glTF] Fixing model import crash containing TRIANGLE_STRIP and TRIANGLE_FAN with no provided UVs.

#rnx

diff --git a/Engine/Plugins/Interchange/Runtime/Source/Parsers/GLTFCore/Private/GLTF/GLTFMesh.cpp b/Engine/Plugins/Interchange/Runtime/Source/Parsers/GLTFCore/Private/GLTF/GLTFMesh.cpp
--- a/Engine/Plugins/Interchange/Runtime/Source/Parsers/GLTFCore/Private/GLTF/GLTFMesh.cpp
+++ b/Engine/Plugins/Interchange/Runtime/Source/Parsers/GLTFCore/Private/GLTF/GLTFMesh.cpp
@@ -237,6 +237,25 @@ namespace GLTF
 		return GetAttributeAccessor(EMeshAttributeType::COLOR_0).Count;
 	}
 
+
+	const TArray<FPrimitive::EMode> FPrimitive::SupportedModes = { FPrimitive::EMode::Triangles, FPrimitive::EMode::TriangleStrip, FPrimitive::EMode::TriangleFan };
+	FString FPrimitive::ToString(const FPrimitive::EMode& Mode)
+	{
+		switch (Mode)
+		{
+			case EMode::Points:         return TEXT("POINTS");
+			case EMode::Lines:          return TEXT("LINES");
+			case EMode::LineLoop:       return TEXT("LINE_LOOP");
+			case EMode::LineStrip:      return TEXT("LINE_STRIP");
+			case EMode::Triangles:      return TEXT("TRIANGLES");
+			case EMode::TriangleStrip:  return TEXT("TRIANGLE_STRIP");
+			case EMode::TriangleFan:    return TEXT("TRIANGLE_FAN");
+
+			case EMode::Unknown:
+			default:                    return TEXT("UNKNOWN");
+		}
+	}
+
 	FPrimitive::FPrimitive(EMode InMode, int32 InMaterial, const FAccessor& InIndices, const FAccessor& InPosition, const FAccessor& InNormal,
 	                       const FAccessor& InTangent, const FAccessor& InTexCoord0, const FAccessor& InTexCoord1, const FAccessor& InColor0,
 	                       const FAccessor& InJoints0, const FAccessor& InWeights0)
diff --git a/Engine/Plugins/Interchange/Runtime/Source/Parsers/GLTFCore/Private/GLTF/GLTFReader.cpp b/Engine/Plugins/Interchange/Runtime/Source/Parsers/GLTFCore/Private/GLTF/GLTFReader.cpp
--- a/Engine/Plugins/Interchange/Runtime/Source/Parsers/GLTFCore/Private/GLTF/GLTFReader.cpp
+++ b/Engine/Plugins/Interchange/Runtime/Source/Parsers/GLTFCore/Private/GLTF/GLTFReader.cpp
@@ -363,7 +363,7 @@ namespace GLTF
 		Primitive.MorphTargets.Emplace(Position, Normal, Tangent, TexCoord0, TexCoord1, Color0);
 	}
 
-	void FFileReader::SetupPrimitive(const FJsonObject& Object, FMesh& Mesh, const bool bMeshQuantized) const
+	void FFileReader::SetupPrimitive(const FJsonObject& Object, FMesh& Mesh, const bool bMeshQuantized, const uint32& PrimitiveIndex) const
 	{
 		const FPrimitive::EMode       Mode = PrimitiveModeFromNumber(GetUnsignedInt(Object, TEXT("mode"), (uint32)FPrimitive::EMode::Triangles));
 		if (Mode == FPrimitive::EMode::Unknown)
@@ -371,6 +371,11 @@ namespace GLTF
 			return;
 		}
 
+		if (!FPrimitive::SupportedModes.Contains(Mode))
+		{
+			Messages.Emplace(EMessageSeverity::Warning, FString::Printf(TEXT("Primitive Mode[%s] in Primitive[%i] (in Mesh[%s]) is currently not supported. Geometry won't be imported."), *FPrimitive::ToString(Mode), PrimitiveIndex, *Mesh.Name));
+		}
+
 		const int32                   MaterialIndex = GetIndex(Object, TEXT("material"));
 		const TArray<FAccessor>& A             = Asset->Accessors;
 
@@ -420,10 +425,11 @@ namespace GLTF
 		Mesh.Primitives.Reserve(PrimArray.Num());
 
 		int32 NumberOfMorphTargets = -1;
+		uint32 PrimitiveIndex = 0;
 		for (TSharedPtr<FJsonValue> Value : PrimArray)
 		{
 			const FJsonObject& PrimObject = *Value->AsObject();
-			SetupPrimitive(PrimObject, Mesh, bMeshQuantized);
+			SetupPrimitive(PrimObject, Mesh, bMeshQuantized, PrimitiveIndex);
 
 			if (NumberOfMorphTargets == -1)
 			{
@@ -437,6 +443,8 @@ namespace GLTF
 					Messages.Emplace(EMessageSeverity::Error, TEXT("Number of Primitive.Targets is not consistent across the Mesh."));
 				}
 			}
+
+			PrimitiveIndex++;
 		}
 
 		// Morph Target Weights:
diff --git a/Engine/Plugins/Interchange/Runtime/Source/Parsers/GLTFCore/Private/GLTFMeshFactory.cpp b/Engine/Plugins/Interchange/Runtime/Source/Parsers/GLTFCore/Private/GLTFMeshFactory.cpp
--- a/Engine/Plugins/Interchange/Runtime/Source/Parsers/GLTFCore/Private/GLTFMeshFactory.cpp
+++ b/Engine/Plugins/Interchange/Runtime/Source/Parsers/GLTFCore/Private/GLTFMeshFactory.cpp
@@ -575,7 +575,7 @@ namespace GLTF
 				// Unreal StaticMesh must have UV channel 0.
 				// glTF doesn't require this since not all materials need texture coordinates.
 				// We also fill UV channel > 1 for this primitive if other primitives have it, to avoid gaps.
-				(*UVs[UVIndex]).AddZeroed(Primitive.VertexCount());
+				(*UVs[UVIndex]).AddZeroed(TriCount * 3);
 			}
 		}
 
diff --git a/Engine/Plugins/Interchange/Runtime/Source/Parsers/GLTFCore/Public/GLTFMesh.h b/Engine/Plugins/Interchange/Runtime/Source/Parsers/GLTFCore/Public/GLTFMesh.h
--- a/Engine/Plugins/Interchange/Runtime/Source/Parsers/GLTFCore/Public/GLTFMesh.h
+++ b/Engine/Plugins/Interchange/Runtime/Source/Parsers/GLTFCore/Public/GLTFMesh.h
@@ -107,7 +107,6 @@ namespace GLTF
 			LineStrip = 3,
 			// initially supported
 			Triangles = 4,
-			// will be supported prior to release
 			TriangleStrip = 5,
 			TriangleFan   = 6,
 
@@ -115,6 +114,9 @@ namespace GLTF
 			Unknown = 7
 		};
 
+		static const TArray<EMode> SupportedModes;
+		static FString ToString(const EMode& Mode);
+
 		const EMode Mode;
 		const int32 MaterialIndex;
 		TArray<FVariantMapping>	VariantMappings;
diff --git a/Engine/Plugins/Interchange/Runtime/Source/Parsers/GLTFCore/Public/GLTFReader.h b/Engine/Plugins/Interchange/Runtime/Source/Parsers/GLTFCore/Public/GLTFReader.h
--- a/Engine/Plugins/Interchange/Runtime/Source/Parsers/GLTFCore/Public/GLTFReader.h
+++ b/Engine/Plugins/Interchange/Runtime/Source/Parsers/GLTFCore/Public/GLTFReader.h
@@ -44,7 +44,7 @@ namespace GLTF
 		void SetupBufferView(const FJsonObject& Object) const;
 		void SetupAccessor(const FJsonObject& Object) const;
 		void SetupMorphTarget(const FJsonObject& Object, GLTF::FPrimitive& Primitive, const bool bMeshQuantized) const;
-		void SetupPrimitive(const FJsonObject& Object, GLTF::FMesh& Mesh, const bool bMeshQuantized) const;
+		void SetupPrimitive(const FJsonObject& Object, GLTF::FMesh& Mesh, const bool bMeshQuantized, const uint32& PrimitiveIndex) const;
 		void SetupMesh(const FJsonObject& Object, const bool bMeshQuantized) const;
 
 		void SetupScene(const FJsonObject& Object) const;
