Fix for Issue 398366: Compile errors in InterchangeMaterialFactory.cpp

#rnx

diff --git a/Engine/Plugins/Interchange/Runtime/Source/Import/Private/Material/InterchangeMaterialFactory.cpp b/Engine/Plugins/Interchange/Runtime/Source/Import/Private/Material/InterchangeMaterialFactory.cpp
--- a/Engine/Plugins/Interchange/Runtime/Source/Import/Private/Material/InterchangeMaterialFactory.cpp
+++ b/Engine/Plugins/Interchange/Runtime/Source/Import/Private/Material/InterchangeMaterialFactory.cpp
@@ -1115,12 +1115,13 @@ void UInterchangeMaterialFactory::SetupObject_GameThread(const FSetupObjectParam
 			{
 				if (TextureParameterValue.ParameterValue)
 				{
-					UTexture* DefaultTexture;
-					ImportedMaterial->GetTextureParameterDefaultValue(TextureParameterValue.ParameterInfo.Name, DefaultTexture);
-
-					if (UTexture* ConvertedTexture = UE::Interchange::MaterialFactory::Internal::GetVirtualTextureStreamingMatchedTexture(DefaultTexture->VirtualTextureStreaming, TextureParameterValue.ParameterValue))
+					UTexture* DefaultTexture = nullptr;
+					if (ImportedMaterial->GetTextureParameterDefaultValue(TextureParameterValue.ParameterInfo.Name, DefaultTexture))
 					{
-						TextureParameterValue.ParameterValue = ConvertedTexture;
+						if (UTexture* ConvertedTexture = UE::Interchange::MaterialFactory::Internal::GetVirtualTextureStreamingMatchedTexture(DefaultTexture->VirtualTextureStreaming, TextureParameterValue.ParameterValue))
+						{
+							TextureParameterValue.ParameterValue = ConvertedTexture;
+						}
 					}
 				}
 			}
