Interchange: Add pipeline conflict preview when showing pipeline properties.

#rnx

diff --git a/Engine/Plugins/Interchange/Editor/Source/Pipelines/Private/InterchangePipelineConfigurationGeneric.cpp b/Engine/Plugins/Interchange/Editor/Source/Pipelines/Private/InterchangePipelineConfigurationGeneric.cpp
--- a/Engine/Plugins/Interchange/Editor/Source/Pipelines/Private/InterchangePipelineConfigurationGeneric.cpp
+++ b/Engine/Plugins/Interchange/Editor/Source/Pipelines/Private/InterchangePipelineConfigurationGeneric.cpp
@@ -98,7 +98,8 @@ EInterchangePipelineConfigurationDialogResult UInterchangePipelineConfigurationG
 EInterchangePipelineConfigurationDialogResult UInterchangePipelineConfigurationGeneric::ShowReimportPipelineConfigurationDialog(TArray<FInterchangeStackInfo>& PipelineStacks
 	, TArray<UInterchangePipelineBase*>& OutPipelines
 	, TWeakObjectPtr<UInterchangeSourceData> SourceData
-	, TWeakObjectPtr<UInterchangeBaseNodeContainer> BaseNodeContainer)
+	, TWeakObjectPtr<UInterchangeBaseNodeContainer> BaseNodeContainer
+	, TWeakObjectPtr <UObject> ReimportAsset)
 {
 	//Create and show the graph inspector UI dialog
 	TSharedPtr<SWindow> ParentWindow;
@@ -121,6 +122,7 @@ EInterchangePipelineConfigurationDialogResult UInterchangePipelineConfigurationG
 		.PipelineStacks(PipelineStacks)
 		.OutPipelines(&OutPipelines)
 		.BaseNodeContainer(BaseNodeContainer)
+		.ReimportObject(ReimportAsset)
 	);
 
 	FSlateApplication::Get().AddModalWindow(Window, ParentWindow, false);
diff --git a/Engine/Plugins/Interchange/Editor/Source/Pipelines/Private/SInterchangePipelineConfigurationDialog.cpp b/Engine/Plugins/Interchange/Editor/Source/Pipelines/Private/SInterchangePipelineConfigurationDialog.cpp
--- a/Engine/Plugins/Interchange/Editor/Source/Pipelines/Private/SInterchangePipelineConfigurationDialog.cpp
+++ b/Engine/Plugins/Interchange/Editor/Source/Pipelines/Private/SInterchangePipelineConfigurationDialog.cpp
@@ -8,20 +8,20 @@
 #include "Framework/Views/TableViewMetadata.h"
 #include "IDetailsView.h"
 #include "IDocumentation.h"
-#include "Nodes/InterchangeBaseNodeContainer.h"
 #include "InterchangeManager.h"
 #include "InterchangePipelineConfigurationBase.h"
 #include "InterchangeProjectSettings.h"
+#include "Interfaces/IMainFrameModule.h"
 #include "Misc/App.h"
 #include "Misc/ConfigCacheIni.h"
+#include "Nodes/InterchangeBaseNodeContainer.h"
 #include "PropertyEditorModule.h"
 #include "SPrimaryButton.h"
 #include "Styling/SlateIconFinder.h"
 #include "Widgets/Input/SCheckBox.h"
 #include "Widgets/Input/STextComboBox.h"
 #include "Widgets/Layout/SUniformGridPanel.h"
-
-#include "Widgets/SWindow.h"
+#include "Layout/Visibility.h"
 
 #define LOCTEXT_NAMESPACE "InterchangePipelineConfiguration"
 
@@ -59,11 +59,28 @@ void SInterchangePipelineItem::Construct(
 	TObjectPtr<UInterchangePipelineBase> PipelineElementPtr = PipelineElement->Pipeline;
 	check(PipelineElementPtr.Get());
 	FText PipelineName = LOCTEXT("InvalidPipelineName", "Invalid Pipeline");
+	FText ConflictsComboEntryText = LOCTEXT("SInterchangePipelineItem::Conflicts", "Conflicts");
+	ConflictsComboEntry = MakeShared<FString>(ConflictsComboEntryText.ToString());
 	if (PipelineElementPtr.Get())
 	{
 		FString PipelineNameString = FString::Printf(TEXT("%s (%s)"), *PipelineElement->DisplayName, *PipelineElementPtr->GetClass()->GetName());
 		PipelineName = FText::FromString(PipelineNameString);
+		ConflictInfos = PipelineElementPtr->GetConflictInfos(PipelineElement->ReimportObject, PipelineElement->Container, PipelineElement->SourceData);
+		if (ConflictInfos.Num() > 0)
+		{
+			ConflictNameList.Reset(ConflictInfos.Num() + 1);
+			ConflictNameList.Add(ConflictsComboEntry);
+			for (const FInterchangeConflictInfo& ConflictInfo : ConflictInfos)
+			{
+				TSharedPtr<FString> ConflictNamePtr = MakeShared<FString>(ConflictInfo.DisplayName);
+				ConflictNameList.Add(ConflictNamePtr);
+			}
+		}
 	}
+	
+	ConflictComboBox = nullptr;
+
+	FText ConflictsComboBoxTooltip = LOCTEXT("ConflictsComboBoxTooltip", "If there is some conflict, simply select one to see more details.");
 		
 	STableRow<TSharedPtr<FInterchangePipelineItemType>>::Construct(
 		STableRow<TSharedPtr<FInterchangePipelineItemType>>::FArguments()
@@ -85,7 +102,48 @@ void SInterchangePipelineItem::Construct(
 				SNew(STextBlock)
 				.Text(PipelineName)
 			]
+			+ SHorizontalBox::Slot()
+			.AutoWidth()
+			.Padding(2.0f, 0.0f)
+			.VAlign(VAlign_Center)
+			[
+				SAssignNew(ConflictComboBox, STextComboBox)
+					.Visibility_Lambda([this]()->EVisibility
+						{
+							return ConflictInfos.Num() > 0 ? EVisibility::All : EVisibility::Collapsed;
+						})
+					.OptionsSource(&ConflictNameList)
+					.OnSelectionChanged_Lambda([this](TSharedPtr<FString> String, ESelectInfo::Type)
+						{
+							if (!String.IsValid() || String->IsEmpty() || ConflictsComboEntry.Get()->Equals(*String.Get()))
+							{
+								return;
+							}
+							//Find and display the conflict info
+							for (const FInterchangeConflictInfo& ConflictInfo : ConflictInfos)
+							{
+								if (ConflictInfo.DisplayName.Equals(*String.Get()))
+								{
+									if (ConflictInfo.Pipeline)
+									{
+										ConflictInfo.Pipeline->ShowConflictDialog(ConflictInfo.UniqueId);
+									}
+
+									//Re-select the conflict item after we show the conflict modal dialog
+									ConflictComboBox->SetSelectedItem(ConflictsComboEntry);
+									break;
+								}
+							}
+						})
+					.ToolTipText(ConflictsComboBoxTooltip)
+			]
 		], OwnerTable);
+
+	if (ConflictInfos.Num() > 0)
+	{
+		//Select the conflicts item
+		ConflictComboBox->SetSelectedItem(ConflictsComboEntry);
+	}
 }
 
 const FSlateBrush* SInterchangePipelineItem::GetImageItemIcon() const
@@ -177,7 +235,7 @@ TSharedRef<SBox> SInterchangePipelineConfigurationDialog::SpawnPipelineConfigura
 					{
 						GeneratedPipeline->FilterPropertiesFromTranslatedData(BaseNodeContainer.Get());
 					}
-					PipelineListViewItems.Add(MakeShareable(new FInterchangePipelineItemType{ GetPipelineDisplayName(DefaultPipeline), GeneratedPipeline}));
+					PipelineListViewItems.Add(MakeShareable(new FInterchangePipelineItemType{ GetPipelineDisplayName(DefaultPipeline), GeneratedPipeline, ReimportObject.Get(), BaseNodeContainer.Get(), SourceData.Get()}));
 				}
 			}
 			SelectedStack = StackNamePtr;
@@ -312,6 +370,12 @@ void SInterchangePipelineConfigurationDialog::Construct(const FArguments& InArgs
 	PipelineStacks = InArgs._PipelineStacks;
 	OutPipelines = InArgs._OutPipelines;
 	BaseNodeContainer = InArgs._BaseNodeContainer;
+	ReimportObject = InArgs._ReimportObject;
+	SourceData = InArgs._SourceData;
+	if (ReimportObject.IsValid())
+	{
+		ensure(bReimport);
+	}
 
 	check(OutPipelines);
 
@@ -548,12 +612,9 @@ FReply SInterchangePipelineConfigurationDialog::OnResetToDefault()
 					//We assume the pipelines inside one stack are all different classes, we use the class to know which default asset we need to duplicate
 					if (DefaultPipeline->GetClass() == PipelineClass)
 					{
-						
 						for(int32 PipelineIndex = 0; PipelineIndex < PipelineListViewItems.Num(); ++PipelineIndex)
 						{
-							
 							TObjectPtr<UInterchangePipelineBase> PipelineElement = PipelineListViewItems[PipelineIndex]->Pipeline;
-
 							if (PipelineElement.Get() == Pipeline)
 							{
 								if (UInterchangePipelineBase* GeneratedPipeline = UE::Interchange::GeneratePipelineInstanceInSourceAssetPackage(DefaultPipeline))
@@ -719,13 +780,17 @@ void SInterchangePipelineConfigurationDialog::RefreshStack(bool bStackSelectionC
 				{
 					//Load the settings for this pipeline
 					GeneratedPipeline->LoadSettings(Stack.StackName);
-					GeneratedPipeline->PreDialogCleanup(Stack.StackName);
+					if (bStackSelectionChange)
+					{
+						//Do not reset pipeline value if we are just refreshing the filtering
+						GeneratedPipeline->PreDialogCleanup(Stack.StackName);
+					}
 				}
 				if (bFilterOptions && BaseNodeContainer.IsValid())
 				{
 					GeneratedPipeline->FilterPropertiesFromTranslatedData(BaseNodeContainer.Get());
 				}
-				PipelineListViewItems.Add(MakeShareable(new FInterchangePipelineItemType{ GetPipelineDisplayName(DefaultPipeline), GeneratedPipeline }));
+				PipelineListViewItems.Add(MakeShareable(new FInterchangePipelineItemType{ GetPipelineDisplayName(DefaultPipeline), GeneratedPipeline, ReimportObject.Get(), BaseNodeContainer.Get(), SourceData.Get() }));
 			}
 		}
 	}
diff --git a/Engine/Plugins/Interchange/Editor/Source/Pipelines/Private/SInterchangePipelineConfigurationDialog.h b/Engine/Plugins/Interchange/Editor/Source/Pipelines/Private/SInterchangePipelineConfigurationDialog.h
--- a/Engine/Plugins/Interchange/Editor/Source/Pipelines/Private/SInterchangePipelineConfigurationDialog.h
+++ b/Engine/Plugins/Interchange/Editor/Source/Pipelines/Private/SInterchangePipelineConfigurationDialog.h
@@ -14,13 +14,16 @@ struct FPropertyAndParent;
 struct FSlateBrush;
 class IDetailsView;
 class SCheckBox;
+class STextComboBox;
 
 struct FInterchangePipelineItemType
 {
 public:
 	FString DisplayName;
-
 	UInterchangePipelineBase* Pipeline;
+	UObject* ReimportObject = nullptr;
+	UInterchangeBaseNodeContainer* Container = nullptr;
+	UInterchangeSourceData* SourceData = nullptr;
 };
 
 class SInterchangePipelineItem : public STableRow<TSharedPtr<FInterchangePipelineItemType>>
@@ -33,7 +36,11 @@ public:
 private:
 	const FSlateBrush* GetImageItemIcon() const;
 
-	TSharedPtr<FInterchangePipelineItemType> PipelineElement;
+	TSharedPtr<FInterchangePipelineItemType> PipelineElement = nullptr;
+	TArray<FInterchangeConflictInfo> ConflictInfos;
+	TArray<TSharedPtr<FString>> ConflictNameList;
+	TSharedPtr<FString> ConflictsComboEntry = nullptr;
+	TSharedPtr<STextComboBox> ConflictComboBox = nullptr;
 };
 
 typedef SListView< TSharedPtr<FInterchangePipelineItemType> > SPipelineListViewType;
@@ -59,6 +66,7 @@ public:
 		SLATE_ARGUMENT(TArray<FInterchangeStackInfo>, PipelineStacks)
 		SLATE_ARGUMENT(TArray<UInterchangePipelineBase*>*, OutPipelines)
 		SLATE_ARGUMENT(TWeakObjectPtr<UInterchangeBaseNodeContainer>, BaseNodeContainer)
+		SLATE_ARGUMENT(TWeakObjectPtr<UObject>, ReimportObject)
 	SLATE_END_ARGS()
 
 public:
@@ -107,6 +115,7 @@ private:
 	TWeakPtr< SWindow > OwnerWindow;
 	TWeakObjectPtr<UInterchangeSourceData> SourceData;
 	TWeakObjectPtr<UInterchangeBaseNodeContainer> BaseNodeContainer;
+	TWeakObjectPtr<UObject> ReimportObject;
 	TArray<FInterchangeStackInfo> PipelineStacks;
 	TArray<UInterchangePipelineBase*>* OutPipelines;
 
diff --git a/Engine/Plugins/Interchange/Editor/Source/Pipelines/Public/InterchangePipelineConfigurationGeneric.h b/Engine/Plugins/Interchange/Editor/Source/Pipelines/Public/InterchangePipelineConfigurationGeneric.h
--- a/Engine/Plugins/Interchange/Editor/Source/Pipelines/Public/InterchangePipelineConfigurationGeneric.h
+++ b/Engine/Plugins/Interchange/Editor/Source/Pipelines/Public/InterchangePipelineConfigurationGeneric.h
@@ -29,7 +29,8 @@ protected:
 	virtual EInterchangePipelineConfigurationDialogResult ShowReimportPipelineConfigurationDialog(TArray<FInterchangeStackInfo>& PipelineStacks
 		, TArray<UInterchangePipelineBase*>& OutPipelines
 		, TWeakObjectPtr<UInterchangeSourceData> SourceData
-		, TWeakObjectPtr<UInterchangeBaseNodeContainer> BaseNodeContainer) override;
+		, TWeakObjectPtr<UInterchangeBaseNodeContainer> BaseNodeContainer
+		, TWeakObjectPtr <UObject> ReimportAsset) override;
 };
 
 #if UE_ENABLE_INCLUDE_ORDER_DEPRECATED_IN_5_2
diff --git a/Engine/Plugins/Interchange/Runtime/Source/Pipelines/InterchangePipelines.Build.cs b/Engine/Plugins/Interchange/Runtime/Source/Pipelines/InterchangePipelines.Build.cs
--- a/Engine/Plugins/Interchange/Runtime/Source/Pipelines/InterchangePipelines.Build.cs
+++ b/Engine/Plugins/Interchange/Runtime/Source/Pipelines/InterchangePipelines.Build.cs
@@ -22,10 +22,14 @@ namespace UnrealBuildTool.Rules
 			PrivateDependencyModuleNames.AddRange(
 				new string[]
 				{
+					"ApplicationCore",
 					"CinematicCamera",
 					"InterchangeImport",
 					"DeveloperSettings",
 					"GLTFCore",
+					"InputCore",
+					"Slate",
+					"SlateCore",
 				}
 			);
 
@@ -34,6 +38,7 @@ namespace UnrealBuildTool.Rules
 				PrivateDependencyModuleNames.AddRange(
 					new string[]
 					{
+						"MainFrame",
 						"PhysicsUtilities",
 						"TextureUtilitiesCommon",
 						"UnrealEd",
diff --git a/Engine/Plugins/Interchange/Runtime/Source/Pipelines/Private/InterchangeGenericAssetsPipeline.cpp b/Engine/Plugins/Interchange/Runtime/Source/Pipelines/Private/InterchangeGenericAssetsPipeline.cpp
--- a/Engine/Plugins/Interchange/Runtime/Source/Pipelines/Private/InterchangeGenericAssetsPipeline.cpp
+++ b/Engine/Plugins/Interchange/Runtime/Source/Pipelines/Private/InterchangeGenericAssetsPipeline.cpp
@@ -3,7 +3,12 @@
 #include "InterchangeGenericAssetsPipeline.h"
 
 #include "Animation/Skeleton.h"
+#include "AssetRegistry/AssetData.h"
+#include "AssetRegistry/AssetRegistryModule.h"
 #include "CoreMinimal.h"
+#include "Engine/SkinnedAssetCommon.h"
+#include "Framework/Docking/TabManager.h"
+#include "HAL/PlatformApplicationMisc.h"
 #include "InterchangeAnimationTrackSetNode.h"
 #include "InterchangeAnimSequenceFactoryNode.h"
 #include "InterchangeCommonPipelineDataFactoryNode.h"
@@ -13,26 +18,44 @@
 #include "InterchangeGenericTexturePipeline.h"
 #include "InterchangeMaterialFactoryNode.h"
 #include "InterchangeMeshNode.h"
+#include "InterchangePipelineHelper.h"
 #include "InterchangePipelineLog.h"
 #include "InterchangePipelineMeshesUtilities.h"
 #include "InterchangeSceneNode.h"
 #include "InterchangeSkeletalMeshFactoryNode.h"
 #include "InterchangeSkeletalMeshLodDataNode.h"
 #include "InterchangeSkeletonFactoryNode.h"
+#include "InterchangeSkeletonHelper.h"
 #include "InterchangeSourceData.h"
 #include "InterchangeStaticMeshFactoryNode.h"
 #include "InterchangeStaticMeshLodDataNode.h"
+#include "Misc/PackageName.h"
 #include "Misc/Paths.h"
 #include "Nodes/InterchangeBaseNode.h"
 #include "Nodes/InterchangeBaseNodeContainer.h"
 #include "Nodes/InterchangeUserDefinedAttribute.h"
+#include "Styling/AppStyle.h"
 #include "Templates/SubclassOf.h"
 #include "UObject/Object.h"
 #include "UObject/ObjectMacros.h"
 #include "UObject/MetaData.h"
+#include "UObject/ReferencerFinder.h"
+#include "Widgets/Images/SImage.h"
+#include "Widgets/Input/SButton.h"
+#include "Widgets/Layout/SBorder.h"
+#include "Widgets/Layout/SBox.h"
+#include "Widgets/Layout/SSeparator.h"
+#include "Widgets/Text/STextBlock.h"
+#include "Widgets/Views/STableRow.h"
 
 #include UE_INLINE_GENERATED_CPP_BY_NAME(InterchangeGenericAssetsPipeline)
 
+const FName SInterchangeGenericAssetMaterialConflictWidget::NAME_Import = FName(TEXT("Imported"));
+const FName SInterchangeGenericAssetMaterialConflictWidget::NAME_Asset = FName(TEXT("Asset"));
+
+const FSlateColor SInterchangeGenericAssetMaterialConflictWidget::SlateColorFullConflict = FSlateColor(FLinearColor(0.9f, 0.25f, 0.0f));
+const FSlateColor SInterchangeGenericAssetMaterialConflictWidget::SlateColorSubConflict = FSlateColor(FLinearColor(0.7f, 0.45f, 0.1f));
+
 UInterchangeGenericAssetsPipeline::UInterchangeGenericAssetsPipeline()
 {
 	MaterialPipeline = CreateDefaultSubobject<UInterchangeGenericMaterialPipeline>("MaterialPipeline");
@@ -331,8 +354,651 @@ bool UInterchangeGenericAssetsPipeline::GetPropertyPossibleValues(const FName Pr
 	return Super::GetPropertyPossibleValues(PropertyPath, PossibleValues);
 }
 
+void UInterchangeGenericAssetsPipeline::CreateMaterialConflict(UStaticMesh* StaticMesh, USkeletalMesh* SkeletalMesh, UInterchangeBaseNodeContainer* TransientBaseNodeContainer)
+{
+	UObject* ReimportObject = nullptr;
+	MaterialConflictData.Reset();
+
+	TArray<FString> AssetImportMaterialNames;
+	TMap<FString, FString> MaterialNodePerMaterialSlotName;
+	FString NoName = TEXT("Material_");
+	uint32 NameIndex = 1;
+	if (StaticMesh)
+	{
+		ReimportObject = StaticMesh;
+		for (FStaticMaterial& Material : StaticMesh->GetStaticMaterials())
+		{
+			if (Material.ImportedMaterialSlotName != NAME_None)
+			{
+				AssetImportMaterialNames.Add(Material.ImportedMaterialSlotName.ToString());
+			}
+			else if (Material.MaterialSlotName != NAME_None)
+			{
+				AssetImportMaterialNames.Add(Material.MaterialSlotName.ToString());
+			}
+			else if (Material.MaterialInterface)
+			{
+				AssetImportMaterialNames.Add(Material.MaterialInterface->GetName());
+			}
+			else
+			{
+				AssetImportMaterialNames.Add(NoName + FString::FromInt(NameIndex++));
+			}
+		}
+
+		TArray<FString> StaticMeshes;
+		TransientBaseNodeContainer->GetNodes(UInterchangeStaticMeshFactoryNode::StaticClass(), StaticMeshes);
+		if (StaticMeshes.Num() == 1)
+		{
+			//Grab all imported materials
+			UInterchangeStaticMeshFactoryNode* MeshFactoryNode = Cast<UInterchangeStaticMeshFactoryNode>(TransientBaseNodeContainer->GetFactoryNode(StaticMeshes[0]));
+			MeshFactoryNode->GetSlotMaterialDependencies(MaterialNodePerMaterialSlotName);
+		}
+	}
+	else if (ensure(SkeletalMesh))
+	{
+		ReimportObject = SkeletalMesh;
+		for (FSkeletalMaterial& Material : SkeletalMesh->GetMaterials())
+		{
+#if WITH_EDITORONLY_DATA
+			if (Material.ImportedMaterialSlotName != NAME_None)
+			{
+				AssetImportMaterialNames.Add(Material.ImportedMaterialSlotName.ToString());
+			}
+			else
+#endif // WITH_EDITORONLY_DATA
+				if (Material.MaterialSlotName != NAME_None)
+				{
+					AssetImportMaterialNames.Add(Material.MaterialSlotName.ToString());
+				}
+				else if (Material.MaterialInterface)
+				{
+					AssetImportMaterialNames.Add(Material.MaterialInterface->GetName());
+				}
+				else
+				{
+					AssetImportMaterialNames.Add(NoName + FString::FromInt(NameIndex++));
+				}
+		}
+
+		TArray<FString> SkeletalMeshes;
+		TransientBaseNodeContainer->GetNodes(UInterchangeSkeletalMeshFactoryNode::StaticClass(), SkeletalMeshes);
+		if (SkeletalMeshes.Num() == 1)
+		{
+			//Grab all imported materials
+			UInterchangeSkeletalMeshFactoryNode* MeshFactoryNode = Cast<UInterchangeSkeletalMeshFactoryNode>(TransientBaseNodeContainer->GetFactoryNode(SkeletalMeshes[0]));
+			MeshFactoryNode->GetSlotMaterialDependencies(MaterialNodePerMaterialSlotName);
+		}
+	}
+	//Compare and cache the results
+	bool bHasConflict = false;
+	TArray<int32> MatchMaterials;
+	int32 MatchMaterialCount = 0;
+	int32 ImportMaterialIndex = 0;
+	const bool bImportMoreMaterial = MaterialNodePerMaterialSlotName.Num() > AssetImportMaterialNames.Num();
+	for (TPair<FString, FString> MaterialSlotNameAndMaterialNode : MaterialNodePerMaterialSlotName)
+	{
+		int32& MatchMaterial = MatchMaterials.Add_GetRef(INDEX_NONE);
+		const FString& ImportMaterialName = MaterialSlotNameAndMaterialNode.Key;
+		bool bFoundMatch = false;
+		for (int32 AssetMaterialIndex = 0; AssetMaterialIndex < AssetImportMaterialNames.Num(); ++AssetMaterialIndex)
+		{
+			const FString& AssetMaterialName = AssetImportMaterialNames[AssetMaterialIndex];
+			if (ImportMaterialName.Equals(AssetMaterialName))
+			{
+				bFoundMatch = true;
+				MatchMaterial = AssetMaterialIndex;
+				MatchMaterialCount++;
+				break;
+			}
+		}
+		if (!bFoundMatch)
+		{
+			bHasConflict = true;
+		}
+		ImportMaterialIndex++;
+	}
+
+	//Remove the conflict if we match all original asset materials
+	if (bHasConflict && bImportMoreMaterial && MatchMaterialCount == AssetImportMaterialNames.Num())
+	{
+		bHasConflict = false;
+	}
+
+	//We have a conflict, make sure the MatchMaterials array is in sync with MaterialNodePerMaterialSlotName.
+	if (bHasConflict && MaterialNodePerMaterialSlotName.Num() == MatchMaterials.Num())
+	{
+		//Throw all unmatch material into the conflict data
+		FInterchangeConflictInfo& MaterialConflict = ConflictInfos.AddDefaulted_GetRef();
+		MaterialConflict.DisplayName = TEXT("Materials");
+		MaterialConflict.Description = TEXT("There is some unmatched materials");
+		MaterialConflict.Pipeline = this;
+		MaterialConflict.UniqueId = FGuid::NewGuid();
+
+		//Cache the data so we do not have to redo the works when we will show the conflict
+		MaterialConflictData.ConflictUniqueId = MaterialConflict.UniqueId;
+		MaterialConflictData.AssetMaterialNames = AssetImportMaterialNames;
+		for (TPair<FString, FString> MaterialSlotNameAndMaterialNode : MaterialNodePerMaterialSlotName)
+		{
+			const FString& ImportMaterialName = MaterialSlotNameAndMaterialNode.Key;
+			MaterialConflictData.ImportMaterialNames.Add(ImportMaterialName);
+		}
+		MaterialConflictData.MatchMaterialIndexes = MatchMaterials;
+		MaterialConflictData.ReimportObject = ReimportObject;
+	}
+}
+
+void UInterchangeGenericAssetsPipeline::InternalRecursiveFillJointsFromReferenceSkeleton(TSharedPtr<FSkeletonJoint> ParentJoint, TMap<FString, TSharedPtr<FSkeletonJoint>>& Joints, const int32 BoneIndex, const FReferenceSkeleton& ReferenceSkeleton)
+{
+	const TArray<FMeshBoneInfo>& BoneInfos = ReferenceSkeleton.GetRawRefBoneInfo();
+	if (!BoneInfos.IsValidIndex(BoneIndex))
+	{
+		return;
+	}
+	const FMeshBoneInfo& BoneInfo = BoneInfos[BoneIndex];
+	const FString BoneName = BoneInfo.Name.ToString();
+	//We should not have any name collision
+	if (!ensure(!Joints.Contains(BoneName)))
+	{
+		return;
+	}
+	TSharedPtr<FSkeletonJoint> SkeletonJoint = MakeShared<FSkeletonJoint>();
+	SkeletonJoint->JointName = BoneName;
+	SkeletonJoint->Parent = ParentJoint;
+	if (ParentJoint.IsValid())
+	{
+		ParentJoint->Children.Add(SkeletonJoint);
+	}
+
+	Joints.Add(BoneName, SkeletonJoint);
+
+	TArray<int32> Children;
+	ReferenceSkeleton.GetRawDirectChildBones(BoneIndex, Children);
+	for (int32 ChildIndex = 0; ChildIndex < Children.Num(); ++ChildIndex)
+	{
+		InternalRecursiveFillJointsFromReferenceSkeleton(SkeletonJoint, Joints, Children[ChildIndex], ReferenceSkeleton);
+	}
+}
+
+void UInterchangeGenericAssetsPipeline::InternalRecursiveFillJointsFromNodeContainer(TSharedPtr<FSkeletonJoint> ParentJoint, TMap<FString, TSharedPtr<FSkeletonJoint>>& Joints, const FString& JoinUid, const UInterchangeBaseNodeContainer* BaseNodeContainer, const bool bConvertStaticToSkeletalActive)
+{
+	const UInterchangeSceneNode* SceneNode = Cast<UInterchangeSceneNode>(BaseNodeContainer->GetNode(JoinUid));
+	if (!bConvertStaticToSkeletalActive && (!SceneNode || !SceneNode->IsSpecializedTypeContains(UE::Interchange::FSceneNodeStaticData::GetJointSpecializeTypeString())))
+	{
+		return;
+	}
+	const FString ImportBoneName = SceneNode->GetDisplayLabel();
+	//We should not have any name collision
+	if (!ensure(!Joints.Contains(ImportBoneName)))
+	{
+		return;
+	}
+	TSharedPtr<FSkeletonJoint> SkeletonJoint = MakeShared<FSkeletonJoint>();
+	SkeletonJoint->JointName = ImportBoneName;
+	SkeletonJoint->Parent = ParentJoint;
+	if (ParentJoint.IsValid())
+	{
+		ParentJoint->Children.Add(SkeletonJoint);
+	}
+
+	Joints.Add(ImportBoneName, SkeletonJoint);
+
+	//Iterate childrens
+	const TArray<FString> ChildrenIds = BaseNodeContainer->GetNodeChildrenUids(JoinUid);
+	for (int32 ChildIndex = 0; ChildIndex < ChildrenIds.Num(); ++ChildIndex)
+	{
+		InternalRecursiveFillJointsFromNodeContainer(SkeletonJoint, Joints, ChildrenIds[ChildIndex], BaseNodeContainer, bConvertStaticToSkeletalActive);
+	}
+}
+
+namespace UE::Interchange::Private
+{
+	void RecursivelyFillJointRemoved(TSharedPtr<UInterchangeGenericAssetsPipeline::FSkeletonJoint> AssetJoint
+		, TSharedPtr<UInterchangeGenericAssetsPipeline::FSkeletonJoint> ParentJoint
+		, TMap<FString, TSharedPtr<UInterchangeGenericAssetsPipeline::FSkeletonJoint>>& ConflictDataJoints)
+	{
+		if (!AssetJoint.IsValid())
+		{
+			return;
+		}
+		auto SetRemove = [](TSharedPtr<UInterchangeGenericAssetsPipeline::FSkeletonJoint> JointMatch)
+			{
+				JointMatch->bMatch = false;
+				JointMatch->bAdded = false;
+				JointMatch->bRemoved = true;
+			};
+		TSharedPtr<UInterchangeGenericAssetsPipeline::FSkeletonJoint> NewJoint = MakeShared<UInterchangeGenericAssetsPipeline::FSkeletonJoint>();
+		NewJoint->JointName = AssetJoint->JointName;
+		if (ParentJoint.IsValid())
+		{
+			NewJoint->Parent = ParentJoint;
+			ParentJoint->Children.Add(NewJoint);
+		}
+		else
+		{
+			//Add only the root nodes
+			ConflictDataJoints.Add(NewJoint->JointName, NewJoint);
+		}
+		SetRemove(NewJoint);
+		const int32 ChildCount = AssetJoint->Children.Num();
+		for (int32 ChildIndex = 0; ChildIndex < ChildCount; ++ChildIndex)
+		{
+			TSharedPtr<UInterchangeGenericAssetsPipeline::FSkeletonJoint> AssetJointChild = AssetJoint->Children[ChildIndex];
+			RecursivelyFillJointRemoved(AssetJointChild, NewJoint, ConflictDataJoints);
+		}
+	}
+
+	void RecursivelyFillJointAdded(TSharedPtr<UInterchangeGenericAssetsPipeline::FSkeletonJoint> ImportJoint
+		, TSharedPtr<UInterchangeGenericAssetsPipeline::FSkeletonJoint> ParentJoint
+		, TMap<FString, TSharedPtr<UInterchangeGenericAssetsPipeline::FSkeletonJoint>>& AssetJoints
+		, TMap<FString, TSharedPtr<UInterchangeGenericAssetsPipeline::FSkeletonJoint>>& ConflictDataJoints)
+	{
+		if (!ImportJoint.IsValid())
+		{
+			return;
+		}
+		auto SetAdded = [](TSharedPtr<UInterchangeGenericAssetsPipeline::FSkeletonJoint> Joint)
+			{
+				Joint->bMatch = false;
+				Joint->bAdded = true;
+				Joint->bRemoved = false;
+			};
+		auto SetConflict = [&ParentJoint](TSharedPtr<UInterchangeGenericAssetsPipeline::FSkeletonJoint> Joint )
+			{
+				Joint->bMatch = false;
+				Joint->bAdded = true;
+				Joint->bRemoved = false;
+				Joint->bConflict = true;
+				//Add child conflict to all parent chain
+				TSharedPtr<UInterchangeGenericAssetsPipeline::FSkeletonJoint> ParentJointIter = ParentJoint;
+				while (ParentJointIter.IsValid() && !ParentJointIter->bChildConflict)
+				{
+					ParentJointIter->bChildConflict = true;
+					ParentJointIter = ParentJointIter->Parent;
+				}
+			};
+		TSharedPtr<UInterchangeGenericAssetsPipeline::FSkeletonJoint> NewJoint = MakeShared<UInterchangeGenericAssetsPipeline::FSkeletonJoint>();
+		NewJoint->JointName = ImportJoint->JointName;
+		if (ParentJoint.IsValid())
+		{
+			NewJoint->Parent = ParentJoint;
+			ParentJoint->Children.Add(NewJoint);
+		}
+		else
+		{
+			//Add only the root nodes
+			ConflictDataJoints.Add(NewJoint->JointName, NewJoint);
+		}
+		const FString ImportJointName = NewJoint->JointName;
+		if (AssetJoints.Contains(ImportJointName))
+		{
+			//We have a conflict
+			SetConflict(NewJoint);
+		}
+		else
+		{
+			SetAdded(NewJoint);
+		}
+		const int32 ChildCount = ImportJoint->Children.Num();
+		for (int32 ChildIndex = 0; ChildIndex < ChildCount; ++ChildIndex)
+		{
+			TSharedPtr<UInterchangeGenericAssetsPipeline::FSkeletonJoint> ImportJointChild = ImportJoint->Children[ChildIndex];
+			RecursivelyFillJointAdded(ImportJointChild, NewJoint, AssetJoints, ConflictDataJoints);
+		}
+	}
+
+	void RecursivelyFillJointMatch(TSharedPtr<UInterchangeGenericAssetsPipeline::FSkeletonJoint> AssetJoint
+		, TSharedPtr<UInterchangeGenericAssetsPipeline::FSkeletonJoint> ImportJoint
+		, TSharedPtr<UInterchangeGenericAssetsPipeline::FSkeletonJoint> ParentJoint
+		, TMap<FString, TSharedPtr<UInterchangeGenericAssetsPipeline::FSkeletonJoint>>& AssetJoints
+		, TMap<FString, TSharedPtr<UInterchangeGenericAssetsPipeline::FSkeletonJoint>>& ConflictDataJoints)
+	{
+		if (!ImportJoint.IsValid() || !AssetJoint.IsValid())
+		{
+			return;
+		}
+		auto SetMatch = [](TSharedPtr<UInterchangeGenericAssetsPipeline::FSkeletonJoint> JointMatch)
+			{
+				JointMatch->bMatch = true;
+				JointMatch->bAdded = false;
+				JointMatch->bRemoved = false;
+				JointMatch->bChildConflict = false;
+			};
+		if (ImportJoint->JointName.Equals(AssetJoint->JointName, ESearchCase::IgnoreCase))
+		{
+			TSharedPtr<UInterchangeGenericAssetsPipeline::FSkeletonJoint> NewJoint = MakeShared<UInterchangeGenericAssetsPipeline::FSkeletonJoint>();
+			NewJoint->JointName = ImportJoint->JointName;
+			if (ParentJoint.IsValid())
+			{
+				NewJoint->Parent = ParentJoint;
+				ParentJoint->Children.Add(NewJoint);
+			}
+			else
+			{
+				//Add only the root nodes
+				ConflictDataJoints.Add(NewJoint->JointName, NewJoint);
+			}
+			SetMatch(NewJoint);
+			
+
+			const int32 ChildCount = FMath::Max(ImportJoint->Children.Num(), AssetJoint->Children.Num());
+			for (int32 ChildIndex = 0; ChildIndex < ChildCount; ++ChildIndex)
+			{
+				TSharedPtr<UInterchangeGenericAssetsPipeline::FSkeletonJoint> ImportJointChild = ImportJoint->Children.IsValidIndex(ChildIndex) ? ImportJoint->Children[ChildIndex] : nullptr;
+				TSharedPtr<UInterchangeGenericAssetsPipeline::FSkeletonJoint> AssetJointChild = AssetJoint->Children.IsValidIndex(ChildIndex) ? AssetJoint->Children[ChildIndex] : nullptr;
+				if (ImportJointChild.IsValid())
+				{
+					if (AssetJointChild.IsValid())
+					{
+						//We have both children mean test match
+						RecursivelyFillJointMatch(AssetJointChild, ImportJointChild, NewJoint, AssetJoints, ConflictDataJoints);
+					}
+					else
+					{
+						//Import only joint mean added
+						RecursivelyFillJointAdded(ImportJoint, NewJoint, AssetJoints, ConflictDataJoints);
+					}
+				}
+				else
+				{
+					//Asset only joint mean removed
+					RecursivelyFillJointRemoved(AssetJoint, NewJoint, ConflictDataJoints);
+				}
+			}
+		}
+		else
+		{
+			//when match fail we have two separate branch to display
+			RecursivelyFillJointAdded(ImportJoint, ParentJoint, AssetJoints, ConflictDataJoints);
+			RecursivelyFillJointRemoved(AssetJoint, ParentJoint, ConflictDataJoints);
+		}
+	}
+}
+
+void UInterchangeGenericAssetsPipeline::CreateSkeletonConflict(USkeleton* SpecifiedSkeleton, USkeletalMesh* SkeletalMesh, UInterchangeBaseNodeContainer* TransientBaseNodeContainer)
+{
+	SkeletonConflictData.Reset();
+	if (!SkeletalMesh && !SpecifiedSkeleton)
+	{
+		return;
+	}
+	TArray<FString> SkeletalMeshes;
+	TransientBaseNodeContainer->GetNodes(UInterchangeSkeletalMeshFactoryNode::StaticClass(), SkeletalMeshes);
+	if (SkeletalMeshes.Num() != 1)
+	{
+		return;
+	}
+
+	//Get the skeleton from the container
+	UInterchangeSkeletalMeshFactoryNode* MeshFactoryNode = Cast<UInterchangeSkeletalMeshFactoryNode>(TransientBaseNodeContainer->GetFactoryNode(SkeletalMeshes[0]));
+	int32 LodCount = MeshFactoryNode->GetLodDataCount();
+	TArray<FString> LodDataUniqueIds;
+	MeshFactoryNode->GetLodDataUniqueIds(LodDataUniqueIds);
+	ensure(LodDataUniqueIds.Num() == LodCount);
+	if (LodCount <= 0)
+	{
+		return;
+	}
+
+	TMap<FString, TSharedPtr<FSkeletonJoint>> AssetJoints;
+	TMap<FString, TSharedPtr<FSkeletonJoint>> ImportedJoints;
+
+	//Only test the LOD 0
+	constexpr int32 LodIndex = 0;
+	{
+		FString LodUniqueId = LodDataUniqueIds[LodIndex];
+		const UInterchangeSkeletalMeshLodDataNode* LodDataNode = Cast<UInterchangeSkeletalMeshLodDataNode>(TransientBaseNodeContainer->GetNode(LodUniqueId));
+		if (!LodDataNode)
+		{
+			return;
+		}
+		
+		FString SkeletonNodeUid;
+		if (!LodDataNode->GetCustomSkeletonUid(SkeletonNodeUid))
+		{
+			return;
+		}
+		const UInterchangeSkeletonFactoryNode* SkeletonNode = Cast<UInterchangeSkeletonFactoryNode>(TransientBaseNodeContainer->GetNode(SkeletonNodeUid));
+		if (!SkeletonNode)
+		{
+			return;
+		}
+
+		FString RootJointNodeId;
+		if (!SkeletonNode->GetCustomRootJointUid(RootJointNodeId))
+		{
+			return;
+		}
+
+		const UInterchangeSceneNode* RootJointNode = Cast<UInterchangeSceneNode>(TransientBaseNodeContainer->GetNode(RootJointNodeId));
+		if (!RootJointNode)
+		{
+			return;
+		}
+		const bool bConvertStaticToSkeletalActive = CommonSkeletalMeshesAndAnimationsProperties->bConvertStaticsWithMorphTargetsToSkeletals || CommonMeshesProperties->ForceAllMeshAsType == EInterchangeForceMeshType::IFMT_SkeletalMesh;
+		//If we have a compatible skeleton we do not need to create a conflict
+		if (UE::Interchange::Private::FSkeletonHelper::IsCompatibleSkeleton(SpecifiedSkeleton ? SpecifiedSkeleton : SkeletalMesh->GetSkeleton(), RootJointNodeId, TransientBaseNodeContainer, bConvertStaticToSkeletalActive))
+		{
+			return;
+		}
+		InternalRecursiveFillJointsFromNodeContainer(nullptr, ImportedJoints, RootJointNodeId, TransientBaseNodeContainer, bConvertStaticToSkeletalActive);
+	}
+
+	//Get the current asset reference skeleton
+	const FReferenceSkeleton& ReferenceSkeleton = SpecifiedSkeleton ? SpecifiedSkeleton->GetReferenceSkeleton() : SkeletalMesh->GetSkeleton()->GetReferenceSkeleton();
+	const TArray<FMeshBoneInfo>& BoneInfos = ReferenceSkeleton.GetRawRefBoneInfo();
+
+	for (int32 BoneIndex = 0; BoneIndex < BoneInfos.Num(); ++BoneIndex)
+	{
+		const FMeshBoneInfo& BoneInfo = BoneInfos[BoneIndex];
+		if (BoneInfo.ParentIndex == INDEX_NONE)
+		{
+			InternalRecursiveFillJointsFromReferenceSkeleton(nullptr, AssetJoints, BoneIndex, ReferenceSkeleton);
+		}
+	}
+
+	TSharedPtr<FSkeletonJoint> ImportedRootJoint = nullptr;
+	for (TPair<FString, TSharedPtr<FSkeletonJoint>>& NameAndImportedJoint : ImportedJoints)
+	{
+		if (!NameAndImportedJoint.Value->Parent.IsValid())
+		{
+			ImportedRootJoint = NameAndImportedJoint.Value;
+			//We use the first root we found, we do not support multiple root
+			break;
+		}
+	}
+	
+	TSharedPtr<FSkeletonJoint> AssetRootJoint = nullptr;
+	for (TPair<FString, TSharedPtr<FSkeletonJoint>>& NameAndAssetJoint : AssetJoints)
+	{
+		if (!NameAndAssetJoint.Value->Parent.IsValid())
+		{
+			AssetRootJoint = NameAndAssetJoint.Value;
+			//We use the first root we found, we do not support multiple root
+			break;
+		}
+	}
+
+	if (!ImportedRootJoint.IsValid() || !AssetRootJoint.IsValid())
+	{
+		SkeletonConflictData.Reset();
+		return;
+	}
+	UE::Interchange::Private::RecursivelyFillJointMatch(AssetRootJoint, ImportedRootJoint, nullptr, AssetJoints, SkeletonConflictData.Joints);
+
+	//Compare and cache the results
+	FInterchangeConflictInfo& SkeletonConflict = ConflictInfos.AddDefaulted_GetRef();
+	SkeletonConflict.DisplayName = TEXT("Skeleton");
+	SkeletonConflict.Description = TEXT("Imported skeleton is incompatible with the asset skeleton");
+	SkeletonConflict.Pipeline = this;
+	SkeletonConflict.UniqueId = FGuid::NewGuid();
+
+	SkeletonConflictData.ConflictUniqueId = SkeletonConflict.UniqueId;
+	SkeletonConflictData.ReimportObject = SkeletalMesh;
+}
 #endif //WITH_EDITOR
 
+TArray<FInterchangeConflictInfo> UInterchangeGenericAssetsPipeline::GetConflictInfos(UObject* ReimportObject, UInterchangeBaseNodeContainer* InBaseNodeContainer, UInterchangeSourceData* SourceData)
+{
+	//Dont touch any conflict data outside of the game thread
+	if (!ensure(IsInGameThread()))
+	{
+		return {};
+	}
+
+	ConflictInfos.Reset();
+
+#if WITH_EDITOR
+	USkeleton* SpecifiedSkeleton = CommonSkeletalMeshesAndAnimationsProperties->Skeleton.IsValid() ? CommonSkeletalMeshesAndAnimationsProperties->Skeleton.Get() : nullptr;
+	UStaticMesh* StaticMesh = Cast<UStaticMesh>(ReimportObject);
+	USkeletalMesh* SkeletalMesh = Cast<USkeletalMesh>(ReimportObject);
+	if (!StaticMesh && !SkeletalMesh && !SpecifiedSkeleton)
+	{
+		return ConflictInfos;
+	}
+
+	UInterchangeBaseNodeContainer* TransientBaseNodeContainer = DuplicateObject<UInterchangeBaseNodeContainer>(InBaseNodeContainer, GetTransientPackage());
+
+	TArray<UInterchangeSourceData*> SourceDatas;
+	SourceDatas.Add(SourceData);
+
+	UInterchangeResultsContainer* OldResults = Results;
+	Results = NewObject<UInterchangeResultsContainer>(GetTransientPackage());
+	ExecutePipeline(TransientBaseNodeContainer, SourceDatas);
+	Results = OldResults;
+
+	//Create the materials conflict
+	if (StaticMesh || SkeletalMesh)
+	{
+		CreateMaterialConflict(StaticMesh, SkeletalMesh, TransientBaseNodeContainer);
+	}
+
+	if (SpecifiedSkeleton || SkeletalMesh)
+	{
+		CreateSkeletonConflict(SpecifiedSkeleton, SkeletalMesh, TransientBaseNodeContainer);
+	}
+#endif
+	return ConflictInfos;
+}
+
+void UInterchangeGenericAssetsPipeline::ShowConflictDialog(const FGuid& ConflictUniqueId)
+{
+	for (const FInterchangeConflictInfo& ConflictInfo : ConflictInfos)
+	{
+		if (ConflictInfo.UniqueId == ConflictUniqueId)
+		{
+			if (ensure(ConflictInfo.Pipeline == this))
+			{
+				if (MaterialConflictData.ConflictUniqueId == ConflictUniqueId)
+				{
+					//Create the dialog widget
+					TSharedRef<SInterchangeGenericAssetMaterialConflictWidget> WidgetContent = SNew(SInterchangeGenericAssetMaterialConflictWidget)
+						.AssetMaterialNames(MaterialConflictData.AssetMaterialNames)
+						.ImportMaterialNames(MaterialConflictData.ImportMaterialNames)
+						.MatchMaterialIndexes(MaterialConflictData.MatchMaterialIndexes)
+						.ReimportObject(MaterialConflictData.ReimportObject);
+
+					UE::Interchange::PipelineHelper::ShowModalDialog(WidgetContent, MaterialConflictData.DialogTitle, FVector2D(450.0f, 200.0f));
+				}
+				else if (SkeletonConflictData.ConflictUniqueId == ConflictUniqueId)
+				{
+					TArray<TSharedPtr<UInterchangeGenericAssetsPipeline::FSkeletonJoint>> Joints;
+					Joints.Reserve(SkeletonConflictData.Joints.Num());
+					for (TPair<FString, TSharedPtr<FSkeletonJoint>> NameAndJoint : SkeletonConflictData.Joints)
+					{
+						Joints.Add(NameAndJoint.Value);
+					}
+					
+					//Grab all skeleton reference so user can know the impact of changing the skeleton
+					TArray<TSharedPtr<FString>> AssetReferencingSkeleton;
+					USkeletalMesh* ReimportSkeletalMesh = Cast<USkeletalMesh>(SkeletonConflictData.ReimportObject);
+					USkeleton* Skeleton = CommonSkeletalMeshesAndAnimationsProperties->Skeleton.IsValid() ? CommonSkeletalMeshesAndAnimationsProperties->Skeleton.Get() : ReimportSkeletalMesh ? ReimportSkeletalMesh->GetSkeleton() : nullptr;
+					if (Skeleton)
+					{
+						//Get in memory reference
+						TArray<FName> MemoryDependencies;
+						{
+							TArray<UObject*> ReferencedObjects;
+							ReferencedObjects.Add(Skeleton);
+							TSet<UObject*> ReferencerObjects;
+							// Use the fast reference collector to recursively find referencers until no more new ones are found.
+							int32 LastObjectCount = 0;
+							TArray<UObject*> FoundReferencerObjects = FReferencerFinder::GetAllReferencers(ReferencedObjects, nullptr, EReferencerFinderFlags::None);
+							do
+							{
+								LastObjectCount = ReferencerObjects.Num();
+								//Do not list transient package objects
+								for (UObject* FoundReferencerObject : FoundReferencerObjects)
+								{
+									if (FoundReferencerObject->GetOutermost() != GetTransientPackage())
+									{
+										ReferencerObjects.Add(FoundReferencerObject);
+									}
+								}
+								FoundReferencerObjects = FReferencerFinder::GetAllReferencers(FoundReferencerObjects, nullptr, EReferencerFinderFlags::SkipInnerReferences);
+
+							} while (LastObjectCount != ReferencerObjects.Num());
+
+							// Add the full path name to the memory dependencies array.
+							MemoryDependencies.Reserve(ReferencerObjects.Num());
+							for (UObject* ReferencerObject : ReferencerObjects)
+							{
+								MemoryDependencies.Add(*ReferencerObject->GetFullName());
+							}
+						}
+
+						//Get not in memory assets dependencies
+						FAssetRegistryModule& AssetRegistryModule = FModuleManager::LoadModuleChecked<FAssetRegistryModule>(TEXT("AssetRegistry"));
+						const FName SelectedPackageName = Skeleton->GetOutermost()->GetFName();
+						//Get the Hard dependencies
+						TArray<FName> HardDependencies;
+						AssetRegistryModule.Get().GetReferencers(SelectedPackageName, HardDependencies, UE::AssetRegistry::EDependencyCategory::Package, UE::AssetRegistry::EDependencyQuery::Hard);
+						//Get the Soft dependencies
+						TArray<FName> SoftDependencies;
+						AssetRegistryModule.Get().GetReferencers(SelectedPackageName, SoftDependencies, UE::AssetRegistry::EDependencyCategory::Package, UE::AssetRegistry::EDependencyQuery::Soft);
+
+						//Compose the All dependencies array with unique entries
+						TArray<FName> AllDependencies = MemoryDependencies;
+
+						auto AddAssetName = [&AllDependencies, &AssetRegistryModule](TArray<FName>& AssetDependencies)
+							{
+								for (const FName& AssetDependencyName : AssetDependencies)
+								{
+									const FString PackageString = AssetDependencyName.ToString();
+									const FSoftObjectPath FullAssetPath(*PackageString, *FPackageName::GetLongPackageAssetName(PackageString), {});
+									FAssetData AssetData = AssetRegistryModule.Get().GetAssetByObjectPath(FullAssetPath);
+									if (AssetData.GetClass() != nullptr)
+									{
+										AllDependencies.AddUnique(*FString(AssetData.GetClass()->GetName() + TEXT(" ") + AssetData.GetObjectPathString()));
+									}
+								}
+							};
+
+						AddAssetName(HardDependencies);
+						AddAssetName(SoftDependencies);
+
+						//Construct shared string for each dependencies 
+						for (const FName& AssetDependencyName : AllDependencies)
+						{
+							TSharedPtr<FString> AssetReferencing = MakeShareable(new FString(AssetDependencyName.ToString()));
+							AssetReferencingSkeleton.Add(AssetReferencing);
+						}
+					}
+
+					//Create the dialog widget
+					TSharedRef<SInterchangeGenericAssetSkeletonConflictWidget> WidgetContent = SNew(SInterchangeGenericAssetSkeletonConflictWidget)
+						.Joints(Joints)
+						.ReimportObject(SkeletonConflictData.ReimportObject)
+						.AssetReferencingSkeleton(AssetReferencingSkeleton);
+
+					UE::Interchange::PipelineHelper::ShowModalDialog(WidgetContent, SkeletonConflictData.DialogTitle, FVector2D(650.0f, 600.0f));
+				}
+			}
+			break;
+		}
+	}
+}
+
 void UInterchangeGenericAssetsPipeline::ExecutePipeline(UInterchangeBaseNodeContainer* InBaseNodeContainer, const TArray<UInterchangeSourceData*>& InSourceDatas)
 {
 	check(CommonSkeletalMeshesAndAnimationsProperties);
@@ -377,8 +1043,6 @@ void UInterchangeGenericAssetsPipeline::ExecutePipeline(UInterchangeBaseNodeCont
 	//////////////////////////////////////////////////////////////////////////
 
 
-	BaseNodeContainer = InBaseNodeContainer;
-
 	//Setup the Global import offset
 	{
 		//Make sure the scale value is greater than zero, warn the user in this case and set the scale to the default value 1.0f
@@ -400,17 +1064,11 @@ void UInterchangeGenericAssetsPipeline::ExecutePipeline(UInterchangeBaseNodeCont
 		ImportOffsetTransform.SetRotation(FQuat(ImportOffsetRotation));
 		ImportOffsetTransform.SetScale3D(FVector(ImportOffsetUniformScale));
 
-		UInterchangeCommonPipelineDataFactoryNode* CommonPipelineDataFactoryNode = UInterchangeCommonPipelineDataFactoryNode::FindOrCreateUniqueInstance(BaseNodeContainer);
-		CommonPipelineDataFactoryNode->SetCustomGlobalOffsetTransform(BaseNodeContainer, ImportOffsetTransform);
+		UInterchangeCommonPipelineDataFactoryNode* CommonPipelineDataFactoryNode = UInterchangeCommonPipelineDataFactoryNode::FindOrCreateUniqueInstance(InBaseNodeContainer);
+		CommonPipelineDataFactoryNode->SetCustomGlobalOffsetTransform(InBaseNodeContainer, ImportOffsetTransform);
 
 		// In case all mesh types are forced to Static/Skeletal we bake the scene instance hierarchy transforms
-		CommonPipelineDataFactoryNode->SetBakeMeshes(BaseNodeContainer, CommonMeshesProperties->ForceAllMeshAsType != EInterchangeForceMeshType::IFMT_None || CommonMeshesProperties->bBakeMeshes);
-	}
-
-	SourceDatas.Empty(InSourceDatas.Num());
-	for (const UInterchangeSourceData* SourceData : InSourceDatas)
-	{
-		SourceDatas.Add(SourceData);
+		CommonPipelineDataFactoryNode->SetBakeMeshes(InBaseNodeContainer, CommonMeshesProperties->ForceAllMeshAsType != EInterchangeForceMeshType::IFMT_None || CommonMeshesProperties->bBakeMeshes);
 	}
 
 	if (MaterialPipeline)
@@ -426,9 +1084,9 @@ void UInterchangeGenericAssetsPipeline::ExecutePipeline(UInterchangeBaseNodeCont
 		AnimationPipeline->ScriptedExecutePipeline(InBaseNodeContainer, InSourceDatas);
 	}
 
-	ImplementUseSourceNameForAssetOption();
+	ImplementUseSourceNameForAssetOption(InBaseNodeContainer, InSourceDatas);
 	//Make sure all factory nodes have the specified strategy
-	BaseNodeContainer->IterateNodesOfType<UInterchangeFactoryBaseNode>([ReimportStrategyClosure = ReimportStrategy](const FString& NodeUid, UInterchangeFactoryBaseNode* FactoryNode)
+	InBaseNodeContainer->IterateNodesOfType<UInterchangeFactoryBaseNode>([ReimportStrategyClosure = ReimportStrategy](const FString& NodeUid, UInterchangeFactoryBaseNode* FactoryNode)
 		{
 			FactoryNode->SetReimportStrategyFlags(ReimportStrategyClosure);
 		});
@@ -717,19 +1375,19 @@ void UInterchangeGenericAssetsPipeline::SetReimportSourceIndex(UClass* ReimportO
 	}
 }
 
-void UInterchangeGenericAssetsPipeline::ImplementUseSourceNameForAssetOption()
+void UInterchangeGenericAssetsPipeline::ImplementUseSourceNameForAssetOption(UInterchangeBaseNodeContainer* InBaseNodeContainer, const TArray<UInterchangeSourceData*>& InSourceDatas)
 {
 	const UClass* SkeletalMeshFactoryNodeClass = UInterchangeSkeletalMeshFactoryNode::StaticClass();
 	TArray<FString> SkeletalMeshNodeUids;
-	BaseNodeContainer->GetNodes(SkeletalMeshFactoryNodeClass, SkeletalMeshNodeUids);
+	InBaseNodeContainer->GetNodes(SkeletalMeshFactoryNodeClass, SkeletalMeshNodeUids);
 
 	const UClass* StaticMeshFactoryNodeClass = UInterchangeStaticMeshFactoryNode::StaticClass();
 	TArray<FString> StaticMeshNodeUids;
-	BaseNodeContainer->GetNodes(StaticMeshFactoryNodeClass, StaticMeshNodeUids);
+	InBaseNodeContainer->GetNodes(StaticMeshFactoryNodeClass, StaticMeshNodeUids);
 
 	const UClass* AnimSequenceFactoryNodeClass = UInterchangeAnimSequenceFactoryNode::StaticClass();
 	TArray<FString> AnimSequenceNodeUids;
-	BaseNodeContainer->GetNodes(AnimSequenceFactoryNodeClass, AnimSequenceNodeUids);
+	InBaseNodeContainer->GetNodes(AnimSequenceFactoryNodeClass, AnimSequenceNodeUids);
 
 	//If we import only one mesh, we want to rename the mesh using the file name.
 	const int32 MeshesImportedNodeCount = SkeletalMeshNodeUids.Num() + StaticMeshNodeUids.Num();
@@ -748,18 +1406,571 @@ void UInterchangeGenericAssetsPipeline::ImplementUseSourceNameForAssetOption()
 		//StaticMesh
 		if (MeshesImportedNodeCount == 1 && StaticMeshNodeUids.Num() > 0)
 		{
-			UInterchangeStaticMeshFactoryNode* StaticMeshNode = Cast<UInterchangeStaticMeshFactoryNode>(BaseNodeContainer->GetFactoryNode(StaticMeshNodeUids[0]));
-			const FString DisplayLabelName = OverrideAssetName.IsEmpty() ? FPaths::GetBaseFilename(SourceDatas[0]->GetFilename()) : OverrideAssetName;
+			UInterchangeStaticMeshFactoryNode* StaticMeshNode = Cast<UInterchangeStaticMeshFactoryNode>(InBaseNodeContainer->GetFactoryNode(StaticMeshNodeUids[0]));
+			const FString DisplayLabelName = OverrideAssetName.IsEmpty() ? FPaths::GetBaseFilename(InSourceDatas[0]->GetFilename()) : OverrideAssetName;
 			StaticMeshNode->SetDisplayLabel(DisplayLabelName);
 		}
 
 		//Animation, simply look if we import only 1 animation before applying the option to animation
 		if (AnimSequenceNodeUids.Num() == 1)
 		{
-			UInterchangeAnimSequenceFactoryNode* AnimSequenceNode = Cast<UInterchangeAnimSequenceFactoryNode>(BaseNodeContainer->GetFactoryNode(AnimSequenceNodeUids[0]));
-			const FString DisplayLabelName = (OverrideAssetName.IsEmpty() ? FPaths::GetBaseFilename(SourceDatas[0]->GetFilename()) : OverrideAssetName) + TEXT("_Anim");
+			UInterchangeAnimSequenceFactoryNode* AnimSequenceNode = Cast<UInterchangeAnimSequenceFactoryNode>(InBaseNodeContainer->GetFactoryNode(AnimSequenceNodeUids[0]));
+			const FString DisplayLabelName = (OverrideAssetName.IsEmpty() ? FPaths::GetBaseFilename(InSourceDatas[0]->GetFilename()) : OverrideAssetName) + TEXT("_Anim");
 			AnimSequenceNode->SetDisplayLabel(DisplayLabelName);
 		}
 	}
 }
 
+class SInterchangeGenericAssetMaterialConflictListRow : public SMultiColumnTableRow<TSharedPtr<SInterchangeGenericAssetMaterialConflictWidget::FListItem>>
+{
+public:
+	SLATE_BEGIN_ARGS(SInterchangeGenericAssetMaterialConflictListRow) {}
+		SLATE_ARGUMENT(TSharedPtr<SInterchangeGenericAssetMaterialConflictWidget::FListItem>, Item)
+	SLATE_END_ARGS()
+
+	TSharedPtr<SInterchangeGenericAssetMaterialConflictWidget::FListItem> Item;
+
+	void Construct(const FArguments& InArgs, const TSharedRef<STableViewBase>& OwnerTable)
+	{
+		Item = InArgs._Item;
+
+		SMultiColumnTableRow<TSharedPtr<SInterchangeGenericAssetMaterialConflictWidget::FListItem>>::Construct(
+			FSuperRowType::FArguments()
+			.Style(FAppStyle::Get(), "DataTableEditor.CellListViewRow"),
+			OwnerTable
+		);
+	}
+
+	virtual TSharedRef<SWidget> GenerateWidgetForColumn(const FName& ColumnName) override
+	{
+		FSlateColor UnmatchedColor(FLinearColor(0.7f, 0.3f, 0.0f));
+		if (ColumnName == SInterchangeGenericAssetMaterialConflictWidget::NAME_Import)
+		{
+			bool bConflict = Item->bMatched == INDEX_NONE;
+			FSlateColor SlateColor = bConflict ? UnmatchedColor : FSlateColor::UseForeground();
+			FText Tooltip;
+			if (bConflict)
+			{
+				Tooltip = NSLOCTEXT("InterchangeGenericAssetPipeline", "SInterchangeGenericAssetMaterialConflictListRow_Conflict", "Import material is unmatched");
+			}
+			else
+			{
+				Tooltip = FText::FromString(
+					FText(NSLOCTEXT("InterchangeGenericAssetPipeline", "SInterchangeGenericAssetMaterialConflictListRow_Conflict", "Import material is matched with original asset index: ")).ToString()
+					+ FString::FromInt(Item->bMatched));
+			}
+			return SNew(SBox)
+			.Padding(2.0f)
+			.HAlign(HAlign_Left)
+			.VAlign(VAlign_Center)
+			[
+				SNew(STextBlock)
+				.Text(FText::FromString(Item->ImportName))
+				.ToolTipText(Tooltip)
+				.ColorAndOpacity(SlateColor)
+			];
+		}
+		else
+		{
+			return SNew(SBox)
+				.Padding(2.0f)
+				.HAlign(HAlign_Left)
+				.VAlign(VAlign_Center)
+				[
+					SNew(STextBlock)
+						.Text(FText::FromString(Item->AssetName))
+				];
+		}
+	}
+};
+
+void SInterchangeGenericAssetMaterialConflictWidget::Construct(const FArguments& InArgs)
+{
+	AssetMaterialNames = InArgs._AssetMaterialNames;
+	ImportMaterialNames = InArgs._ImportMaterialNames;
+	MatchMaterialIndexes = InArgs._MatchMaterialIndexes;
+	ReimportObject = InArgs._ReimportObject;
+
+	if (!ensure(ReimportObject))
+	{
+		return;
+	}
+
+	int32 RowNumber = FMath::Max(AssetMaterialNames.Num(), ImportMaterialNames.Num());
+	for (int32 RowIndex = 0; RowIndex < RowNumber; ++RowIndex)
+	{
+		TSharedRef<FListItem> Item = MakeShared<FListItem>();
+		Item->ImportName = ImportMaterialNames.IsValidIndex(RowIndex) ? ImportMaterialNames[RowIndex] : FString();
+		Item->bMatched = INDEX_NONE;
+		if (MatchMaterialIndexes.IsValidIndex(RowIndex))
+		{
+			Item->bMatched = MatchMaterialIndexes[RowIndex];
+			if (AssetMaterialNames.IsValidIndex(Item->bMatched))
+			{
+				Item->AssetMatchedName = AssetMaterialNames[Item->bMatched];
+			}
+		}
+		Item->AssetName = AssetMaterialNames.IsValidIndex(RowIndex) ? AssetMaterialNames[RowIndex] : FString();
+		RowItems.Add(Item);
+	}
+
+	this->ChildSlot
+	[
+		SNew(SBox)
+		[
+			SNew(SVerticalBox)
+			+ SVerticalBox::Slot()
+			.HAlign(HAlign_Fill)
+			.VAlign(VAlign_Fill)
+			.Padding(0.0f, 0.0f, 0.0f, 8.0f)
+			.FillHeight(1.0f)
+			[
+				SNew(SBorder)
+				[
+					SAssignNew(MaterialList, SListView<TSharedPtr<FListItem>>)
+					.ListItemsSource(&RowItems)
+					.OnGenerateRow(this, &SInterchangeGenericAssetMaterialConflictWidget::OnGenerateRow)
+					.HeaderRow
+					(
+						SNew(SHeaderRow)
+
+						+ SHeaderRow::Column(NAME_Import)
+						.FillWidth(0.5f)
+						.DefaultLabel(NSLOCTEXT("GenericAssetPipeline", "SInterchangeGenericAssetMaterialConflictWidget_ImportName", "Import"))
+
+						+ SHeaderRow::Column(NAME_Asset)
+						.FillWidth(0.5f)
+						.DefaultLabel(NSLOCTEXT("GenericAssetPipeline", "SInterchangeGenericAssetMaterialConflictWidget_AssetName", "Asset"))
+					)
+				]
+			]
+			+ SVerticalBox::Slot()
+			.AutoHeight()
+			.HAlign(HAlign_Right)
+			.VAlign(VAlign_Bottom)
+			.Padding(2)
+			[
+				SNew(SButton)
+					.HAlign(HAlign_Center)
+					.Text(NSLOCTEXT("SInterchangeGenericAssetMaterialConflictWidget", "SInterchangeGenericAssetMaterialConflictWidget_Done", "Done"))
+					.OnClicked(this, &SInterchangeGenericAssetMaterialConflictWidget::OnDone)
+			]
+		]
+	];
+}
+
+TSharedRef<ITableRow> SInterchangeGenericAssetMaterialConflictWidget::OnGenerateRow(TSharedPtr<SInterchangeGenericAssetMaterialConflictWidget::FListItem> Item, const TSharedRef<STableViewBase>& OwnerTable)
+{
+	return SNew(SInterchangeGenericAssetMaterialConflictListRow, OwnerTable)
+		.Item(Item);
+}
+
+FReply SInterchangeGenericAssetMaterialConflictWidget::OnDone()
+{
+	if (WidgetWindow.IsValid())
+	{
+		WidgetWindow->RequestDestroyWindow();
+	}
+
+	return FReply::Handled();
+}
+
+FReply SInterchangeGenericAssetMaterialConflictWidget::OnKeyDown(const FGeometry& MyGeometry, const FKeyEvent& InKeyEvent)
+{
+	if (InKeyEvent.GetKey() == EKeys::Escape)
+	{
+		return OnDone();
+	}
+	return FReply::Unhandled();
+}
+
+void SInterchangeGenericAssetSkeletonConflictWidget::Construct(const FArguments& InArgs)
+{
+	bShowSectionFlag[EInterchangeSkeletonCompareSection::Skeleton] = true;
+	bShowSectionFlag[EInterchangeSkeletonCompareSection::References] = true;
+
+	AssetReferencingSkeleton = InArgs._AssetReferencingSkeleton;
+	Joints = InArgs._Joints;
+	ReimportObject = InArgs._ReimportObject;
+
+	// Skeleton comparison
+	TSharedPtr<SWidget> SkeletonCompareSection = ConstructSkeletonComparison();
+	TSharedPtr<SWidget> SkeletonReferencesSection = ConstructSkeletonReference();
+
+	this->ChildSlot
+	[
+		SNew(SBox)
+		[
+			SNew(SVerticalBox)
+			+SVerticalBox::Slot()
+			.FillHeight(1.0f)
+			[
+				SNew(SBorder)
+				.BorderImage(FAppStyle::GetBrush("ToolPanel.DarkGroupBorder"))
+				[
+					SNew(SVerticalBox)
+					+ SVerticalBox::Slot()
+					.FillHeight(1.0f)
+					.Padding(2.0f)
+					[
+						SNew(SSplitter)
+						.Orientation(Orient_Vertical)
+						.ResizeMode(ESplitterResizeMode::Fill)
+						+ SSplitter::Slot()
+						.Value(0.8f)
+						[
+							// Skeleton Compare section
+							SkeletonCompareSection.ToSharedRef()
+						]
+						+ SSplitter::Slot()
+						.Value(0.2f)
+						[
+							// Skeleton Compare section
+							SkeletonReferencesSection.ToSharedRef()
+						]
+					]
+				]
+			]
+			+ SVerticalBox::Slot()
+			.AutoHeight()
+			.HAlign(HAlign_Right)
+			.Padding(2.0f)
+			[
+				SNew(SHorizontalBox)
+				+ SHorizontalBox::Slot()
+				.AutoWidth()
+				.Padding(2.0f, 0.0f)
+				[
+					SNew(SButton)
+					.HAlign(HAlign_Center)
+					.Text(NSLOCTEXT("SInterchangeGenericAssetSkeletonConflictWidget", "ConstructDoneButton", "Done"))
+					.OnClicked(this, &SInterchangeGenericAssetSkeletonConflictWidget::OnDone)
+				]
+			]
+		]
+	];
+}
+
+FReply SInterchangeGenericAssetSkeletonConflictWidget::SetSectionVisible(EInterchangeSkeletonCompareSection SectionIndex)
+{
+	bShowSectionFlag[SectionIndex] = !bShowSectionFlag[SectionIndex];
+	return FReply::Handled();
+}
+
+EVisibility SInterchangeGenericAssetSkeletonConflictWidget::IsSectionVisible(EInterchangeSkeletonCompareSection SectionIndex)
+{
+	return bShowSectionFlag[SectionIndex] ? EVisibility::All : EVisibility::Collapsed;
+}
+
+const FSlateBrush* SInterchangeGenericAssetSkeletonConflictWidget::GetCollapsableArrow(EInterchangeSkeletonCompareSection SectionIndex) const
+{
+	return bShowSectionFlag[SectionIndex] ? FAppStyle::GetBrush("Symbols.DownArrow") : FAppStyle::GetBrush("Symbols.RightArrow");
+}
+
+namespace UE::Interchange::Private
+{
+	void RecursivelyExpandTreeItem(TSharedPtr<STreeView<TSharedPtr<UInterchangeGenericAssetsPipeline::FSkeletonJoint>>> CompareTree
+		, TSharedPtr<UInterchangeGenericAssetsPipeline::FSkeletonJoint> JointItem)
+	{
+		if (JointItem->bInitialAutoExpand || !JointItem->bMatch || !JointItem->bChildConflict)
+		{
+			return;
+		}
+		JointItem->bInitialAutoExpand = true;
+		CompareTree->SetItemExpansion(JointItem, true);
+		for (TSharedPtr<UInterchangeGenericAssetsPipeline::FSkeletonJoint> ChildJoint : JointItem->Children)
+		{
+			RecursivelyExpandTreeItem(CompareTree, ChildJoint);
+		}
+	}
+}
+
+TSharedPtr<SWidget> SInterchangeGenericAssetSkeletonConflictWidget::ConstructSkeletonComparison()
+{
+	FText SkeletonStatus = NSLOCTEXT("SInterchangeGenericAssetSkeletonConflictWidget", "ConstructSkeletonComparison_SkeletonStatus", "The skeleton has some conflicts");
+	
+	CompareTree = SNew(STreeView<TSharedPtr<UInterchangeGenericAssetsPipeline::FSkeletonJoint>>)
+		.ItemHeight(24)
+		.SelectionMode(ESelectionMode::None)
+		.TreeItemsSource(&Joints)
+		.OnGenerateRow(this, &SInterchangeGenericAssetSkeletonConflictWidget::OnGenerateRowCompareTreeView)
+		.OnGetChildren(this, &SInterchangeGenericAssetSkeletonConflictWidget::OnGetChildrenRowCompareTreeView);
+	
+
+	for (TSharedPtr<UInterchangeGenericAssetsPipeline::FSkeletonJoint> Joint : Joints)
+	{
+		UE::Interchange::Private::RecursivelyExpandTreeItem(CompareTree, Joint);
+	}
+	
+	return SNew(SBox)
+	[
+		SNew(SBorder)
+		.Padding(FMargin(3))
+		.BorderImage(FAppStyle::GetBrush("ToolPanel.GroupBorder"))
+		[
+			SNew(SVerticalBox)
+			+SVerticalBox::Slot()
+			.AutoHeight()
+			.Padding(2)
+			[
+				SNew(SHorizontalBox)
+				+ SHorizontalBox::Slot()
+				.AutoWidth()
+				[
+					SNew(SButton)
+					.HAlign(HAlign_Center)
+					.VAlign(VAlign_Center)
+					.IsFocusable(false)
+					.ButtonStyle(FAppStyle::Get(), "NoBorder")
+					.OnClicked(this, &SInterchangeGenericAssetSkeletonConflictWidget::SetSectionVisible, EInterchangeSkeletonCompareSection::Skeleton)
+					[
+						SNew(SImage).Image(this, &SInterchangeGenericAssetSkeletonConflictWidget::GetCollapsableArrow, EInterchangeSkeletonCompareSection::Skeleton)
+					]
+				]
+				+ SHorizontalBox::Slot()
+				.AutoWidth()
+				[
+					SNew(STextBlock)
+					.Font(FAppStyle::GetFontStyle("DetailsView.CategoryFontStyle"))
+					.Text(NSLOCTEXT("SInterchangeGenericAssetSkeletonConflictWidget", "ConstructSkeletonComparison_SkeletonCompareHeader", "Skeleton"))
+				]
+			]
+			+SVerticalBox::Slot()
+			.FillHeight(1.0f)
+			.Padding(2)
+			[
+				SNew(SBox)
+				.Visibility(TAttribute<EVisibility>::Create(TAttribute<EVisibility>::FGetter::CreateSP(this, &SInterchangeGenericAssetSkeletonConflictWidget::IsSectionVisible, EInterchangeSkeletonCompareSection::Skeleton)))
+				[
+					SNew(SBorder)
+					.Padding(FMargin(3))
+					.BorderImage(FAppStyle::GetBrush("ToolPanel.DarkGroupBorder"))
+					[
+						SNew(SVerticalBox)
+						+SVerticalBox::Slot()
+						.AutoHeight()
+						.Padding(2)
+						[
+							SNew(STextBlock)
+							.Font(FAppStyle::GetFontStyle("DetailsView.CategoryFontStyle"))
+							.Text(SkeletonStatus)
+							.ColorAndOpacity(SInterchangeGenericAssetMaterialConflictWidget::SlateColorFullConflict)
+						]
+						+SVerticalBox::Slot()
+						.AutoHeight()
+						.Padding(2)
+						[
+							SNew(SSeparator)
+							.Orientation(EOrientation::Orient_Horizontal)
+						]
+						+SVerticalBox::Slot()
+						.FillHeight(1.0f)
+						.Padding(2)
+						[
+							CompareTree.ToSharedRef()
+						]
+					]
+				]
+			]
+		]
+	];
+}
+
+TSharedPtr<SWidget> SInterchangeGenericAssetSkeletonConflictWidget::ConstructSkeletonReference()
+{
+	FString SkeletonReferenceStatistic;
+	if (AssetReferencingSkeleton.Num() > 0)
+	{
+		SkeletonReferenceStatistic += TEXT("Skeleton is references by ") + FString::FromInt(AssetReferencingSkeleton.Num()) + TEXT(" assets.");
+	}
+	
+	return SNew(SBox)
+	[
+		SNew(SBorder)
+		.Padding(FMargin(3))
+		.BorderImage(FAppStyle::GetBrush("ToolPanel.GroupBorder"))
+		[
+			SNew(SVerticalBox)
+			+SVerticalBox::Slot()
+			.AutoHeight()
+			.Padding(2)
+			[
+				SNew(SHorizontalBox)
+				+ SHorizontalBox::Slot()
+				.AutoWidth()
+				[
+					SNew(SButton)
+					.HAlign(HAlign_Center)
+					.VAlign(VAlign_Center)
+					.IsFocusable(false)
+					.ButtonStyle(FAppStyle::Get(), "NoBorder")
+					.OnClicked(this, &SInterchangeGenericAssetSkeletonConflictWidget::SetSectionVisible, EInterchangeSkeletonCompareSection::References)
+					[
+						SNew(SImage).Image(this, &SInterchangeGenericAssetSkeletonConflictWidget::GetCollapsableArrow, EInterchangeSkeletonCompareSection::References)
+					]
+				]
+				+ SHorizontalBox::Slot()
+				.AutoWidth()
+				[
+					SNew(STextBlock)
+					.Font(FAppStyle::GetFontStyle("DetailsView.CategoryFontStyle"))
+					.Text(NSLOCTEXT("SInterchangeGenericAssetSkeletonConflictWidget", "SkeletonReferencesHeader", "References"))
+				]
+			]
+			+SVerticalBox::Slot()
+			.FillHeight(1.0f)
+			.Padding(2)
+			[
+				SNew(SBox)
+				.Visibility(TAttribute<EVisibility>::Create(TAttribute<EVisibility>::FGetter::CreateSP(this, &SInterchangeGenericAssetSkeletonConflictWidget::IsSectionVisible, EInterchangeSkeletonCompareSection::References)))
+				[
+					SNew(SBorder)
+					.Padding(FMargin(3))
+					.BorderImage(FAppStyle::GetBrush("ToolPanel.DarkGroupBorder"))
+					[
+						SNew(SVerticalBox)
+						+SVerticalBox::Slot()
+						.AutoHeight()
+						.Padding(2)
+						[
+							SNew(STextBlock)
+							.Font(FAppStyle::GetFontStyle("DetailsView.CategoryFontStyle"))
+							.Text(FText::FromString(SkeletonReferenceStatistic))
+						]
+						+SVerticalBox::Slot()
+						.AutoHeight()
+						.Padding(2)
+						[
+							SNew(SSeparator)
+							.Orientation(EOrientation::Orient_Horizontal)
+						]
+						+SVerticalBox::Slot()
+						.FillHeight(1.0f)
+						.Padding(2)
+						[
+							//Show the asset referencing this skeleton
+							SNew(SListView<TSharedPtr<FString>>)
+							.ListItemsSource(&AssetReferencingSkeleton)
+							.OnGenerateRow(this, &SInterchangeGenericAssetSkeletonConflictWidget::OnGenerateRowAssetReferencingSkeleton)
+						]
+					]
+				]
+			]
+		]
+	];
+}
+
+class SInterchangeCompareSkeletonTreeViewItem : public STableRow< TSharedPtr<UInterchangeGenericAssetsPipeline::FSkeletonJoint> >
+{
+public:
+
+	SLATE_BEGIN_ARGS(SInterchangeCompareSkeletonTreeViewItem)
+		: _SkeletonCompareData(nullptr)
+		{}
+
+		/** The item content. */
+		SLATE_ARGUMENT(TSharedPtr<UInterchangeGenericAssetsPipeline::FSkeletonJoint>, SkeletonCompareData)
+	SLATE_END_ARGS()
+
+	/**
+	* Construct the widget
+	*
+	* @param InArgs   A declaration from which to construct the widget
+	*/
+	void Construct(const FArguments& InArgs, const TSharedRef<STableViewBase>& InOwnerTableView)
+	{
+		SkeletonCompareData = InArgs._SkeletonCompareData;
+
+		//This is suppose to always be valid
+		check(SkeletonCompareData.IsValid());
+
+		const FSlateBrush* JointIcon = FAppStyle::GetDefaultBrush();
+		FString Tooltip = NSLOCTEXT("SInterchangeCompareSkeletonTreeViewItem", "Construct_Joint_tooltip", "Re-import joint match skeletal mesh skeleton joint.").ToString();
+		if(SkeletonCompareData->bAdded)
+		{
+			JointIcon = FAppStyle::GetBrush("FBXIcon.ReimportCompareAdd");
+			Tooltip = NSLOCTEXT("SInterchangeCompareSkeletonTreeViewItem", "Construct_AddJoint_tooltip", "Re-import add this joint.").ToString();
+		}
+		else if(SkeletonCompareData->bRemoved)
+		{
+			JointIcon = FAppStyle::GetBrush("FBXIcon.ReimportCompareRemoved");
+			Tooltip = NSLOCTEXT("SInterchangeCompareSkeletonTreeViewItem", "Construct_RemoveJoint_tooltip", "Re-import remove this joint.").ToString();
+		}
+		const bool bConflictRoot = SkeletonCompareData->bConflict && (SkeletonCompareData->Parent.IsValid() && !SkeletonCompareData->Parent->bConflict);
+		FSlateColor ForegroundTextColor = SkeletonCompareData->bMatch ? FSlateColor::UseForeground() : (!bConflictRoot ? SInterchangeGenericAssetMaterialConflictWidget::SlateColorSubConflict : SInterchangeGenericAssetMaterialConflictWidget::SlateColorFullConflict);
+
+		this->ChildSlot
+			[
+				SNew(SHorizontalBox)
+					+ SHorizontalBox::Slot()
+					.AutoWidth()
+					[
+						SNew(SExpanderArrow, SharedThis(this))
+					]
+
+					+ SHorizontalBox::Slot()
+					.AutoWidth()
+					.Padding(0.0f, 2.0f, 6.0f, 2.0f)
+					[
+						SNew(SImage)
+							.Image(JointIcon)
+							.Visibility(JointIcon != FAppStyle::GetDefaultBrush() ? EVisibility::Visible : EVisibility::Collapsed)
+					]
+					+ SHorizontalBox::Slot()
+					.FillWidth(1.0f)
+					.Padding(0.0f, 3.0f, 6.0f, 3.0f)
+					.VAlign(VAlign_Center)
+					[
+						SNew(STextBlock)
+							.Text(FText::FromString(SkeletonCompareData->JointName))
+							.ToolTipText(FText::FromString(Tooltip))
+							.ColorAndOpacity(ForegroundTextColor)
+					]
+			];
+
+		STableRow< TSharedPtr<UInterchangeGenericAssetsPipeline::FSkeletonJoint> >::ConstructInternal(
+			STableRow::FArguments()
+			.ShowSelection(true),
+			InOwnerTableView
+		);
+	}
+
+private:
+	/** The node info to build the tree view row from. */
+	TSharedPtr<UInterchangeGenericAssetsPipeline::FSkeletonJoint> SkeletonCompareData;
+};
+
+
+TSharedRef<ITableRow> SInterchangeGenericAssetSkeletonConflictWidget::OnGenerateRowCompareTreeView(TSharedPtr<UInterchangeGenericAssetsPipeline::FSkeletonJoint> RowData, const TSharedRef<STableViewBase>& Table)
+{
+	TSharedRef<SInterchangeCompareSkeletonTreeViewItem> ReturnRow = SNew(SInterchangeCompareSkeletonTreeViewItem, Table)
+		.SkeletonCompareData(RowData);
+	return ReturnRow;
+}
+
+void SInterchangeGenericAssetSkeletonConflictWidget::OnGetChildrenRowCompareTreeView(TSharedPtr<UInterchangeGenericAssetsPipeline::FSkeletonJoint> InParent, TArray< TSharedPtr<UInterchangeGenericAssetsPipeline::FSkeletonJoint> >& OutChildren)
+{
+	for (int32 ChildIndex = 0; ChildIndex < InParent->Children.Num(); ++ChildIndex)
+	{
+		TSharedPtr<UInterchangeGenericAssetsPipeline::FSkeletonJoint> ChildJoint = InParent->Children[ChildIndex];
+		if (ChildJoint.IsValid())
+		{
+			OutChildren.Add(ChildJoint);
+		}
+	}
+}
+
+TSharedRef<ITableRow> SInterchangeGenericAssetSkeletonConflictWidget::OnGenerateRowAssetReferencingSkeleton(TSharedPtr<FString> InItem, const TSharedRef<STableViewBase>& OwnerTable)
+{
+	int32 AssetListIndex = AssetReferencingSkeleton.Find(InItem);
+	bool LightBackgroundColor = AssetListIndex % 2 == 0;
+	return SNew(STableRow<TSharedPtr<FString> >, OwnerTable)
+		[
+			SNew(SBorder)
+			.BorderImage(LightBackgroundColor ? FAppStyle::GetBrush("ToolPanel.GroupBorder") : FAppStyle::GetBrush("ToolPanel.DarkGroupBorder"))
+			[
+				SNew(STextBlock)
+				.Text(FText::FromString(*(InItem.Get())))
+			]
+		];
+}
\ No newline at end of file
diff --git a/Engine/Plugins/Interchange/Runtime/Source/Pipelines/Private/InterchangePipelineHelper.cpp b/Engine/Plugins/Interchange/Runtime/Source/Pipelines/Private/InterchangePipelineHelper.cpp
new file mode 100644
--- /dev/null
+++ b/Engine/Plugins/Interchange/Runtime/Source/Pipelines/Private/InterchangePipelineHelper.cpp
@@ -0,0 +1,48 @@
+// Copyright Epic Games, Inc. All Rights Reserved. 
+
+#include "InterchangePipelineHelper.h"
+
+#include "CoreMinimal.h"
+#include "Framework/Application/SlateApplication.h"
+#include "Framework/Docking/TabManager.h"
+#include "HAL/PlatformApplicationMisc.h"
+#include "InterchangePipelineBase.h"
+#include "Widgets/SCompoundWidget.h"
+#include "Widgets/SWindow.h"
+
+namespace UE::Interchange::PipelineHelper
+{
+	void ShowModalDialog(TSharedRef<SInterchangeBaseConflictWidget> ConflictWidget, const FText& Title, const FVector2D& WindowSize)
+	{
+		TSharedPtr<SWindow> ParentWindow = FGlobalTabmanager::Get()->GetRootWindow();
+
+		const float FbxImportWindowWidth = WindowSize.X > 150.0f ? WindowSize.X : 150.0f;
+		const float FbxImportWindowHeight = WindowSize.Y > 50.0f ? WindowSize.Y : 50.0f;
+		FVector2D FbxImportWindowSize = FVector2D(FbxImportWindowWidth, FbxImportWindowHeight); // Max window size it can get based on current slate
+
+		FSlateRect WorkAreaRect = FSlateApplicationBase::Get().GetPreferredWorkArea();
+		FVector2D DisplayTopLeft(WorkAreaRect.Left, WorkAreaRect.Top);
+		FVector2D DisplaySize(WorkAreaRect.Right - WorkAreaRect.Left, WorkAreaRect.Bottom - WorkAreaRect.Top);
+
+		float ScaleFactor = FPlatformApplicationMisc::GetDPIScaleFactorAtPoint(DisplayTopLeft.X, DisplayTopLeft.Y);
+		FbxImportWindowSize *= ScaleFactor;
+
+		FVector2D WindowPosition = (DisplayTopLeft + (DisplaySize - FbxImportWindowSize) / 2.0f) / ScaleFactor;
+
+		TSharedRef<SWindow> Window = SNew(SWindow)
+			.Title(Title)
+			.SizingRule(ESizingRule::UserSized)
+			.AutoCenter(EAutoCenter::None)
+			.ClientSize(FbxImportWindowSize)
+			.ScreenPosition(WindowPosition);
+
+		ConflictWidget->SetWidgetWindow(Window);
+		Window->SetContent
+		(
+			ConflictWidget
+		);
+
+		// @todo: we can make this slow as showing progress bar later
+		FSlateApplication::Get().AddModalWindow(Window, ParentWindow, false);
+	}
+}
diff --git a/Engine/Plugins/Interchange/Runtime/Source/Pipelines/Public/InterchangeGenericAssetsPipeline.h b/Engine/Plugins/Interchange/Runtime/Source/Pipelines/Public/InterchangeGenericAssetsPipeline.h
--- a/Engine/Plugins/Interchange/Runtime/Source/Pipelines/Public/InterchangeGenericAssetsPipeline.h
+++ b/Engine/Plugins/Interchange/Runtime/Source/Pipelines/Public/InterchangeGenericAssetsPipeline.h
@@ -10,13 +10,25 @@
 #include "Nodes/InterchangeBaseNodeContainer.h"
 #include "UObject/Object.h"
 #include "UObject/ObjectMacros.h"
+#include "Widgets/SCompoundWidget.h"
+#include "Widgets/SWindow.h"
+#include "Widgets/Views/SListView.h"
+#include "Widgets/Views/STreeView.h"
 
 #include "InterchangeGenericAssetsPipeline.generated.h"
 
+class STableViewBase;
+class ITableRow;
 class UInterchangeGenericAnimationPipeline;
 class UInterchangeGenericMaterialPipeline;
 class UInterchangeGenericMeshPipeline;
 class UInterchangeGenericTexturePipeline;
+class USkeletalMesh;
+class USkeleton;
+class UStaticMesh;
+
+struct FReferenceSkeleton;
+struct FMeshBoneInfo;
 
 /**
  * This pipeline is the generic pipeline option for all meshes type and should be call before specialized Mesh pipeline (like generic static mesh or skeletal mesh pipelines)
@@ -87,6 +99,11 @@ public:
 
 	virtual bool IsPropertyChangeNeedRefresh(const FPropertyChangedEvent& PropertyChangedEvent) override;
 #endif //WITH_EDITOR
+
+	virtual TArray<FInterchangeConflictInfo> GetConflictInfos(UObject* ReimportObject, UInterchangeBaseNodeContainer* InBaseNodeContainer, UInterchangeSourceData* SourceData) override;
+
+	virtual void ShowConflictDialog(const FGuid& ConflictUniqueId) override;
+
 	virtual bool IsScripted() override
 	{
 		return false;
@@ -96,6 +113,19 @@ public:
 	virtual bool GetPropertyPossibleValues(const FName PropertyPath, TArray<FString>& PossibleValues) override;
 #endif
 
+	struct FSkeletonJoint : public TSharedFromThis<FSkeletonJoint>
+	{
+		FString JointName;
+		bool bAdded = false;
+		bool bRemoved = false;
+		bool bMatch = false;
+		bool bConflict = false;
+		bool bChildConflict = false;
+		bool bInitialAutoExpand = false;
+		TSharedPtr<FSkeletonJoint> Parent;
+		TArray<TSharedPtr<FSkeletonJoint>> Children;
+	};
+
 protected:
 
 	virtual void ExecutePipeline(UInterchangeBaseNodeContainer* InBaseNodeContainer, const TArray<UInterchangeSourceData*>& InSourceDatas) override;
@@ -114,11 +144,18 @@ protected:
 
 	//virtual bool ExecuteExportPipeline(UInterchangeBaseNodeContainer* BaseNodeContainer) override;
 private:
-	
+
+#if WITH_EDITOR
+	void CreateMaterialConflict(UStaticMesh* StaticMesh, USkeletalMesh* SkeletalMesh, UInterchangeBaseNodeContainer* TransientBaseNodeContainer);
+	void InternalRecursiveFillJointsFromReferenceSkeleton(TSharedPtr<FSkeletonJoint> ParentJoint, TMap<FString, TSharedPtr<FSkeletonJoint>>& Joints, const int32 BoneIndex, const FReferenceSkeleton& ReferenceSkeleton);
+	void InternalRecursiveFillJointsFromNodeContainer(TSharedPtr<FSkeletonJoint> ParentJoint, TMap<FString, TSharedPtr<FSkeletonJoint>>& Joints, const FString& JoinUid, const UInterchangeBaseNodeContainer* BaseNodeContainer, const bool bConvertStaticToSkeletalActive);
+	void CreateSkeletonConflict(USkeleton* SpecifiedSkeleton, USkeletalMesh* SkeletalMesh, UInterchangeBaseNodeContainer* TransientBaseNodeContainer);
+#endif
+
 	/**
 	 * Implement pipeline option bUseSourceNameForAsset
 	 */
-	void ImplementUseSourceNameForAssetOption();
+	void ImplementUseSourceNameForAssetOption(UInterchangeBaseNodeContainer* InBaseNodeContainer, const TArray<UInterchangeSourceData*>& InSourceDatas);
 
 #if WITH_EDITORONLY_DATA
 	/**
@@ -127,7 +164,166 @@ private:
 	void AddPackageMetaData(UObject* CreatedAsset, const UInterchangeBaseNode* Node);
 #endif // WITH_EDITORONLY_DATA
 
-	UInterchangeBaseNodeContainer* BaseNodeContainer;
-	TArray<const UInterchangeSourceData*> SourceDatas;
+	struct FMaterialConflictData
+	{
+		FGuid ConflictUniqueId;
+		TArray<FString> AssetMaterialNames;
+		TArray<FString> ImportMaterialNames;
+		TArray<int32> MatchMaterialIndexes;
+		UObject* ReimportObject = nullptr;
+		const FText DialogTitle = NSLOCTEXT("UInterchangeGenericAssetsPipeline", "GetConflictInfos_MaterialTitle", "Material Conflicts");
+
+		void Reset()
+		{
+			ConflictUniqueId.Invalidate();
+			AssetMaterialNames.Empty();
+			ImportMaterialNames.Empty();
+			MatchMaterialIndexes.Empty();
+			ReimportObject = nullptr;
+		}
+	};
+	FMaterialConflictData MaterialConflictData;
+
+	struct FSkeletonConflictData
+	{
+		FGuid ConflictUniqueId;
+		TMap<FString, TSharedPtr<FSkeletonJoint>> Joints;
+		UObject* ReimportObject = nullptr;
+		const FText DialogTitle = NSLOCTEXT("UInterchangeGenericAssetsPipeline", "GetConflictInfos_SkeletonTitle", "Skeleton Conflicts");
+
+		void Reset()
+		{
+			ConflictUniqueId.Invalidate();
+			Joints.Empty();
+			ReimportObject = nullptr;
+		}
+	};
+	FSkeletonConflictData SkeletonConflictData;
+};
+
+class SInterchangeGenericAssetMaterialConflictWidget : public SInterchangeBaseConflictWidget
+{
+public:
+
+	struct FListItem
+	{
+		FString ImportName;
+		int32 bMatched = INDEX_NONE;
+		FString AssetMatchedName;
+		FString AssetName;
+	};
+
+	static const FName NAME_Import;
+	static const FName NAME_Asset;
+
+	static const FSlateColor SlateColorFullConflict;
+	static const FSlateColor SlateColorSubConflict;
+
+	SLATE_BEGIN_ARGS(SInterchangeGenericAssetMaterialConflictWidget)
+		: _AssetMaterialNames()
+		, _ImportMaterialNames()
+		, _MatchMaterialIndexes()
+		, _ReimportObject(nullptr)
+		{}
+		SLATE_ARGUMENT(TArray<FString>, AssetMaterialNames)
+		SLATE_ARGUMENT(TArray<FString>, ImportMaterialNames)
+		SLATE_ARGUMENT(TArray<int32>, MatchMaterialIndexes)
+		SLATE_ARGUMENT(UObject*, ReimportObject)
+
+	SLATE_END_ARGS()
+	void Construct(const FArguments& InArgs);
+	virtual bool SupportsKeyboardFocus() const override { return true; }
+
+	FReply OnDone();
+
+	virtual FReply OnKeyDown(const FGeometry& MyGeometry, const FKeyEvent& InKeyEvent) override;
+
+protected:
 	
+	TSharedRef<ITableRow> OnGenerateRow(TSharedPtr<FListItem> Item, const TSharedRef<STableViewBase>& OwnerTable);
+
+	TArray<FString> AssetMaterialNames;
+	TArray<FString> ImportMaterialNames;
+	TArray<int32> MatchMaterialIndexes;
+	UObject* ReimportObject = nullptr;
+
+	TArray<TSharedPtr<FListItem>> RowItems;
+	TSharedPtr<SListView<TSharedPtr<FListItem>>> MaterialList;
+};
+
+enum EInterchangeSkeletonCompareSection
+{
+	Skeleton = 0,
+	References,
+	Count
+};
+
+class SInterchangeGenericAssetSkeletonConflictWidget : public SInterchangeBaseConflictWidget
+{
+public:
+
+	SLATE_BEGIN_ARGS(SInterchangeGenericAssetSkeletonConflictWidget)
+		: _AssetReferencingSkeleton()
+		, _Joints()
+		, _ReimportObject(nullptr)
+		{}
+
+		SLATE_ARGUMENT(TArray<TSharedPtr<FString>>, AssetReferencingSkeleton)
+		SLATE_ARGUMENT(TArray<TSharedPtr<UInterchangeGenericAssetsPipeline::FSkeletonJoint>>, Joints)
+		SLATE_ARGUMENT(UObject*, ReimportObject)
+
+	SLATE_END_ARGS()
+
+public:
+	void Construct(const FArguments& InArgs);
+	virtual bool SupportsKeyboardFocus() const override { return true; }
+
+	FReply OnDone()
+	{
+		if (WidgetWindow.IsValid())
+		{
+			WidgetWindow->RequestDestroyWindow();
+		}
+		return FReply::Handled();
+	}
+
+
+	virtual FReply OnKeyDown(const FGeometry& MyGeometry, const FKeyEvent& InKeyEvent) override
+	{
+		if (InKeyEvent.GetKey() == EKeys::Escape)
+		{
+			return OnDone();
+		}
+		return FReply::Unhandled();
+	}
+
+	SInterchangeGenericAssetSkeletonConflictWidget()
+	{}
+
+private:
+	TArray<TSharedPtr<FString>> AssetReferencingSkeleton;
+	TArray<TSharedPtr<UInterchangeGenericAssetsPipeline::FSkeletonJoint>> Joints;
+	UObject* ReimportObject;
+
+	//////////////////////////////////////////////////////////////////////////
+	//Collapse generic
+	bool bShowSectionFlag[EInterchangeSkeletonCompareSection::Count];
+	FReply SetSectionVisible(EInterchangeSkeletonCompareSection SectionIndex);
+	EVisibility IsSectionVisible(EInterchangeSkeletonCompareSection SectionIndex);
+	const FSlateBrush* GetCollapsableArrow(EInterchangeSkeletonCompareSection SectionIndex) const;
+	//////////////////////////////////////////////////////////////////////////
+
+
+	//////////////////////////////////////////////////////////////////////////
+	// Skeleton Data
+	TSharedPtr<STreeView<TSharedPtr<UInterchangeGenericAssetsPipeline::FSkeletonJoint>>> CompareTree;
+
+	//Construct slate
+	TSharedPtr<SWidget> ConstructSkeletonComparison();
+	TSharedPtr<SWidget> ConstructSkeletonReference();
+	//Slate events
+	TSharedRef<ITableRow> OnGenerateRowCompareTreeView(TSharedPtr<UInterchangeGenericAssetsPipeline::FSkeletonJoint> RowData, const TSharedRef<STableViewBase>& Table);
+	void OnGetChildrenRowCompareTreeView(TSharedPtr<UInterchangeGenericAssetsPipeline::FSkeletonJoint> InParent, TArray< TSharedPtr<UInterchangeGenericAssetsPipeline::FSkeletonJoint> >& OutChildren);
+	TSharedRef<ITableRow> OnGenerateRowAssetReferencingSkeleton(TSharedPtr<FString> InItem, const TSharedRef<STableViewBase>& OwnerTable);
+	//////////////////////////////////////////////////////////////////////////
 };
diff --git a/Engine/Plugins/Interchange/Runtime/Source/Pipelines/Public/InterchangeGenericAssetsPipelineSharedSettings.h b/Engine/Plugins/Interchange/Runtime/Source/Pipelines/Public/InterchangeGenericAssetsPipelineSharedSettings.h
--- a/Engine/Plugins/Interchange/Runtime/Source/Pipelines/Public/InterchangeGenericAssetsPipelineSharedSettings.h
+++ b/Engine/Plugins/Interchange/Runtime/Source/Pipelines/Public/InterchangeGenericAssetsPipelineSharedSettings.h
@@ -166,6 +166,10 @@ public:
 		{
 			return true;
 		}
+		else if (PropertyChangedEvent.Property->GetName() == GET_MEMBER_NAME_CHECKED(UInterchangeGenericCommonSkeletalMeshesAndAnimationsProperties, Skeleton))
+		{
+			return true;
+		}
 		return Super::IsPropertyChangeNeedRefresh(PropertyChangedEvent);
 	}
 #endif //WITH_EDITOR
diff --git a/Engine/Plugins/Interchange/Runtime/Source/Pipelines/Public/InterchangePipelineHelper.h b/Engine/Plugins/Interchange/Runtime/Source/Pipelines/Public/InterchangePipelineHelper.h
new file mode 100644
--- /dev/null
+++ b/Engine/Plugins/Interchange/Runtime/Source/Pipelines/Public/InterchangePipelineHelper.h
@@ -0,0 +1,11 @@
+// Copyright Epic Games, Inc. All Rights Reserved.
+
+#pragma once
+
+#include "CoreMinimal.h"
+#include "InterchangePipelineBase.h"
+
+namespace UE::Interchange::PipelineHelper
+{
+	void ShowModalDialog(TSharedRef<SInterchangeBaseConflictWidget> ConflictWidget, const FText& Title, const FVector2D& WindowSize);
+}
diff --git a/Engine/Source/Runtime/Interchange/Core/InterchangeCore.Build.cs b/Engine/Source/Runtime/Interchange/Core/InterchangeCore.Build.cs
--- a/Engine/Source/Runtime/Interchange/Core/InterchangeCore.Build.cs
+++ b/Engine/Source/Runtime/Interchange/Core/InterchangeCore.Build.cs
@@ -13,5 +13,12 @@ public class InterchangeCore : ModuleRules
 				"Json"
 			}
 		);
+
+		PrivateDependencyModuleNames.AddRange(
+			new string[]
+			{
+				"SlateCore",
+			}
+		);
 	}
 }
diff --git a/Engine/Source/Runtime/Interchange/Core/Public/InterchangePipelineBase.h b/Engine/Source/Runtime/Interchange/Core/Public/InterchangePipelineBase.h
--- a/Engine/Source/Runtime/Interchange/Core/Public/InterchangePipelineBase.h
+++ b/Engine/Source/Runtime/Interchange/Core/Public/InterchangePipelineBase.h
@@ -18,6 +18,8 @@
 #include "UObject/ObjectMacros.h"
 #include "UObject/ObjectPtr.h"
 #include "UObject/UObjectGlobals.h"
+#include "Widgets/SCompoundWidget.h"
+#include "Widgets/SWindow.h"
 
 #include "InterchangePipelineBase.generated.h"
 
@@ -101,6 +103,38 @@ struct FInterchangePipelinePropertyStates
 	FInterchangePipelinePropertyStatePerContext ReimportStates;
 };
 
+struct FInterchangeConflictInfo
+{
+	FString DisplayName;
+	FString Description;
+	FGuid UniqueId;
+	TObjectPtr<UInterchangePipelineBase> Pipeline = nullptr;
+};
+
+class SInterchangeBaseConflictWidget : public SCompoundWidget
+{
+public:
+	SLATE_BEGIN_ARGS(SInterchangeBaseConflictWidget)
+		: _WidgetWindow(nullptr)
+		{}
+
+		SLATE_ARGUMENT(TSharedPtr<SWindow>, WidgetWindow)
+
+	SLATE_END_ARGS()
+
+	virtual void SetWidgetWindow(TSharedPtr<SWindow> InWidgetWindow)
+	{
+		WidgetWindow = InWidgetWindow;
+	}
+
+	FVector2D GetMinimumSize(float ApplicationScale) const
+	{
+		return ComputeDesiredSize(ApplicationScale);
+	}
+protected:
+	TSharedPtr<SWindow> WidgetWindow = nullptr;
+};
+
 UCLASS(BlueprintType, Blueprintable, editinlinenew, Abstract, MinimalAPI)
 class UInterchangePipelineBase : public UObject
 {
@@ -255,8 +289,25 @@ public:
 	{
 		return false;
 	}
+
 #endif //WITH_EDITOR
 
+	/**
+	 * Return all conflict the pipeline find in the translated data. This function is call by the import dialog.
+	 * If some specific options change, the UI must refresh the conflicts, see function IsPropertyChangeNeedRefresh.
+	 */
+	INTERCHANGECORE_API virtual TArray<FInterchangeConflictInfo> GetConflictInfos(UObject* ReimportObject, UInterchangeBaseNodeContainer* InBaseNodeContainer, UInterchangeSourceData* SourceData)
+	{
+		//The base class function do not have any conflict
+		ConflictInfos.Empty();
+		return ConflictInfos;
+	}
+
+	INTERCHANGECORE_API virtual void ShowConflictDialog(const FGuid& ConflictUniqueId)
+	{
+		return;
+	}
+
 	/**
 	 * This function is used to add the given message object directly into the results for this operation.
 	 */
@@ -443,4 +494,6 @@ protected:
 	mutable TMap<FName, FInterchangePipelinePropertyStates> CachePropertiesStates;
 	mutable EInterchangePipelineContext CachePipelineContext = EInterchangePipelineContext::None;
 	mutable TWeakObjectPtr<UObject> CacheReimportObject = nullptr;
+
+	TArray<FInterchangeConflictInfo> ConflictInfos;
 };
diff --git a/Engine/Source/Runtime/Interchange/Engine/Private/InterchangeManager.cpp b/Engine/Source/Runtime/Interchange/Engine/Private/InterchangeManager.cpp
--- a/Engine/Source/Runtime/Interchange/Engine/Private/InterchangeManager.cpp
+++ b/Engine/Source/Runtime/Interchange/Engine/Private/InterchangeManager.cpp
@@ -1741,7 +1741,7 @@ UInterchangeManager::ImportInternal(const FString& ContentPath, const UInterchan
 				TranslateSourceFile();
 				UInterchangeBaseNodeContainer* BaseNodeContainer = AsyncHelper->BaseNodeContainers[SourceIndex].Get();
 				//Show the dialog, a plugin should have registered this dialog. We use a plugin to be able to use editor code when doing UI
-				EInterchangePipelineConfigurationDialogResult DialogResult = RegisteredPipelineConfiguration->ScriptedShowReimportPipelineConfigurationDialog(PipelineStacks, OutPipelines, DuplicateSourceData, BaseNodeContainer);
+				EInterchangePipelineConfigurationDialogResult DialogResult = RegisteredPipelineConfiguration->ScriptedShowReimportPipelineConfigurationDialog(PipelineStacks, OutPipelines, DuplicateSourceData, BaseNodeContainer, ImportAssetParameters.ReimportAsset);
 				if (DialogResult == EInterchangePipelineConfigurationDialogResult::Cancel)
 				{
 					bImportCanceled = true;
diff --git a/Engine/Source/Runtime/Interchange/Engine/Public/InterchangePipelineConfigurationBase.h b/Engine/Source/Runtime/Interchange/Engine/Public/InterchangePipelineConfigurationBase.h
--- a/Engine/Source/Runtime/Interchange/Engine/Public/InterchangePipelineConfigurationBase.h
+++ b/Engine/Source/Runtime/Interchange/Engine/Public/InterchangePipelineConfigurationBase.h
@@ -82,16 +82,18 @@ public:
 	INTERCHANGEENGINE_API EInterchangePipelineConfigurationDialogResult ScriptedShowReimportPipelineConfigurationDialog(TArray<FInterchangeStackInfo>& PipelineStacks
 		, TArray<UInterchangePipelineBase*>& OutPipelines
 		, UInterchangeSourceData* SourceData
-		, UInterchangeBaseNodeContainer* BaseNodeContainer);
+		, UInterchangeBaseNodeContainer* BaseNodeContainer
+		, UObject* ReimportAsset);
 
 	/** The default implementation (call if the blueprint do not have any implementation) will call the virtual ExecuteImportPipeline */
 	EInterchangePipelineConfigurationDialogResult ScriptedShowReimportPipelineConfigurationDialog_Implementation(TArray<FInterchangeStackInfo>& PipelineStacks
 		, TArray<UInterchangePipelineBase*>& OutPipelines
 		, UInterchangeSourceData* SourceData
-		, UInterchangeBaseNodeContainer* BaseNodeContainer)
+		, UInterchangeBaseNodeContainer* BaseNodeContainer
+		, UObject* ReimportAsset)
 	{
 		//By default we call the virtual import pipeline execution
-		return ShowReimportPipelineConfigurationDialog(PipelineStacks, OutPipelines, SourceData, BaseNodeContainer);
+		return ShowReimportPipelineConfigurationDialog(PipelineStacks, OutPipelines, SourceData, BaseNodeContainer, ReimportAsset);
 	}
 
 protected:
@@ -126,7 +128,8 @@ protected:
 	virtual EInterchangePipelineConfigurationDialogResult ShowReimportPipelineConfigurationDialog(TArray<FInterchangeStackInfo>& PipelineStacks
 		, TArray<UInterchangePipelineBase*>& OutPipelines
 		, TWeakObjectPtr<UInterchangeSourceData> SourceData
-		, TWeakObjectPtr <UInterchangeBaseNodeContainer> BaseNodeContainer)
+		, TWeakObjectPtr <UInterchangeBaseNodeContainer> BaseNodeContainer
+		, TWeakObjectPtr <UObject> ReimportAsset)
 	{
 		//Not implemented
 		return EInterchangePipelineConfigurationDialogResult::Cancel;
