Make a preview import button

#rnx

diff --git a/Engine/Plugins/Interchange/Editor/Source/Pipelines/Private/InterchangeEditorPipelineDetails.cpp b/Engine/Plugins/Interchange/Editor/Source/Pipelines/Private/InterchangeEditorPipelineDetails.cpp
--- a/Engine/Plugins/Interchange/Editor/Source/Pipelines/Private/InterchangeEditorPipelineDetails.cpp
+++ b/Engine/Plugins/Interchange/Editor/Source/Pipelines/Private/InterchangeEditorPipelineDetails.cpp
@@ -670,6 +670,31 @@ void FInterchangeBaseNodeDetailsCustomization::CustomizeDetails( IDetailLayoutBu
 		return;
 	}
 
+	//Add a row to describe the node type and class
+	{
+		const FText NodeInformationCategoryText = LOCTEXT("NodeInformationCategoryName", "Node Information");
+		const FName NodeInformationCategoryName = FName(TEXT("Node Information"));
+		IDetailCategoryBuilder& AttributeCategoryBuilder = DetailBuilder.EditCategory(NodeInformationCategoryName, NodeInformationCategoryText);
+		FText NodeInformationText = FText::GetEmpty();
+		if (UInterchangeFactoryBaseNode* FactoryNode = Cast<UInterchangeFactoryBaseNode>(InterchangeBaseNode))
+		{
+			const FString ClassName = FactoryNode->GetObjectClass() ? FactoryNode->GetObjectClass()->GetName() : InterchangeBaseNode->GetClass()->GetName();
+			NodeInformationText = FText::Format(LOCTEXT("NodeFactoryInformationText", "Factory Node ({0})"), FText::FromString(ClassName));
+		}
+		else
+		{
+			NodeInformationText = FText::Format(LOCTEXT("NodeFactoryInformationText", "Translated Node ({0})"), FText::FromString(InterchangeBaseNode->GetClass()->GetName()));
+		}
+
+		FDetailWidgetRow& CustomRow = AttributeCategoryBuilder.AddCustomRow(NodeInformationCategoryText)
+		.NameContent()
+		[
+			SNew(STextBlock)
+				.Text(NodeInformationText)
+				.Font(IDetailLayoutBuilder::GetDetailFont())
+		];
+	}
+
 	TArray< UE::Interchange::FAttributeKey> AttributeKeys;
 	InterchangeBaseNode->GetAttributeKeys(AttributeKeys);
 
@@ -890,6 +915,7 @@ void FInterchangeBaseNodeDetailsCustomization::BuildBoolValueContent(IDetailCate
 			SNew(SBox)
 			[
 				SNew(SCheckBox)
+				.IsEnabled(InterchangeBaseNode->UserInterfaceContext == EInterchangeNodeUserInterfaceContext::None)
 				.OnCheckStateChanged_Lambda([this, AttributeKey](ECheckBoxState CheckType)
 				{
 					const bool IsChecked = CheckType == ECheckBoxState::Checked;
@@ -1689,6 +1715,7 @@ TSharedRef<SWidget> FInterchangeBaseNodeDetailsCustomization::MakeNumericWidget(
 	
 	return
 		SNew(SNumericEntryBox<NumericType>)
+			.IsEnabled(InterchangeBaseNode->UserInterfaceContext == EInterchangeNodeUserInterfaceContext::None)
 			.EditableTextBoxStyle(&FCoreStyle::Get().GetWidgetStyle<FEditableTextBoxStyle>("NormalEditableTextBox"))
 			.Value_Lambda([this, &GetValue, ComponentIndex, AttributeKey]()
 			{
diff --git a/Engine/Plugins/Interchange/Editor/Source/Pipelines/Private/SInterchangeGraphInspectorWindow.cpp b/Engine/Plugins/Interchange/Editor/Source/Pipelines/Private/SInterchangeGraphInspectorWindow.cpp
--- a/Engine/Plugins/Interchange/Editor/Source/Pipelines/Private/SInterchangeGraphInspectorWindow.cpp
+++ b/Engine/Plugins/Interchange/Editor/Source/Pipelines/Private/SInterchangeGraphInspectorWindow.cpp
@@ -10,6 +10,7 @@
 #include "IDocumentation.h"
 #include "Nodes/InterchangeBaseNodeContainer.h"
 #include "Nodes/InterchangeFactoryBaseNode.h"
+#include "PropertyEditorDelegates.h"
 #include "PropertyEditorModule.h"
 #include "Styling/SlateIconFinder.h"
 #include "Styling/StyleColors.h"
@@ -32,6 +33,7 @@ SInterchangeGraphInspectorTreeView::~SInterchangeGraphInspectorTreeView()
 void SInterchangeGraphInspectorTreeView::Construct(const FArguments& InArgs)
 {
 	InterchangeBaseNodeContainer = InArgs._InterchangeBaseNodeContainer;
+	bPreview = InArgs._bPreview;
 	OnSelectionChangedDelegate = InArgs._OnSelectionChangedDelegate;
 	//Build the FbxNodeInfoPtr tree data
 	check(InterchangeBaseNodeContainer != nullptr);
@@ -41,6 +43,14 @@ void SInterchangeGraphInspectorTreeView::Construct(const FArguments& InArgs)
 	for (FString RootID : Roots)
 	{
 		const UInterchangeBaseNode* RootNode = InterchangeBaseNodeContainer->GetNode(RootID);
+		if (bPreview)
+		{
+			const UInterchangeFactoryBaseNode* FactoryNode = Cast<UInterchangeFactoryBaseNode>(RootNode);
+			if (!FactoryNode || !FactoryNode->GetObjectClass())
+			{
+				continue;
+			}
+		}
 		RootNodeArray.Add(const_cast<UInterchangeBaseNode*>(RootNode));
 	}
 
@@ -64,11 +74,13 @@ public:
 	SLATE_BEGIN_ARGS(SInterchangeGraphInspectorTreeViewItem)
 		: _InterchangeNode(nullptr)
 		, _InterchangeBaseNodeContainer(nullptr)
+		, _bPreview(false)
 	{}
 
 	/** The item content. */
 	SLATE_ARGUMENT(UInterchangeBaseNode*, InterchangeNode)
 	SLATE_ARGUMENT(UInterchangeBaseNodeContainer*, InterchangeBaseNodeContainer)
+	SLATE_ARGUMENT(bool, bPreview)
 	SLATE_END_ARGS()
 
 	/**
@@ -79,6 +91,7 @@ public:
 	void Construct(const FArguments& InArgs, const TSharedRef<STableViewBase>& InOwnerTableView)
 	{
 		InterchangeNode = InArgs._InterchangeNode;
+		bPreview = InArgs._bPreview;
 		InterchangeBaseNodeContainer = InArgs._InterchangeBaseNodeContainer;
 
 		// This is supposed to always be valid
@@ -110,23 +123,38 @@ public:
 			TypeIcon = FSlateIconFinder::FindIconBrushForClass(AActor::StaticClass());
 		}
 
-		FSlateColor TypeIconColor;
+		TSharedPtr<SImage> IconWidget = nullptr;
 
-		switch(InterchangeNode->GetNodeContainerType())
+		if (!bPreview)
+		{
+			FSlateColor TypeIconColor;
+			switch (InterchangeNode->GetNodeContainerType())
+			{
+			case EInterchangeNodeContainerType::TranslatedAsset:
+				TypeIconColor = FStyleColors::AccentBlue;
+				break;
+			case EInterchangeNodeContainerType::TranslatedScene:
+				TypeIconColor = FStyleColors::AccentGreen;
+				break;
+			case EInterchangeNodeContainerType::FactoryData:
+				TypeIconColor = FStyleColors::AccentPurple;
+				break;
+			case EInterchangeNodeContainerType::None:
+			default:
+				TypeIconColor = FStyleColors::AccentRed;
+				break;
+			}
+			
+			IconWidget = SNew(SImage)
+				.Image(TypeIcon)
+				.ColorAndOpacity(TypeIconColor)
+				.Visibility(TypeIcon != FAppStyle::GetDefaultBrush() ? EVisibility::Visible : EVisibility::Collapsed);
+		}
+		else
 		{
-		case EInterchangeNodeContainerType::TranslatedAsset:
-			TypeIconColor = FStyleColors::AccentBlue;
-			break;
-		case EInterchangeNodeContainerType::TranslatedScene:
-			TypeIconColor = FStyleColors::AccentGreen;
-			break;
-		case EInterchangeNodeContainerType::FactoryData:
-			TypeIconColor = FStyleColors::AccentPurple;
-			break;
-		case EInterchangeNodeContainerType::None:
-		default:
-			TypeIconColor = FStyleColors::AccentRed;
-			break;
+			IconWidget = SNew(SImage)
+				.Image(TypeIcon)
+				.Visibility(TypeIcon != FAppStyle::GetDefaultBrush() ? EVisibility::Visible : EVisibility::Collapsed);
 		}
 
 		this->ChildSlot
@@ -139,6 +167,7 @@ public:
 				SNew(SCheckBox)
 				.OnCheckStateChanged(this, &SInterchangeGraphInspectorTreeViewItem::OnItemCheckChanged)
 				.IsChecked(this, &SInterchangeGraphInspectorTreeViewItem::IsItemChecked)
+				.IsEnabled(!bPreview)
 			]
 			+ SHorizontalBox::Slot()
 			.AutoWidth()
@@ -149,10 +178,7 @@ public:
 			.AutoWidth()
 			.Padding(0.0f, 2.0f, 6.0f, 2.0f)
 			[
-				SNew(SImage)
-				.Image(TypeIcon)
-				.ColorAndOpacity(TypeIconColor)
-				.Visibility(TypeIcon != FAppStyle::GetDefaultBrush() ? EVisibility::Visible : EVisibility::Collapsed)
+				IconWidget.ToSharedRef()
 			]
 			+ SHorizontalBox::Slot()
 			.FillWidth(1.0f)
@@ -205,12 +231,14 @@ private:
 	/** The node to build the tree view row from. */
 	UInterchangeBaseNode* InterchangeNode = nullptr;
 	UInterchangeBaseNodeContainer* InterchangeBaseNodeContainer = nullptr;
+	bool bPreview = false;
 };
 
 TSharedRef< ITableRow > SInterchangeGraphInspectorTreeView::OnGenerateRowGraphInspectorTreeView(UInterchangeBaseNode* Item, const TSharedRef< STableViewBase >& OwnerTable)
 {
 	TSharedRef< SInterchangeGraphInspectorTreeViewItem > ReturnRow = SNew(SInterchangeGraphInspectorTreeViewItem, OwnerTable)
 		.InterchangeNode(Item)
+		.bPreview(bPreview)
 		.InterchangeBaseNodeContainer(InterchangeBaseNodeContainer);
 	return ReturnRow;
 }
@@ -220,7 +248,7 @@ void SInterchangeGraphInspectorTreeView::OnGetChildrenGraphInspectorTreeView(UIn
 	for (const FString& ChildID : Childrens)
 	{
 		const UInterchangeBaseNode* ChildNode = InterchangeBaseNodeContainer->GetNode(ChildID);
-		if (!ChildNode)
+		if (!ChildNode || (bPreview && !ChildNode->IsA<UInterchangeFactoryBaseNode>()))
 			continue;
 		OutChildren.Add(const_cast<UInterchangeBaseNode*>(ChildNode));
 	}
@@ -292,6 +320,10 @@ FReply SInterchangeGraphInspectorTreeView::OnCollapseAll()
 
 TSharedPtr<SWidget> SInterchangeGraphInspectorTreeView::OnOpenContextMenu()
 {
+	if (bPreview)
+	{
+		return SNullWidget::NullWidget;
+	}
 	// Build up the menu for a selection
 	const bool bCloseAfterSelection = true;
 	FMenuBuilder MenuBuilder(bCloseAfterSelection, TSharedPtr<FUICommandList>());
@@ -371,6 +403,7 @@ TSharedRef<SBox> SInterchangeGraphInspectorWindow::SpawnGraphInspector()
 	//Create the treeview
 	GraphInspectorTreeview = SNew(SInterchangeGraphInspectorTreeView)
 		.InterchangeBaseNodeContainer(InterchangeBaseNodeContainer)
+		.bPreview(bPreview)
 		.OnSelectionChangedDelegate(this, &SInterchangeGraphInspectorWindow::OnSelectionChanged);
 
 	TSharedPtr<SBox> InspectorBox;
@@ -397,6 +430,7 @@ TSharedRef<SBox> SInterchangeGraphInspectorWindow::SpawnGraphInspector()
 						SNew(SCheckBox)
 						.HAlign(HAlign_Center)
 						.OnCheckStateChanged(GraphInspectorTreeview.Get(), &SInterchangeGraphInspectorTreeView::OnToggleSelectAll)
+						.Visibility(bPreview ? EVisibility::Collapsed : EVisibility::All)
 					]
 					+ SHorizontalBox::Slot()
 					.FillWidth(1.0f)
@@ -405,6 +439,7 @@ TSharedRef<SBox> SInterchangeGraphInspectorWindow::SpawnGraphInspector()
 					[
 						SNew(STextBlock)
 						.Text(LOCTEXT("GraphInspectorWindow_Scene_All", "All"))
+						.Visibility(bPreview ? EVisibility::Collapsed : EVisibility::All)
 					]
 				]
 				+ SUniformGridPanel::Slot(1, 0)
@@ -440,11 +475,13 @@ TSharedRef<SBox> SInterchangeGraphInspectorWindow::SpawnGraphInspector()
 
 	FPropertyEditorModule& PropertyEditorModule = FModuleManager::GetModuleChecked<FPropertyEditorModule>("PropertyEditor");
 	FDetailsViewArgs DetailsViewArgs;
-	DetailsViewArgs.bAllowSearch = false;
+	DetailsViewArgs.bAllowSearch = bPreview;
 	DetailsViewArgs.NameAreaSettings = FDetailsViewArgs::HideNameArea;
 	GraphInspectorDetailsView = PropertyEditorModule.CreateDetailView(DetailsViewArgs);
 	InspectorBox->SetContent(GraphInspectorDetailsView->AsShared());
 	GraphInspectorDetailsView->SetObject(nullptr);
+//	GraphInspectorDetailsView->SetIsPropertyReadOnlyDelegate(FIsPropertyReadOnly::CreateLambda([this](const FPropertyAndParent&) { return bPreview; }));
+//	GraphInspectorDetailsView->SetIsCustomRowReadOnlyDelegate(FIsCustomRowReadOnly::CreateLambda([this](FName InRowName, FName InParentName) { return bPreview; }));
 	return GraphInspectorPanelBox;
 }
 
@@ -452,6 +489,7 @@ TSharedRef<SBox> SInterchangeGraphInspectorWindow::SpawnGraphInspector()
 void SInterchangeGraphInspectorWindow::Construct(const FArguments& InArgs)
 {
 	InterchangeBaseNodeContainer = InArgs._InterchangeBaseNodeContainer;
+	bPreview = InArgs._bPreview;
 	OwnerWindow = InArgs._OwnerWindow;
 
 	check(InterchangeBaseNodeContainer != nullptr);
@@ -475,19 +513,10 @@ void SInterchangeGraphInspectorWindow::Construct(const FArguments& InArgs)
 			.HAlign(HAlign_Right)
 			.Padding(2)
 			[
-				SNew(SUniformGridPanel)
-				.SlotPadding(2)
-				+ SUniformGridPanel::Slot(0, 0)
-				[
-					IDocumentation::Get()->CreateAnchor(FString("Engine/Content/Interchange/GraphInspector"))
-				]
-				+ SUniformGridPanel::Slot(1, 0)
-				[
-					SNew(SButton)
-					.HAlign(HAlign_Center)
-					.Text(LOCTEXT("InspectorGraphWindow_Close", "Ok"))
-					.OnClicked(this, &SInterchangeGraphInspectorWindow::OnCloseDialog)
-				]
+				SNew(SButton)
+				.HAlign(HAlign_Center)
+				.Text(LOCTEXT("InspectorGraphWindow_Close", "Ok"))
+				.OnClicked(this, &SInterchangeGraphInspectorWindow::OnCloseDialog)
 			]
 		]
 	];
diff --git a/Engine/Plugins/Interchange/Editor/Source/Pipelines/Private/SInterchangeGraphInspectorWindow.h b/Engine/Plugins/Interchange/Editor/Source/Pipelines/Private/SInterchangeGraphInspectorWindow.h
--- a/Engine/Plugins/Interchange/Editor/Source/Pipelines/Private/SInterchangeGraphInspectorWindow.h
+++ b/Engine/Plugins/Interchange/Editor/Source/Pipelines/Private/SInterchangeGraphInspectorWindow.h
@@ -20,9 +20,11 @@ public:
 
 	SLATE_BEGIN_ARGS(SInterchangeGraphInspectorTreeView)
 		: _InterchangeBaseNodeContainer(nullptr)
+		, _bPreview(false)
 		, _OnSelectionChangedDelegate()
 		{}
 		SLATE_ARGUMENT(UInterchangeBaseNodeContainer*, InterchangeBaseNodeContainer)
+		SLATE_ARGUMENT(bool, bPreview)
 		SLATE_EVENT(FOnGraphInspectorSelectionChanged, OnSelectionChangedDelegate)
 	SLATE_END_ARGS()
 
@@ -36,9 +38,9 @@ public:
 	FReply OnCollapseAll();
 protected:
 	UInterchangeBaseNodeContainer* InterchangeBaseNodeContainer = nullptr;
+	bool bPreview = false;
 	/** Delegate to invoke when selection changes. */
 	FOnGraphInspectorSelectionChanged OnSelectionChangedDelegate;
-
 	/** the elements we show in the tree view */
 	TArray<UInterchangeBaseNode*> RootNodeArray;
 
@@ -57,10 +59,12 @@ class SInterchangeGraphInspectorWindow : public SCompoundWidget
 public:
 	SLATE_BEGIN_ARGS(SInterchangeGraphInspectorWindow)
 		: _InterchangeBaseNodeContainer(nullptr)
+		, _bPreview(false)
 		, _OwnerWindow()
 		{}
 
 		SLATE_ARGUMENT( UInterchangeBaseNodeContainer*, InterchangeBaseNodeContainer)
+		SLATE_ARGUMENT(bool, bPreview)
 		SLATE_ARGUMENT(TWeakPtr<SWindow>, OwnerWindow)
 	SLATE_END_ARGS()
 
@@ -99,6 +103,7 @@ private:
 	//SInterchangeGraphInspectorWindow Arguments
 	
 	UInterchangeBaseNodeContainer* InterchangeBaseNodeContainer;
+	bool bPreview = false;
 	
 	
 	//Graph Inspector UI elements
diff --git a/Engine/Plugins/Interchange/Editor/Source/Pipelines/Private/SInterchangePipelineConfigurationDialog.cpp b/Engine/Plugins/Interchange/Editor/Source/Pipelines/Private/SInterchangePipelineConfigurationDialog.cpp
--- a/Engine/Plugins/Interchange/Editor/Source/Pipelines/Private/SInterchangePipelineConfigurationDialog.cpp
+++ b/Engine/Plugins/Interchange/Editor/Source/Pipelines/Private/SInterchangePipelineConfigurationDialog.cpp
@@ -16,6 +16,7 @@
 #include "Misc/ConfigCacheIni.h"
 #include "Nodes/InterchangeBaseNodeContainer.h"
 #include "PropertyEditorModule.h"
+#include "SInterchangeGraphInspectorWindow.h"
 #include "SPrimaryButton.h"
 #include "Styling/SlateIconFinder.h"
 #include "Widgets/Input/SCheckBox.h"
@@ -520,6 +521,15 @@ void SInterchangePipelineConfigurationDialog::Construct(const FArguments& InArgs
 				]
 				+ SHorizontalBox::Slot()
 				.AutoWidth()
+				.Padding(4.f, 0.f)
+				[
+					SNew(SButton)
+					.HAlign(HAlign_Center)
+					.Text(LOCTEXT("InspectorGraphWindow_Preview", "Preview..."))
+					.OnClicked(this, &SInterchangePipelineConfigurationDialog::OnPreviewImport)
+				]
+				+ SHorizontalBox::Slot()
+				.AutoWidth()
 				.Padding(4.f, 0.f) 
 				[
 					SNew(SButton)
@@ -872,4 +882,45 @@ void SInterchangePipelineConfigurationDialog::OnFilterOptionsChanged(ECheckBoxSt
 	GConfig->SetBool(TEXT("InterchangeImportDialogOptions"), TEXT("FilterOptions"), bFilterOptions, GEditorPerProjectIni);
 }
 
+FReply SInterchangePipelineConfigurationDialog::OnPreviewImport() const
+{
+	UInterchangeBaseNodeContainer* DuplicateBaseNodeContainer = DuplicateObject<UInterchangeBaseNodeContainer>(BaseNodeContainer.Get(), GetTransientPackage());
+
+	TArray<UInterchangeSourceData*> SourceDatas;
+	SourceDatas.Add(SourceData.Get());
+
+	//Execute all pipelines on the duplicated container
+	UInterchangeResultsContainer* Results = NewObject<UInterchangeResultsContainer>(GetTransientPackage());
+	for (int32 PipelineIndex = 0; PipelineIndex < PipelineListViewItems.Num(); ++PipelineIndex)
+	{
+		const TSharedPtr<FInterchangePipelineItemType> PipelineItem = PipelineListViewItems[PipelineIndex];
+		PipelineItem->Pipeline->SetResultsContainer(Results);
+		PipelineItem->Pipeline->ScriptedExecutePipeline(DuplicateBaseNodeContainer, SourceDatas);
+	}
+
+	//Set all node in preview mode so hide the internal data attributes
+	DuplicateBaseNodeContainer->IterateNodesOfType<UInterchangeFactoryBaseNode>([](const FString& NodeUid, UInterchangeFactoryBaseNode* Node)
+		{
+			Node->UserInterfaceContext = EInterchangeNodeUserInterfaceContext::Preview;
+		});
+
+	//Create and show the graph inspector UI dialog
+	TSharedRef<SWindow> Window = SNew(SWindow)
+		.ClientSize(FVector2D(800.f, 650.f))
+		.Title(NSLOCTEXT("Interchange", "GraphInspectorTitle", "Interchange node graph inspector"));
+	TSharedPtr<SInterchangeGraphInspectorWindow> InterchangeGraphInspectorWindow;
+
+	Window->SetContent
+	(
+		SAssignNew(InterchangeGraphInspectorWindow, SInterchangeGraphInspectorWindow)
+		.InterchangeBaseNodeContainer(DuplicateBaseNodeContainer)
+		.bPreview(true)
+		.OwnerWindow(Window)
+	);
+
+	FSlateApplication::Get().AddModalWindow(Window, OwnerWindow.Pin(), false);
+
+	return FReply::Handled();
+}
+
 #undef LOCTEXT_NAMESPACE
diff --git a/Engine/Plugins/Interchange/Editor/Source/Pipelines/Private/SInterchangePipelineConfigurationDialog.h b/Engine/Plugins/Interchange/Editor/Source/Pipelines/Private/SInterchangePipelineConfigurationDialog.h
--- a/Engine/Plugins/Interchange/Editor/Source/Pipelines/Private/SInterchangePipelineConfigurationDialog.h
+++ b/Engine/Plugins/Interchange/Editor/Source/Pipelines/Private/SInterchangePipelineConfigurationDialog.h
@@ -145,6 +145,8 @@ private:
 
 	void OnFilterOptionsChanged(ECheckBoxState CheckState);
 
+	FReply OnPreviewImport() const;
+
 	TSharedPtr<IDetailsView> PipelineConfigurationDetailsView;
 	TSharedPtr<SCheckBox> UseSameSettingsForAllCheckBox;
 
diff --git a/Engine/Plugins/Interchange/Runtime/Source/FactoryNodes/Private/InterchangeAnimSequenceFactoryNode.cpp b/Engine/Plugins/Interchange/Runtime/Source/FactoryNodes/Private/InterchangeAnimSequenceFactoryNode.cpp
--- a/Engine/Plugins/Interchange/Runtime/Source/FactoryNodes/Private/InterchangeAnimSequenceFactoryNode.cpp
+++ b/Engine/Plugins/Interchange/Runtime/Source/FactoryNodes/Private/InterchangeAnimSequenceFactoryNode.cpp
@@ -98,6 +98,8 @@ FString UInterchangeAnimSequenceFactoryNode::GetTypeName() const
 	return TypeName;
 }
 
+#if WITH_EDITOR
+
 FString UInterchangeAnimSequenceFactoryNode::GetKeyDisplayName(const UE::Interchange::FAttributeKey& NodeAttributeKey) const
 {
 	FString KeyDisplayName = NodeAttributeKey.ToString();
@@ -109,6 +111,182 @@ FString UInterchangeAnimSequenceFactoryNode::GetKeyDisplayName(const UE::Interch
 	{
 		KeyDisplayName = TEXT("Specified Existing Skeleton");
 	}
+	else if (NodeAttributeKey == UE::Interchange::FAnimSequenceNodeStaticData::GetAnimatedMorphTargetDependenciesKey())
+	{
+		KeyDisplayName = TEXT("Animated Morph Targets Count");
+		return KeyDisplayName;
+	}
+	else if (KeyDisplayName.StartsWith(UE::Interchange::FAnimSequenceNodeStaticData::GetAnimatedMorphTargetDependenciesKey().ToString()))
+	{
+		KeyDisplayName = TEXT("Animated Morph Targets Index ");
+		const FString IndexKey = UE::Interchange::TArrayAttributeHelper<FString>::IndexKey();
+		int32 IndexPosition = KeyDisplayName.Find(IndexKey) + IndexKey.Len();
+		if (IndexPosition < KeyDisplayName.Len())
+		{
+			KeyDisplayName += KeyDisplayName.RightChop(IndexPosition);
+		}
+		return KeyDisplayName;
+	}
+	else if (NodeAttributeKey == UE::Interchange::FAnimSequenceNodeStaticData::GetAnimatedAttributeCurveNamesKey())
+	{
+		KeyDisplayName = TEXT("Animated Attribute Curvse Count");
+		return KeyDisplayName;
+	}
+	else if (KeyDisplayName.StartsWith(UE::Interchange::FAnimSequenceNodeStaticData::GetAnimatedAttributeCurveNamesKey().ToString()))
+	{
+		KeyDisplayName = TEXT("Animated Attribute Curves Index ");
+		const FString IndexKey = UE::Interchange::TArrayAttributeHelper<FString>::IndexKey();
+		int32 IndexPosition = KeyDisplayName.Find(IndexKey) + IndexKey.Len();
+		if (IndexPosition < KeyDisplayName.Len())
+		{
+			KeyDisplayName += KeyDisplayName.RightChop(IndexPosition);
+		}
+		return KeyDisplayName;
+	}
+	else if (NodeAttributeKey == UE::Interchange::FAnimSequenceNodeStaticData::GetAnimatedAttributeStepCurveNamesKey())
+	{
+		KeyDisplayName = TEXT("Animated Attribute Step Curves Count");
+		return KeyDisplayName;
+	}
+	else if (KeyDisplayName.StartsWith(UE::Interchange::FAnimSequenceNodeStaticData::GetAnimatedAttributeStepCurveNamesKey().ToString()))
+	{
+		KeyDisplayName = TEXT("Animated Attribute Step Curves Index ");
+		const FString IndexKey = UE::Interchange::TArrayAttributeHelper<FString>::IndexKey();
+		int32 IndexPosition = KeyDisplayName.Find(IndexKey) + IndexKey.Len();
+		if (IndexPosition < KeyDisplayName.Len())
+		{
+			KeyDisplayName += KeyDisplayName.RightChop(IndexPosition);
+		}
+		return KeyDisplayName;
+	}
+	else if (NodeAttributeKey == UE::Interchange::FAnimSequenceNodeStaticData::GetAnimatedMaterialCurveSuffixesKey())
+	{
+		KeyDisplayName = TEXT("Animated Material Curve suffixes Count");
+		return KeyDisplayName;
+	}
+	else if (KeyDisplayName.StartsWith(UE::Interchange::FAnimSequenceNodeStaticData::GetAnimatedMaterialCurveSuffixesKey().ToString()))
+	{
+		KeyDisplayName = TEXT("Animated Material Curve Suffixes Index ");
+		const FString IndexKey = UE::Interchange::TArrayAttributeHelper<FString>::IndexKey();
+		int32 IndexPosition = KeyDisplayName.Find(IndexKey) + IndexKey.Len();
+		if (IndexPosition < KeyDisplayName.Len())
+		{
+			KeyDisplayName += KeyDisplayName.RightChop(IndexPosition);
+		}
+		return KeyDisplayName;
+	}
+	else if (NodeAttributeKey == UE::Interchange::FAnimSequenceNodeStaticData::GetSceneNodeAnimationPayloadKeyUidMapKey())
+	{
+		KeyDisplayName = TEXT("Scene Node Animation Payload Key Uid Count");
+		return KeyDisplayName;
+	}
+	else if (KeyDisplayName.StartsWith(UE::Interchange::FAnimSequenceNodeStaticData::GetSceneNodeAnimationPayloadKeyUidMapKey().ToString()))
+	{
+		FString MapKeyIndex = UE::Interchange::FAnimSequenceNodeStaticData::GetSceneNodeAnimationPayloadKeyUidMapKey().ToString() + TEXT("_KeyIndex_");
+		FString MapKey = UE::Interchange::FAnimSequenceNodeStaticData::GetSceneNodeAnimationPayloadKeyUidMapKey().ToString() + TEXT("_Key_");
+
+		int32 IndexPosition = 0;
+		if (KeyDisplayName.StartsWith(MapKeyIndex))
+		{
+			KeyDisplayName = TEXT("Scene Node Animation Payload Key Uid Key ");
+			IndexPosition = KeyDisplayName.Find(MapKeyIndex) + MapKeyIndex.Len();
+		}
+		else if (KeyDisplayName.StartsWith(MapKey))
+		{
+			KeyDisplayName = TEXT("Scene Node Animation Payload Key Uid Value ");
+			IndexPosition = KeyDisplayName.Find(MapKey) + MapKey.Len();
+		}
+
+		if (IndexPosition < KeyDisplayName.Len())
+		{
+			KeyDisplayName += KeyDisplayName.RightChop(IndexPosition);
+		}
+		return KeyDisplayName;
+	}
+	else if (NodeAttributeKey == UE::Interchange::FAnimSequenceNodeStaticData::GetSceneNodeAnimationPayloadKeyTypeMapKey())
+	{
+		KeyDisplayName = TEXT("Scene Node Animation Payload Key Type Count");
+		return KeyDisplayName;
+	}
+	else if (KeyDisplayName.StartsWith(UE::Interchange::FAnimSequenceNodeStaticData::GetSceneNodeAnimationPayloadKeyTypeMapKey().ToString()))
+	{
+		FString MapKeyIndex = UE::Interchange::FAnimSequenceNodeStaticData::GetSceneNodeAnimationPayloadKeyTypeMapKey().ToString() + TEXT("_KeyIndex_");
+		FString MapKey = UE::Interchange::FAnimSequenceNodeStaticData::GetSceneNodeAnimationPayloadKeyTypeMapKey().ToString() + TEXT("_Key_");
+
+		int32 IndexPosition = 0;
+		if (KeyDisplayName.StartsWith(MapKeyIndex))
+		{
+			KeyDisplayName = TEXT("Scene Node Animation Payload Key Type Key ");
+			IndexPosition = KeyDisplayName.Find(MapKeyIndex) + MapKeyIndex.Len();
+		}
+		else if (KeyDisplayName.StartsWith(MapKey))
+		{
+			KeyDisplayName = TEXT("Scene Node Animation Payload Key Type Value ");
+			IndexPosition = KeyDisplayName.Find(MapKey) + MapKey.Len();
+		}
+
+		if (IndexPosition < KeyDisplayName.Len())
+		{
+			KeyDisplayName += KeyDisplayName.RightChop(IndexPosition);
+		}
+		return KeyDisplayName;
+	}
+	else if (NodeAttributeKey == UE::Interchange::FAnimSequenceNodeStaticData::GetMorphTargetNodePayloadKeyUidMapKey())
+	{
+		KeyDisplayName = TEXT("Morph Target Node Payload Key Uid Count");
+		return KeyDisplayName;
+	}
+	else if (KeyDisplayName.StartsWith(UE::Interchange::FAnimSequenceNodeStaticData::GetMorphTargetNodePayloadKeyUidMapKey().ToString()))
+	{
+		FString MapKeyIndex = UE::Interchange::FAnimSequenceNodeStaticData::GetMorphTargetNodePayloadKeyUidMapKey().ToString() + TEXT("_KeyIndex_");
+		FString MapKey = UE::Interchange::FAnimSequenceNodeStaticData::GetMorphTargetNodePayloadKeyUidMapKey().ToString() + TEXT("_Key_");
+
+		int32 IndexPosition = 0;
+		if (KeyDisplayName.StartsWith(MapKeyIndex))
+		{
+			KeyDisplayName = TEXT("Morph Target Node Payload Key Uid Key ");
+			IndexPosition = KeyDisplayName.Find(MapKeyIndex) + MapKeyIndex.Len();
+		}
+		else if (KeyDisplayName.StartsWith(MapKey))
+		{
+			KeyDisplayName = TEXT("Morph Target Node Payload Key Uid Value ");
+			IndexPosition = KeyDisplayName.Find(MapKey) + MapKey.Len();
+		}
+
+		if (IndexPosition < KeyDisplayName.Len())
+		{
+			KeyDisplayName += KeyDisplayName.RightChop(IndexPosition);
+		}
+		return KeyDisplayName;
+	}
+	else if (NodeAttributeKey == UE::Interchange::FAnimSequenceNodeStaticData::GetMorphTargetNodePayloadKeyTypeMapKey())
+	{
+		KeyDisplayName = TEXT("Morph Target Node Payload Key Type Count");
+		return KeyDisplayName;
+	}
+	else if (KeyDisplayName.StartsWith(UE::Interchange::FAnimSequenceNodeStaticData::GetMorphTargetNodePayloadKeyTypeMapKey().ToString()))
+	{
+		FString MapKeyIndex = UE::Interchange::FAnimSequenceNodeStaticData::GetMorphTargetNodePayloadKeyTypeMapKey().ToString() + TEXT("_KeyIndex_");
+		FString MapKey = UE::Interchange::FAnimSequenceNodeStaticData::GetMorphTargetNodePayloadKeyTypeMapKey().ToString() + TEXT("_Key_");
+
+		int32 IndexPosition = 0;
+		if (KeyDisplayName.StartsWith(MapKeyIndex))
+		{
+			KeyDisplayName = TEXT("Morph Target Node Payload Key Type Key ");
+			IndexPosition = KeyDisplayName.Find(MapKeyIndex) + MapKeyIndex.Len();
+		}
+		else if (KeyDisplayName.StartsWith(MapKey))
+		{
+			KeyDisplayName = TEXT("Morph Target Node Payload Key Type Value ");
+			IndexPosition = KeyDisplayName.Find(MapKey) + MapKey.Len();
+		}
+
+		if (IndexPosition < KeyDisplayName.Len())
+		{
+			KeyDisplayName += KeyDisplayName.RightChop(IndexPosition);
+		}
+		return KeyDisplayName;
+	}
 	else
 	{
 		KeyDisplayName = Super::GetKeyDisplayName(NodeAttributeKey);
@@ -116,6 +294,56 @@ FString UInterchangeAnimSequenceFactoryNode::GetKeyDisplayName(const UE::Interch
 	return KeyDisplayName;
 }
 
+FString UInterchangeAnimSequenceFactoryNode::GetAttributeCategory(const UE::Interchange::FAttributeKey& NodeAttributeKey) const
+{
+	const FString NodeAttributeKeyString = NodeAttributeKey.ToString();
+	if (NodeAttributeKeyString.StartsWith(UE::Interchange::FAnimSequenceNodeStaticData::GetAnimatedMorphTargetDependenciesKey().ToString()))
+	{
+		return TEXT("Animated Morph Target Dependencies");
+	}
+	else if (NodeAttributeKeyString.StartsWith(UE::Interchange::FAnimSequenceNodeStaticData::GetAnimatedAttributeCurveNamesKey().ToString()))
+	{
+		return TEXT("Animated Attribute Curve Names");
+	}
+	else if (NodeAttributeKeyString.StartsWith(UE::Interchange::FAnimSequenceNodeStaticData::GetAnimatedAttributeStepCurveNamesKey().ToString()))
+	{
+		return TEXT("Animated Attribute Step Curve Names");
+	}
+	else if (NodeAttributeKeyString.StartsWith(UE::Interchange::FAnimSequenceNodeStaticData::GetAnimatedMaterialCurveSuffixesKey().ToString()))
+	{
+		return TEXT("Animated Material Curve Suffixes");
+	}
+	else if (NodeAttributeKeyString.StartsWith(UE::Interchange::FAnimSequenceNodeStaticData::GetSceneNodeAnimationPayloadKeyUidMapKey().ToString())
+		|| NodeAttributeKeyString.StartsWith(UE::Interchange::FAnimSequenceNodeStaticData::GetSceneNodeAnimationPayloadKeyTypeMapKey().ToString()))
+	{
+		return TEXT("Scene Node Animation Payload Keys");
+	}
+	else if (NodeAttributeKeyString.StartsWith(UE::Interchange::FAnimSequenceNodeStaticData::GetMorphTargetNodePayloadKeyUidMapKey().ToString())
+		|| NodeAttributeKeyString.StartsWith(UE::Interchange::FAnimSequenceNodeStaticData::GetMorphTargetNodePayloadKeyTypeMapKey().ToString()))
+	{
+		return TEXT("Morph Target Node Payload Keys");
+	}
+	else
+	{
+		return Super::GetAttributeCategory(NodeAttributeKey);
+	}
+}
+
+bool UInterchangeAnimSequenceFactoryNode::ShouldHideAttribute(const UE::Interchange::FAttributeKey& NodeAttributeKey) const
+{
+	if (UserInterfaceContext == EInterchangeNodeUserInterfaceContext::Preview)
+	{
+		if (NodeAttributeKey == Macro_CustomSkeletonFactoryNodeUidKey)
+		{
+			return true;
+		}
+	}
+
+	return Super::ShouldHideAttribute(NodeAttributeKey);
+}
+
+#endif //WITH_EDITOR
+
 UClass* UInterchangeAnimSequenceFactoryNode::GetObjectClass() const
 {
 	return UAnimSequence::StaticClass();
diff --git a/Engine/Plugins/Interchange/Runtime/Source/FactoryNodes/Private/InterchangeMeshFactoryNode.cpp b/Engine/Plugins/Interchange/Runtime/Source/FactoryNodes/Private/InterchangeMeshFactoryNode.cpp
--- a/Engine/Plugins/Interchange/Runtime/Source/FactoryNodes/Private/InterchangeMeshFactoryNode.cpp
+++ b/Engine/Plugins/Interchange/Runtime/Source/FactoryNodes/Private/InterchangeMeshFactoryNode.cpp
@@ -84,6 +84,8 @@ UInterchangeMeshFactoryNode::UInterchangeMeshFactoryNode()
 	SlotMaterialDependencies.Initialize(Attributes.ToSharedRef(), UE::Interchange::FMeshFactoryNodeStaticData::GetSlotMaterialDependencyBaseKey().ToString());
 }
 
+#if WITH_EDITOR
+
 FString UInterchangeMeshFactoryNode::GetKeyDisplayName(const UE::Interchange::FAttributeKey& NodeAttributeKey) const
 {
 	FString KeyDisplayName = NodeAttributeKey.ToString();
@@ -104,9 +106,65 @@ FString UInterchangeMeshFactoryNode::GetKeyDisplayName(const UE::Interchange::FA
 		}
 		return KeyDisplayName;
 	}
+	else if (NodeAttributeKey == UE::Interchange::FMeshFactoryNodeStaticData::GetSlotMaterialDependencyBaseKey())
+	{
+		KeyDisplayName = TEXT("Material Dependencies Count");
+		return KeyDisplayName;
+	}
+	else if (NodeAttributeKeyString.StartsWith(UE::Interchange::FMeshFactoryNodeStaticData::GetSlotMaterialDependencyBaseKey().ToString()))
+	{
+		FString MapKeyIndex = UE::Interchange::FMeshFactoryNodeStaticData::GetSlotMaterialDependencyBaseKey().ToString() + TEXT("_KeyIndex_");
+		FString MapKey = UE::Interchange::FMeshFactoryNodeStaticData::GetSlotMaterialDependencyBaseKey().ToString() + TEXT("_Key_");
+
+		int32 IndexPosition = 0;
+		if (NodeAttributeKeyString.StartsWith(MapKeyIndex))
+		{
+			KeyDisplayName = TEXT("Material Dependencies Key ");
+			IndexPosition = NodeAttributeKeyString.Find(MapKeyIndex) + MapKeyIndex.Len();
+		}
+		else if (NodeAttributeKeyString.StartsWith(MapKey))
+		{
+			KeyDisplayName = TEXT("Material Dependencies Value ");
+			IndexPosition = NodeAttributeKeyString.Find(MapKey) + MapKey.Len();
+		}
+		
+		if (IndexPosition < NodeAttributeKeyString.Len())
+		{
+			KeyDisplayName += NodeAttributeKeyString.RightChop(IndexPosition);
+		}
+		return KeyDisplayName;
+	}
 	return Super::GetKeyDisplayName(NodeAttributeKey);
 }
 
+FString UInterchangeMeshFactoryNode::GetAttributeCategory(const UE::Interchange::FAttributeKey& NodeAttributeKey) const
+{
+	if (NodeAttributeKey.ToString().StartsWith(UE::Interchange::FMeshFactoryNodeStaticData::GetLodDependenciesBaseKey().ToString()))
+	{
+		return TEXT("LOD Dependencies");
+	}
+	else if (NodeAttributeKey.ToString().StartsWith(UE::Interchange::FMeshFactoryNodeStaticData::GetSlotMaterialDependencyBaseKey().ToString()))
+	{
+		return TEXT("Materials Dependencies");
+	}
+	else
+	{
+		return Super::GetAttributeCategory(NodeAttributeKey);
+	}
+}
+
+bool UInterchangeMeshFactoryNode::ShouldHideAttribute(const UE::Interchange::FAttributeKey& NodeAttributeKey) const
+{
+	const FString NodeAttributeString = NodeAttributeKey.ToString();
+	if (NodeAttributeString.StartsWith(UE::Interchange::FMeshFactoryNodeStaticData::GetLodDependenciesBaseKey().ToString()))
+	{
+		return true;
+	}
+	return Super::ShouldHideAttribute(NodeAttributeKey);
+}
+
+#endif //WITH_EDITOR
+
 int32 UInterchangeMeshFactoryNode::GetLodDataCount() const
 {
 	return LodDependencies.GetCount();
diff --git a/Engine/Plugins/Interchange/Runtime/Source/FactoryNodes/Private/InterchangeSkeletalMeshLodDataNode.cpp b/Engine/Plugins/Interchange/Runtime/Source/FactoryNodes/Private/InterchangeSkeletalMeshLodDataNode.cpp
--- a/Engine/Plugins/Interchange/Runtime/Source/FactoryNodes/Private/InterchangeSkeletalMeshLodDataNode.cpp
+++ b/Engine/Plugins/Interchange/Runtime/Source/FactoryNodes/Private/InterchangeSkeletalMeshLodDataNode.cpp
@@ -28,6 +28,8 @@ FString UInterchangeSkeletalMeshLodDataNode::GetTypeName() const
 	return TypeName;
 }
 
+#if WITH_EDITOR
+
 FString UInterchangeSkeletalMeshLodDataNode::GetKeyDisplayName(const UE::Interchange::FAttributeKey& NodeAttributeKey) const
 {
 	FString KeyDisplayName = NodeAttributeKey.ToString();
@@ -65,6 +67,8 @@ FString UInterchangeSkeletalMeshLodDataNode::GetAttributeCategory(const UE::Inte
 	return Super::GetAttributeCategory(NodeAttributeKey);
 }
 
+#endif //WITH_EDITOR
+
 bool UInterchangeSkeletalMeshLodDataNode::GetCustomSkeletonUid(FString& AttributeValue) const
 {
 	IMPLEMENT_NODE_ATTRIBUTE_GETTER(SkeletonUid, FString);
diff --git a/Engine/Plugins/Interchange/Runtime/Source/FactoryNodes/Private/InterchangeStaticMeshFactoryNode.cpp b/Engine/Plugins/Interchange/Runtime/Source/FactoryNodes/Private/InterchangeStaticMeshFactoryNode.cpp
--- a/Engine/Plugins/Interchange/Runtime/Source/FactoryNodes/Private/InterchangeStaticMeshFactoryNode.cpp
+++ b/Engine/Plugins/Interchange/Runtime/Source/FactoryNodes/Private/InterchangeStaticMeshFactoryNode.cpp
@@ -90,6 +90,48 @@ UClass* UInterchangeStaticMeshFactoryNode::GetObjectClass() const
 #endif
 }
 
+#if WITH_EDITOR
+
+FString UInterchangeStaticMeshFactoryNode::GetKeyDisplayName(const UE::Interchange::FAttributeKey& NodeAttributeKey) const
+{
+	FString KeyDisplayName = NodeAttributeKey.ToString();
+	const FString OriginalKeyName = KeyDisplayName;
+	if (NodeAttributeKey == UE::Interchange::FStaticMeshNodeStaticData::GetSocketUidsBaseKey())
+	{
+		KeyDisplayName = TEXT("Socket Count");
+	}
+	else if (NodeAttributeKey == UE::Interchange::FStaticMeshNodeStaticData::GetSocketUidsBaseKey())
+	{
+		KeyDisplayName = TEXT("Socket Index ");
+		const FString IndexKey = UE::Interchange::TArrayAttributeHelper<FString>::IndexKey();
+		int32 IndexPosition = OriginalKeyName.Find(IndexKey) + IndexKey.Len();
+		if (IndexPosition < OriginalKeyName.Len())
+		{
+			KeyDisplayName += OriginalKeyName.RightChop(IndexPosition);
+		}
+	}
+	else
+	{
+		KeyDisplayName = Super::GetKeyDisplayName(NodeAttributeKey);
+	}
+
+	return KeyDisplayName;
+}
+
+FString UInterchangeStaticMeshFactoryNode::GetAttributeCategory(const UE::Interchange::FAttributeKey& NodeAttributeKey) const
+{
+	if (NodeAttributeKey.ToString().StartsWith(UE::Interchange::FStaticMeshNodeStaticData::GetSocketUidsBaseKey().ToString()))
+	{
+		return TEXT("Sockets");
+	}
+	else
+	{
+		return Super::GetAttributeCategory(NodeAttributeKey);
+	}
+}
+
+#endif //WITH_EDITOR
+
 bool UInterchangeStaticMeshFactoryNode::GetCustomBuildNanite(bool& AttributeValue) const
 {
 	IMPLEMENT_NODE_ATTRIBUTE_GETTER(BuildNanite, bool)
diff --git a/Engine/Plugins/Interchange/Runtime/Source/FactoryNodes/Private/InterchangeStaticMeshLodDataNode.cpp b/Engine/Plugins/Interchange/Runtime/Source/FactoryNodes/Private/InterchangeStaticMeshLodDataNode.cpp
--- a/Engine/Plugins/Interchange/Runtime/Source/FactoryNodes/Private/InterchangeStaticMeshLodDataNode.cpp
+++ b/Engine/Plugins/Interchange/Runtime/Source/FactoryNodes/Private/InterchangeStaticMeshLodDataNode.cpp
@@ -60,6 +60,8 @@ FString UInterchangeStaticMeshLodDataNode::GetTypeName() const
 	return TypeName;
 }
 
+#if WITH_EDITOR
+
 FString UInterchangeStaticMeshLodDataNode::GetKeyDisplayName(const UE::Interchange::FAttributeKey& NodeAttributeKey) const
 {
 	FString KeyDisplayName = NodeAttributeKey.ToString();
@@ -80,6 +82,80 @@ FString UInterchangeStaticMeshLodDataNode::GetKeyDisplayName(const UE::Interchan
 		}
 		return KeyDisplayName;
 	}
+	else if (NodeAttributeKey == UE::Interchange::FStaticMeshNodeLodDataStaticData::GetBoxCollisionMeshUidsBaseKey())
+	{
+		KeyDisplayName = TEXT("Box collision count");
+		return KeyDisplayName;
+	}
+	else if (NodeAttributeKeyString.StartsWith(UE::Interchange::FStaticMeshNodeLodDataStaticData::GetBoxCollisionMeshUidsBaseKey().ToString()))
+	{
+		KeyDisplayName = TEXT("Box collision index ");
+		const FString IndexKey = UE::Interchange::TArrayAttributeHelper<FString>::IndexKey();
+		int32 IndexPosition = NodeAttributeKeyString.Find(IndexKey) + IndexKey.Len();
+		if (IndexPosition < NodeAttributeKeyString.Len())
+		{
+			KeyDisplayName += NodeAttributeKeyString.RightChop(IndexPosition);
+		}
+		return KeyDisplayName;
+	}
+	else if (NodeAttributeKey == UE::Interchange::FStaticMeshNodeLodDataStaticData::GetCapsuleCollisionMeshUidsBaseKey())
+	{
+		KeyDisplayName = TEXT("Capsule collision count");
+		return KeyDisplayName;
+	}
+	else if (NodeAttributeKeyString.StartsWith(UE::Interchange::FStaticMeshNodeLodDataStaticData::GetCapsuleCollisionMeshUidsBaseKey().ToString()))
+	{
+		KeyDisplayName = TEXT("Capsule collision index ");
+		const FString IndexKey = UE::Interchange::TArrayAttributeHelper<FString>::IndexKey();
+		int32 IndexPosition = NodeAttributeKeyString.Find(IndexKey) + IndexKey.Len();
+		if (IndexPosition < NodeAttributeKeyString.Len())
+		{
+			KeyDisplayName += NodeAttributeKeyString.RightChop(IndexPosition);
+		}
+		return KeyDisplayName;
+	}
+	else if (NodeAttributeKey == UE::Interchange::FStaticMeshNodeLodDataStaticData::GetSphereCollisionMeshUidsBaseKey())
+	{
+		KeyDisplayName = TEXT("Sphere collision count");
+		return KeyDisplayName;
+	}
+	else if (NodeAttributeKeyString.StartsWith(UE::Interchange::FStaticMeshNodeLodDataStaticData::GetSphereCollisionMeshUidsBaseKey().ToString()))
+	{
+		KeyDisplayName = TEXT("Sphere collision index ");
+		const FString IndexKey = UE::Interchange::TArrayAttributeHelper<FString>::IndexKey();
+		int32 IndexPosition = NodeAttributeKeyString.Find(IndexKey) + IndexKey.Len();
+		if (IndexPosition < NodeAttributeKeyString.Len())
+		{
+			KeyDisplayName += NodeAttributeKeyString.RightChop(IndexPosition);
+		}
+		return KeyDisplayName;
+	}
+	else if (NodeAttributeKey == UE::Interchange::FStaticMeshNodeLodDataStaticData::GetConvexCollisionMeshUidsBaseKey())
+	{
+		KeyDisplayName = TEXT("Convex collision count");
+		return KeyDisplayName;
+	}
+	else if (NodeAttributeKeyString.StartsWith(UE::Interchange::FStaticMeshNodeLodDataStaticData::GetConvexCollisionMeshUidsBaseKey().ToString()))
+	{
+		KeyDisplayName = TEXT("Convex collision index ");
+		const FString IndexKey = UE::Interchange::TArrayAttributeHelper<FString>::IndexKey();
+		int32 IndexPosition = NodeAttributeKeyString.Find(IndexKey) + IndexKey.Len();
+		if (IndexPosition < NodeAttributeKeyString.Len())
+		{
+			KeyDisplayName += NodeAttributeKeyString.RightChop(IndexPosition);
+		}
+		return KeyDisplayName;
+	}
+	else if (NodeAttributeKey == Macro_CustomOneConvexHullPerUCXKey)
+	{
+		KeyDisplayName = TEXT("One Convex Hull Per UCX");
+		return KeyDisplayName;
+	}
+	else if (NodeAttributeKey == Macro_CustomImportCollisionKey)
+	{
+		KeyDisplayName = TEXT("Import Collision");
+		return KeyDisplayName;
+	}
 
 	return Super::GetKeyDisplayName(NodeAttributeKey);
 }
@@ -90,9 +166,27 @@ FString UInterchangeStaticMeshLodDataNode::GetAttributeCategory(const UE::Interc
 	{
 		return FString(TEXT("Meshes"));
 	}
+	else if (NodeAttributeKey.ToString().StartsWith(UE::Interchange::FStaticMeshNodeLodDataStaticData::GetBoxCollisionMeshUidsBaseKey().ToString()))
+	{
+		return FString(TEXT("Box Collisions"));
+	}
+	else if (NodeAttributeKey.ToString().StartsWith(UE::Interchange::FStaticMeshNodeLodDataStaticData::GetCapsuleCollisionMeshUidsBaseKey().ToString()))
+	{
+		return FString(TEXT("Capsule Collisions"));
+	}
+	else if (NodeAttributeKey.ToString().StartsWith(UE::Interchange::FStaticMeshNodeLodDataStaticData::GetSphereCollisionMeshUidsBaseKey().ToString()))
+	{
+		return FString(TEXT("Sphere Collisions"));
+	}
+	else if (NodeAttributeKey.ToString().StartsWith(UE::Interchange::FStaticMeshNodeLodDataStaticData::GetConvexCollisionMeshUidsBaseKey().ToString()))
+	{
+		return FString(TEXT("Convex Collisions"));
+	}
 	return Super::GetAttributeCategory(NodeAttributeKey);
 }
 
+#endif //WITH_EDITOR
+
 int32 UInterchangeStaticMeshLodDataNode::GetMeshUidsCount() const
 {
 	return MeshUids.GetCount();
diff --git a/Engine/Plugins/Interchange/Runtime/Source/FactoryNodes/Public/InterchangeAnimSequenceFactoryNode.h b/Engine/Plugins/Interchange/Runtime/Source/FactoryNodes/Public/InterchangeAnimSequenceFactoryNode.h
--- a/Engine/Plugins/Interchange/Runtime/Source/FactoryNodes/Public/InterchangeAnimSequenceFactoryNode.h
+++ b/Engine/Plugins/Interchange/Runtime/Source/FactoryNodes/Public/InterchangeAnimSequenceFactoryNode.h
@@ -68,7 +68,11 @@ public:
 	 */
 	virtual FString GetTypeName() const override;
 
+#if WITH_EDITOR
 	virtual FString GetKeyDisplayName(const UE::Interchange::FAttributeKey& NodeAttributeKey) const override;
+	virtual FString GetAttributeCategory(const UE::Interchange::FAttributeKey& NodeAttributeKey) const override;
+	virtual bool ShouldHideAttribute(const UE::Interchange::FAttributeKey& NodeAttributeKey) const override;
+#endif //WITH_EDITOR
 
 	/** Get the class this node want to create */
 	UFUNCTION(BlueprintCallable, Category = "Interchange | Node | AnimSequence")
diff --git a/Engine/Plugins/Interchange/Runtime/Source/FactoryNodes/Public/InterchangeMeshFactoryNode.h b/Engine/Plugins/Interchange/Runtime/Source/FactoryNodes/Public/InterchangeMeshFactoryNode.h
--- a/Engine/Plugins/Interchange/Runtime/Source/FactoryNodes/Public/InterchangeMeshFactoryNode.h
+++ b/Engine/Plugins/Interchange/Runtime/Source/FactoryNodes/Public/InterchangeMeshFactoryNode.h
@@ -43,7 +43,11 @@ public:
 		}
 	}
 
+#if WITH_EDITOR
 	virtual FString GetKeyDisplayName(const UE::Interchange::FAttributeKey& NodeAttributeKey) const override;
+	virtual FString GetAttributeCategory(const UE::Interchange::FAttributeKey& NodeAttributeKey) const override;
+	virtual bool ShouldHideAttribute(const UE::Interchange::FAttributeKey& NodeAttributeKey) const override;
+#endif //WITH_EDITOR
 
 public:
 	/** Return The number of LOD this static mesh has.*/
diff --git a/Engine/Plugins/Interchange/Runtime/Source/FactoryNodes/Public/InterchangePhysicsAssetFactoryNode.h b/Engine/Plugins/Interchange/Runtime/Source/FactoryNodes/Public/InterchangePhysicsAssetFactoryNode.h
--- a/Engine/Plugins/Interchange/Runtime/Source/FactoryNodes/Public/InterchangePhysicsAssetFactoryNode.h
+++ b/Engine/Plugins/Interchange/Runtime/Source/FactoryNodes/Public/InterchangePhysicsAssetFactoryNode.h
@@ -68,6 +68,19 @@ public:
 #endif
 	}
 
+#if WITH_EDITOR
+
+	virtual bool ShouldHideAttribute(const UE::Interchange::FAttributeKey& NodeAttributeKey) const override
+	{
+		if (NodeAttributeKey == Macro_CustomSkeletalMeshUidKey)
+		{
+			return true;
+		}
+
+		return Super::ShouldHideAttribute(NodeAttributeKey);
+	}
+#endif
+
 	virtual FGuid GetHash() const override
 	{
 		return Attributes->GetStorageHash();
diff --git a/Engine/Plugins/Interchange/Runtime/Source/FactoryNodes/Public/InterchangeSkeletalMeshLodDataNode.h b/Engine/Plugins/Interchange/Runtime/Source/FactoryNodes/Public/InterchangeSkeletalMeshLodDataNode.h
--- a/Engine/Plugins/Interchange/Runtime/Source/FactoryNodes/Public/InterchangeSkeletalMeshLodDataNode.h
+++ b/Engine/Plugins/Interchange/Runtime/Source/FactoryNodes/Public/InterchangeSkeletalMeshLodDataNode.h
@@ -34,9 +34,11 @@ public:
 	 */
 	virtual FString GetTypeName() const override;
 
+#if WITH_EDITOR
 	virtual FString GetKeyDisplayName(const UE::Interchange::FAttributeKey& NodeAttributeKey) const override;
 
 	virtual FString GetAttributeCategory(const UE::Interchange::FAttributeKey& NodeAttributeKey) const override;
+#endif //WITH_EDITOR
 
 public:
 	/** Query the LOD skeletal mesh factory skeleton reference. Return false if the attribute was not set.*/
diff --git a/Engine/Plugins/Interchange/Runtime/Source/FactoryNodes/Public/InterchangeStaticMeshFactoryNode.h b/Engine/Plugins/Interchange/Runtime/Source/FactoryNodes/Public/InterchangeStaticMeshFactoryNode.h
--- a/Engine/Plugins/Interchange/Runtime/Source/FactoryNodes/Public/InterchangeStaticMeshFactoryNode.h
+++ b/Engine/Plugins/Interchange/Runtime/Source/FactoryNodes/Public/InterchangeStaticMeshFactoryNode.h
@@ -51,6 +51,13 @@ public:
 	UFUNCTION(BlueprintCallable, Category = "Interchange | Node | StaticMesh")
 	virtual class UClass* GetObjectClass() const override;
 
+#if WITH_EDITOR
+
+	virtual FString GetKeyDisplayName(const UE::Interchange::FAttributeKey& NodeAttributeKey) const override;
+	virtual FString GetAttributeCategory(const UE::Interchange::FAttributeKey& NodeAttributeKey) const override;
+
+#endif //WITH_EDITOR
+
 public:
 	/** Get weather the static mesh factory should set the nanite build settings. Return false if the attribute was not set.*/
 	UFUNCTION(BlueprintCallable, Category = "Interchange | Node | StaticMesh")
diff --git a/Engine/Plugins/Interchange/Runtime/Source/FactoryNodes/Public/InterchangeStaticMeshLodDataNode.h b/Engine/Plugins/Interchange/Runtime/Source/FactoryNodes/Public/InterchangeStaticMeshLodDataNode.h
--- a/Engine/Plugins/Interchange/Runtime/Source/FactoryNodes/Public/InterchangeStaticMeshLodDataNode.h
+++ b/Engine/Plugins/Interchange/Runtime/Source/FactoryNodes/Public/InterchangeStaticMeshLodDataNode.h
@@ -38,9 +38,10 @@ public:
 	 */
 	virtual FString GetTypeName() const override;
 
+#if WITH_EDITOR
 	virtual FString GetKeyDisplayName(const UE::Interchange::FAttributeKey& NodeAttributeKey) const override;
-
 	virtual FString GetAttributeCategory(const UE::Interchange::FAttributeKey& NodeAttributeKey) const override;
+#endif //WITH_EDITOR
 
 public:
 	/* Mesh Uids: It can be either a scene or a mesh node uid. If its a scene it mean we want the mesh factory to bake the geo payload with the global transform of the scene node. */
diff --git a/Engine/Plugins/Interchange/Runtime/Source/Nodes/Private/InterchangeMeshNode.cpp b/Engine/Plugins/Interchange/Runtime/Source/Nodes/Private/InterchangeMeshNode.cpp
--- a/Engine/Plugins/Interchange/Runtime/Source/Nodes/Private/InterchangeMeshNode.cpp
+++ b/Engine/Plugins/Interchange/Runtime/Source/Nodes/Private/InterchangeMeshNode.cpp
@@ -74,6 +74,8 @@ UInterchangeMeshNode::UInterchangeMeshNode()
 	SlotMaterialDependencies.Initialize(Attributes.ToSharedRef(), UE::Interchange::FMeshNodeStaticData::GetSlotMaterialDependenciesKey().ToString());
 }
 
+#if WITH_EDITOR
+
 FString UInterchangeMeshNode::GetKeyDisplayName(const UE::Interchange::FAttributeKey& NodeAttributeKey) const
 {
 	FString KeyDisplayName = NodeAttributeKey.ToString();
@@ -180,6 +182,8 @@ FString UInterchangeMeshNode::GetAttributeCategory(const UE::Interchange::FAttri
 	return Super::GetAttributeCategory(NodeAttributeKey);
 }
 
+#endif //WITH_EDITOR
+
 FString UInterchangeMeshNode::GetTypeName() const
 {
 	const FString TypeName = TEXT("MeshNode");
diff --git a/Engine/Plugins/Interchange/Runtime/Source/Nodes/Private/InterchangeSceneNode.cpp b/Engine/Plugins/Interchange/Runtime/Source/Nodes/Private/InterchangeSceneNode.cpp
--- a/Engine/Plugins/Interchange/Runtime/Source/Nodes/Private/InterchangeSceneNode.cpp
+++ b/Engine/Plugins/Interchange/Runtime/Source/Nodes/Private/InterchangeSceneNode.cpp
@@ -71,6 +71,8 @@ FString UInterchangeSceneNode::GetTypeName() const
 	return TypeName;
 }
 
+#if WITH_EDITOR
+
 FString UInterchangeSceneNode::GetKeyDisplayName(const UE::Interchange::FAttributeKey& NodeAttributeKey) const
 {
 	FString KeyDisplayName = NodeAttributeKey.ToString();
@@ -136,6 +138,8 @@ FString UInterchangeSceneNode::GetAttributeCategory(const UE::Interchange::FAttr
 	return Super::GetAttributeCategory(NodeAttributeKey);
 }
 
+#endif //WITH_EDITOR
+
 FName UInterchangeSceneNode::GetIconName() const
 {
 	FString SpecializedType;
diff --git a/Engine/Plugins/Interchange/Runtime/Source/Nodes/Public/InterchangeMeshNode.h b/Engine/Plugins/Interchange/Runtime/Source/Nodes/Public/InterchangeMeshNode.h
--- a/Engine/Plugins/Interchange/Runtime/Source/Nodes/Public/InterchangeMeshNode.h
+++ b/Engine/Plugins/Interchange/Runtime/Source/Nodes/Public/InterchangeMeshNode.h
@@ -70,9 +70,11 @@ class INTERCHANGENODES_API UInterchangeMeshNode : public UInterchangeBaseNode
 public:
 	UInterchangeMeshNode();
 
+#if WITH_EDITOR
 	virtual FString GetKeyDisplayName(const UE::Interchange::FAttributeKey& NodeAttributeKey) const override;
 
 	virtual FString GetAttributeCategory(const UE::Interchange::FAttributeKey& NodeAttributeKey) const override;
+#endif //WITH_EDITOR
 
 	/**
 	 * Return the node type name of the class, we use this when reporting error
diff --git a/Engine/Plugins/Interchange/Runtime/Source/Nodes/Public/InterchangeSceneNode.h b/Engine/Plugins/Interchange/Runtime/Source/Nodes/Public/InterchangeSceneNode.h
--- a/Engine/Plugins/Interchange/Runtime/Source/Nodes/Public/InterchangeSceneNode.h
+++ b/Engine/Plugins/Interchange/Runtime/Source/Nodes/Public/InterchangeSceneNode.h
@@ -59,10 +59,11 @@ public:
 	 */
 	virtual FString GetTypeName() const override;
 
+#if WITH_EDITOR
 	virtual FString GetKeyDisplayName(const UE::Interchange::FAttributeKey& NodeAttributeKey) const override;
 
 	virtual FString GetAttributeCategory(const UE::Interchange::FAttributeKey& NodeAttributeKey) const override;
-
+#endif //WITH_EDITOR
 	/**
 	 * Icon name are simply create by adding "InterchangeIcon_" in front of the specialized type. If there is no special type the function will return NAME_None which will use the default icon.
 	 */
diff --git a/Engine/Plugins/Interchange/Runtime/Source/Nodes/Public/InterchangeTextureNode.h b/Engine/Plugins/Interchange/Runtime/Source/Nodes/Public/InterchangeTextureNode.h
--- a/Engine/Plugins/Interchange/Runtime/Source/Nodes/Public/InterchangeTextureNode.h
+++ b/Engine/Plugins/Interchange/Runtime/Source/Nodes/Public/InterchangeTextureNode.h
@@ -63,7 +63,7 @@ public:
 		const FString TypeName = TEXT("TextureNode");
 		return TypeName;
 	}
-
+#if WITH_EDITOR
 	virtual FString GetKeyDisplayName(const UE::Interchange::FAttributeKey& NodeAttributeKey) const override
 	{
 		FString KeyDisplayName = NodeAttributeKey.ToString();
@@ -73,6 +73,7 @@ public:
 		}
 		return Super::GetKeyDisplayName(NodeAttributeKey);
 	}
+#endif //WITH_EDITOR
 
 	virtual FGuid GetHash() const override
 	{
diff --git a/Engine/Source/Runtime/Interchange/Core/Private/Nodes/InterchangeBaseNode.cpp b/Engine/Source/Runtime/Interchange/Core/Private/Nodes/InterchangeBaseNode.cpp
--- a/Engine/Source/Runtime/Interchange/Core/Private/Nodes/InterchangeBaseNode.cpp
+++ b/Engine/Source/Runtime/Interchange/Core/Private/Nodes/InterchangeBaseNode.cpp
@@ -107,6 +107,8 @@ FName UInterchangeBaseNode::GetIconName() const
 	return NAME_None;
 }
 
+#if WITH_EDITOR
+
 FString UInterchangeBaseNode::GetKeyDisplayName(const UE::Interchange::FAttributeKey& NodeAttributeKey) const
 {
 	FString KeyDisplayName = NodeAttributeKey.ToString();
@@ -159,15 +161,52 @@ FString UInterchangeBaseNode::GetKeyDisplayName(const UE::Interchange::FAttribut
 
 bool UInterchangeBaseNode::ShouldHideAttribute(const UE::Interchange::FAttributeKey& NodeAttributeKey) const
 {
+	//If context is preview we hide internal data
+	if (UserInterfaceContext == EInterchangeNodeUserInterfaceContext::Preview)
+	{
+		if (NodeAttributeKey == UE::Interchange::FBaseNodeStaticData::ParentIDKey())
+		{
+			return true;
+		}
+		else if (NodeAttributeKey == UE::Interchange::FBaseNodeStaticData::IsEnabledKey())
+		{
+			return true;
+		}
+		else if (NodeAttributeKey == UE::Interchange::FBaseNodeStaticData::UniqueIDKey())
+		{
+			return true;
+		}
+		else if (NodeAttributeKey == UE::Interchange::FBaseNodeStaticData::NodeContainerTypeKey())
+		{
+			return true;
+		}
+		else if (NodeAttributeKey == UE::Interchange::FBaseNodeStaticData::TargetAssetIDsKey())
+		{
+			return true;
+		}
+		else if (NodeAttributeKey.ToString().StartsWith(UE::Interchange::FBaseNodeStaticData::TargetAssetIDsKey().ToString()))
+		{
+			return true;
+		}
+	}
+
+	//Show anything else
 	return false;
 }
 
 FString UInterchangeBaseNode::GetAttributeCategory(const UE::Interchange::FAttributeKey& NodeAttributeKey) const
 {
 	FString CategoryName = TEXT("Attributes");
+	const FString NodeAttributeString = NodeAttributeKey.ToString();
+	if (NodeAttributeString.StartsWith(TEXT("UserDefined_")) && NodeAttributeString.EndsWith(TEXT("_Value")))
+	{
+		CategoryName = TEXT("Custom Attributes");
+	}
 	return CategoryName;
 }
 
+#endif //WITH_EDITOR
+
 bool UInterchangeBaseNode::HasAttribute(const UE::Interchange::FAttributeKey& NodeAttributeKey) const
 {
 	return Attributes->ContainAttribute(NodeAttributeKey);
diff --git a/Engine/Source/Runtime/Interchange/Core/Private/Nodes/InterchangeFactoryBaseNode.cpp b/Engine/Source/Runtime/Interchange/Core/Private/Nodes/InterchangeFactoryBaseNode.cpp
--- a/Engine/Source/Runtime/Interchange/Core/Private/Nodes/InterchangeFactoryBaseNode.cpp
+++ b/Engine/Source/Runtime/Interchange/Core/Private/Nodes/InterchangeFactoryBaseNode.cpp
@@ -50,7 +50,7 @@ UInterchangeFactoryBaseNode::UInterchangeFactoryBaseNode()
 	RegisterAttribute<uint8>(UE::Interchange::FFactoryBaseNodeStaticData::ReimportStrategyFlagsKey(), static_cast<uint8>(EReimportStrategyFlags::ApplyNoProperties));
 }
 
-
+#if WITH_EDITOR
 FString UInterchangeFactoryBaseNode::GetKeyDisplayName(const UE::Interchange::FAttributeKey& NodeAttributeKey) const
 {
 	FString KeyDisplayName = NodeAttributeKey.ToString();
@@ -77,6 +77,18 @@ FString UInterchangeFactoryBaseNode::GetKeyDisplayName(const UE::Interchange::FA
 	{
 		KeyDisplayName = TEXT("Re-Import Strategy");
 	}
+	else if (NodeAttributeKey == UE::Interchange::FFactoryBaseNodeStaticData::SkipNodeImportKey())
+	{
+		KeyDisplayName = TEXT("Skip Node Import");
+	}
+	else if (NodeAttributeKey == UE::Interchange::FFactoryBaseNodeStaticData::ForceNodeReimportKey())
+	{
+		KeyDisplayName = TEXT("Force Node Reimport");
+	}
+	else if (NodeAttributeKey == Macro_CustomSubPathKey)
+	{
+		KeyDisplayName = TEXT("Asset Sub Path");
+	}
 	else
 	{
 		KeyDisplayName = Super::GetKeyDisplayName(NodeAttributeKey);
@@ -98,6 +110,40 @@ FString UInterchangeFactoryBaseNode::GetAttributeCategory(const UE::Interchange:
 	}
 }
 
+bool UInterchangeFactoryBaseNode::ShouldHideAttribute(const UE::Interchange::FAttributeKey& NodeAttributeKey) const
+{
+	if (UserInterfaceContext == EInterchangeNodeUserInterfaceContext::Preview)
+	{
+		const FString KeyDisplayName = NodeAttributeKey.ToString();
+		if (KeyDisplayName.Equals(UE::Interchange::FFactoryBaseNodeStaticData::FactoryDependenciesBaseKey()))
+		{
+			return true;
+		}
+		else if (KeyDisplayName.StartsWith(UE::Interchange::FFactoryBaseNodeStaticData::FactoryDependenciesBaseKey()))
+		{
+			return true;
+		}
+		else if (NodeAttributeKey == UE::Interchange::FFactoryBaseNodeStaticData::ReimportStrategyFlagsKey())
+		{
+			return true;
+		}
+		else if (NodeAttributeKey == UE::Interchange::FFactoryBaseNodeStaticData::SkipNodeImportKey())
+		{
+			return true;
+		}
+		else if (NodeAttributeKey == UE::Interchange::FFactoryBaseNodeStaticData::ForceNodeReimportKey())
+		{
+			return true;
+		}
+		else if (NodeAttributeKey == Macro_CustomReferenceObjectKey)
+		{
+			return true;
+		}
+	}
+
+	return Super::ShouldHideAttribute(NodeAttributeKey);
+}
+#endif //WITH_EDITOR
 
 UClass* UInterchangeFactoryBaseNode::GetObjectClass() const
 {
diff --git a/Engine/Source/Runtime/Interchange/Core/Public/Nodes/InterchangeBaseNode.h b/Engine/Source/Runtime/Interchange/Core/Public/Nodes/InterchangeBaseNode.h
--- a/Engine/Source/Runtime/Interchange/Core/Public/Nodes/InterchangeBaseNode.h
+++ b/Engine/Source/Runtime/Interchange/Core/Public/Nodes/InterchangeBaseNode.h
@@ -173,6 +173,13 @@ enum class EInterchangeNodeContainerType : uint8
 	FactoryData
 };
 
+UENUM(BlueprintType)
+enum class EInterchangeNodeUserInterfaceContext : uint8
+{
+	None,
+	Preview //When we want to preview the node data, a preview is normaly read only and should not show internal data
+};
+
 /**
  * This struct is used to store and retrieve key value attributes. The attributes are store in a generic FAttributeStorage which serialize the value in a TArray64<uint8>
  * See UE::Interchange::EAttributeTypes to know the supported template types
@@ -207,6 +214,7 @@ public:
 	 */
 	INTERCHANGECORE_API virtual FName GetIconName() const;
 
+#if WITH_EDITOR
 	/**
 	 * UI that inspect node attribute call this to give a readable name to attribute key
 	 */
@@ -222,6 +230,16 @@ public:
 	 */
 	INTERCHANGECORE_API virtual FString GetAttributeCategory(const UE::Interchange::FAttributeKey& NodeAttributeKey) const;
 
+#endif //WITH_EDITOR
+
+#if WITH_EDITORONLY_DATA
+	/**
+	 * Temporary property set by UI to have the context for ShouldHideAttribute. We use this because property editor cannot add this custom context when it call CustomizeDetails
+	 */
+	UPROPERTY(Transient, DuplicateTransient)
+	EInterchangeNodeUserInterfaceContext UserInterfaceContext = EInterchangeNodeUserInterfaceContext::None;
+#endif
+
 	/**
 	 * Add an attribute to the node
 	 * @param NodeAttributeKey - The key of the attribute
diff --git a/Engine/Source/Runtime/Interchange/Core/Public/Nodes/InterchangeFactoryBaseNode.h b/Engine/Source/Runtime/Interchange/Core/Public/Nodes/InterchangeFactoryBaseNode.h
--- a/Engine/Source/Runtime/Interchange/Core/Public/Nodes/InterchangeFactoryBaseNode.h
+++ b/Engine/Source/Runtime/Interchange/Core/Public/Nodes/InterchangeFactoryBaseNode.h
@@ -164,8 +164,11 @@ class UInterchangeFactoryBaseNode : public UInterchangeBaseNode
 public:
 	INTERCHANGECORE_API UInterchangeFactoryBaseNode();
 
+#if WITH_EDITOR
 	INTERCHANGECORE_API virtual FString GetKeyDisplayName(const UE::Interchange::FAttributeKey& NodeAttributeKey) const override;
 	INTERCHANGECORE_API virtual FString GetAttributeCategory(const UE::Interchange::FAttributeKey& NodeAttributeKey) const override;
+	INTERCHANGECORE_API virtual bool ShouldHideAttribute(const UE::Interchange::FAttributeKey& NodeAttributeKey) const override;
+#endif
 
 	/**
 	 * Return the reimport strategy flags.
