[Interchange - glTF] Fixing tiling method in glTF Material Functions. Using Texture Source.
[Interchange - glTF] Make refraction dependent on KHR_materials_volume.
As KHR_materials_volume is not supported, disabling Refraction setups.
Also set the Base Interchange-glTF Materials as 'Used with Skeletal Mesh'.

#rnx

#virtualized

diff --git a/Engine/Plugins/Interchange/Runtime/Source/Import/Private/Gltf/InterchangeGLTFMaterial.cpp b/Engine/Plugins/Interchange/Runtime/Source/Import/Private/Gltf/InterchangeGLTFMaterial.cpp
--- a/Engine/Plugins/Interchange/Runtime/Source/Import/Private/Gltf/InterchangeGLTFMaterial.cpp
+++ b/Engine/Plugins/Interchange/Runtime/Source/Import/Private/Gltf/InterchangeGLTFMaterial.cpp
@@ -138,6 +138,29 @@ namespace UE::Interchange::GLTFMaterials
 		return UInterchangeShaderPortsAPI::MakeInputValueKey(Name);
 	}
 
+	FVector4f GetTilingMethod(const GLTF::FSampler& Sampler)
+	{
+		float TilingU = 0;
+		switch (Sampler.WrapS)
+		{
+			case GLTF::FSampler::EWrap::Repeat:         TilingU = 0; break;
+			case GLTF::FSampler::EWrap::ClampToEdge:    TilingU = 1; break;
+			case GLTF::FSampler::EWrap::MirroredRepeat: TilingU = 2; break;
+			default:                                    TilingU = 0; break;;
+		}
+
+		float TilingV = 0;
+		switch (Sampler.WrapT)
+		{
+			case GLTF::FSampler::EWrap::Repeat:         TilingV = 0; break;
+			case GLTF::FSampler::EWrap::ClampToEdge:    TilingV = 1; break;
+			case GLTF::FSampler::EWrap::MirroredRepeat: TilingV = 2; break;
+			default:                                    TilingV = 0; break;;
+		}
+
+		return FVector4f(TilingU, TilingV, 0, 0);
+	}
+
 	struct FGLTFMaterialProcessor
 	{
 		UInterchangeBaseNodeContainer& NodeContainer;
@@ -346,6 +369,9 @@ namespace UE::Interchange::GLTFMaterials
 				// [3...4) -> UV3
 				// else    -> UV0 (defaults to 0)
 				SetScalar(Name + TEXT("_TexCoord"), TextureMap.TexCoord, 0.f);
+
+				//Set the TilingMethod:
+				SetVec4(Name + TEXT("_TilingMethod"), GetTilingMethod(Textures[TextureMap.TextureIndex].Sampler), FVector4f(0, 0, 0, 0));
 			}
 
 			if (TextureMap.bHasTextureTransform)
diff --git a/Engine/Plugins/Interchange/Runtime/Source/Import/Public/Gltf/InterchangeGLTFMaterial.h b/Engine/Plugins/Interchange/Runtime/Source/Import/Public/Gltf/InterchangeGLTFMaterial.h
--- a/Engine/Plugins/Interchange/Runtime/Source/Import/Public/Gltf/InterchangeGLTFMaterial.h
+++ b/Engine/Plugins/Interchange/Runtime/Source/Import/Public/Gltf/InterchangeGLTFMaterial.h
@@ -22,7 +22,8 @@ const FString InterchangeGltfMaterialAttributeIdentifier = TEXT("Gltf_MI_Attribu
 												 const FString MapName##Texture_Scale_X = TEXT(INTERCHANGE_GLTF_STRINGIFY(MapName ## Texture_Scale_X)); \
 												 const FString MapName##Texture_Scale_Y = TEXT(INTERCHANGE_GLTF_STRINGIFY(MapName ## Texture_Scale_Y)); \
 												 const FString MapName##Texture_Rotation = TEXT(INTERCHANGE_GLTF_STRINGIFY(MapName ## Texture_Rotation)); \
-												 const FString MapName##Texture_TexCoord = TEXT(INTERCHANGE_GLTF_STRINGIFY(MapName ## Texture_TexCoord));
+												 const FString MapName##Texture_TexCoord = TEXT(INTERCHANGE_GLTF_STRINGIFY(MapName ## Texture_TexCoord)); \
+												 const FString MapName##Texture_TilingMethod = TEXT(INTERCHANGE_GLTF_STRINGIFY(MapName ## Texture_TilingMethod));
 
 namespace UE::Interchange::GLTFMaterials
 {
@@ -183,8 +184,7 @@ namespace UE::Interchange::GLTFMaterials
 					TEXT("EmissiveColor"),
 					TEXT("Opacity"),
 					TEXT("Normal"),
-					TEXT("Occlusion"),
-					TEXT("Refraction")})},
+					TEXT("Occlusion")})},
 
 		{EShadingModel::SPECULARGLOSSINESS,
 			TPair<FString, TArray<FString>>(
diff --git a/Engine/Plugins/Interchange/Runtime/Source/Pipelines/Private/InterchangeglTFPipeline.cpp b/Engine/Plugins/Interchange/Runtime/Source/Pipelines/Private/InterchangeglTFPipeline.cpp
--- a/Engine/Plugins/Interchange/Runtime/Source/Pipelines/Private/InterchangeglTFPipeline.cpp
+++ b/Engine/Plugins/Interchange/Runtime/Source/Pipelines/Private/InterchangeglTFPipeline.cpp
@@ -89,12 +89,18 @@ TSet<FString> UGLTFPipelineSettings::GenerateExpectedParametersList(const FStrin
 
 	TSet<FString> ExpectedParameters;
 
+	auto AddTextureAndRelated = [&ExpectedParameters](const FString& TextureName)
+		{
+			ExpectedParameters.Add(TextureName);
+			ExpectedParameters.Add(TextureName + TEXT("_OffsetScale"));
+			ExpectedParameters.Add(TextureName + TEXT("_Rotation"));
+			ExpectedParameters.Add(TextureName + TEXT("_TexCoord"));
+			ExpectedParameters.Add(TextureName + TEXT("_TilingMethod"));
+		};
+
 	if (Identifier.Contains(TEXT("_Unlit")))
 	{
-		ExpectedParameters.Add(Inputs::BaseColorTexture);
-		ExpectedParameters.Add(Inputs::BaseColorTexture_OffsetScale);
-		ExpectedParameters.Add(Inputs::BaseColorTexture_Rotation);
-		ExpectedParameters.Add(Inputs::BaseColorTexture_TexCoord);
+		AddTextureAndRelated(Inputs::BaseColorTexture);
 		ExpectedParameters.Add(Inputs::BaseColorFactor);
 
 		return ExpectedParameters;
@@ -102,36 +108,24 @@ TSet<FString> UGLTFPipelineSettings::GenerateExpectedParametersList(const FStrin
 
 	//Generic ones:
 	{
-		ExpectedParameters.Add(Inputs::NormalTexture);
-		ExpectedParameters.Add(Inputs::NormalTexture_OffsetScale);
-		ExpectedParameters.Add(Inputs::NormalTexture_Rotation);
-		ExpectedParameters.Add(Inputs::NormalTexture_TexCoord);
+		AddTextureAndRelated(Inputs::NormalTexture);
 		ExpectedParameters.Add(Inputs::NormalScale);
 
 		if (!Identifier.Contains(TEXT("Transmission")))
 		{
-			ExpectedParameters.Add(Inputs::EmissiveTexture);
-			ExpectedParameters.Add(Inputs::EmissiveTexture_OffsetScale);
-			ExpectedParameters.Add(Inputs::EmissiveTexture_Rotation);
-			ExpectedParameters.Add(Inputs::EmissiveTexture_TexCoord);
+			AddTextureAndRelated(Inputs::EmissiveTexture);
 			ExpectedParameters.Add(Inputs::EmissiveFactor);
 			ExpectedParameters.Add(Inputs::EmissiveStrength);
 		}
 
-		ExpectedParameters.Add(Inputs::OcclusionTexture);
-		ExpectedParameters.Add(Inputs::OcclusionTexture_OffsetScale);
-		ExpectedParameters.Add(Inputs::OcclusionTexture_Rotation);
-		ExpectedParameters.Add(Inputs::OcclusionTexture_TexCoord);
+		AddTextureAndRelated(Inputs::OcclusionTexture);
 		ExpectedParameters.Add(Inputs::OcclusionStrength);
 
 		if (!Identifier.Contains(TEXT("SpecularGlossiness")))
 		{
 			ExpectedParameters.Add(Inputs::IOR);
 
-			ExpectedParameters.Add(Inputs::SpecularTexture);
-			ExpectedParameters.Add(Inputs::SpecularTexture_OffsetScale);
-			ExpectedParameters.Add(Inputs::SpecularTexture_Rotation);
-			ExpectedParameters.Add(Inputs::SpecularTexture_TexCoord);
+			AddTextureAndRelated(Inputs::SpecularTexture);
 			ExpectedParameters.Add(Inputs::SpecularFactor);
 		}
 	}
@@ -142,73 +136,43 @@ TSet<FString> UGLTFPipelineSettings::GenerateExpectedParametersList(const FStrin
 	{
 		//MetalRoughness Specific:
 
-		ExpectedParameters.Add(Inputs::BaseColorTexture);
-		ExpectedParameters.Add(Inputs::BaseColorTexture_OffsetScale);
-		ExpectedParameters.Add(Inputs::BaseColorTexture_Rotation);
-		ExpectedParameters.Add(Inputs::BaseColorTexture_TexCoord);
+		AddTextureAndRelated(Inputs::BaseColorTexture);
 		ExpectedParameters.Add(Inputs::BaseColorFactor);
 
-		ExpectedParameters.Add(Inputs::MetallicRoughnessTexture);
-		ExpectedParameters.Add(Inputs::MetallicRoughnessTexture_OffsetScale);
-		ExpectedParameters.Add(Inputs::MetallicRoughnessTexture_Rotation);
-		ExpectedParameters.Add(Inputs::MetallicRoughnessTexture_TexCoord);
+		AddTextureAndRelated(Inputs::MetallicRoughnessTexture);
 		ExpectedParameters.Add(Inputs::MetallicFactor);
 		ExpectedParameters.Add(Inputs::RoughnessFactor);
 	}
 	else if (Identifier.Contains(TEXT("ClearCoat")))
 	{
-		ExpectedParameters.Add(Inputs::ClearCoatTexture);
-		ExpectedParameters.Add(Inputs::ClearCoatTexture_OffsetScale);
-		ExpectedParameters.Add(Inputs::ClearCoatTexture_Rotation);
-		ExpectedParameters.Add(Inputs::ClearCoatTexture_TexCoord);
+		AddTextureAndRelated(Inputs::ClearCoatTexture);
 		ExpectedParameters.Add(Inputs::ClearCoatFactor);
 
-		ExpectedParameters.Add(Inputs::ClearCoatRoughnessTexture);
-		ExpectedParameters.Add(Inputs::ClearCoatRoughnessTexture_OffsetScale);
-		ExpectedParameters.Add(Inputs::ClearCoatRoughnessTexture_Rotation);
-		ExpectedParameters.Add(Inputs::ClearCoatRoughnessTexture_TexCoord);
+		AddTextureAndRelated(Inputs::ClearCoatRoughnessTexture);
 		ExpectedParameters.Add(Inputs::ClearCoatRoughnessFactor);
 
-		ExpectedParameters.Add(Inputs::ClearCoatNormalTexture);
-		ExpectedParameters.Add(Inputs::ClearCoatNormalTexture_OffsetScale);
-		ExpectedParameters.Add(Inputs::ClearCoatNormalTexture_Rotation);
-		ExpectedParameters.Add(Inputs::ClearCoatNormalTexture_TexCoord);
+		AddTextureAndRelated(Inputs::ClearCoatNormalTexture);
 		ExpectedParameters.Add(Inputs::ClearCoatNormalScale);
 	}
 	else if (Identifier.Contains(TEXT("Sheen")))
 	{
-		ExpectedParameters.Add(Inputs::SheenColorTexture);
-		ExpectedParameters.Add(Inputs::SheenColorTexture_OffsetScale);
-		ExpectedParameters.Add(Inputs::SheenColorTexture_Rotation);
-		ExpectedParameters.Add(Inputs::SheenColorTexture_TexCoord);
+		AddTextureAndRelated(Inputs::SheenColorTexture);
 		ExpectedParameters.Add(Inputs::SheenColorFactor);
 
-		ExpectedParameters.Add(Inputs::SheenRoughnessTexture);
-		ExpectedParameters.Add(Inputs::SheenRoughnessTexture_OffsetScale);
-		ExpectedParameters.Add(Inputs::SheenRoughnessTexture_Rotation);
-		ExpectedParameters.Add(Inputs::SheenRoughnessTexture_TexCoord);
+		AddTextureAndRelated(Inputs::SheenRoughnessTexture);
 		ExpectedParameters.Add(Inputs::SheenRoughnessFactor);
 	}
 	else if (Identifier.Contains(TEXT("Transmission")))
 	{
-		ExpectedParameters.Add(Inputs::TransmissionTexture);
-		ExpectedParameters.Add(Inputs::TransmissionTexture_OffsetScale);
-		ExpectedParameters.Add(Inputs::TransmissionTexture_Rotation);
-		ExpectedParameters.Add(Inputs::TransmissionTexture_TexCoord);
+		AddTextureAndRelated(Inputs::TransmissionTexture);
 		ExpectedParameters.Add(Inputs::TransmissionFactor);
 	}
 	else if (Identifier.Contains(TEXT("SpecularGlossiness")))
 	{
-		ExpectedParameters.Add(Inputs::DiffuseTexture);
-		ExpectedParameters.Add(Inputs::DiffuseTexture_OffsetScale);
-		ExpectedParameters.Add(Inputs::DiffuseTexture_Rotation);
-		ExpectedParameters.Add(Inputs::DiffuseTexture_TexCoord);
+		AddTextureAndRelated(Inputs::DiffuseTexture);
 		ExpectedParameters.Add(Inputs::DiffuseFactor);
 
-		ExpectedParameters.Add(Inputs::SpecularGlossinessTexture);
-		ExpectedParameters.Add(Inputs::SpecularGlossinessTexture_OffsetScale);
-		ExpectedParameters.Add(Inputs::SpecularGlossinessTexture_Rotation);
-		ExpectedParameters.Add(Inputs::SpecularGlossinessTexture_TexCoord);
+		AddTextureAndRelated(Inputs::SpecularGlossinessTexture);
 		ExpectedParameters.Add(Inputs::SpecFactor);
 		ExpectedParameters.Add(Inputs::GlossinessFactor);
 	}
