[Interchange - glTF] Negative scale support fix for asset import.

#rnx

diff --git a/Engine/Plugins/Interchange/Runtime/Source/Parsers/GLTFCore/Private/GLTFMeshFactory.cpp b/Engine/Plugins/Interchange/Runtime/Source/Parsers/GLTFCore/Private/GLTFMeshFactory.cpp
--- a/Engine/Plugins/Interchange/Runtime/Source/Parsers/GLTFCore/Private/GLTFMeshFactory.cpp
+++ b/Engine/Plugins/Interchange/Runtime/Source/Parsers/GLTFCore/Private/GLTFMeshFactory.cpp
@@ -456,6 +456,9 @@ namespace GLTF
 		FMatrix TotalMatrixForNormal;
 		TotalMatrixForNormal = TotalMatrix.Inverse();
 		TotalMatrixForNormal = TotalMatrixForNormal.GetTransposed();
+		TotalMatrixForNormal.RemoveScaling(); //based on accepted solution in StaticMeshOperations::AppendMeshDescriptions / StaticMeshOperations::ApplyTransform
+
+		const bool bIsMirrored = TotalMatrix.Determinant() < 0.f;
 
 		auto TransformPosition = [](const FMatrix& Matrix, FVector3f& Position)
 		{
@@ -580,69 +583,143 @@ namespace GLTF
 		}
 
 		bool bHasDegenerateTriangles = false;
+		
 		// Now add all vertexInstances
 		FVertexID         CornerVertexIDs[3];
-		for (uint32 TriangleIndex = 0; TriangleIndex < TriCount; ++TriangleIndex)
+		uint32 TriangleCornerMaxIndex = (TriCount * 3) - 1;
+
+		//Note: below code is mostly duplicated, 
+		// difference is the IndiceIndex calculations in the 2 inner for loops.
+		// Duplicated for performance reasons.
+		if (bIsMirrored)
 		{
-			for (int32 Corner = 0; Corner < 3; ++Corner)
+			for (uint32 TriangleIndex = 0; TriangleIndex < TriCount; ++TriangleIndex)
 			{
-				const int32 VertexIndex = Indices[TriangleIndex * 3 + Corner];
-				const FVertexID VertexID = PositionIndexToVertexIdPerPrim[PrimitiveIndex][VertexIndex];
-				CornerVertexIDs[Corner] = VertexID;
-			}
+				for (uint32 Corner = 0; Corner < 3; ++Corner)
+				{
+					const int32 IndiceIndex = Indices[TriangleCornerMaxIndex - (TriangleIndex * 3 + Corner)];
+					const FVertexID VertexID = PositionIndexToVertexIdPerPrim[PrimitiveIndex][IndiceIndex];
+					CornerVertexIDs[Corner] = VertexID;
+				}
 
-			// Check for degenerate triangles
-			const FVertexID& Vertex1 = CornerVertexIDs[0];
-			const FVertexID& Vertex2 = CornerVertexIDs[1];
-			const FVertexID& Vertex3 = CornerVertexIDs[2];
+				// Check for degenerate triangles
+				const FVertexID& Vertex1 = CornerVertexIDs[0];
+				const FVertexID& Vertex2 = CornerVertexIDs[1];
+				const FVertexID& Vertex3 = CornerVertexIDs[2];
 
-			if (Vertex1 == Vertex2 || Vertex2 == Vertex3 || Vertex1 == Vertex3)
-			{
-				bHasDegenerateTriangles = true;
-				continue; // Triangle is degenerate, skip it
-			}
+				if (Vertex1 == Vertex2 || Vertex2 == Vertex3 || Vertex1 == Vertex3)
+				{
+					bHasDegenerateTriangles = true;
+					continue; // Triangle is degenerate, skip it
+				}
 
-			for (int32 Corner = 0; Corner < 3; ++Corner)
-			{
-				const FVertexInstanceID VertexInstanceID = MeshDescription->CreateVertexInstance(CornerVertexIDs[Corner]);
-				CornerVertexInstanceIDs[Corner] = VertexInstanceID;
+				for (uint32 Corner = 0; Corner < 3; ++Corner)
+				{
+					const FVertexInstanceID VertexInstanceID = MeshDescription->CreateVertexInstance(CornerVertexIDs[Corner]);
+					CornerVertexInstanceIDs[Corner] = VertexInstanceID;
 
-				const uint32 IndiceIndex = TriangleIndex * 3 + Corner;
+					const uint32 IndiceIndex = TriangleCornerMaxIndex - (TriangleIndex * 3 + Corner);
 
-				VertexInstanceNormals[VertexInstanceID] = Normals[IndiceIndex];
+					VertexInstanceNormals[VertexInstanceID] = Normals[IndiceIndex];
 
-				if (!bIgnoreTangents && Tangents.Num() > 0)
+					if (!bIgnoreTangents && Tangents.Num() > 0)
+					{
+						VertexInstanceTangents[VertexInstanceID] = Tangents[IndiceIndex];
+						VertexInstanceBinormalSigns[VertexInstanceID] = Tangents[IndiceIndex].W;
+					}
+
+					for (int32 UVIndex = 0; UVIndex < NumUVs; ++UVIndex)
+					{
+						VertexInstanceUVs.Set(VertexInstanceID, UVIndex, (*UVs[UVIndex])[IndiceIndex]);
+					}
+
+					if (Colors.Num() > 0)
+					{
+						FVector4f Color = Colors[IndiceIndex];
+						// Because the engine will auto convert our vertex colors to sRGB, we need to counter the conversion here to ensure the end result is still linear space
+						UE::Geometry::LinearColors::SRGBToLinear(Color);
+						VertexInstanceColors[VertexInstanceID] = Color;
+					}
+				}
+
+				// Insert a polygon into the mesh
+				TArray<FEdgeID> NewEdgeIDs;
+				const FPolygonID NewPolygonID = MeshDescription->CreatePolygon(CurrentPolygonGroupID, CornerVertexInstanceIDs, &NewEdgeIDs);
+
+				for (const FEdgeID& NewEdgeID : NewEdgeIDs)
 				{
-					VertexInstanceTangents[VertexInstanceID] = Tangents[IndiceIndex];
-					VertexInstanceBinormalSigns[VertexInstanceID] = Tangents[IndiceIndex].W;
+					// Make all faces part of the same smoothing group, so Unreal will combine identical adjacent verts.
+					// (Is there a way to set auto-gen smoothing threshold? glTF spec says to generate flat normals if they're not specified.
+					//   We want to combine identical verts whether they're smooth neighbors or triangles belonging to the same flat polygon.)
+					EdgeHardnesses[NewEdgeID] = false;
+				}
+			}
+		}
+		else
+		{
+			for (uint32 TriangleIndex = 0; TriangleIndex < TriCount; ++TriangleIndex)
+			{
+				for (uint32 Corner = 0; Corner < 3; ++Corner)
+				{
+					const int32 VertexIndex = Indices[(TriangleIndex * 3 + Corner)];
+					const FVertexID VertexID = PositionIndexToVertexIdPerPrim[PrimitiveIndex][VertexIndex];
+					CornerVertexIDs[Corner] = VertexID;
 				}
 
-				for (int32 UVIndex = 0; UVIndex < NumUVs; ++UVIndex)
+				// Check for degenerate triangles
+				const FVertexID& Vertex1 = CornerVertexIDs[0];
+				const FVertexID& Vertex2 = CornerVertexIDs[1];
+				const FVertexID& Vertex3 = CornerVertexIDs[2];
+
+				if (Vertex1 == Vertex2 || Vertex2 == Vertex3 || Vertex1 == Vertex3)
 				{
-					VertexInstanceUVs.Set(VertexInstanceID, UVIndex, (*UVs[UVIndex])[IndiceIndex]);
+					bHasDegenerateTriangles = true;
+					continue; // Triangle is degenerate, skip it
 				}
 
-				if (Colors.Num() > 0)
+				for (uint32 Corner = 0; Corner < 3; ++Corner)
 				{
-					FVector4f Color = Colors[IndiceIndex];
-					// Because the engine will auto convert our vertex colors to sRGB, we need to counter the conversion here to ensure the end result is still linear space
-					UE::Geometry::LinearColors::SRGBToLinear(Color);
-					VertexInstanceColors[VertexInstanceID] = Color;
+					const FVertexInstanceID VertexInstanceID = MeshDescription->CreateVertexInstance(CornerVertexIDs[Corner]);
+					CornerVertexInstanceIDs[Corner] = VertexInstanceID;
+
+					const uint32 IndiceIndex = (TriangleIndex * 3 + Corner);
+
+					VertexInstanceNormals[VertexInstanceID] = Normals[IndiceIndex];
+
+					if (!bIgnoreTangents && Tangents.Num() > 0)
+					{
+						VertexInstanceTangents[VertexInstanceID] = Tangents[IndiceIndex];
+						VertexInstanceBinormalSigns[VertexInstanceID] = Tangents[IndiceIndex].W;
+					}
+
+					for (int32 UVIndex = 0; UVIndex < NumUVs; ++UVIndex)
+					{
+						VertexInstanceUVs.Set(VertexInstanceID, UVIndex, (*UVs[UVIndex])[IndiceIndex]);
+					}
+
+					if (Colors.Num() > 0)
+					{
+						FVector4f Color = Colors[IndiceIndex];
+						// Because the engine will auto convert our vertex colors to sRGB, we need to counter the conversion here to ensure the end result is still linear space
+						UE::Geometry::LinearColors::SRGBToLinear(Color);
+						VertexInstanceColors[VertexInstanceID] = Color;
+					}
 				}
-			}
 
-			// Insert a polygon into the mesh
-			TArray<FEdgeID> NewEdgeIDs;
-			const FPolygonID NewPolygonID = MeshDescription->CreatePolygon(CurrentPolygonGroupID, CornerVertexInstanceIDs, &NewEdgeIDs);
+				// Insert a polygon into the mesh
+				TArray<FEdgeID> NewEdgeIDs;
+				const FPolygonID NewPolygonID = MeshDescription->CreatePolygon(CurrentPolygonGroupID, CornerVertexInstanceIDs, &NewEdgeIDs);
 
-			for (const FEdgeID& NewEdgeID : NewEdgeIDs)
-			{
-				// Make all faces part of the same smoothing group, so Unreal will combine identical adjacent verts.
-				// (Is there a way to set auto-gen smoothing threshold? glTF spec says to generate flat normals if they're not specified.
-				//   We want to combine identical verts whether they're smooth neighbors or triangles belonging to the same flat polygon.)
-				EdgeHardnesses[NewEdgeID] = false;
+				for (const FEdgeID& NewEdgeID : NewEdgeIDs)
+				{
+					// Make all faces part of the same smoothing group, so Unreal will combine identical adjacent verts.
+					// (Is there a way to set auto-gen smoothing threshold? glTF spec says to generate flat normals if they're not specified.
+					//   We want to combine identical verts whether they're smooth neighbors or triangles belonging to the same flat polygon.)
+					EdgeHardnesses[NewEdgeID] = false;
+				}
 			}
 		}
+		
 		return bHasDegenerateTriangles;
 	}
 
