Interchange: Fix unreachable code
#rnx

diff --git a/Engine/Plugins/Interchange/Runtime/Source/FactoryNodes/Private/InterchangeCommonPipelineDataFactoryNode.cpp b/Engine/Plugins/Interchange/Runtime/Source/FactoryNodes/Private/InterchangeCommonPipelineDataFactoryNode.cpp
--- a/Engine/Plugins/Interchange/Runtime/Source/FactoryNodes/Private/InterchangeCommonPipelineDataFactoryNode.cpp
+++ b/Engine/Plugins/Interchange/Runtime/Source/FactoryNodes/Private/InterchangeCommonPipelineDataFactoryNode.cpp
@@ -42,9 +42,17 @@ bool UInterchangeCommonPipelineDataFactoryNode::GetCustomGlobalOffsetTransform(F
 UFUNCTION(BlueprintCallable, Category = "Interchange | Node | Common Pipeline Data")
 bool UInterchangeCommonPipelineDataFactoryNode::SetCustomGlobalOffsetTransform(const UInterchangeBaseNodeContainer* NodeContainer, const FTransform& AttributeValue)
 {
-	IMPLEMENT_NODE_ATTRIBUTE_SETTER_NODELEGATE(GlobalOffsetTransform, FTransform);
-	//Reset all scene node container cache
-	UInterchangeSceneNode::ResetAllGlobalTransformCaches(NodeContainer);
+	auto ImplementationFunction = [this, &NodeContainer, &AttributeValue]()
+		{
+			IMPLEMENT_NODE_ATTRIBUTE_SETTER_NODELEGATE(GlobalOffsetTransform, FTransform);
+		};
+	if (ImplementationFunction())
+	{
+		//Reset all scene node container cache
+		UInterchangeSceneNode::ResetAllGlobalTransformCaches(NodeContainer);
+		return true;
+	}
+	return false;
 }
 
 bool UInterchangeCommonPipelineDataFactoryNode::GetBakeMeshes(bool& AttributeValue) const
diff --git a/Engine/Plugins/Interchange/Runtime/Source/FactoryNodes/Private/InterchangeMaterialFactoryNode.cpp b/Engine/Plugins/Interchange/Runtime/Source/FactoryNodes/Private/InterchangeMaterialFactoryNode.cpp
--- a/Engine/Plugins/Interchange/Runtime/Source/FactoryNodes/Private/InterchangeMaterialFactoryNode.cpp
+++ b/Engine/Plugins/Interchange/Runtime/Source/FactoryNodes/Private/InterchangeMaterialFactoryNode.cpp
@@ -449,8 +449,17 @@ bool UInterchangeMaterialFunctionCallExpressionFactoryNode::GetCustomMaterialFun
 
 bool UInterchangeMaterialFunctionCallExpressionFactoryNode::SetCustomMaterialFunctionDependency(const FString& AttributeValue)
 {
-	IMPLEMENT_NODE_ATTRIBUTE_SETTER_NODELEGATE(MaterialFunctionDependency, FString);
-	AddFactoryDependencyUid(AttributeValue);
+	
+	auto ImplementationFunction = [this, &AttributeValue]()
+		{
+			IMPLEMENT_NODE_ATTRIBUTE_SETTER_NODELEGATE(MaterialFunctionDependency, FString);
+		};
+	if (ImplementationFunction())
+	{
+		AddFactoryDependencyUid(AttributeValue);
+		return true;
+	}
+	return false;
 }
 
 FString UInterchangeMaterialFunctionFactoryNode::GetTypeName() const
diff --git a/Engine/Plugins/Interchange/Runtime/Source/Nodes/Private/InterchangeAnimationTrackSetNode.cpp b/Engine/Plugins/Interchange/Runtime/Source/Nodes/Private/InterchangeAnimationTrackSetNode.cpp
--- a/Engine/Plugins/Interchange/Runtime/Source/Nodes/Private/InterchangeAnimationTrackSetNode.cpp
+++ b/Engine/Plugins/Interchange/Runtime/Source/Nodes/Private/InterchangeAnimationTrackSetNode.cpp
@@ -80,7 +80,6 @@ bool UInterchangeAnimationTrackSetNode::SetCustomFrameRate(const float& Attribut
 bool UInterchangeAnimationTrackBaseNode::GetCustomCompletionMode(int32& AttributeValue) const
 {
 	IMPLEMENT_NODE_ATTRIBUTE_GETTER(CompletionMode, int32);
-	AttributeValue = FMath::Clamp(AttributeValue, (int32)EInterchangeAimationCompletionMode::KeepState, (int32)EInterchangeAimationCompletionMode::ProjectDefault);
 }
 
 bool UInterchangeAnimationTrackBaseNode::SetCustomCompletionMode(const int32& AttributeValue)
