[Interchange - glTF] Added validation step for Bind Pose Transforms.
Falling back to no BindPoseTransforms on import in case of invalid BindPose found.

#rnx

diff --git a/Engine/Plugins/Interchange/Runtime/Source/Parsers/GLTFCore/Private/GLTF/GLTFReader.cpp b/Engine/Plugins/Interchange/Runtime/Source/Parsers/GLTFCore/Private/GLTF/GLTFReader.cpp
--- a/Engine/Plugins/Interchange/Runtime/Source/Parsers/GLTFCore/Private/GLTF/GLTFReader.cpp
+++ b/Engine/Plugins/Interchange/Runtime/Source/Parsers/GLTFCore/Private/GLTF/GLTFReader.cpp
@@ -1283,13 +1283,8 @@ namespace GLTF
 							//Scenario is that the a Skin is instantiated at the end of another skin
 							//(Prime example is the RecursiveSkeleton gltf sample file.)
 							ParentGlobalInverseBindTransform = ParentNode.SkinIndexToGlobalInverseBindTransform.begin().Value();
-
-							if (ParentNode.SkinIndexToGlobalInverseBindTransform.Num() > 1)
-							{
-								Messages.Emplace(EMessageSeverity::Error, FString::Printf(TEXT("The same Joint [%s] is used in multiple Skins, which is currently not supported."), *ParentNode.Name));
-							}
 						}
-						
+
 						FTransform LocalBindPose = CurrentNode.SkinIndexToGlobalInverseBindTransform[SkinIndex].Inverse() * ParentGlobalInverseBindTransform;
 
 						CurrentNode.SkinIndexToLocalBindPose.Add(SkinIndex, LocalBindPose);
@@ -1308,8 +1303,44 @@ namespace GLTF
 			}
 		}
 	}
+
 	void FFileReader::SetLocalBindPosesForJoints() const
 	{
+		//Validate generated SkinIndexToLocalBindPose values, before setting.
+		TArray<FString> OffendingJointsNames;
+
+		for (const FNode& CurrentNode : Asset->Nodes)
+		{
+			TMap<int, FTransform>::TRangedForConstIterator Iter(CurrentNode.SkinIndexToLocalBindPose.begin());
+			FTransform ToCompareAgainst = Iter ? Iter.Value() : FTransform();
+			++Iter;
+			for (; Iter; ++Iter)
+			{
+				if (!ToCompareAgainst.Equals(Iter.Value()))
+				{
+					OffendingJointsNames.Add(CurrentNode.Name);
+					break;
+				}
+			}
+		}
+		
+		if (OffendingJointsNames.Num() > 0)
+		{
+			FString OffendingJointsNamesString;
+			for (const FString& OffendingJointName : OffendingJointsNames)
+			{
+				if (OffendingJointsNamesString.Len() > 0)
+				{
+					OffendingJointsNamesString += TEXT(", ");
+				}
+				OffendingJointsNamesString += OffendingJointName;
+			}
+			
+			Messages.Emplace(EMessageSeverity::Warning, FString::Printf(TEXT("The same Joint(s) are used in multiple Skins with multiple different InverseBindMatrix values, which is not supported. Ignoring InverseBindMatrices for the entire Import. Offending Joints' Names: %s."), *OffendingJointsNamesString));
+
+			return;
+		}
+
 		for (size_t SkinIndex = 0; SkinIndex < Asset->Skins.Num(); SkinIndex++)
 		{
 			const FSkinInfo& Skin = Asset->Skins[SkinIndex];
@@ -1320,14 +1351,6 @@ namespace GLTF
 				{
 					FNode& CurrentNode = Asset->Nodes[Skin.Joints[JointCounter]];
 
-					if (CurrentNode.bHasLocalBindPose && CurrentNode.SkinIndexToLocalBindPose.Contains(SkinIndex))
-					{
-						if (!CurrentNode.LocalBindPose.Equals(CurrentNode.SkinIndexToLocalBindPose[SkinIndex]))
-						{
-							Messages.Emplace(EMessageSeverity::Error, FString::Printf(TEXT("The same Joint [%s] is used in multiple Skins with multiple different InverseBindMatrix values, which is currently not supported."), *CurrentNode.Name));
-						}
-					}
-
 					if (!CurrentNode.bHasLocalBindPose
 						&& CurrentNode.SkinIndexToLocalBindPose.Contains(SkinIndex))
 					{
