[Interchange - glTF] Added Analytics to track abnormal InverseBindMatrices.

#rnx

diff --git a/Engine/Plugins/Interchange/Runtime/Source/Import/Private/Gltf/InterchangeGltfTranslator.cpp b/Engine/Plugins/Interchange/Runtime/Source/Import/Private/Gltf/InterchangeGltfTranslator.cpp
--- a/Engine/Plugins/Interchange/Runtime/Source/Import/Private/Gltf/InterchangeGltfTranslator.cpp
+++ b/Engine/Plugins/Interchange/Runtime/Source/Import/Private/Gltf/InterchangeGltfTranslator.cpp
@@ -146,22 +146,20 @@ namespace UE::Interchange::Gltf::Private
 		NOTSUPPORTED_EXTENSION_FOUND
 	};
 	void SendAnalytics(const TranslationResult& TranslationResult,
-		const TArray<FString>& ExtensionsUsed = TArray<FString>(),
-		const TArray<FString>& ExtensionsRequired = TArray<FString>(),
-		GLTF::FMetadata Metadata = GLTF::FMetadata(),
+		const GLTF::FAsset& Asset = GLTF::FAsset(),
 		const FString& GLTFReaderLogMessage = "")
 	{
 		if (FEngineAnalytics::IsAvailable())
 		{
 			TMap<FString, FString> MetadataExtras;
-			for (GLTF::FMetadata::FExtraData ExtraData : Metadata.Extras)
+			for (GLTF::FMetadata::FExtraData ExtraData : Asset.Metadata.Extras)
 			{
 				MetadataExtras.Add(ExtraData.Name, ExtraData.Value);
 			}
 
 			TSet<FString> AllExtensions;
-			AllExtensions.Append(ExtensionsUsed);
-			AllExtensions.Append(ExtensionsRequired);
+			AllExtensions.Append(Asset.ExtensionsUsed);
+			AllExtensions.Append(Asset.ExtensionsRequired);
 
 			TArray<FString> ExtensionsSupported;
 			TArray<FString> ExtensionsUnsupported;
@@ -179,13 +177,14 @@ namespace UE::Interchange::Gltf::Private
 			}
 
 			TArray<FAnalyticsEventAttribute> GLTFAnalytics;
-			if (ExtensionsUsed.Num() > 0)					GLTFAnalytics.Add(FAnalyticsEventAttribute(TEXT("ExtensionsUsed"), ExtensionsUsed));
-			if (ExtensionsRequired.Num() > 0)				GLTFAnalytics.Add(FAnalyticsEventAttribute(TEXT("ExtensionsRequired"), ExtensionsRequired));
+			if (Asset.ExtensionsUsed.Num() > 0)				GLTFAnalytics.Add(FAnalyticsEventAttribute(TEXT("ExtensionsUsed"), Asset.ExtensionsUsed));
+			if (Asset.ExtensionsRequired.Num() > 0)			GLTFAnalytics.Add(FAnalyticsEventAttribute(TEXT("ExtensionsRequired"), Asset.ExtensionsRequired));
 			if (ExtensionsSupported.Num() > 0)				GLTFAnalytics.Add(FAnalyticsEventAttribute(TEXT("ExtensionsSupported"), ExtensionsSupported));
 			if (ExtensionsUnsupported.Num() > 0)			GLTFAnalytics.Add(FAnalyticsEventAttribute(TEXT("ExtensionsUnsupported"), ExtensionsUnsupported));
-			if (Metadata.GeneratorName.Len() > 0)			GLTFAnalytics.Add(FAnalyticsEventAttribute(TEXT("MetaData.GeneratorName"), Metadata.GeneratorName));
+			if (Asset.Metadata.GeneratorName.Len() > 0)		GLTFAnalytics.Add(FAnalyticsEventAttribute(TEXT("MetaData.GeneratorName"), Asset.Metadata.GeneratorName));
 			if (MetadataExtras.Num() > 0)					GLTFAnalytics.Add(FAnalyticsEventAttribute(TEXT("MetaData.Extras"), MetadataExtras));
-			/*Version is always set at this point.*/		GLTFAnalytics.Add(FAnalyticsEventAttribute(TEXT("MetaData.Version"), Metadata.Version));
+			/*Version is always set at this point.*/		GLTFAnalytics.Add(FAnalyticsEventAttribute(TEXT("MetaData.Version"), Asset.Metadata.Version));
+			if (Asset.HasAbnormalInverseBindMatrices)		GLTFAnalytics.Add(FAnalyticsEventAttribute(TEXT("HasAbnormalInverseBindMatrices"), true));
 
 			switch (TranslationResult)
 			{
@@ -429,7 +428,7 @@ bool UInterchangeGLTFTranslator::Translate( UInterchangeBaseNodeContainer& NodeC
 			ErrorResult->SourceAssetName = FileName;
 			ErrorResult->Text = FText::Format(LOCTEXT("GLTF::FFileReader::ReadFile Failed.", "LogMessage: {0}"), FText::FromString(LogMessage.Value));
 
-			SendAnalytics(TranslationResult::GLTFREADER_FAILED, GltfAsset.ExtensionsUsed,GltfAsset.ExtensionsRequired, GltfAsset.Metadata, LogMessage.Value);
+			SendAnalytics(TranslationResult::GLTFREADER_FAILED, GltfAsset, LogMessage.Value);
 			return false;
 		}
 	}
@@ -453,7 +452,7 @@ bool UInterchangeGLTFTranslator::Translate( UInterchangeBaseNodeContainer& NodeC
 			LOCTEXT("UnsupportedRequiredExtensions", "Not All Required Extensions are supported. (Unsupported extensions: {0})"),
 			FText::FromString(NotSupportedRequiredExtensionsStringified));
 
-		SendAnalytics(TranslationResult::NOTSUPPORTED_EXTENSION_FOUND, GltfAsset.ExtensionsUsed, GltfAsset.ExtensionsRequired, GltfAsset.Metadata);
+		SendAnalytics(TranslationResult::NOTSUPPORTED_EXTENSION_FOUND, GltfAsset);
 		return false;
 	}
 
@@ -791,7 +790,7 @@ bool UInterchangeGLTFTranslator::Translate( UInterchangeBaseNodeContainer& NodeC
 		UE_LOG(LogInterchangeImport, Warning, TEXT("GLTF Mesh Import Warning. Gltf Mesh Usage expectation is not met."));
 	}
 
-	SendAnalytics(TranslationResult::SUCCESSFULL, GltfAsset.ExtensionsUsed, GltfAsset.ExtensionsRequired, GltfAsset.Metadata);
+	SendAnalytics(TranslationResult::SUCCESSFULL, GltfAsset);
 	return true;
 }
 
diff --git a/Engine/Plugins/Interchange/Runtime/Source/Parsers/GLTFCore/Private/GLTF/GLTFReader.cpp b/Engine/Plugins/Interchange/Runtime/Source/Parsers/GLTFCore/Private/GLTF/GLTFReader.cpp
--- a/Engine/Plugins/Interchange/Runtime/Source/Parsers/GLTFCore/Private/GLTF/GLTFReader.cpp
+++ b/Engine/Plugins/Interchange/Runtime/Source/Parsers/GLTFCore/Private/GLTF/GLTFReader.cpp
@@ -1337,6 +1337,8 @@ namespace GLTF
 			}
 			
 			Messages.Emplace(EMessageSeverity::Warning, FString::Printf(TEXT("The same Joint(s) are used in multiple Skins with multiple different InverseBindMatrix values, which is not supported. Ignoring InverseBindMatrices for the entire Import. Offending Joints' Names: %s."), *OffendingJointsNamesString));
+			
+			Asset->HasAbnormalInverseBindMatrices = true;
 
 			return;
 		}
diff --git a/Engine/Plugins/Interchange/Runtime/Source/Parsers/GLTFCore/Public/GLTFAsset.h b/Engine/Plugins/Interchange/Runtime/Source/Parsers/GLTFCore/Public/GLTFAsset.h
--- a/Engine/Plugins/Interchange/Runtime/Source/Parsers/GLTFCore/Public/GLTFAsset.h
+++ b/Engine/Plugins/Interchange/Runtime/Source/Parsers/GLTFCore/Public/GLTFAsset.h
@@ -108,6 +108,7 @@ namespace GLTF
 		TSet<EExtension> ProcessedExtensions;
 		FMetadata        Metadata;
 
+		bool             HasAbnormalInverseBindMatrices = false; //True: in cases where at least 1 Joint node has multiple InverseBindMatrices that do not equal.
 
 		/**
 		 * Will clear the asset's buffers.
