- Use input name to name parameters in Material graph
- Create scalars instead of constants

diff --git a/Engine/Plugins/Interchange/Runtime/Source/Import/Private/MaterialX/MaterialXUtils/MaterialXSurfaceShaderAbstract.cpp b/Engine/Plugins/Interchange/Runtime/Source/Import/Private/MaterialX/MaterialXUtils/MaterialXSurfaceShaderAbstract.cpp
--- a/Engine/Plugins/Interchange/Runtime/Source/Import/Private/MaterialX/MaterialXUtils/MaterialXSurfaceShaderAbstract.cpp
+++ b/Engine/Plugins/Interchange/Runtime/Source/Import/Private/MaterialX/MaterialXUtils/MaterialXSurfaceShaderAbstract.cpp
@@ -32,9 +32,14 @@ bool FMaterialXSurfaceShaderAbstract::AddAttribute(MaterialX::InputPtr Input, co
 {
 	if(Input)
 	{
+		// The parent is either a node, or it's an interfacename and we just take the name of the input
+		mx::NodePtr Node = Input->getParent()->asA<mx::Node>();
+		FString NodeName = Node ? Node->getName().c_str() + FString{ TEXT("_") } : FString{};
+		ShaderNode->SetDisplayLabel(NodeName + Input->getName().c_str());
+
 		if(Input->getType() == mx::Type::Float)
 		{
-			return ShaderNode->AddFloatAttribute(UInterchangeShaderPortsAPI::MakeInputValueKey(InputChannelName), mx::fromValueString<float>(Input->getValueString()));
+			return ShaderNode->AddFloatInput(InputChannelName, mx::fromValueString<float>(Input->getValueString()), true);
 		}
 		else if(Input->getType() == mx::Type::Integer) //Let's add it a Float attribute, because Interchange doesn't create a scalar if it's an int
 		{
@@ -43,7 +48,7 @@ bool FMaterialXSurfaceShaderAbstract::AddAttribute(MaterialX::InputPtr Input, co
 		else if(Input->getType() == mx::Type::Color3 || Input->getType() == mx::Type::Color4)
 		{
 			FLinearColor LinearColor = GetLinearColor(Input);
-			return ShaderNode->AddLinearColorAttribute(UInterchangeShaderPortsAPI::MakeInputValueKey(InputChannelName), LinearColor);
+			return ShaderNode->AddLinearColorInput(InputChannelName, LinearColor, true);
 		}
 		else if(Input->getType() == mx::Type::Vector2)
 		{
@@ -53,7 +58,7 @@ bool FMaterialXSurfaceShaderAbstract::AddAttribute(MaterialX::InputPtr Input, co
 		else if(Input->getType() == mx::Type::Vector3 || Input->getType() == mx::Type::Vector4)
 		{
 			FLinearColor Vector = GetVector(Input);
-			return ShaderNode->AddLinearColorAttribute(UInterchangeShaderPortsAPI::MakeInputValueKey(InputChannelName), Vector);
+			return ShaderNode->AddLinearColorInput(InputChannelName, Vector, true);
 		}
 	}
 
@@ -106,7 +111,7 @@ bool FMaterialXSurfaceShaderAbstract::AddFloatAttribute(MaterialX::InputPtr Inpu
 
 		if(!FMath::IsNearlyEqual(Value, DefaultValue))
 		{
-			return ShaderNode->AddFloatAttribute(UInterchangeShaderPortsAPI::MakeInputValueKey(InputChannelName), Value);
+			return ShaderNode->AddFloatInput(InputChannelName, Value, true);
 		}
 	}
 
@@ -121,7 +126,7 @@ bool FMaterialXSurfaceShaderAbstract::AddLinearColorAttribute(MaterialX::InputPt
 
 		if(!Value.Equals(DefaultValue))
 		{
-			return ShaderNode->AddLinearColorAttribute(UInterchangeShaderPortsAPI::MakeInputValueKey(InputChannelName), Value);
+			return ShaderNode->AddLinearColorInput(InputChannelName, Value, true);
 		}
 	}
 
@@ -136,7 +141,7 @@ bool FMaterialXSurfaceShaderAbstract::AddVectorAttribute(MaterialX::InputPtr Inp
 
 		if(!Value.Equals(DefaultValue))
 		{
-			return ShaderNode->AddLinearColorAttribute(UInterchangeShaderPortsAPI::MakeInputValueKey(InputChannelName), Value);
+			return ShaderNode->AddLinearColorInput(InputChannelName, Value, true);
 		}
 	}
 
diff --git a/Engine/Plugins/Interchange/Runtime/Source/Pipelines/Private/InterchangeGenericMaterialPipeline.cpp b/Engine/Plugins/Interchange/Runtime/Source/Pipelines/Private/InterchangeGenericMaterialPipeline.cpp
--- a/Engine/Plugins/Interchange/Runtime/Source/Pipelines/Private/InterchangeGenericMaterialPipeline.cpp
+++ b/Engine/Plugins/Interchange/Runtime/Source/Pipelines/Private/InterchangeGenericMaterialPipeline.cpp
@@ -2691,7 +2691,14 @@ UInterchangeMaterialExpressionFactoryNode* UInterchangeGenericMaterialPipeline::
 		MaterialExpressionFactoryNode->AddApplyAndFillDelegates<float>(DefaultValueMemberName.ToString(), UMaterialExpressionScalarParameter::StaticClass(), DefaultValueMemberName);
 	}
 
-	MaterialExpressionFactoryNode->SetDisplayLabel(InputName);
+	if(FString DisplayLabel = ShaderNode->GetDisplayLabel(); DisplayLabel.IsEmpty())
+	{
+		MaterialExpressionFactoryNode->SetDisplayLabel(InputName);
+	}
+	else
+	{
+		MaterialExpressionFactoryNode->SetDisplayLabel(DisplayLabel);
+	}
 
 	return MaterialExpressionFactoryNode;
 }
@@ -2735,7 +2742,14 @@ UInterchangeMaterialExpressionFactoryNode* UInterchangeGenericMaterialPipeline::
 		MaterialExpressionFactoryNode->AddApplyAndFillDelegates<FLinearColor>(DefaultValueName.ToString(), UMaterialExpressionVectorParameter::StaticClass(), DefaultValueName);
 	}
 
-	MaterialExpressionFactoryNode->SetDisplayLabel(InputName);
+	if(FString DisplayLabel = ShaderNode->GetDisplayLabel(); DisplayLabel.IsEmpty())
+	{
+		MaterialExpressionFactoryNode->SetDisplayLabel(InputName);
+	}
+	else
+	{
+		MaterialExpressionFactoryNode->SetDisplayLabel(DisplayLabel);
+	}
 
 	return MaterialExpressionFactoryNode;
 }
