Notify the user the import data contain too long name (more then 1023 charcter) for the meta data key.

#rnx

diff --git a/Engine/Plugins/Interchange/Runtime/Source/Pipelines/Private/InterchangeGenericAssetsPipeline.cpp b/Engine/Plugins/Interchange/Runtime/Source/Pipelines/Private/InterchangeGenericAssetsPipeline.cpp
--- a/Engine/Plugins/Interchange/Runtime/Source/Pipelines/Private/InterchangeGenericAssetsPipeline.cpp
+++ b/Engine/Plugins/Interchange/Runtime/Source/Pipelines/Private/InterchangeGenericAssetsPipeline.cpp
@@ -1348,9 +1348,21 @@ void UInterchangeGenericAssetsPipeline::AddPackageMetaData(UObject* CreatedAsset
 			if (MetaDataValue.IsSet())
 			{
 				const FString& MetaDataStringValue = MetaDataValue.GetValue();
-				const FName& MetaDataKey = FName(InterchangeMetaDataPrefix + UserAttributeInfo.Name);
-				//SetValue either add the key or set the new value
-				MetaData->SetValue(CreatedAsset, MetaDataKey, *MetaDataStringValue);
+				FString MetaDataKeyString = InterchangeMetaDataPrefix + UserAttributeInfo.Name;
+				if (MetaDataKeyString.Len() < NAME_SIZE)
+				{
+					const FName& MetaDataKey = FName(MetaDataKeyString);
+					//SetValue either add the key or set the new value
+					MetaData->SetValue(CreatedAsset, MetaDataKey, *MetaDataStringValue);
+				}
+				else if(!bHasNotify_MetaDataAttributeKeyNameTooLong)
+				{
+					bHasNotify_MetaDataAttributeKeyNameTooLong = true;
+					//We cannot add this meta data, notify the user the meta attribute key name is too long
+					UInterchangeResultWarning_Generic* Message = AddMessage<UInterchangeResultWarning_Generic>();
+					Message->Text = FText::Format(NSLOCTEXT("UInterchangeGenericAssetsPipeline", "MetadataKeyNameTooLong", "One or more metadata key(s) cannot be added because the name exceeds the maximum length ({0}) allowed by the engine. The metadata is provided by the source file node's custom attributes."),
+						FText::AsNumber(NAME_SIZE));
+				}
 			}
 		}
 	}
diff --git a/Engine/Plugins/Interchange/Runtime/Source/Pipelines/Public/InterchangeGenericAssetsPipeline.h b/Engine/Plugins/Interchange/Runtime/Source/Pipelines/Public/InterchangeGenericAssetsPipeline.h
--- a/Engine/Plugins/Interchange/Runtime/Source/Pipelines/Public/InterchangeGenericAssetsPipeline.h
+++ b/Engine/Plugins/Interchange/Runtime/Source/Pipelines/Public/InterchangeGenericAssetsPipeline.h
@@ -199,6 +199,9 @@ private:
 		}
 	};
 	FSkeletonConflictData SkeletonConflictData;
+
+	//Make sure we notify the user only once for metadata attribute key name too long
+	bool bHasNotify_MetaDataAttributeKeyNameTooLong = false;
 };
 
 class SInterchangeGenericAssetMaterialConflictWidget : public SInterchangeBaseConflictWidget
