Fixing Debug NonUnity compile errors.

diff --git a/Engine/Plugins/Interchange/Runtime/Source/Pipelines/Private/InterchangeGenericMaterialPipeline.cpp b/Engine/Plugins/Interchange/Runtime/Source/Pipelines/Private/InterchangeGenericMaterialPipeline.cpp
--- a/Engine/Plugins/Interchange/Runtime/Source/Pipelines/Private/InterchangeGenericMaterialPipeline.cpp
+++ b/Engine/Plugins/Interchange/Runtime/Source/Pipelines/Private/InterchangeGenericMaterialPipeline.cpp
@@ -61,6 +61,10 @@
 #include "UObject/Object.h"
 #include "UObject/ObjectMacros.h"
 
+#if UE_BUILD_DEBUG
+#include "HAL/PlatformFileManager.h"
+#endif
+
 // Material Hash Utils
 #include "Material/InterchangeMaterialFactory.h"
 
diff --git a/Engine/Source/Runtime/Landscape/Private/LandscapeGrassWeightExporter.cpp b/Engine/Source/Runtime/Landscape/Private/LandscapeGrassWeightExporter.cpp
--- a/Engine/Source/Runtime/Landscape/Private/LandscapeGrassWeightExporter.cpp
+++ b/Engine/Source/Runtime/Landscape/Private/LandscapeGrassWeightExporter.cpp
@@ -15,6 +15,10 @@
 #include "RenderCaptureInterface.h"
 #include "ShaderPlatformCachedIniValue.h"
 
+#if UE_BUILD_DEBUG
+#include "Misc/FileHelper.h"
+#endif
+
 class FLandscapeGrassWeightVS;
 class FLandscapeGrassWeightPS;
 
