Interchange: Support inbetween morph target curves.

#rnx

diff --git a/Engine/Plugins/Interchange/Runtime/Source/Dispatcher/Private/InterchangeCommands.cpp b/Engine/Plugins/Interchange/Runtime/Source/Dispatcher/Private/InterchangeCommands.cpp
--- a/Engine/Plugins/Interchange/Runtime/Source/Dispatcher/Private/InterchangeCommands.cpp
+++ b/Engine/Plugins/Interchange/Runtime/Source/Dispatcher/Private/InterchangeCommands.cpp
@@ -12,7 +12,7 @@ namespace UE
 	{
 		int32 DispatcherCommandVersion::GetMajor()
 		{
-			return 18;
+			return 19;
 		}
 
 		int32 DispatcherCommandVersion::GetMinor()
diff --git a/Engine/Plugins/Interchange/Runtime/Source/Import/Private/Animation/InterchangeAnimSequenceFactory.cpp b/Engine/Plugins/Interchange/Runtime/Source/Import/Private/Animation/InterchangeAnimSequenceFactory.cpp
--- a/Engine/Plugins/Interchange/Runtime/Source/Import/Private/Animation/InterchangeAnimSequenceFactory.cpp
+++ b/Engine/Plugins/Interchange/Runtime/Source/Import/Private/Animation/InterchangeAnimSequenceFactory.cpp
@@ -149,7 +149,7 @@ namespace UE::Interchange::Private
 
 	bool InternalCreateCurve(UAnimSequence* TargetSequence
 		, TArray<FRichCurve>& Curves
-		, const FString& CurveName
+		, const TArray<FString>& CurveNames
 		, const int32 CurveFlags
 		, const bool bDoNotImportCurveWithZero
 		, const bool bAddCurveMetadataToSkeleton
@@ -157,16 +157,21 @@ namespace UE::Interchange::Private
 		, const bool bMaterialCurve
 		, const bool bShouldTransact)
 	{
-		if (!TargetSequence || Curves.Num() != 1 || CurveName.IsEmpty() || Curves[0].IsEmpty())
+		bool bResult = false;
+		if (!TargetSequence || CurveNames.Num() <= 0 || CurveNames.Num() != Curves.Num())
 		{
-			return false;
+			return bResult;
 		}
 
-		if (bDoNotImportCurveWithZero)
+
+		for (int32 CurveIndex = 0; CurveIndex < Curves.Num(); ++CurveIndex)
 		{
-			bool bAllCurveValueAreZero = true;
-			for (const FRichCurve& Curve : Curves)
+			FRichCurve& Curve = Curves[CurveIndex];
+			FName Name = *CurveNames[CurveIndex];
+
+			if (bDoNotImportCurveWithZero)
 			{
+				bool bAllCurveValueAreZero = true;
 				FKeyHandle KeyHandle = Curve.GetFirstKeyHandle();
 				while (KeyHandle != FKeyHandle::Invalid())
 				{
@@ -177,68 +182,204 @@ namespace UE::Interchange::Private
 					}
 					KeyHandle = Curve.GetNextKey(KeyHandle);
 				}
+				if (bAllCurveValueAreZero)
+				{
+					continue;
+				}
 			}
-			if (bAllCurveValueAreZero)
+			
+			FAnimationCurveIdentifier FloatCurveId(Name, ERawCurveTrackTypes::RCT_Float);
+
+			IAnimationDataModel* DataModel = TargetSequence->GetDataModel();
+			IAnimationDataController& Controller = TargetSequence->GetController();
+
+			const FFloatCurve* TargetCurve = DataModel->FindFloatCurve(FloatCurveId);
+			if (TargetCurve == nullptr)
 			{
-				//Avoid importing morph target curve with only zero value
-				return false;
+				// Need to add the curve first
+				Controller.AddCurve(FloatCurveId, AACF_DefaultCurve | CurveFlags, bShouldTransact);
+				TargetCurve = DataModel->FindFloatCurve(FloatCurveId);
+			}
+			else
+			{
+				// Need to update any of the flags
+				Controller.SetCurveFlags(FloatCurveId, CurveFlags | TargetCurve->GetCurveTypeFlags(), bShouldTransact);
+			}
+
+			// Should be valid at this point
+			ensure(TargetCurve);
+
+			//MorphTarget curves are shift to 0 second
+			const float CurveStartTime = Curve.GetFirstKey().Time;
+			if (CurveStartTime > UE_SMALL_NUMBER)
+			{
+				Curve.ShiftCurve(-CurveStartTime);
+			}
+			// Set actual keys on curve within the model
+			Controller.SetCurveKeys(FloatCurveId, Curve.GetConstRefOfKeys(), bShouldTransact);
+
+			if (bMaterialCurve || bMorphTargetCurve)
+			{
+				if (bAddCurveMetadataToSkeleton)
+				{
+					USkeleton* Skeleton = TargetSequence->GetSkeleton();
+					Skeleton->AccumulateCurveMetaData(Name, bMaterialCurve, bMorphTargetCurve);
+				}
 			}
+			bResult = true;
 		}
-		
-		FName Name = *CurveName;
-		FAnimationCurveIdentifier FloatCurveId(Name, ERawCurveTrackTypes::RCT_Float);
-		
-		IAnimationDataModel* DataModel = TargetSequence->GetDataModel();
-		IAnimationDataController& Controller = TargetSequence->GetController();
+		return bResult;
+	}
 
-		const FFloatCurve* TargetCurve = DataModel->FindFloatCurve(FloatCurveId);
-		if (TargetCurve == nullptr)
+	void ResolveWeightsForBlendShape(const TArray<float>& InbetweenFullWeights, float InWeight, float& OutMainWeight, TArray<float>& OutInbetweenWeights)
+	{
+		int32 NumInbetweens = InbetweenFullWeights.Num();
+		if (NumInbetweens == 0)
 		{
-			// Need to add the curve first
-			Controller.AddCurve(FloatCurveId, AACF_DefaultCurve | CurveFlags, bShouldTransact);
-			TargetCurve = DataModel->FindFloatCurve(FloatCurveId);
+			OutMainWeight = InWeight;
+			return;
 		}
-		else
+
+		OutInbetweenWeights.SetNumUninitialized(NumInbetweens);
+		for (float& OutInbetweenWeight : OutInbetweenWeights)
+		{
+			OutInbetweenWeight = 0.0f;
+		}
+
+		if (FMath::IsNearlyEqual(InWeight, 0.0f))
+		{
+			OutMainWeight = 0.0f;
+			return;
+		}
+		else if (FMath::IsNearlyEqual(InWeight, 1.0f))
+		{
+			OutMainWeight = 1.0f;
+			return;
+		}
+
+		// Note how we don't care if UpperIndex/LowerIndex are beyond the bounds of the array here,
+		// as that signals when we're above/below all inbetweens
+		int32 UpperIndex = Algo::UpperBoundBy(InbetweenFullWeights, InWeight, [](const double& InbetweenWeight)
+			{
+				return InbetweenWeight;
+			});
+		int32 LowerIndex = UpperIndex - 1;
+
+		float UpperWeight = 1.0f;
+		if (UpperIndex <= NumInbetweens - 1)
 		{
-			// Need to update any of the flags
-			Controller.SetCurveFlags(FloatCurveId, CurveFlags | TargetCurve->GetCurveTypeFlags(), bShouldTransact);
+			UpperWeight = InbetweenFullWeights[UpperIndex];
 		}
 
-		// Should be valid at this point
-		ensure(TargetCurve);
+		float LowerWeight = 0.0f;
+		if (LowerIndex >= 0)
+		{
+			LowerWeight = InbetweenFullWeights[LowerIndex];
+		}
+
+		UpperWeight = (InWeight - LowerWeight) / (UpperWeight - LowerWeight);
+		LowerWeight = (1.0f - UpperWeight);
 
-		//MorphTarget curves are shift to 0 second
-		const float CurveStartTime = Curves[0].GetFirstKey().Time;
-		if (CurveStartTime > UE_SMALL_NUMBER)
+		// We're between upper inbetween and the 1.0 weight
+		if (UpperIndex > NumInbetweens - 1)
+		{
+			OutMainWeight = UpperWeight;
+			OutInbetweenWeights[NumInbetweens - 1] = LowerWeight;
+		}
+		// We're between 0.0 and the first inbetween weight
+		else if (LowerIndex < 0)
 		{
-			Curves[0].ShiftCurve(-CurveStartTime);
+			OutMainWeight = 0;
+			OutInbetweenWeights[0] = UpperWeight;
 		}
-		// Set actual keys on curve within the model
-		Controller.SetCurveKeys(FloatCurveId, Curves[0].GetConstRefOfKeys(), bShouldTransact);
+		// We're between two inbetweens
+		else
+		{
+			OutInbetweenWeights[UpperIndex] = UpperWeight;
+			OutInbetweenWeights[LowerIndex] = LowerWeight;
+		}
+	}
+
+	TArray<FRichCurve> ResolveWeightsForBlendShapeCurve(FRichCurve& ChannelWeightCurve, const TArray<float>& InbetweenFullWeights)
+	{
+		int32 NumInbetweens = InbetweenFullWeights.Num();
+		if (NumInbetweens == 0)
+		{
+			return { ChannelWeightCurve };
+		}
+
+		TArray<FRichCurve> Result;
+		Result.SetNum(NumInbetweens + 1);
 
-		if (bMaterialCurve || bMorphTargetCurve)
+		TArray<float> ResolvedInbetweenWeightsSample;
+		ResolvedInbetweenWeightsSample.SetNum(NumInbetweens);
+
+		for (const FRichCurveKey& SourceKey : ChannelWeightCurve.Keys)
 		{
-			if(bAddCurveMetadataToSkeleton)
+			const float SourceTime = SourceKey.Time;
+			const float SourceValue = SourceKey.Value;
+
+			float ResolvedPrimarySample = 0.0f;
+
+			ResolveWeightsForBlendShape(InbetweenFullWeights, SourceValue, ResolvedPrimarySample, ResolvedInbetweenWeightsSample);
+
+			FRichCurve& PrimaryCurve = Result[0];
+			FKeyHandle PrimaryHandle = PrimaryCurve.AddKey(SourceTime, ResolvedPrimarySample);
+			PrimaryCurve.SetKeyInterpMode(PrimaryHandle, SourceKey.InterpMode);
+
+			for (int32 InbetweenIndex = 0; InbetweenIndex < NumInbetweens; ++InbetweenIndex)
 			{
-				USkeleton* Skeleton = TargetSequence->GetSkeleton();
-				Skeleton->AccumulateCurveMetaData(Name, bMaterialCurve, bMorphTargetCurve);
+				FRichCurve& InbetweenCurve = Result[InbetweenIndex + 1];
+				FKeyHandle InbetweenHandle = InbetweenCurve.AddKey(SourceTime, ResolvedInbetweenWeightsSample[InbetweenIndex]);
+				InbetweenCurve.SetKeyInterpMode(InbetweenHandle, SourceKey.InterpMode);
 			}
 		}
-		return true;
+
+		return Result;
 	}
 
-	bool CreateMorphTargetCurve(UAnimSequence* TargetSequence, TArray<FRichCurve>& Curves, const FString& CurveName, int32 CurveFlags, bool bDoNotImportCurveWithZero, bool bAddCurveMetadataToSkeleton, bool bShouldTransact)
+	bool CreateMorphTargetCurve(UAnimSequence* TargetSequence
+		, TArray<FRichCurve>& Curves
+		, const FString& CurveName
+		, TArray<FString>& InbetweenCurveNames
+		, TArray<float>& InbetweenFullWeights
+		, bool bRemoveCurveRedundantKeys
+		, int32 CurveFlags
+		, bool bDoNotImportCurveWithZero
+		, bool bAddCurveMetadataToSkeleton
+		, bool bShouldTransact)
 	{
 		constexpr bool bIsMorphTargetCurve = true;
 		constexpr bool bIsMaterialCurve = false;
-		return InternalCreateCurve(TargetSequence, Curves, CurveName, CurveFlags, bDoNotImportCurveWithZero, bAddCurveMetadataToSkeleton, bIsMorphTargetCurve, bIsMaterialCurve, bShouldTransact);
+		if (Curves.Num() == 1 && InbetweenCurveNames.Num() == InbetweenFullWeights.Num()+1)
+		{
+			//We must create inbetween shape curve to simulate the result
+			//First bake the channel weight curves
+			FRichCurve& ChannelWeightCurve = Curves[0];
+#if WITH_EDITORONLY_DATA
+			ChannelWeightCurve.BakeCurve(1.0f / TargetSequence->ImportResampleFramerate);
+#endif
+
+			// use the primary curve to generate inbetween shape curves + a modified primary curve
+			TArray<FRichCurve> Results = ResolveWeightsForBlendShapeCurve(ChannelWeightCurve, InbetweenFullWeights);
+
+			for (FRichCurve& Result : Results)
+			{
+				if (bRemoveCurveRedundantKeys)
+				{
+					Result.RemoveRedundantAutoTangentKeys(SMALL_NUMBER);
+				}
+			}
+			return InternalCreateCurve(TargetSequence, Results, InbetweenCurveNames, CurveFlags, bDoNotImportCurveWithZero, bAddCurveMetadataToSkeleton, bIsMorphTargetCurve, bIsMaterialCurve, bShouldTransact);
+		}
+		return InternalCreateCurve(TargetSequence, Curves, { CurveName }, CurveFlags, bDoNotImportCurveWithZero, bAddCurveMetadataToSkeleton, bIsMorphTargetCurve, bIsMaterialCurve, bShouldTransact);
 	}
 
 	bool CreateMaterialCurve(UAnimSequence* TargetSequence, TArray<FRichCurve>& Curves, const FString& CurveName, int32 CurveFlags, bool bDoNotImportCurveWithZero, bool bAddCurveMetadataToSkeleton, bool bShouldTransact)
 	{
 		constexpr bool bIsMorphTargetCurve = false;
 		constexpr bool bIsMaterialCurve = true;
-		return InternalCreateCurve(TargetSequence, Curves, CurveName, CurveFlags, bDoNotImportCurveWithZero, bAddCurveMetadataToSkeleton, bIsMorphTargetCurve, bIsMaterialCurve, bShouldTransact);
+		return InternalCreateCurve(TargetSequence, Curves, { CurveName }, CurveFlags, bDoNotImportCurveWithZero, bAddCurveMetadataToSkeleton, bIsMorphTargetCurve, bIsMaterialCurve, bShouldTransact);
 	}
 
 	bool CreateAttributeCurve(UAnimSequence* TargetSequence, TArray<FRichCurve>& Curves, const FString& CurveName, int32 CurveFlags, bool bDoNotImportCurveWithZero, bool bAddCurveMetadataToSkeleton, bool bShouldTransact)
@@ -246,7 +387,7 @@ namespace UE::Interchange::Private
 		//This curve don't animate morph target or material parameter.
 		constexpr bool bIsMorphTargetCurve = false;
 		constexpr bool bIsMaterialCurve = false;
-		return InternalCreateCurve(TargetSequence, Curves, CurveName, CurveFlags, bDoNotImportCurveWithZero, bAddCurveMetadataToSkeleton, bIsMorphTargetCurve, bIsMaterialCurve, bShouldTransact);
+		return InternalCreateCurve(TargetSequence, Curves, { CurveName }, CurveFlags, bDoNotImportCurveWithZero, bAddCurveMetadataToSkeleton, bIsMorphTargetCurve, bIsMaterialCurve, bShouldTransact);
 	}
 
 	void RetrieveAnimationPayloads(UAnimSequence* AnimSequence
@@ -507,7 +648,16 @@ namespace UE::Interchange::Private
 						}
 					}
 					constexpr int32 CurveFlags = 0;
-					CreateMorphTargetCurve(AnimSequence, AnimationCurvePayload.Curves, MorphTargetData.CurveNodeNamePerPayloadKey.FindChecked(CurveNameAndPayload.Key), CurveFlags, bDoNotImportCurveWithZero, bAddCurveMetadataToSkeleton, bShouldTransact);
+					CreateMorphTargetCurve(AnimSequence
+						, AnimationCurvePayload.Curves
+						, MorphTargetData.CurveNodeNamePerPayloadKey.FindChecked(CurveNameAndPayload.Key)
+						, AnimationCurvePayload.InbetweenCurveNames
+						, AnimationCurvePayload.InbetweenFullWeights
+						, bRemoveCurveRedundantKeys
+						, CurveFlags
+						, bDoNotImportCurveWithZero
+						, bAddCurveMetadataToSkeleton
+						, bShouldTransact);
 				}
 			}
 
diff --git a/Engine/Plugins/Interchange/Runtime/Source/Import/Private/Fbx/InterchangeFbxTranslator.cpp b/Engine/Plugins/Interchange/Runtime/Source/Import/Private/Fbx/InterchangeFbxTranslator.cpp
--- a/Engine/Plugins/Interchange/Runtime/Source/Import/Private/Fbx/InterchangeFbxTranslator.cpp
+++ b/Engine/Plugins/Interchange/Runtime/Source/Import/Private/Fbx/InterchangeFbxTranslator.cpp
@@ -401,10 +401,23 @@ TFuture<TOptional<UE::Interchange::FAnimationPayloadData>> UInterchangeFbxTransl
 			switch (PayLoadKey.Type)
 			{
 			case EInterchangeAnimationPayLoadType::CURVE:
+				{
+					TArray<FInterchangeCurve> InterchangeCurves;
+					Ar << InterchangeCurves;
+					AnimationTransformPayload.Curves.AddDefaulted(InterchangeCurves.Num());
+					for (int32 CurveIndex = 0; CurveIndex < InterchangeCurves.Num(); ++CurveIndex)
+					{
+						const FInterchangeCurve& InterchangeCurve = InterchangeCurves[CurveIndex];
+						InterchangeCurve.ToRichCurve(AnimationTransformPayload.Curves[CurveIndex]);
+					}
+				}
+				break;
 			case EInterchangeAnimationPayLoadType::MORPHTARGETCURVE:
 				{
 					TArray<FInterchangeCurve> InterchangeCurves;
 					Ar << InterchangeCurves;
+					Ar << AnimationTransformPayload.InbetweenCurveNames;
+					Ar << AnimationTransformPayload.InbetweenFullWeights;
 					AnimationTransformPayload.Curves.AddDefaulted(InterchangeCurves.Num());
 					for (int32 CurveIndex = 0; CurveIndex < InterchangeCurves.Num(); ++CurveIndex)
 					{
diff --git a/Engine/Plugins/Interchange/Runtime/Source/Parsers/CommonParser/Public/InterchangeCommonAnimationPayload.h b/Engine/Plugins/Interchange/Runtime/Source/Parsers/CommonParser/Public/InterchangeCommonAnimationPayload.h
--- a/Engine/Plugins/Interchange/Runtime/Source/Parsers/CommonParser/Public/InterchangeCommonAnimationPayload.h
+++ b/Engine/Plugins/Interchange/Runtime/Source/Parsers/CommonParser/Public/InterchangeCommonAnimationPayload.h
@@ -22,6 +22,10 @@ namespace UE::Interchange
 #if WITH_ENGINE
 		//CURVE
 		TArray<FRichCurve> Curves;
+
+		//In-between blend shape data
+		TArray<FString> InbetweenCurveNames;
+		TArray<float> InbetweenFullWeights;
 #endif
 		//STEP CURVE
 		TArray<FInterchangeStepCurve> StepCurves;
diff --git a/Engine/Plugins/Interchange/Runtime/Source/Parsers/Fbx/Private/FbxAnimation.cpp b/Engine/Plugins/Interchange/Runtime/Source/Parsers/Fbx/Private/FbxAnimation.cpp
--- a/Engine/Plugins/Interchange/Runtime/Source/Parsers/Fbx/Private/FbxAnimation.cpp
+++ b/Engine/Plugins/Interchange/Runtime/Source/Parsers/Fbx/Private/FbxAnimation.cpp
@@ -787,6 +787,8 @@ namespace UE::Interchange::Private
 		{
 			FLargeMemoryWriter Ar;
 			Ar << InterchangeCurves;
+			Ar << FetchPayloadData.InbetweenCurveNames;
+			Ar << FetchPayloadData.InbetweenFullWeights;
 			uint8* ArchiveData = Ar.GetData();
 			int64 ArchiveSize = Ar.TotalSize();
 			TArray64<uint8> Buffer(ArchiveData, ArchiveSize);
@@ -896,6 +898,8 @@ namespace UE::Interchange::Private
 			FetchPayloadData.MorphTargetIndex = MorphTargetAnimationBuildingData.MorphTargetIndex;
 			FetchPayloadData.ChannelIndex = MorphTargetAnimationBuildingData.ChannelIndex;
 			FetchPayloadData.AnimLayer = MorphTargetAnimationBuildingData.AnimLayer;
+			FetchPayloadData.InbetweenCurveNames = MorphTargetAnimationBuildingData.InbetweenCurveNames;
+			FetchPayloadData.InbetweenFullWeights = MorphTargetAnimationBuildingData.InbetweenFullWeights;
 
 			AnimPayload->MorphTargetFetchPayloadData = FetchPayloadData;
 
diff --git a/Engine/Plugins/Interchange/Runtime/Source/Parsers/Fbx/Private/FbxAnimation.h b/Engine/Plugins/Interchange/Runtime/Source/Parsers/Fbx/Private/FbxAnimation.h
--- a/Engine/Plugins/Interchange/Runtime/Source/Parsers/Fbx/Private/FbxAnimation.h
+++ b/Engine/Plugins/Interchange/Runtime/Source/Parsers/Fbx/Private/FbxAnimation.h
@@ -49,7 +49,7 @@ namespace UE::Interchange::Private
 	struct FMorphTargetFetchPayloadData
 	{
 		//We cannot just store the FbxAnimCurve
-		// becase! it is refernced from FbxMesh which gets replaced when we are triangulating it upon PayloadData acquisition
+		// because! it is referenced from FbxMesh which gets replaced when we are triangulating it upon PayloadData acquisition
 		// for that reason we need to re-acquire it all the way from the Scene, based on GeometryIndex, MorphTargetIndex, ChannelIndex and AnimLayer.
 		// (This is happening because upon Triangulation we are passing bReplace = true, see @FMeshDescriptionImporter::FillMeshDescriptionFromFbxMesh)
 		//FbxAnimCurve* MorphTargetAnimCurve = nullptr;
@@ -60,6 +60,11 @@ namespace UE::Interchange::Private
 		int32 MorphTargetIndex;
 		int32 ChannelIndex;
 		FbxAnimLayer* AnimLayer;
+
+		//In between blend shape animation support
+		//The count of the curve names should match the in between full weights
+		TArray<FString> InbetweenCurveNames;
+		TArray<float> InbetweenFullWeights;
 	};
 
 	
diff --git a/Engine/Plugins/Interchange/Runtime/Source/Parsers/Fbx/Private/FbxMesh.cpp b/Engine/Plugins/Interchange/Runtime/Source/Parsers/Fbx/Private/FbxMesh.cpp
--- a/Engine/Plugins/Interchange/Runtime/Source/Parsers/Fbx/Private/FbxMesh.cpp
+++ b/Engine/Plugins/Interchange/Runtime/Source/Parsers/Fbx/Private/FbxMesh.cpp
@@ -1789,27 +1789,75 @@ void FFbxMesh::AddAllMeshes(FbxScene* SDKScene, FbxGeometryConverter* SDKGeometr
 							PayloadContexts.Add(MorphTargetPayLoadKey, GeoPayload);
 						}
 						MorphTargetNode->SetPayLoadKey(MorphTargetPayLoadKey, EInterchangeMeshPayLoadType::MORPHTARGET);
+					}
+				}
+
+				FbxShape* Shape = Channel->GetTargetShape(0);
+				FString MorphTargetUniqueID = Parser.GetFbxHelper()->GetMeshUniqueID(Shape);
+				for (int32 AnimationIndex = 0; AnimationIndex < NumAnimations; AnimationIndex++)
+				{
+					FbxAnimStack* AnimStack = (FbxAnimStack*)SDKScene->GetSrcObject<FbxAnimStack>(AnimationIndex);
+					if (AnimStack)
+					{
+						FbxAnimLayer* AnimLayer = (FbxAnimLayer*)AnimStack->GetMember(0);
 
-						for (int32 AnimationIndex = 0; AnimationIndex < NumAnimations; AnimationIndex++)
+						FbxAnimCurve* Curve = Mesh->GetShapeChannel(MorphTargetIndex, ChannelIndex, AnimLayer);
+						if (Curve && Curve->KeyGetCount() > 0)
 						{
-							FbxAnimStack* AnimStack = (FbxAnimStack*)SDKScene->GetSrcObject<FbxAnimStack>(AnimationIndex);
-							if (AnimStack)
+							FbxTimeSpan TimeInterval;
+							Curve->GetTimeInterval(TimeInterval);
+							if (CurrentChannelMorphTargetCount == 1)
 							{
-								FbxAnimLayer* AnimLayer = (FbxAnimLayer*)AnimStack->GetMember(0);
-								
-								FbxAnimCurve* Curve = Mesh->GetShapeChannel(MorphTargetIndex, ChannelIndex, AnimLayer);
-								if (Curve && Curve->KeyGetCount() > 0)
-								{
-									FbxTimeSpan TimeInterval;
-									Curve->GetTimeInterval(TimeInterval);
+								MorphTargetAnimationsBuildingData.Add(FMorphTargetAnimationBuildingData(TimeInterval.GetStart().GetSecondDouble()
+									, TimeInterval.GetStop().GetSecondDouble()
+									, MeshNode
+									, GeometryIndex
+									, AnimationIndex
+									, AnimLayer
+									, MorphTargetIndex
+									, ChannelIndex
+									, MorphTargetUniqueID));
+							}
+							else
+							{
+								TArray<FString> InbetweenCurveNames;
+								InbetweenCurveNames.Reserve(CurrentChannelMorphTargetCount);
+
+								// in fbx the primary shape is the last shape, however to make
+								// the code more similar to usd importer, we deal with the primary shape separately
+								InbetweenCurveNames.Add(ChannelName);
+
+								// ignoring the last shape because it is not a inbetween, i.e. it is the primary shape
+								int32 InbetweenCount = CurrentChannelMorphTargetCount - 1;
+
+								TArrayView<double> FbxInbetweenFullWeights = { Channel->GetTargetShapeFullWeights(), InbetweenCount };
 
-									//In order to appropriately identify the Skeleton Node Uids we have to process the MorphTarget animations once the hierarchy is processed
-									MorphTargetAnimationsBuildingData.Add(FMorphTargetAnimationBuildingData(TimeInterval.GetStart().GetSecondDouble(), TimeInterval.GetStop().GetSecondDouble(), MeshNode, GeometryIndex, AnimationIndex, AnimLayer, MorphTargetIndex, ChannelIndex, MorphTargetUniqueID));
+								TArray<float> InbetweenFullWeights;
+								InbetweenFullWeights.Reserve(InbetweenCount);
+								/** for some reason blend shape values are coming as 100 scaled, so a transform is needed to scale it to 0-1 **/
+								Algo::Transform(FbxInbetweenFullWeights, InbetweenFullWeights, [](double Input) { return Input * 0.01f; });
+
+								// collect inbetween shape names
+								for (int32 InbetweenIndex = 0; InbetweenIndex < InbetweenCount; ++InbetweenIndex)
+								{
+									FbxShape* InbetweenShape = Channel->GetTargetShape(InbetweenIndex);
+									InbetweenCurveNames.Add(Parser.GetFbxHelper()->GetFbxObjectName(InbetweenShape));
 								}
+								MorphTargetAnimationsBuildingData.Add(FMorphTargetAnimationBuildingData(TimeInterval.GetStart().GetSecondDouble()
+									, TimeInterval.GetStop().GetSecondDouble()
+									, MeshNode
+									, GeometryIndex
+									, AnimationIndex
+									, AnimLayer
+									, MorphTargetIndex
+									, ChannelIndex
+									, MorphTargetUniqueID
+									, InbetweenCurveNames
+									, InbetweenFullWeights));
 							}
 						}
 					}
-				} // for CurrentChannelMorphTargetCount
+				}
 			} // for MorphTargetChannelCount
 		} // for MorphTargetCount
 	} // for GeometryCount
diff --git a/Engine/Plugins/Interchange/Runtime/Source/Parsers/Fbx/Private/FbxMesh.h b/Engine/Plugins/Interchange/Runtime/Source/Parsers/Fbx/Private/FbxMesh.h
--- a/Engine/Plugins/Interchange/Runtime/Source/Parsers/Fbx/Private/FbxMesh.h
+++ b/Engine/Plugins/Interchange/Runtime/Source/Parsers/Fbx/Private/FbxMesh.h
@@ -29,6 +29,8 @@ namespace UE
 				int32 MorphTargetIndex;
 				int32 ChannelIndex;
 				FString MorphTargetNodeUid;
+				TArray<FString> InbetweenCurveNames;
+				TArray<float> InbetweenFullWeights;
 
 				FMorphTargetAnimationBuildingData(double InStartTime
 						, double InStopTime
@@ -51,6 +53,32 @@ namespace UE
 				{
 
 				}
+
+				FMorphTargetAnimationBuildingData(double InStartTime
+						, double InStopTime
+						, UInterchangeMeshNode* InInterchangeMeshNode
+						, int32 InGeometryIndex
+						, int32 InAnimationIndex
+						, FbxAnimLayer* InAnimLayer
+						, int32 InMorphTargetIndex
+						, int32 InChannelIndex
+						, FString InMorphTargetNodeUid
+						, TArray<FString> InInbetweenCurveNames
+						, TArray<float> InInbetweenFullWeights)
+					: StartTime(InStartTime)
+					, StopTime(InStopTime)
+					, InterchangeMeshNode(InInterchangeMeshNode)
+					, GeometryIndex(InGeometryIndex)
+					, AnimationIndex(InAnimationIndex)
+					, AnimLayer(InAnimLayer)
+					, MorphTargetIndex(InMorphTargetIndex)
+					, ChannelIndex(InChannelIndex)
+					, MorphTargetNodeUid(InMorphTargetNodeUid)
+					, InbetweenCurveNames(InInbetweenCurveNames)
+					, InbetweenFullWeights(InInbetweenFullWeights)
+				{
+
+				}
 			};
 
 			class FFbxMesh;
