Interchange:
- Clean up some API for the Fbx translator: The in-process meshes payload are not using the temporary file and pass directly the structure in memory. This speed up the import when in-process.
- Make sure the temporary fbx translator folder is deleted.

#rnx

diff --git a/Engine/Plugins/Interchange/Runtime/Source/Import/Private/Fbx/InterchangeFbxTranslator.cpp b/Engine/Plugins/Interchange/Runtime/Source/Import/Private/Fbx/InterchangeFbxTranslator.cpp
--- a/Engine/Plugins/Interchange/Runtime/Source/Import/Private/Fbx/InterchangeFbxTranslator.cpp
+++ b/Engine/Plugins/Interchange/Runtime/Source/Import/Private/Fbx/InterchangeFbxTranslator.cpp
@@ -60,10 +60,16 @@ namespace UE::Interchange::Private
 	}
 } //ns UE::Interchange::Private
 
+#define INTERCHANGE_FBX_PATH TEXT("Interchange/Fbx")
 UInterchangeFbxTranslator::UInterchangeFbxTranslator()
 {
 	Dispatcher = nullptr;
 	bUseWorkerImport = false;
+	if (HasAnyFlags(RF_ClassDefaultObject | RF_ArchetypeObject))
+	{
+		//Class default object should not use any resources
+		return;
+	}
 	
 	FGuid RandomGuid;
 	FPlatformMisc::CreateGuid(RandomGuid);
@@ -74,7 +80,7 @@ UInterchangeFbxTranslator::UInterchangeFbxTranslator()
 	{
 		PlatformFile.CreateDirectory(*ProjectSavedDir);
 	}
-	const FString InterchangeDir = FPaths::Combine(ProjectSavedDir, TEXT("Interchange"));
+	const FString InterchangeDir = FPaths::Combine(ProjectSavedDir, INTERCHANGE_FBX_PATH);
 	if (!PlatformFile.DirectoryExists(*InterchangeDir))
 	{
 		PlatformFile.CreateDirectory(*InterchangeDir);
@@ -107,6 +113,24 @@ UInterchangeFbxTranslator::UInterchangeFbxTranslator()
 	}
 }
 
+void UInterchangeFbxTranslator::CleanUpTemporaryFolder()
+{
+	//Clean up the interchange fbx temporary folder.
+	IPlatformFile& PlatformFile = FPlatformFileManager::Get().GetPlatformFile();
+	const FString ProjectSavedDir = FPaths::ConvertRelativePathToFull(FPaths::ProjectSavedDir());
+	if (PlatformFile.DirectoryExists(*ProjectSavedDir))
+	{
+		const FString InterchangeDir = FPaths::Combine(ProjectSavedDir, INTERCHANGE_FBX_PATH);
+		if (PlatformFile.DirectoryExists(*InterchangeDir))
+		{
+			constexpr bool RequireExists = false;
+			//Delete recursively folder's content
+			constexpr bool Tree = true;
+			IFileManager::Get().DeleteDirectory(*InterchangeDir, RequireExists, Tree);
+		}
+	}
+}
+
 bool UInterchangeFbxTranslator::IsThreadSafe() const
 {
 	return bUseWorkerImport;
@@ -138,6 +162,16 @@ TArray<FString> UInterchangeFbxTranslator::GetSupportedFormats() const
 bool UInterchangeFbxTranslator::Translate(UInterchangeBaseNodeContainer& BaseNodeContainer) const
 {
 	TRACE_CPUPROFILER_EVENT_SCOPE(UInterchangeFbxTranslator::Translate);
+	if (HasAnyFlags(RF_ClassDefaultObject | RF_ArchetypeObject))
+	{
+		//Class default should never be use has an active translator
+		ensure(!HasAnyFlags(RF_ClassDefaultObject | RF_ArchetypeObject));
+		UInterchangeResultError_Generic* ErrorResult = AddMessage<UInterchangeResultError_Generic>();
+		ErrorResult->SourceAssetName = SourceData ? SourceData->GetFilename() : FString();
+		ErrorResult->Text = NSLOCTEXT("UInterchangeFbxTranslator", "Translate_DoNotUseClassDefault", "Class default should not be use has an active translator.");
+		return false;
+	}
+
 	//Make sure the hash is compute here in asynchronous mode
 	GetSourceData()->GetFileContentHash();
 	FString Filename = GetSourceData()->GetFilename();
@@ -316,9 +350,16 @@ TFuture<TOptional<UE::Interchange::FMeshPayloadData>> UInterchangeFbxTranslator:
 	if (!bUseWorkerImport)
 	{
 #if WITH_EDITOR
-		FString ResultMeshPayloadsUniqueId = FbxParser.FetchMeshPayload(PayLoadKey.UniqueId, MeshGlobalTransform, ResultFolder);
-		FString MeshPayloadFilename = FbxParser.GetResultPayloadFilepath(ResultMeshPayloadsUniqueId);
-		OnPayloadReady(MeshPayloadFilename);
+		UE::Interchange::FMeshPayloadData MeshPayloadData;
+		MeshPayloadData.MeshDescription.Empty();
+		FbxParser.FetchMeshPayload(PayLoadKey.UniqueId, MeshGlobalTransform, MeshPayloadData);
+		if (!FStaticMeshOperations::ValidateAndFixData(MeshPayloadData.MeshDescription, PayLoadKey.UniqueId))
+		{
+			UInterchangeResultError_Generic* ErrorResult = AddMessage<UInterchangeResultError_Generic>();
+			ErrorResult->SourceAssetName = SourceData ? SourceData->GetFilename() : FString();
+			ErrorResult->Text = NSLOCTEXT("UInterchangeFbxTranslator", "GetMeshPayloadData_ValidateMeshDescriptionFail", "Invalid mesh data (NAN) was found and fix to zero. Mesh render can be bad.");
+		}
+		Promise->SetValue(MoveTemp(MeshPayloadData));
 #endif
 	}
 	else
diff --git a/Engine/Plugins/Interchange/Runtime/Source/Import/Private/InterchangeImportModule.cpp b/Engine/Plugins/Interchange/Runtime/Source/Import/Private/InterchangeImportModule.cpp
--- a/Engine/Plugins/Interchange/Runtime/Source/Import/Private/InterchangeImportModule.cpp
+++ b/Engine/Plugins/Interchange/Runtime/Source/Import/Private/InterchangeImportModule.cpp
@@ -78,6 +78,7 @@ void FInterchangeImportModule::StartupModule()
 		//Register the translators
 		//Scenes
 #if WITH_EDITOR
+		UInterchangeFbxTranslator::CleanUpTemporaryFolder();
 		InterchangeManager.RegisterTranslator(UInterchangeFbxTranslator::StaticClass());
 #endif
 		InterchangeManager.RegisterTranslator(UInterchangeGLTFTranslator::StaticClass());
@@ -132,6 +133,7 @@ void FInterchangeImportModule::StartupModule()
 void FInterchangeImportModule::ShutdownModule()
 {
 	UInterchangeManager::SetInterchangeImportEnabled(false);
+	UInterchangeFbxTranslator::CleanUpTemporaryFolder();
 }
 
 
diff --git a/Engine/Plugins/Interchange/Runtime/Source/Import/Public/Fbx/InterchangeFbxTranslator.h b/Engine/Plugins/Interchange/Runtime/Source/Import/Public/Fbx/InterchangeFbxTranslator.h
--- a/Engine/Plugins/Interchange/Runtime/Source/Import/Public/Fbx/InterchangeFbxTranslator.h
+++ b/Engine/Plugins/Interchange/Runtime/Source/Import/Public/Fbx/InterchangeFbxTranslator.h
@@ -30,6 +30,8 @@ class INTERCHANGEIMPORT_API UInterchangeFbxTranslator : public UInterchangeTrans
 public:
 	UInterchangeFbxTranslator();
 
+	static void CleanUpTemporaryFolder();
+
 	/** Begin UInterchangeTranslatorBase API*/
 	virtual bool IsThreadSafe() const override;
 	virtual EInterchangeTranslatorType GetTranslatorType() const override;
diff --git a/Engine/Plugins/Interchange/Runtime/Source/Parsers/Fbx/InterchangeFbxParser.Build.cs b/Engine/Plugins/Interchange/Runtime/Source/Parsers/Fbx/InterchangeFbxParser.Build.cs
--- a/Engine/Plugins/Interchange/Runtime/Source/Parsers/Fbx/InterchangeFbxParser.Build.cs
+++ b/Engine/Plugins/Interchange/Runtime/Source/Parsers/Fbx/InterchangeFbxParser.Build.cs
@@ -48,7 +48,8 @@ namespace UnrealBuildTool.Rules
 				PrivateDependencyModuleNames.AddRange(
 					new string[] {
 						"Engine",
-						"InterchangeEngine"
+						"InterchangeEngine",
+						"InterchangeImport"
 					}
 				);
 			}
diff --git a/Engine/Plugins/Interchange/Runtime/Source/Parsers/Fbx/Private/FbxAPI.cpp b/Engine/Plugins/Interchange/Runtime/Source/Parsers/Fbx/Private/FbxAPI.cpp
--- a/Engine/Plugins/Interchange/Runtime/Source/Parsers/Fbx/Private/FbxAPI.cpp
+++ b/Engine/Plugins/Interchange/Runtime/Source/Parsers/Fbx/Private/FbxAPI.cpp
@@ -12,6 +12,9 @@
 #include "FbxMesh.h"
 #include "FbxScene.h"
 #include "InterchangeTextureNode.h"
+#if WITH_ENGINE
+#include "Mesh/InterchangeMeshPayload.h"
+#endif
 #include "Nodes/InterchangeBaseNodeContainer.h"
 #include "Misc/SecureHash.h"
 
@@ -182,6 +185,25 @@ namespace UE
 				}
 			}
 
+#if WITH_ENGINE
+			bool FFbxParser::FetchMeshPayloadData(const FString& PayloadKey, const FTransform& MeshGlobalTransform, FMeshPayloadData& OutMeshPayloadData)
+			{
+				if (!PayloadContexts.Contains(PayloadKey))
+				{
+					UInterchangeResultError_Generic* Message = AddMessage<UInterchangeResultError_Generic>();
+					Message->Text = LOCTEXT("CannotRetrievePayload", "Cannot retrieve payload; payload key doesn't have any context.");
+					return false;
+				}
+
+				{
+					//Critical section to force payload to be fetch one by one with no concurrency.
+					FScopeLock Lock(&PayloadCriticalSection);
+					TSharedPtr<FPayloadContextBase>& PayloadContext = PayloadContexts.FindChecked(PayloadKey);
+					return PayloadContext->FetchMeshPayload(*this, MeshGlobalTransform, OutMeshPayloadData);
+				}
+			}
+#endif
+
 			bool FFbxParser::FetchAnimationBakeTransformPayload(const FString& PayloadKey, const double BakeFrequency, const double RangeStartTime, const double RangeEndTime, const FString& PayloadFilepath)
 			{
 				if (!PayloadContexts.Contains(PayloadKey))
diff --git a/Engine/Plugins/Interchange/Runtime/Source/Parsers/Fbx/Private/FbxAPI.h b/Engine/Plugins/Interchange/Runtime/Source/Parsers/Fbx/Private/FbxAPI.h
--- a/Engine/Plugins/Interchange/Runtime/Source/Parsers/Fbx/Private/FbxAPI.h
+++ b/Engine/Plugins/Interchange/Runtime/Source/Parsers/Fbx/Private/FbxAPI.h
@@ -20,6 +20,9 @@ namespace UE
 			class FPayloadContextBase;
 			struct FFbxHelper;
 		}
+#if WITH_ENGINE
+		struct FMeshPayloadData;
+#endif
 	}
 }
 class UInterchangeBaseNodeContainer;
@@ -58,6 +61,9 @@ namespace UE
 
 				/* Extract the fbx mesh data from the sdk into our node container */
 				bool FetchMeshPayloadData(const FString& PayloadKey, const FTransform& MeshGlobalTransform, const FString& PayloadFilepath);
+#if WITH_ENGINE
+				bool FetchMeshPayloadData(const FString& PayloadKey, const FTransform& MeshGlobalTransform, FMeshPayloadData& OutMeshPayloadData);
+#endif
 
 				/* Extract the fbx data from the sdk into our node container */
 				bool FetchAnimationBakeTransformPayload(const FString& PayloadKey, const double BakeFrequency, const double RangeStartTime, const double RangeEndTime, const FString& PayloadFilepath);
diff --git a/Engine/Plugins/Interchange/Runtime/Source/Parsers/Fbx/Private/FbxAnimation.cpp b/Engine/Plugins/Interchange/Runtime/Source/Parsers/Fbx/Private/FbxAnimation.cpp
--- a/Engine/Plugins/Interchange/Runtime/Source/Parsers/Fbx/Private/FbxAnimation.cpp
+++ b/Engine/Plugins/Interchange/Runtime/Source/Parsers/Fbx/Private/FbxAnimation.cpp
@@ -423,7 +423,7 @@ namespace UE::Interchange::Private
 				NodeTransform = ParentTransform.Inverse() * NodeTransform;
 			}
 
-			AnimationBakeTransformPayloadData.Transforms.Add(UE::Interchange::Private::FFbxConvert::ConvertTransform(NodeTransform));
+			AnimationBakeTransformPayloadData.Transforms.Add(UE::Interchange::Private::FFbxConvert::ConvertTransform<FTransform, FVector, FQuat>(NodeTransform));
 		}
 
 		return true;
@@ -508,7 +508,7 @@ namespace UE::Interchange::Private
 		}
 	}
 
-	bool FAnimationPayloadContextTransform::FetchPayloadToFile(FFbxParser& Parser, const FString& PayloadFilepath)
+	bool FAnimationPayloadContext::FetchPayloadToFile(FFbxParser& Parser, const FString& PayloadFilepath)
 	{
 		if (AttributeFetchPayloadData.IsSet() || AttributeNodeTransformFetchPayloadData.IsSet())
 		{
@@ -521,7 +521,7 @@ namespace UE::Interchange::Private
 		return false;
 	}
 
-	bool FAnimationPayloadContextTransform::FetchAnimationBakeTransformPayloadToFile(FFbxParser& Parser, const double BakeFrequency, const double RangeStartTime, const double RangeEndTime, const FString& PayloadFilepath)
+	bool FAnimationPayloadContext::FetchAnimationBakeTransformPayloadToFile(FFbxParser& Parser, const double BakeFrequency, const double RangeStartTime, const double RangeEndTime, const FString& PayloadFilepath)
 	{
 		if (!ensure(NodeTransformFetchPayloadData.IsSet()))
 		{
@@ -563,7 +563,7 @@ namespace UE::Interchange::Private
 		return true;
 	}
 
-	bool FAnimationPayloadContextTransform::InternalFetchCurveNodePayloadToFile(FFbxParser& Parser, const FString& PayloadFilepath)
+	bool FAnimationPayloadContext::InternalFetchCurveNodePayloadToFile(FFbxParser& Parser, const FString& PayloadFilepath)
 	{
 		if (AttributeFetchPayloadData.IsSet())
 		{
@@ -757,7 +757,28 @@ namespace UE::Interchange::Private
 		return false;
 	}
 
-	bool FAnimationPayloadContextTransform::InternalFetchMorphTargetCurvePayloadToFile(FFbxParser& Parser, const FString& PayloadFilepath)
+	bool FAnimationPayloadContext::InternalFetchMorphTargetCurvePayloadToFile(FFbxParser& Parser, const FString& PayloadFilepath)
+	{
+		check(MorphTargetFetchPayloadData.IsSet());
+		FMorphTargetFetchPayloadData& FetchPayloadData = MorphTargetFetchPayloadData.GetValue();
+		
+		TArray<FInterchangeCurve> InterchangeCurves;
+		if (InternalFetchMorphTargetCurvePayload(Parser, InterchangeCurves))
+		{
+			FLargeMemoryWriter Ar;
+			Ar << InterchangeCurves;
+			Ar << FetchPayloadData.InbetweenCurveNames;
+			Ar << FetchPayloadData.InbetweenFullWeights;
+			uint8* ArchiveData = Ar.GetData();
+			int64 ArchiveSize = Ar.TotalSize();
+			TArray64<uint8> Buffer(ArchiveData, ArchiveSize);
+			FFileHelper::SaveArrayToFile(Buffer, *PayloadFilepath);
+			return true;
+		}
+		return false;
+	}
+
+	bool FAnimationPayloadContext::InternalFetchMorphTargetCurvePayload(FFbxParser& Parser, TArray<FInterchangeCurve>& InterchangeCurves)
 	{
 		check(MorphTargetFetchPayloadData.IsSet());
 		FMorphTargetFetchPayloadData& FetchPayloadData = MorphTargetFetchPayloadData.GetValue();
@@ -779,22 +800,10 @@ namespace UE::Interchange::Private
 		}
 		FbxAnimCurve* AnimCurve = Geometry->GetShapeChannel(FetchPayloadData.MorphTargetIndex, FetchPayloadData.ChannelIndex, FetchPayloadData.AnimLayer);
 
-		TArray<FInterchangeCurve> InterchangeCurves;
 		//Morph target curve in fbx are between 0 and 100, in Unreal we are between 0 and 1, so we must scale
 		//The curve with 0.01
 		constexpr float ScaleCurve = 0.01f;
-		ImportCurve(AnimCurve, ScaleCurve, InterchangeCurves.AddDefaulted_GetRef().Keys);
-		{
-			FLargeMemoryWriter Ar;
-			Ar << InterchangeCurves;
-			Ar << FetchPayloadData.InbetweenCurveNames;
-			Ar << FetchPayloadData.InbetweenFullWeights;
-			uint8* ArchiveData = Ar.GetData();
-			int64 ArchiveSize = Ar.TotalSize();
-			TArray64<uint8> Buffer(ArchiveData, ArchiveSize);
-			FFileHelper::SaveArrayToFile(Buffer, *PayloadFilepath);
-		}
-		return true;
+		return ImportCurve(AnimCurve, ScaleCurve, InterchangeCurves.AddDefaulted_GetRef().Keys);
 	}
 
 	bool FFbxAnimation::AddSkeletalTransformAnimation(FbxScene* SDKScene
@@ -813,7 +822,7 @@ namespace UE::Interchange::Private
 			FString PayLoadKey = Parser.GetFbxHelper()->GetFbxNodeHierarchyName(Node) + TEXT("_") + FString::FromInt(AnimationIndex) + TEXT("_SkeletalAnimationPayloadKey");
 			if (ensure(!PayloadContexts.Contains(PayLoadKey)))
 			{
-				TSharedPtr<FAnimationPayloadContextTransform> AnimPayload = MakeShared<FAnimationPayloadContextTransform>();
+				TSharedPtr<FAnimationPayloadContext> AnimPayload = MakeShared<FAnimationPayloadContext>();
 				FNodeTransformFetchPayloadData FetchPayloadData;
 				FetchPayloadData.Node = Node;
 				FetchPayloadData.CurrentAnimStack = Parameters.CurrentAnimStack;
@@ -840,7 +849,7 @@ namespace UE::Interchange::Private
 		const FString PayLoadKey = Parser.GetFbxHelper()->GetFbxNodeHierarchyName(Node) + PropertyName + TEXT("_AnimationPayloadKey");
 		if (ensure(!PayloadContexts.Contains(PayLoadKey)))
 		{
-			TSharedPtr<FAnimationPayloadContextTransform> AnimPayload = MakeShared<FAnimationPayloadContextTransform>();
+			TSharedPtr<FAnimationPayloadContext> AnimPayload = MakeShared<FAnimationPayloadContext>();
 			FAttributeFetchPayloadData FetchPayloadData;
 			FetchPayloadData.Node = Node;
 			FetchPayloadData.AnimCurves = AnimCurveNode;
@@ -866,7 +875,7 @@ namespace UE::Interchange::Private
 		const FString PayLoadKey = Parser.GetFbxHelper()->GetFbxNodeHierarchyName(Node) + TEXT("_RigidAnimationPayloadKey");
 		if (ensure(!PayloadContexts.Contains(PayLoadKey)))
 		{
-			TSharedPtr<FAnimationPayloadContextTransform> AnimPayload = MakeShared<FAnimationPayloadContextTransform>();
+			TSharedPtr<FAnimationPayloadContext> AnimPayload = MakeShared<FAnimationPayloadContext>();
 			FAttributeNodeTransformFetchPayloadData FetchPayloadData;
 			FetchPayloadData.FrameRate = Parser.GetFrameRate();
 			FetchPayloadData.Node = Node;
@@ -891,7 +900,7 @@ namespace UE::Interchange::Private
 
 		if (ensure(!PayloadContexts.Contains(PayLoadKey)))
 		{
-			TSharedPtr<FAnimationPayloadContextTransform> AnimPayload = MakeShared<FAnimationPayloadContextTransform>();
+			TSharedPtr<FAnimationPayloadContext> AnimPayload = MakeShared<FAnimationPayloadContext>();
 			FMorphTargetFetchPayloadData FetchPayloadData;
 			FetchPayloadData.SDKScene = SDKScene;
 			FetchPayloadData.GeometryIndex = MorphTargetAnimationBuildingData.GeometryIndex;
diff --git a/Engine/Plugins/Interchange/Runtime/Source/Parsers/Fbx/Private/FbxAnimation.h b/Engine/Plugins/Interchange/Runtime/Source/Parsers/Fbx/Private/FbxAnimation.h
--- a/Engine/Plugins/Interchange/Runtime/Source/Parsers/Fbx/Private/FbxAnimation.h
+++ b/Engine/Plugins/Interchange/Runtime/Source/Parsers/Fbx/Private/FbxAnimation.h
@@ -8,6 +8,7 @@
 #include "FbxInclude.h"
 
 /** Forward declarations */
+struct FInterchangeCurve;
 struct FMeshDescription;
 class UInterchangeMeshNode;
 class UInterchangeSceneNode;
@@ -69,10 +70,10 @@ namespace UE::Interchange::Private
 
 	
 
-	class FAnimationPayloadContextTransform : public FPayloadContextBase
+	class FAnimationPayloadContext : public FPayloadContextBase
 	{
 	public:
-		virtual ~FAnimationPayloadContextTransform() {}
+		virtual ~FAnimationPayloadContext() {}
 		virtual FString GetPayloadType() const override { return TEXT("TransformAnimation-PayloadContext"); }
 		virtual bool FetchPayloadToFile(FFbxParser& Parser, const FString& PayloadFilepath) override;
 		virtual bool FetchAnimationBakeTransformPayloadToFile(FFbxParser& Parser, const double BakeFrequency, const double RangeStartTime, const double RangeEndTime, const FString& PayloadFilepath) override;
@@ -84,6 +85,7 @@ namespace UE::Interchange::Private
 	private:
 		bool InternalFetchCurveNodePayloadToFile(FFbxParser& Parser, const FString& PayloadFilepath);
 		bool InternalFetchMorphTargetCurvePayloadToFile(FFbxParser& Parser, const FString& PayloadFilepath);
+		bool InternalFetchMorphTargetCurvePayload(FFbxParser& Parser, TArray<FInterchangeCurve>& InterchangeCurves);
 	};
 
 	class FFbxAnimation
diff --git a/Engine/Plugins/Interchange/Runtime/Source/Parsers/Fbx/Private/FbxConvert.cpp b/Engine/Plugins/Interchange/Runtime/Source/Parsers/Fbx/Private/FbxConvert.cpp
--- a/Engine/Plugins/Interchange/Runtime/Source/Parsers/Fbx/Private/FbxConvert.cpp
+++ b/Engine/Plugins/Interchange/Runtime/Source/Parsers/Fbx/Private/FbxConvert.cpp
@@ -72,20 +72,6 @@ namespace UE
 				SDKScene->GetAnimationEvaluator()->Reset();
 			}
 
-			FTransform FFbxConvert::ConvertTransform(const FbxAMatrix& Matrix)
-			{
-				FTransform Out;
-
-				FQuat Rotation = ConvertRotToQuat(Matrix.GetQ());
-				FVector Origin = ConvertPos(Matrix.GetT());
-				FVector Scale = ConvertScale(Matrix.GetS());
-
-				Out.SetTranslation(Origin);
-				Out.SetScale3D(Scale);
-				Out.SetRotation(Rotation);
-
-				return Out;
-			}
 
 			FTransform FFbxConvert::AdjustCameraTransform(const FTransform& Transform)
 			{
@@ -107,99 +93,6 @@ namespace UE
 				return LightTransform;
 			}
 
-			FMatrix FFbxConvert::ConvertMatrix(const FbxAMatrix& Matrix)
-			{
-				FMatrix UEMatrix;
-
-				for (int i = 0; i < 4; ++i)
-				{
-					const FbxVector4 Row = Matrix.GetRow(i);
-					if (i == 1)
-					{
-						UEMatrix.M[i][0] = -Row[0];
-						UEMatrix.M[i][1] = Row[1];
-						UEMatrix.M[i][2] = -Row[2];
-						UEMatrix.M[i][3] = -Row[3];
-					}
-					else
-					{
-						UEMatrix.M[i][0] = Row[0];
-						UEMatrix.M[i][1] = -Row[1];
-						UEMatrix.M[i][2] = Row[2];
-						UEMatrix.M[i][3] = Row[3];
-					}
-				}
-
-				return UEMatrix;
-			}
-
-			FbxAMatrix FFbxConvert::ConvertMatrix(const FMatrix& UEMatrix)
-			{
-				FbxAMatrix FbxMatrix;
-
-				for (int i = 0; i < 4; ++i)
-				{
-					FbxVector4 Row;
-					if (i == 1)
-					{
-						Row[0] = -UEMatrix.M[i][0];
-						Row[1] = UEMatrix.M[i][1];
-						Row[2] = -UEMatrix.M[i][2];
-						Row[3] = -UEMatrix.M[i][3];
-					}
-					else
-					{
-						Row[0] = UEMatrix.M[i][0];
-						Row[1] = -UEMatrix.M[i][1];
-						Row[2] = UEMatrix.M[i][2];
-						Row[3] = UEMatrix.M[i][3];
-					}
-					FbxMatrix.SetRow(i, Row);
-				}
-
-				return FbxMatrix;
-			}
-
-			FQuat FFbxConvert::ConvertRotToQuat(FbxQuaternion Quaternion)
-			{
-				FQuat UnrealQuat;
-				UnrealQuat.X = Quaternion[0];
-				UnrealQuat.Y = -Quaternion[1];
-				UnrealQuat.Z = Quaternion[2];
-				UnrealQuat.W = -Quaternion[3];
-
-				return UnrealQuat;
-			}
-
-			FRotator FFbxConvert::ConvertEuler(FbxDouble3 Euler)
-			{
-				return FRotator::MakeFromEuler(FVector(Euler[0], -Euler[1], Euler[2]));
-			}
-
-			FVector FFbxConvert::ConvertScale(FbxVector4 Vector)
-			{
-				FVector Out;
-				Out[0] = Vector[0];
-				Out[1] = Vector[1];
-				Out[2] = Vector[2];
-				return Out;
-			}
-
-			FRotator FFbxConvert::ConvertRotation(FbxQuaternion Quaternion)
-			{
-				FRotator Out(ConvertRotToQuat(Quaternion));
-				return Out;
-			}
-			FVector FFbxConvert::ConvertPos(const FbxVector4& Vector)
-			{
-				return FVector(Vector[0], -Vector[1], Vector[2]);
-			}
-
-			FVector FFbxConvert::ConvertDir(const FbxVector4& Vector)
-			{
-				return FVector(Vector[0], -Vector[1], Vector[2]);
-			}
-
 			FLinearColor FFbxConvert::ConvertColor(const FbxDouble3& Color)
 			{
 				FLinearColor LinearColor;
diff --git a/Engine/Plugins/Interchange/Runtime/Source/Parsers/Fbx/Private/FbxConvert.h b/Engine/Plugins/Interchange/Runtime/Source/Parsers/Fbx/Private/FbxConvert.h
--- a/Engine/Plugins/Interchange/Runtime/Source/Parsers/Fbx/Private/FbxConvert.h
+++ b/Engine/Plugins/Interchange/Runtime/Source/Parsers/Fbx/Private/FbxConvert.h
@@ -18,18 +18,129 @@ namespace UE
 				//////////////////////////////////////////////////////////////////////////
 				/** Transform Conversion API Begin */
 
-				/**
-				 * Create a UE transform from a fbx matrix
-				 */
-				static FTransform ConvertTransform(const FbxAMatrix& Matrix);
-				static FMatrix ConvertMatrix(const FbxAMatrix& Matrix);
-				static FbxAMatrix ConvertMatrix(const FMatrix& UEMatrix);
-				static FQuat ConvertRotToQuat(FbxQuaternion Quaternion);
-				static FRotator ConvertEuler(FbxDouble3 Euler);
-				static FVector ConvertScale(FbxVector4 Vector);
-				static FRotator ConvertRotation(FbxQuaternion Quaternion);
-				static FVector ConvertPos(const FbxVector4& Vector);
-				static FVector ConvertDir(const FbxVector4& Vector);
+				/* Return a FTransform float or double*/
+				template<typename TransformType, typename VectorType, typename QuatType>
+				static TransformType ConvertTransform(const FbxAMatrix& Matrix)
+				{
+					TransformType Out;
+
+					Out.SetTranslation(ConvertPos<VectorType>(Matrix.GetT()));
+					Out.SetScale3D(ConvertScale<VectorType>(Matrix.GetS()));
+					Out.SetRotation(ConvertRotToQuat<QuatType>(Matrix.GetQ()));
+
+					return Out;
+				}
+
+				/* Return a FMatrix float or double*/
+				template<typename MatrixType>
+				static MatrixType ConvertMatrix(const FbxAMatrix& Matrix)
+				{
+					MatrixType UEMatrix;
+
+					for (int i = 0; i < 4; ++i)
+					{
+						const FbxVector4 Row = Matrix.GetRow(i);
+						if (i == 1)
+						{
+							UEMatrix.M[i][0] = -Row[0];
+							UEMatrix.M[i][1] = Row[1];
+							UEMatrix.M[i][2] = -Row[2];
+							UEMatrix.M[i][3] = -Row[3];
+						}
+						else
+						{
+							UEMatrix.M[i][0] = Row[0];
+							UEMatrix.M[i][1] = -Row[1];
+							UEMatrix.M[i][2] = Row[2];
+							UEMatrix.M[i][3] = Row[3];
+						}
+					}
+
+					return UEMatrix;
+				}
+
+				/* Take a FMatrix float or double and return a fbx affine matrix*/
+				template<typename MatrixType>
+				static FbxAMatrix ConvertMatrix(const MatrixType& UEMatrix)
+				{
+					FbxAMatrix FbxMatrix;
+
+					for (int i = 0; i < 4; ++i)
+					{
+						FbxVector4 Row;
+						if (i == 1)
+						{
+							Row[0] = -UEMatrix.M[i][0];
+							Row[1] = UEMatrix.M[i][1];
+							Row[2] = -UEMatrix.M[i][2];
+							Row[3] = -UEMatrix.M[i][3];
+						}
+						else
+						{
+							Row[0] = UEMatrix.M[i][0];
+							Row[1] = -UEMatrix.M[i][1];
+							Row[2] = UEMatrix.M[i][2];
+							Row[3] = UEMatrix.M[i][3];
+						}
+						FbxMatrix.SetRow(i, Row);
+					}
+
+					return FbxMatrix;
+				}
+
+				/* Return a FQuat float or double */
+				template<typename QuatType>
+				static QuatType ConvertRotToQuat(FbxQuaternion Quaternion)
+				{
+					QuatType UnrealQuat;
+					UnrealQuat.X = Quaternion[0];
+					UnrealQuat.Y = -Quaternion[1];
+					UnrealQuat.Z = Quaternion[2];
+					UnrealQuat.W = -Quaternion[3];
+
+					return UnrealQuat;
+				}
+
+				/* Return a FRotator float or double*/
+				template<typename RotatorType, typename VectorType>
+				static RotatorType ConvertEuler(FbxDouble3 Euler)
+				{
+					return RotatorType::MakeFromEuler(VectorType(Euler[0], -Euler[1], Euler[2]));
+				}
+				
+				/* Return a FVector float or double */
+				template<typename VectorType>
+				static VectorType ConvertScale(FbxVector4 Vector)
+				{
+					VectorType Out;
+					Out[0] = Vector[0];
+					Out[1] = Vector[1];
+					Out[2] = Vector[2];
+					return Out;
+				}
+				
+				/* Return a FRotator float or double */
+				template<typename RotatorType, typename QuatType>
+				static RotatorType ConvertRotation(FbxQuaternion Quaternion)
+				{
+					RotatorType Out(ConvertRotToQuat<QuatType>(Quaternion));
+					return Out;
+				}
+
+				/* Return a FVector float or double */
+				template<typename VectorType>
+				static VectorType ConvertPos(const FbxVector4& Vector)
+				{
+					return VectorType(Vector[0], -Vector[1], Vector[2]);
+				}
+
+				/* Return a FVector float or double */
+				template<typename VectorType>
+				static VectorType ConvertDir(const FbxVector4& Vector)
+				{
+					return VectorType(Vector[0], -Vector[1], Vector[2]);
+				}
+
 				static FLinearColor ConvertColor(const FbxDouble3& Color);
 
 				static FTransform AdjustCameraTransform(const FTransform& Transform);
diff --git a/Engine/Plugins/Interchange/Runtime/Source/Parsers/Fbx/Private/FbxHelper.h b/Engine/Plugins/Interchange/Runtime/Source/Parsers/Fbx/Private/FbxHelper.h
--- a/Engine/Plugins/Interchange/Runtime/Source/Parsers/Fbx/Private/FbxHelper.h
+++ b/Engine/Plugins/Interchange/Runtime/Source/Parsers/Fbx/Private/FbxHelper.h
@@ -6,51 +6,61 @@
 #include "Containers/UnrealString.h"
 #include "FbxInclude.h"
 
-namespace UE::Interchange::Private
+namespace UE::Interchange
 {
-	class FFbxParser;
+#if WITH_ENGINE
+	struct FMeshPayloadData;
+#endif
 
-	class FPayloadContextBase
+	namespace Private
 	{
-	public:
-		virtual ~FPayloadContextBase() {}
-		virtual FString GetPayloadType() const { return FString(); }
-		virtual bool FetchPayloadToFile(FFbxParser& Parser, const FString& PayloadFilepath) { return false; }
-		virtual bool FetchMeshPayloadToFile(FFbxParser& Parser, const FTransform& MeshGlobalTransform, const FString& PayloadFilepath) { return false; }
-		virtual bool FetchAnimationBakeTransformPayloadToFile(FFbxParser& Parser, const double BakeFrequency, const double RangeStartTime, const double RangeEndTime, const FString& PayloadFilepath) { return false; }
-	};
-
-	struct FFbxHelper
-	{
-	public:
-
-		FString GetMeshName(FbxGeometryBase* Mesh) const;
-		FString GetMeshUniqueID(FbxGeometryBase* Mesh) const;
-
-		FString GetNodeAttributeName(FbxNodeAttribute* NodeAttribute, const FStringView DefaultNamePrefix) const;
-		FString GetNodeAttributeUniqueID(FbxNodeAttribute* NodeAttribute, const FStringView Prefix) const;
-
-		/**
-			* Return the name of an FbxProperty, return empty string if the property is null.
-			*/
-		FString GetFbxPropertyName(const FbxProperty& Property) const;
-
-		/**
-			* Return the name of an FbxObject, return empty string if the object is null.
-			*/
-		FString GetFbxObjectName(const FbxObject* Object, bool bIsJoint = false) const;
-
-		/**
-			* Return a string with the name of all the parent in the hierarchy separate by a dot( . ) from the fbx root node to the specified node.
-			* This is a way to have a valid unique ID for a fbx node that will be the same if the fbx change when we re-import.
-			* Using the fbx sdk int32 uniqueID is not valid anymore if the fbx is re-exported.
-			*/
-		FString GetFbxNodeHierarchyName(const FbxNode* Node) const;
-
-	protected:
-		FString GetUniqueIDString(const uint64 UniqueID) const;
-
-	private:
-		mutable TMap<FString, const FbxObject*> MaterialNameClashMap;
-	};
-}//ns UE::Interchange::Private
+		class FFbxParser;
+
+		class FPayloadContextBase
+		{
+		public:
+			virtual ~FPayloadContextBase() {}
+			virtual FString GetPayloadType() const { return FString(); }
+			virtual bool FetchPayloadToFile(FFbxParser& Parser, const FString& PayloadFilepath) { return false; }
+			virtual bool FetchMeshPayloadToFile(FFbxParser& Parser, const FTransform& MeshGlobalTransform, const FString& PayloadFilepath) { return false; }
+#if WITH_ENGINE
+			virtual bool FetchMeshPayload(FFbxParser& Parser, const FTransform& MeshGlobalTransform, FMeshPayloadData& OutMeshPayloadData) { return false; };
+#endif
+			virtual bool FetchAnimationBakeTransformPayloadToFile(FFbxParser& Parser, const double BakeFrequency, const double RangeStartTime, const double RangeEndTime, const FString& PayloadFilepath) { return false; }
+		};
+
+		struct FFbxHelper
+		{
+		public:
+
+			FString GetMeshName(FbxGeometryBase* Mesh) const;
+			FString GetMeshUniqueID(FbxGeometryBase* Mesh) const;
+
+			FString GetNodeAttributeName(FbxNodeAttribute* NodeAttribute, const FStringView DefaultNamePrefix) const;
+			FString GetNodeAttributeUniqueID(FbxNodeAttribute* NodeAttribute, const FStringView Prefix) const;
+
+			/**
+				* Return the name of an FbxProperty, return empty string if the property is null.
+				*/
+			FString GetFbxPropertyName(const FbxProperty& Property) const;
+
+			/**
+				* Return the name of an FbxObject, return empty string if the object is null.
+				*/
+			FString GetFbxObjectName(const FbxObject* Object, bool bIsJoint = false) const;
+
+			/**
+				* Return a string with the name of all the parent in the hierarchy separate by a dot( . ) from the fbx root node to the specified node.
+				* This is a way to have a valid unique ID for a fbx node that will be the same if the fbx change when we re-import.
+				* Using the fbx sdk int32 uniqueID is not valid anymore if the fbx is re-exported.
+				*/
+			FString GetFbxNodeHierarchyName(const FbxNode* Node) const;
+
+		protected:
+			FString GetUniqueIDString(const uint64 UniqueID) const;
+
+		private:
+			mutable TMap<FString, const FbxObject*> MaterialNameClashMap;
+		};
+	}//ns Private
+}//ns UE::Interchange
diff --git a/Engine/Plugins/Interchange/Runtime/Source/Parsers/Fbx/Private/FbxMesh.cpp b/Engine/Plugins/Interchange/Runtime/Source/Parsers/Fbx/Private/FbxMesh.cpp
--- a/Engine/Plugins/Interchange/Runtime/Source/Parsers/Fbx/Private/FbxMesh.cpp
+++ b/Engine/Plugins/Interchange/Runtime/Source/Parsers/Fbx/Private/FbxMesh.cpp
@@ -12,6 +12,9 @@
 #include "InterchangeMeshNode.h"
 #include "InterchangeResultsContainer.h"
 #include "MeshDescription.h"
+#if WITH_ENGINE
+#include "Mesh/InterchangeMeshPayload.h"
+#endif
 #include "Misc/FileHelper.h"
 #include "Nodes/InterchangeBaseNodeContainer.h"
 #include "Serialization/LargeMemoryWriter.h"
@@ -23,7 +26,8 @@
 #endif
 #define LOCTEXT_NAMESPACE "InterchangeFbxMesh"
 
-namespace UE::Interchange::Private {
+namespace UE::Interchange::Private
+{
 
 	//Import vertex Attribute from other vertex color layer
 	void GetVertexAttributeFromMeshVertexColor(FMeshDescription& MeshDescription, FbxMesh* Mesh)
@@ -485,16 +489,14 @@ bool FMeshDescriptionImporter::FillMeshDescriptionFromFbxShape(FbxShape* Shape,
 		auto GetFinalPosition = [&TotalMatrix](FbxVector4& FbxPosition)
 		{
 			FbxPosition = TotalMatrix.MultT(FbxPosition);
-			return FFbxConvert::ConvertPos(FbxPosition);
+			return FFbxConvert::ConvertPos<FVector3f>(FbxPosition);
 		};
 
 		MeshDescription->ReserveNewVertices(VertexCount);
 		for (int32 VertexIndex = 0; VertexIndex < VertexCount; ++VertexIndex)
 		{
 			int32 RealVertexIndex = VertexOffset + VertexIndex;
-			const FVector VertexPosition = GetFinalPosition(GeoBase->GetControlPoints()[VertexIndex]);
-			//Maybe we want to do some work here like computing the deltas
-			//const FVector MeshVertexPosition = GetFinalPosition(Mesh->GetControlPoints()[VertexIndex]);
+			const FVector3f VertexPosition = GetFinalPosition(GeoBase->GetControlPoints()[VertexIndex]);
 			FVertexID AddedVertexId = MeshDescription->CreateVertex();
 			if (AddedVertexId.GetValue() != RealVertexIndex)
 			{
@@ -504,7 +506,7 @@ bool FMeshDescriptionImporter::FillMeshDescriptionFromFbxShape(FbxShape* Shape,
 				return false;
 			}
 			//Add the delta position, so we do not have to recompute it later
-			VertexPositions[AddedVertexId] = (FVector3f)VertexPosition;// -MeshVertexPosition;
+			VertexPositions[AddedVertexId] = VertexPosition;// -MeshVertexPosition;
 		}
 		MeshDescription->ResumeVertexIndexing();
 	}
@@ -725,7 +727,7 @@ bool FMeshDescriptionImporter::FillMeshDescriptionFromFbxMesh(FbxMesh* Mesh, con
 		MeshDescription->SuspendPolygonGroupIndexing();
 		MeshDescription->SuspendUVIndexing();
 
-		TRACE_CPUPROFILER_EVENT_SCOPE(BuildTriangles);
+		TRACE_CPUPROFILER_EVENT_SCOPE(Interchange_ImportMeshDescription);
 
 		// Construct the matrices for the conversion from right handed to left handed system
 		FbxAMatrix TotalMatrix = FFbxConvert::ConvertMatrix(MeshGlobalTransform.ToMatrixWithScale());
@@ -811,52 +813,58 @@ bool FMeshDescriptionImporter::FillMeshDescriptionFromFbxMesh(FbxMesh* Mesh, con
 		}
 
 		//Fill the vertex array
-		MeshDescription->ReserveNewVertices(VertexCount);
-		for (int32 VertexIndex = 0; VertexIndex < VertexCount; ++VertexIndex)
 		{
-			int32 RealVertexIndex = VertexOffset + VertexIndex;
-			FbxVector4 FbxPosition = Mesh->GetControlPoints()[VertexIndex];
-			FbxPosition = TotalMatrix.MultT(FbxPosition);
-			const FVector VertexPosition = FFbxConvert::ConvertPos(FbxPosition);
-
-			FVertexID AddedVertexId = MeshDescription->CreateVertex();
-			VertexPositions[AddedVertexId] = (FVector3f)VertexPosition;
-			if (AddedVertexId.GetValue() != RealVertexIndex)
+			TRACE_CPUPROFILER_EVENT_SCOPE(Interchange_ImportVertices);
+			MeshDescription->ReserveNewVertices(VertexCount);
+			for (int32 VertexIndex = 0; VertexIndex < VertexCount; ++VertexIndex)
 			{
-				UInterchangeResultMeshError_Generic* Message = AddMessage<UInterchangeResultMeshError_Generic>(Mesh);
-				Message->Text = LOCTEXT("CannotCreateVertex", "Cannot create valid vertex for mesh '{MeshName}'.");
+				int32 RealVertexIndex = VertexOffset + VertexIndex;
+				FbxVector4 FbxPosition = Mesh->GetControlPoints()[VertexIndex];
+				FbxPosition = TotalMatrix.MultT(FbxPosition);
+				const FVector3f VertexPosition = FFbxConvert::ConvertPos<FVector3f>(FbxPosition);
+
+				FVertexID AddedVertexId = MeshDescription->CreateVertex();
+				VertexPositions[AddedVertexId] = VertexPosition;
+				if (AddedVertexId.GetValue() != RealVertexIndex)
+				{
+					UInterchangeResultMeshError_Generic* Message = AddMessage<UInterchangeResultMeshError_Generic>(Mesh);
+					Message->Text = LOCTEXT("CannotCreateVertex", "Cannot create valid vertex for mesh '{MeshName}'.");
 
-				return false;
+					return false;
+				}
 			}
 		}
 
-		// Fill the UV arrays
-		for (int32 UVLayerIndex = 0; UVLayerIndex < FBXUVs.UniqueUVCount; UVLayerIndex++)
 		{
-			check(FBXUVs.LayerElementUV[UVLayerIndex]);
-			if (FBXUVs.LayerElementUV[UVLayerIndex] != nullptr)
+			TRACE_CPUPROFILER_EVENT_SCOPE(Interchange_ImportUVs);
+			// Fill the UV arrays
+			for (int32 UVLayerIndex = 0; UVLayerIndex < FBXUVs.UniqueUVCount; UVLayerIndex++)
 			{
-				int32 UVCount = FBXUVs.LayerElementUV[UVLayerIndex]->GetDirectArray().GetCount();
-				if (UVCount == 0)
+				check(FBXUVs.LayerElementUV[UVLayerIndex]);
+				if (FBXUVs.LayerElementUV[UVLayerIndex] != nullptr)
 				{
-					UInterchangeResultMeshWarning_Generic* Message = AddMessage<UInterchangeResultMeshWarning_Generic>(Mesh);
-					Message->Text = LOCTEXT("CreateUVs_UVCorrupted", "Found invalid UVs value when importing mesh '{MeshName}'.");
-				}
+					int32 UVCount = FBXUVs.LayerElementUV[UVLayerIndex]->GetDirectArray().GetCount();
+					if (UVCount == 0)
+					{
+						UInterchangeResultMeshWarning_Generic* Message = AddMessage<UInterchangeResultMeshWarning_Generic>(Mesh);
+						Message->Text = LOCTEXT("CreateUVs_UVCorrupted", "Found invalid UVs value when importing mesh '{MeshName}'.");
+					}
 
-				TUVAttributesRef<FVector2f> UVCoordinates = MeshDescription->UVAttributes(UVLayerIndex).GetAttributesRef<FVector2f>(MeshAttribute::UV::UVCoordinate);
-				MeshDescription->ReserveNewUVs(UVCount, UVLayerIndex);
-				for (int32 UVIndex = 0; UVIndex < UVCount; UVIndex++)
-				{
-					FUVID UVID = MeshDescription->CreateUV(UVLayerIndex);
-					FbxVector2 UVVector = FBXUVs.LayerElementUV[UVLayerIndex]->GetDirectArray().GetAt(UVIndex);
-					UVCoordinates[UVID] = FVector2f(static_cast<float>(UVVector[0]), 1.0f - static_cast<float>(UVVector[1]));	// flip the Y of UVs for DirectX
+					TUVAttributesRef<FVector2f> UVCoordinates = MeshDescription->UVAttributes(UVLayerIndex).GetAttributesRef<FVector2f>(MeshAttribute::UV::UVCoordinate);
+					MeshDescription->ReserveNewUVs(UVCount, UVLayerIndex);
+					for (int32 UVIndex = 0; UVIndex < UVCount; UVIndex++)
+					{
+						FUVID UVID = MeshDescription->CreateUV(UVLayerIndex);
+						FbxVector2 UVVector = FBXUVs.LayerElementUV[UVLayerIndex]->GetDirectArray().GetAt(UVIndex);
+						UVCoordinates[UVID] = FVector2f(static_cast<float>(UVVector[0]), 1.0f - static_cast<float>(UVVector[1]));	// flip the Y of UVs for DirectX
+					}
 				}
 			}
 		}
 
 		TMap<uint64, int32> RemapEdgeID;
 		{
-			TRACE_CPUPROFILER_EVENT_SCOPE(BuildMeshEdgeVertices);
+			TRACE_CPUPROFILER_EVENT_SCOPE(Interchange_ImportEdgeVertices);
 			Mesh->BeginGetMeshEdgeVertices();
 
 			//Fill the edge array
@@ -889,7 +897,7 @@ bool FMeshDescriptionImporter::FillMeshDescriptionFromFbxMesh(FbxMesh* Mesh, con
 		}
 
 		{
-			TRACE_CPUPROFILER_EVENT_SCOPE(BuildMeshEdgeIndexForPolygon);
+			TRACE_CPUPROFILER_EVENT_SCOPE(Interchange_ImportPolygons);
 
 			// Compute and reserve memory to be used for vertex instances
 			{
@@ -912,7 +920,7 @@ bool FMeshDescriptionImporter::FillMeshDescriptionFromFbxMesh(FbxMesh* Mesh, con
 			// keep those for all iterations to avoid heap allocations
 			TArray<FVertexInstanceID> CornerInstanceIDs;
 			TArray<FVertexID> CornerVerticesIDs;
-			TArray<FVector, TInlineAllocator<3>> P;
+			TArray<FVector3f, TInlineAllocator<3>> P;
 
 			bool bCorruptedMsgDone = false;
 			//Polygons
@@ -934,7 +942,7 @@ bool FMeshDescriptionImporter::FillMeshDescriptionFromFbxMesh(FbxMesh* Mesh, con
 							bAllCornerValid = false;
 							break;
 						}
-						P[CornerIndex] = (FVector)VertexPositions[VertexID];
+						P[CornerIndex] = VertexPositions[VertexID];
 					}
 					if (!bAllCornerValid)
 					{
@@ -949,7 +957,7 @@ bool FMeshDescriptionImporter::FillMeshDescriptionFromFbxMesh(FbxMesh* Mesh, con
 						continue;
 					}
 					check(P.Num() > 2); //triangle is the smallest polygon we can have
-					const FVector Normal = ((P[1] - P[2]) ^ (P[0] - P[2])).GetSafeNormal(ComparisonThreshold);
+					const FVector3f Normal = ((P[1] - P[2]) ^ (P[0] - P[2])).GetSafeNormal(ComparisonThreshold);
 					//Check for degenerated polygons, avoid NAN
 					if (Normal.IsNearlyZero(ComparisonThreshold) || Normal.ContainsNaN())
 					{
@@ -971,7 +979,6 @@ bool FMeshDescriptionImporter::FillMeshDescriptionFromFbxMesh(FbxMesh* Mesh, con
 					CornerInstanceIDs[CornerIndex] = VertexInstanceID;
 					const int32 ControlPointIndex = Mesh->GetPolygonVertex(PolygonIndex, CornerIndex);
 					const FVertexID VertexID(VertexOffset + ControlPointIndex);
-					const FVector VertexPosition = (FVector)VertexPositions[VertexID];
 					CornerVerticesIDs[CornerIndex] = VertexID;
 
 					FVertexInstanceID AddedVertexInstanceId = MeshDescription->CreateVertexInstance(VertexID);
@@ -1034,8 +1041,8 @@ bool FMeshDescriptionImporter::FillMeshDescriptionFromFbxMesh(FbxMesh* Mesh, con
 
 						FbxVector4 TempValue = LayerElementNormal->GetDirectArray().GetAt(NormalValueIndex);
 						TempValue = TotalMatrixForNormal.MultT(TempValue);
-						FVector TangentZ = FFbxConvert::ConvertDir(TempValue);
-						VertexInstanceNormals[AddedVertexInstanceId] = (FVector3f)TangentZ.GetSafeNormal();
+						FVector3f TangentZ = FFbxConvert::ConvertDir<FVector3f>(TempValue);
+						VertexInstanceNormals[AddedVertexInstanceId] = TangentZ.GetSafeNormal();
 						//tangents and binormals share the same reference, mapping mode and index array
 						if (bHasNTBInformation)
 						{
@@ -1046,8 +1053,8 @@ bool FMeshDescriptionImporter::FillMeshDescriptionFromFbxMesh(FbxMesh* Mesh, con
 
 							TempValue = LayerElementTangent->GetDirectArray().GetAt(TangentValueIndex);
 							TempValue = TotalMatrixForNormal.MultT(TempValue);
-							FVector TangentX = FFbxConvert::ConvertDir(TempValue);
-							VertexInstanceTangents[AddedVertexInstanceId] = (FVector3f)TangentX.GetSafeNormal();
+							FVector3f TangentX = FFbxConvert::ConvertDir<FVector3f>(TempValue);
+							VertexInstanceTangents[AddedVertexInstanceId] = TangentX.GetSafeNormal();
 
 							int BinormalMapIndex = (BinormalMappingMode == FbxLayerElement::eByControlPoint) ?
 								ControlPointIndex : RealFbxVertexIndex;
@@ -1056,7 +1063,7 @@ bool FMeshDescriptionImporter::FillMeshDescriptionFromFbxMesh(FbxMesh* Mesh, con
 
 							TempValue = LayerElementBinormal->GetDirectArray().GetAt(BinormalValueIndex);
 							TempValue = TotalMatrixForNormal.MultT(TempValue);
-							FVector TangentY = -FFbxConvert::ConvertDir(TempValue);
+							FVector3f TangentY = -FFbxConvert::ConvertDir<FVector3f>(TempValue);
 							VertexInstanceBinormalSigns[AddedVertexInstanceId] = FbxGetBasisDeterminantSign(TangentX.GetSafeNormal(), TangentY.GetSafeNormal(), TangentZ.GetSafeNormal());
 						}
 					}
@@ -1067,10 +1074,10 @@ bool FMeshDescriptionImporter::FillMeshDescriptionFromFbxMesh(FbxMesh* Mesh, con
 				if (!bHasNonDegeneratePolygons)
 				{
 					float TriangleComparisonThreshold = THRESH_POINTS_ARE_SAME;
-					FVector VertexPosition[3];
-					VertexPosition[0] = (FVector)VertexPositions[CornerVerticesIDs[0]];
-					VertexPosition[1] = (FVector)VertexPositions[CornerVerticesIDs[1]];
-					VertexPosition[2] = (FVector)VertexPositions[CornerVerticesIDs[2]];
+					FVector3f VertexPosition[3];
+					VertexPosition[0] = VertexPositions[CornerVerticesIDs[0]];
+					VertexPosition[1] = VertexPositions[CornerVerticesIDs[1]];
+					VertexPosition[2] = VertexPositions[CornerVerticesIDs[2]];
 					if (!(VertexPosition[0].Equals(VertexPosition[1], TriangleComparisonThreshold)
 						|| VertexPosition[0].Equals(VertexPosition[2], TriangleComparisonThreshold)
 						|| VertexPosition[1].Equals(VertexPosition[2], TriangleComparisonThreshold)))
@@ -1269,12 +1276,6 @@ bool FMeshDescriptionImporter::FillMeshDescriptionFromFbxMesh(FbxMesh* Mesh, con
 				Mesh->EndGetMeshEdgeIndexForPolygon();
 			}
 
-			MeshDescription->ResumeVertexInstanceIndexing();
-			MeshDescription->ResumeEdgeIndexing();
-			MeshDescription->ResumePolygonIndexing();
-			MeshDescription->ResumePolygonGroupIndexing();
-			MeshDescription->ResumeUVIndexing();
-
 			if (SkippedVertexInstance > 0)
 			{
 				check(MeshDescription->Triangles().Num() == MeshDescription->Triangles().GetArraySize());
@@ -1283,6 +1284,7 @@ bool FMeshDescriptionImporter::FillMeshDescriptionFromFbxMesh(FbxMesh* Mesh, con
 
 		if (MeshType == EMeshType::Skinned)
 		{
+			TRACE_CPUPROFILER_EVENT_SCOPE(Interchange_ImportSkin);
 			FSkeletalMeshAttributes SkeletalMeshAttributes(*MeshDescription);
 			SkeletalMeshAttributes.Register(true);
 
@@ -1379,6 +1381,12 @@ bool FMeshDescriptionImporter::FillMeshDescriptionFromFbxMesh(FbxMesh* Mesh, con
 				VertexSkinWeights.Set(Item.Key, Item.Value);
 			}
 		}
+
+		MeshDescription->ResumeVertexInstanceIndexing();
+		MeshDescription->ResumeEdgeIndexing();
+		MeshDescription->ResumePolygonIndexing();
+		MeshDescription->ResumePolygonGroupIndexing();
+		MeshDescription->ResumeUVIndexing();
 	}
 
 	TArray<FPolygonGroupID> EmptyPolygonGroups;
@@ -1427,7 +1435,10 @@ bool FMeshDescriptionImporter::IsOddNegativeScale(FbxAMatrix& TotalMatrix)
 //////////////////////////////////////////////////////////////////////////
 /// FMeshPayloadContext implementation
 
-bool FMeshPayloadContext::FetchMeshPayloadToFile(FFbxParser& Parser, const FTransform& MeshGlobalTransform, const FString& PayloadFilepath)
+bool FMeshPayloadContext::FetchMeshPayloadInternal(FFbxParser& Parser
+	, const FTransform& MeshGlobalTransform
+	, FMeshDescription& OutMeshDescription
+	, TArray<FString>& OutJointNames)
 {
 	if (!ensure(SDKScene != nullptr))
 	{
@@ -1453,14 +1464,12 @@ bool FMeshPayloadContext::FetchMeshPayloadToFile(FFbxParser& Parser, const FTran
 		return false;
 	}
 
-	TArray<FString> JointUniqueNames;
-	FMeshDescription MeshDescription;
 	if (bIsSkinnedMesh)
 	{
-		FSkeletalMeshAttributes SkeletalMeshAttribute(MeshDescription);
+		FSkeletalMeshAttributes SkeletalMeshAttribute(OutMeshDescription);
 		SkeletalMeshAttribute.Register();
-		FMeshDescriptionImporter MeshDescriptionImporter(Parser, &MeshDescription, SDKScene, SDKGeometryConverter);
-		if (!MeshDescriptionImporter.FillSkinnedMeshDescriptionFromFbxMesh(Mesh, MeshGlobalTransform, JointUniqueNames))
+		FMeshDescriptionImporter MeshDescriptionImporter(Parser, &OutMeshDescription, SDKScene, SDKGeometryConverter);
+		if (!MeshDescriptionImporter.FillSkinnedMeshDescriptionFromFbxMesh(Mesh, MeshGlobalTransform, OutJointNames))
 		{
 			UInterchangeResultError_Generic* Message = Parser.AddMessage<UInterchangeResultError_Generic>();
 			Message->InterchangeKey = Parser.GetFbxHelper()->GetMeshUniqueID(Mesh);
@@ -1470,9 +1479,9 @@ bool FMeshPayloadContext::FetchMeshPayloadToFile(FFbxParser& Parser, const FTran
 	}
 	else
 	{
-		FStaticMeshAttributes StaticMeshAttribute(MeshDescription);
+		FStaticMeshAttributes StaticMeshAttribute(OutMeshDescription);
 		StaticMeshAttribute.Register();
-		FMeshDescriptionImporter MeshDescriptionImporter(Parser, &MeshDescription, SDKScene, SDKGeometryConverter);
+		FMeshDescriptionImporter MeshDescriptionImporter(Parser, &OutMeshDescription, SDKScene, SDKGeometryConverter);
 		if (!MeshDescriptionImporter.FillStaticMeshDescriptionFromFbxMesh(Mesh, MeshGlobalTransform))
 		{
 			UInterchangeResultError_Generic* Message = Parser.AddMessage<UInterchangeResultError_Generic>();
@@ -1481,6 +1490,18 @@ bool FMeshPayloadContext::FetchMeshPayloadToFile(FFbxParser& Parser, const FTran
 			return false;
 		}
 	}
+	return true;
+}
+
+bool FMeshPayloadContext::FetchMeshPayloadToFile(FFbxParser& Parser, const FTransform& MeshGlobalTransform, const FString& PayloadFilepath)
+{
+	FMeshDescription MeshDescription;
+	TArray<FString> JointNames;
+	if (!FetchMeshPayloadInternal(Parser, MeshGlobalTransform, MeshDescription, JointNames))
+	{
+		return false;
+	}
+
 	//Dump the MeshDescription to a file
 	{
 		FLargeMemoryWriter Ar;
@@ -1490,7 +1511,7 @@ bool FMeshPayloadContext::FetchMeshPayloadToFile(FFbxParser& Parser, const FTran
 		if (bIsSkinnedMesh)
 		{
 			//When passing a skinned MeshDescription, We want to pass the joint Node ID so we can know what the influence bone index refer to
-			Ar << JointUniqueNames;
+			Ar << JointNames;
 		}
 		uint8* ArchiveData = Ar.GetData();
 		int64 ArchiveSize = Ar.TotalSize();
@@ -1500,22 +1521,52 @@ bool FMeshPayloadContext::FetchMeshPayloadToFile(FFbxParser& Parser, const FTran
 
 	return true;
 }
-
+#if WITH_ENGINE
+bool FMeshPayloadContext::FetchMeshPayload(FFbxParser& Parser, const FTransform& MeshGlobalTransform, FMeshPayloadData& OutMeshPayloadData)
+{
+	if (FetchMeshPayloadInternal(Parser, MeshGlobalTransform, OutMeshPayloadData.MeshDescription, OutMeshPayloadData.JointNames))
+	{
+		OutMeshPayloadData.GlobalTransform = MeshGlobalTransform;
+		return true;
+	}
+	return false;
+}
+#endif
 //////////////////////////////////////////////////////////////////////////
 /// FMorphTargetPayloadContext implementation
 
-bool FMorphTargetPayloadContext::FetchMeshPayloadToFile(FFbxParser& Parser, const FTransform& MeshGlobalTransform, const FString& PayloadFilepath)
+bool FMorphTargetPayloadContext::FetchMeshPayloadInternal(FFbxParser& Parser
+	, const FTransform& MeshGlobalTransform
+	, FMeshDescription& OutMorphTargetMeshDescription)
 {
+	if (!ensure(SDKScene != nullptr))
+	{
+		UInterchangeResultError_Generic* Message = Parser.AddMessage<UInterchangeResultError_Generic>();
+		Message->InterchangeKey = Parser.GetFbxHelper()->GetMeshUniqueID(Shape);
+		Message->Text = LOCTEXT("FBXSceneNull_Mesh", "Cannot fetch FBX mesh morph shape payload because the FBX scene is null.");
+		return false;
+	}
+
+	if (!ensure(SDKGeometryConverter != nullptr))
+	{
+		UInterchangeResultError_Generic* Message = Parser.AddMessage<UInterchangeResultError_Generic>();
+		Message->InterchangeKey = Parser.GetFbxHelper()->GetMeshUniqueID(Shape);
+		Message->Text = LOCTEXT("FBXConverterNull", "Cannot fetch FBX mesh morph shape payload because the FBX geometry converter is null.");
+		return false;
+	}
+
 	if (!ensure(Shape))
 	{
-		//Todo log an error
+		UInterchangeResultError_Generic* Message = Parser.AddMessage<UInterchangeResultError_Generic>();
+		Message->InterchangeKey = Parser.GetFbxHelper()->GetMeshUniqueID(Shape);
+		Message->Text = LOCTEXT("FBXMeshNull", "Cannot fetch FBX mesh morph shape payload because the FBX shape is null.");
 		return false;
 	}
+
 	//Import the MorphTarget
-	FMeshDescription MorphTargetMeshDescription;
-	FStaticMeshAttributes StaticMeshAttribute(MorphTargetMeshDescription);
+	FStaticMeshAttributes StaticMeshAttribute(OutMorphTargetMeshDescription);
 	StaticMeshAttribute.Register();
-	FMeshDescriptionImporter MeshDescriptionImporter(Parser, &MorphTargetMeshDescription, SDKScene, SDKGeometryConverter);
+	FMeshDescriptionImporter MeshDescriptionImporter(Parser, &OutMorphTargetMeshDescription, SDKScene, SDKGeometryConverter);
 	if (!MeshDescriptionImporter.FillMeshDescriptionFromFbxShape(Shape, MeshGlobalTransform))
 	{
 		UInterchangeResultError_Generic* Message = Parser.AddMessage<UInterchangeResultError_Generic>();
@@ -1524,7 +1575,16 @@ bool FMorphTargetPayloadContext::FetchMeshPayloadToFile(FFbxParser& Parser, cons
 
 		return false;
 	}
+	return true;
+}
 
+bool FMorphTargetPayloadContext::FetchMeshPayloadToFile(FFbxParser& Parser, const FTransform& MeshGlobalTransform, const FString& PayloadFilepath)
+{
+	FMeshDescription MorphTargetMeshDescription;
+	if(!FetchMeshPayloadInternal(Parser, MeshGlobalTransform, MorphTargetMeshDescription))
+	{
+		return false;
+	}
 	//Dump the MeshDescription to a file
 	{
 		FLargeMemoryWriter Ar;
@@ -1537,6 +1597,12 @@ bool FMorphTargetPayloadContext::FetchMeshPayloadToFile(FFbxParser& Parser, cons
 	return true;
 }
 
+#if WITH_ENGINE
+bool FMorphTargetPayloadContext::FetchMeshPayload(FFbxParser& Parser, const FTransform& MeshGlobalTransform, FMeshPayloadData& OutMeshPayloadData)
+{
+	return FetchMeshPayloadInternal(Parser, MeshGlobalTransform, OutMeshPayloadData.MeshDescription);
+}
+#endif
 //////////////////////////////////////////////////////////////////////////
 /// FFbxMesh implementation
 
@@ -1639,7 +1705,7 @@ void FFbxMesh::AddAllMeshes(FbxScene* SDKScene, FbxGeometryConverter* SDKGeometr
 		MeshNode->SetCustomVertexCount(MeshVertexCount);
 		const int32 MeshPolygonCount = Mesh->GetPolygonCount();
 		MeshNode->SetCustomPolygonCount(MeshPolygonCount);
-		const FBox MeshBoundingBox = FBox(FFbxConvert::ConvertPos(Mesh->BBoxMin.Get()), FFbxConvert::ConvertPos(Mesh->BBoxMax.Get()));
+		const FBox MeshBoundingBox = FBox(FFbxConvert::ConvertPos<FVector>(Mesh->BBoxMin.Get()), FFbxConvert::ConvertPos<FVector>(Mesh->BBoxMax.Get()));
 		MeshNode->SetCustomBoundingBox(MeshBoundingBox);
 		const bool bMeshHasVertexNormal = Mesh->GetElementNormalCount() > 0;
 		MeshNode->SetCustomHasVertexNormal(bMeshHasVertexNormal);
diff --git a/Engine/Plugins/Interchange/Runtime/Source/Parsers/Fbx/Private/FbxMesh.h b/Engine/Plugins/Interchange/Runtime/Source/Parsers/Fbx/Private/FbxMesh.h
--- a/Engine/Plugins/Interchange/Runtime/Source/Parsers/Fbx/Private/FbxMesh.h
+++ b/Engine/Plugins/Interchange/Runtime/Source/Parsers/Fbx/Private/FbxMesh.h
@@ -16,6 +16,10 @@ namespace UE
 {
 	namespace Interchange
 	{
+#if WITH_ENGINE
+		struct FMeshPayloadData;
+#endif
+
 		namespace Private
 		{
 			struct FMorphTargetAnimationBuildingData
@@ -128,13 +132,13 @@ namespace UE
 				bool IsOddNegativeScale(FbxAMatrix& TotalMatrix);
 				
 				//TODO move the real function from RenderCore to FVector, so we do not have to add render core to compute such a simple thing
-				float FbxGetBasisDeterminantSign(const FVector& XAxis, const FVector& YAxis, const FVector& ZAxis)
+				float FbxGetBasisDeterminantSign(const FVector3f& XAxis, const FVector3f& YAxis, const FVector3f& ZAxis)
 				{
-					FMatrix Basis(
-						FPlane(XAxis, 0),
-						FPlane(YAxis, 0),
-						FPlane(ZAxis, 0),
-						FPlane(0, 0, 0, 1)
+					FMatrix44f Basis(
+						FPlane4f(XAxis, 0),
+						FPlane4f(YAxis, 0),
+						FPlane4f(ZAxis, 0),
+						FPlane4f(0, 0, 0, 1)
 					);
 					return (Basis.Determinant() < 0) ? -1.0f : +1.0f;
 				}
@@ -152,10 +156,18 @@ namespace UE
 				virtual ~FMeshPayloadContext() {}
 				virtual FString GetPayloadType() const override { return TEXT("Mesh-PayloadContext"); }
 				virtual bool FetchMeshPayloadToFile(FFbxParser& Parser, const FTransform& MeshGlobalTransform, const FString& PayloadFilepath) override;
+#if WITH_ENGINE
+				virtual bool FetchMeshPayload(FFbxParser& Parser, const FTransform& MeshGlobalTransform, FMeshPayloadData& OutMeshPayloadData) override;
+#endif
 				bool bIsSkinnedMesh = false;
 				FbxMesh* Mesh = nullptr;
 				FbxScene* SDKScene = nullptr;
 				FbxGeometryConverter* SDKGeometryConverter = nullptr;
+			private:
+				bool FetchMeshPayloadInternal(FFbxParser& Parser
+					, const FTransform& MeshGlobalTransform
+					, FMeshDescription& OutMeshDescription
+					, TArray<FString>& OutJointNames);
 			};
 
 			class FMorphTargetPayloadContext : public FPayloadContextBase
@@ -164,9 +176,16 @@ namespace UE
 				virtual ~FMorphTargetPayloadContext() {}
 				virtual FString GetPayloadType() const override { return TEXT("MorphTarget-PayloadContext"); }
 				virtual bool FetchMeshPayloadToFile(FFbxParser& Parser, const FTransform& MeshGlobalTransform, const FString& PayloadFilepath) override;
+#if WITH_ENGINE
+				virtual bool FetchMeshPayload(FFbxParser& Parser, const FTransform& MeshGlobalTransform, FMeshPayloadData& OutMeshPayloadData) override;
+#endif
 				FbxShape* Shape = nullptr;
 				FbxScene* SDKScene = nullptr;
 				FbxGeometryConverter* SDKGeometryConverter = nullptr;
+			private:
+				bool FetchMeshPayloadInternal(FFbxParser& Parser
+					, const FTransform& MeshGlobalTransform
+					, FMeshDescription& OutMeshDescription);
 			};
 
 			class FFbxMesh
diff --git a/Engine/Plugins/Interchange/Runtime/Source/Parsers/Fbx/Private/FbxScene.cpp b/Engine/Plugins/Interchange/Runtime/Source/Parsers/Fbx/Private/FbxScene.cpp
--- a/Engine/Plugins/Interchange/Runtime/Source/Parsers/Fbx/Private/FbxScene.cpp
+++ b/Engine/Plugins/Interchange/Runtime/Source/Parsers/Fbx/Private/FbxScene.cpp
@@ -169,7 +169,7 @@ namespace UE
 				
 				auto GetConvertedTransform = [Node](FbxAMatrix& NewFbxMatrix)
 				{
-					FTransform Transform = FTransform(FFbxConvert::ConvertMatrix(NewFbxMatrix));
+					FTransform Transform = FFbxConvert::ConvertTransform<FTransform, FVector, FQuat>(NewFbxMatrix);
 					
 					if (FbxNodeAttribute* NodeAttribute = Node->GetNodeAttribute())
 					{
diff --git a/Engine/Plugins/Interchange/Runtime/Source/Parsers/Fbx/Private/InterchangeFbxParser.cpp b/Engine/Plugins/Interchange/Runtime/Source/Parsers/Fbx/Private/InterchangeFbxParser.cpp
--- a/Engine/Plugins/Interchange/Runtime/Source/Parsers/Fbx/Private/InterchangeFbxParser.cpp
+++ b/Engine/Plugins/Interchange/Runtime/Source/Parsers/Fbx/Private/InterchangeFbxParser.cpp
@@ -198,6 +198,20 @@ namespace UE::Interchange
 		return ResultPayloadUniqueId;
 	}
 
+#if WITH_ENGINE
+	void FInterchangeFbxParser::FetchMeshPayload(const FString& PayloadKey, const FTransform& MeshGlobalTransform, FMeshPayloadData& OutMeshPayloadData)
+	{
+		if (!FbxParserPrivate->FetchMeshPayloadData(PayloadKey, MeshGlobalTransform, OutMeshPayloadData))
+		{
+			if (UInterchangeResultError_Generic* Error = AddMessage<UInterchangeResultError_Generic>())
+			{
+				Error->SourceAssetName = SourceFilename;
+				Error->Text = LOCTEXT("CantFetchPayload", "Cannot fetch FBX payload data.");
+			}
+		}
+	}
+#endif
+
 	FString FInterchangeFbxParser::FetchAnimationBakeTransformPayload(const FString& PayloadKey, const double BakeFrequency, const double RangeStartTime, const double RangeEndTime, const FString& ResultFolder)
 	{
 		ResultsContainer->Empty();
diff --git a/Engine/Plugins/Interchange/Runtime/Source/Parsers/Fbx/Public/InterchangeFbxParser.h b/Engine/Plugins/Interchange/Runtime/Source/Parsers/Fbx/Public/InterchangeFbxParser.h
--- a/Engine/Plugins/Interchange/Runtime/Source/Parsers/Fbx/Public/InterchangeFbxParser.h
+++ b/Engine/Plugins/Interchange/Runtime/Source/Parsers/Fbx/Public/InterchangeFbxParser.h
@@ -4,6 +4,9 @@
 
 #include "CoreMinimal.h"
 #include "HAL/CriticalSection.h"
+#if WITH_ENGINE
+#include "Mesh/InterchangeMeshPayload.h"
+#endif
 #include "Misc/ScopeLock.h"
 #include "Nodes/InterchangeBaseNodeContainer.h"
 #include "InterchangeResultsContainer.h"
@@ -57,6 +60,16 @@ namespace UE
 			 */
 			FString FetchMeshPayload(const FString& PayloadKey, const FTransform& MeshGlobalTransform, const FString& ResultFolder);
 
+#if WITH_ENGINE
+			/**
+			 * Extract mesh payload data from the fbx, the key tell the translator what payload the client ask
+			 * @param - PayloadKey is the key that describe the payload data to extract from the fbx file
+			 * @param - MeshGlobalTransform is the transform we want to apply to the mesh vertex
+			 * @param - OutMeshPayloadData structure receiving the data
+			 */
+			void FetchMeshPayload(const FString& PayloadKey, const FTransform& MeshGlobalTransform, FMeshPayloadData& OutMeshPayloadData);
+#endif
+			
 			/**
 			 * Extract bake transform animation payload data from the fbx, the key tell the translator what payload the client ask
 			 * @param - PayloadKey is the key that describe the payload data to extract from the fbx file
