Interchange: Fbx translator now report texture path that doesn't exist

#rnx

diff --git a/Engine/Plugins/Interchange/Runtime/Source/Import/Private/Fbx/InterchangeFbxTranslator.cpp b/Engine/Plugins/Interchange/Runtime/Source/Import/Private/Fbx/InterchangeFbxTranslator.cpp
--- a/Engine/Plugins/Interchange/Runtime/Source/Import/Private/Fbx/InterchangeFbxTranslator.cpp
+++ b/Engine/Plugins/Interchange/Runtime/Source/Import/Private/Fbx/InterchangeFbxTranslator.cpp
@@ -225,6 +225,7 @@ bool UInterchangeFbxTranslator::Translate(UInterchangeBaseNodeContainer& BaseNod
 	else
 	{
 #if WITH_EDITOR
+		FbxParser.SetResultContainer(Results);
 		FbxParser.LoadFbxFile(Filename, BaseNodeContainer);
 #endif
 	}
diff --git a/Engine/Plugins/Interchange/Runtime/Source/Parsers/Fbx/Private/FbxMaterial.cpp b/Engine/Plugins/Interchange/Runtime/Source/Parsers/Fbx/Private/FbxMaterial.cpp
--- a/Engine/Plugins/Interchange/Runtime/Source/Parsers/Fbx/Private/FbxMaterial.cpp
+++ b/Engine/Plugins/Interchange/Runtime/Source/Parsers/Fbx/Private/FbxMaterial.cpp
@@ -454,6 +454,9 @@ namespace UE
 					//Only import texture that exist on disk
 					if (!FPaths::FileExists(TextureFilename))
 					{
+						UInterchangeResultTextureWarning_TextureFileDoNotExist* Message = Parser.AddMessage<UInterchangeResultTextureWarning_TextureFileDoNotExist>();
+						Message->TextureName = TextureFilename;
+						Message->MaterialName.Empty();
 						continue;
 					}
 					//Create a texture node and make it child of the material node
