Interchange: Make sure to notify the users if an animated bone name exceed 92 characters.

#rnx

diff --git a/Engine/Plugins/Interchange/Runtime/Source/Import/Private/Animation/InterchangeAnimSequenceFactory.cpp b/Engine/Plugins/Interchange/Runtime/Source/Import/Private/Animation/InterchangeAnimSequenceFactory.cpp
--- a/Engine/Plugins/Interchange/Runtime/Source/Import/Private/Animation/InterchangeAnimSequenceFactory.cpp
+++ b/Engine/Plugins/Interchange/Runtime/Source/Import/Private/Animation/InterchangeAnimSequenceFactory.cpp
@@ -399,7 +399,8 @@ namespace UE::Interchange::Private
 		, const IInterchangeAnimationPayloadInterface* AnimSequenceTranslatorPayloadInterface
 		, const FString& AssetName
 		, const bool bIsReimporting
-		, TArray<FString> OutCurvesNotFound)
+		, TArray<FString> OutCurvesNotFound
+		, UInterchangeAnimSequenceFactory* Factory)
 	{
 		TMap<const UInterchangeSceneNode*, TFuture<TOptional<UE::Interchange::FAnimationPayloadData>>> AnimationPayloads;
 
@@ -486,7 +487,10 @@ namespace UE::Interchange::Private
 				if (!ensure(AnimationTransformPayload.Transforms.Num() == BakeKeyCountForAnimationPayload))
 				{
 					FString PayloadKey = PayloadKeys[AnimationPayload.Key->GetUniqueID()].UniqueId;
-					UE_LOG(LogInterchangeImport, Warning, TEXT("Animation Payload [%s] has unexpected number of Baked Transforms."), *PayloadKey);
+					UInterchangeResultWarning_Generic* Message = Factory->AddMessage<UInterchangeResultWarning_Generic>();
+					Message->DestinationAssetName = AssetName;
+					Message->AssetType = UAnimSequence::StaticClass();
+					Message->Text = FText::Format(NSLOCTEXT("UInterchangeAnimSequenceFactory", "AnimationPayloadUnexpectedTransformNumber", "Animation Payload [{0}] has unexpected number of Baked Transforms."), FText::FromString(PayloadKey));
 					break;
 				}
 
@@ -498,7 +502,10 @@ namespace UE::Interchange::Private
 						  ensure(FMath::IsNearlyEqual(AnimationTransformPayload.RangeEndTime, RangeEnd, UE_DOUBLE_KINDA_SMALL_NUMBER))))
 					{
 						FString PayloadKey = PayloadKeys[AnimationPayload.Key->GetUniqueID()].UniqueId;
-						UE_LOG(LogInterchangeImport, Warning, TEXT("Animation Payload [%s] 's BakeFrequency, RangeStartTime and RangeEndTime does not equal with the provided one."), *PayloadKey);
+						UInterchangeResultWarning_Generic* Message = Factory->AddMessage<UInterchangeResultWarning_Generic>();
+						Message->DestinationAssetName = AssetName;
+						Message->AssetType = UAnimSequence::StaticClass();
+						Message->Text = FText::Format(NSLOCTEXT("UInterchangeAnimSequenceFactory", "AnimationPayloadBakeFrequencyNotTheExpected", "Animation Payload [{0}]'s BakeFrequency, RangeStartTime and RangeEndTime does not equal with the provided one."), FText::FromString(PayloadKey));
 					}
 				}
 
@@ -547,11 +554,25 @@ namespace UE::Interchange::Private
 					&& TimeKeys.Num() == BakeKeyCountForAnimationPayload))
 				{
 					FString PayloadKey = PayloadKeys[AnimationPayload.Key->GetUniqueID()].UniqueId;
-					UE_LOG(LogInterchangeImport, Warning, TEXT("Animation Payload [%s] has unexpected number of animation keys. Animation will be incorrect."), *PayloadKey);
+					UInterchangeResultWarning_Generic* Message = Factory->AddMessage<UInterchangeResultWarning_Generic>();
+					Message->DestinationAssetName = AssetName;
+					Message->AssetType = UAnimSequence::StaticClass();
+					Message->Text = FText::Format(NSLOCTEXT("UInterchangeAnimSequenceFactory", "AnimationPayloadBadKeyNumber", "Animation Payload [{0}] has unexpected number of animation keys. Animation will be incorrect."), FText::FromString(PayloadKey));
 					continue;
 				}
 
 				//add new track
+				if (BoneName.GetStringLength() > 92)
+				{
+					//The bone name exceed the maximum length supported by the animation system
+					//The animation system is adding _CONTROL to the bone name to name the animation controller and
+					//the maximum total length is cap at 100, so user should not import bone name longer then 92 characters
+					UInterchangeResultWarning_Generic* Message = Factory->AddMessage<UInterchangeResultWarning_Generic>();
+					Message->DestinationAssetName = AssetName;
+					Message->AssetType = UAnimSequence::StaticClass();
+					Message->Text = FText::Format(NSLOCTEXT("UInterchangeAnimSequenceFactory", "BoneNameExceed92Characters", "Bone with animation cannot have a name exceeding 92 characters: {0}"), FText::FromName(BoneName));
+					continue;
+				}
 				Controller.AddBoneCurve(BoneName, bShouldTransact);
 				Controller.SetBoneTrackKeys(BoneName, RawTrack.PosKeys, RawTrack.RotKeys, RawTrack.ScaleKeys, bShouldTransact);
 			}
@@ -717,7 +738,10 @@ namespace UE::Interchange::Private
 						TOptional<UE::Interchange::FAnimationPayloadData> AnimationCurvePayload = CurveNameAndPayload.Value.Get();
 						if (!AnimationCurvePayload.IsSet())
 						{
-							UE_LOG(LogInterchangeImport, Warning, TEXT("Invalid animation material curve payload key [%s] AnimSequence asset %s"), *CurveName, *AssetName);
+							UInterchangeResultWarning_Generic* Message = Factory->AddMessage<UInterchangeResultWarning_Generic>();
+							Message->DestinationAssetName = AssetName;
+							Message->AssetType = UAnimSequence::StaticClass();
+							Message->Text = FText::Format(NSLOCTEXT("UInterchangeAnimSequenceFactory", "AnimationPayloadInvalidCurve", "Invalid animation curve payload key [{0}] AnimSequence asset {1}"), FText::FromString(CurveName), FText::FromString(AssetName));
 							continue;
 						}
 
@@ -811,7 +835,10 @@ namespace UE::Interchange::Private
 						TOptional<UE::Interchange::FAnimationPayloadData> AnimationCurvePayload = CurveNameAndPayload.Value.Get();
 						if (!AnimationCurvePayload.IsSet())
 						{
-							UE_LOG(LogInterchangeImport, Warning, TEXT("Invalid animation material curve payload key [%s] AnimSequence asset %s"), *CurveNameAndPayload.Key, *AssetName);
+							UInterchangeResultWarning_Generic* Message = Factory->AddMessage<UInterchangeResultWarning_Generic>();
+							Message->DestinationAssetName = AssetName;
+							Message->AssetType = UAnimSequence::StaticClass();
+							Message->Text = FText::Format(NSLOCTEXT("UInterchangeAnimSequenceFactory", "AnimationPayloadInvalidCurvePayloadKey", "Invalid animation curve payload key [{0}] AnimSequence asset {1}"), FText::FromString(CurveNameAndPayload.Key), FText::FromString(AssetName));
 							continue;
 						}
 						FAnimationPayloadData& AnimationPayloadData = AnimationCurvePayload.GetValue();
@@ -824,7 +851,10 @@ namespace UE::Interchange::Private
 						TOptional<UE::Interchange::FAnimationPayloadData> AnimationStepCurvePayload = StepCurveNameAndPayload.Value.Get();
 						if (!AnimationStepCurvePayload.IsSet())
 						{
-							UE_LOG(LogInterchangeImport, Warning, TEXT("Invalid animation material curve payload key [%s] AnimSequence asset %s"), *StepCurveNameAndPayload.Key, *AssetName);
+							UInterchangeResultWarning_Generic* Message = Factory->AddMessage<UInterchangeResultWarning_Generic>();
+							Message->DestinationAssetName = AssetName;
+							Message->AssetType = UAnimSequence::StaticClass();
+							Message->Text = FText::Format(NSLOCTEXT("UInterchangeAnimSequenceFactory", "AnimationPayloadInvalidStepCurvePayloadKey", "Invalid animation curve payload key [{0}] AnimSequence asset {1}"), FText::FromString(StepCurveNameAndPayload.Key), FText::FromString(AssetName));
 							continue;
 						}
 
@@ -1270,7 +1300,8 @@ UInterchangeFactoryBase::FImportAssetResult UInterchangeAnimSequenceFactory::End
 			, AnimSequenceTranslatorPayloadInterface
 			, Arguments.AssetName
 			, bIsReImport
-			, CurvesNotFound);
+			, CurvesNotFound
+			, this);
 
 		if (CurvesNotFound.Num())
 		{
