Interchange: Add warning for bones that are not bind if not using T0 bind pose. Change UE_LOG to interchange result message to make sure the user is notify if anything go wrong when importing a skeletal mesh.

#rnx

diff --git a/Engine/Plugins/Interchange/Runtime/Source/Import/Private/Mesh/InterchangeSkeletalMeshFactory.cpp b/Engine/Plugins/Interchange/Runtime/Source/Import/Private/Mesh/InterchangeSkeletalMeshFactory.cpp
--- a/Engine/Plugins/Interchange/Runtime/Source/Import/Private/Mesh/InterchangeSkeletalMeshFactory.cpp
+++ b/Engine/Plugins/Interchange/Runtime/Source/Import/Private/Mesh/InterchangeSkeletalMeshFactory.cpp
@@ -1013,29 +1013,41 @@ UInterchangeFactoryBase::FImportAssetResult UInterchangeSkeletalMeshFactory::Beg
 	int32 CurrentLodIndex = 0;
 
 	UE::Interchange::Private::FContentInfo ContentInfo = UE::Interchange::Private::GetContentInfo(SkeletalMeshFactoryNode, bIsReImport);
-
 	ImportAssetObjectData.bIsReImport = bIsReImport;
 	ImportAssetObjectData.bApplyGeometryOnly = ContentInfo.bApplyGeometryOnly;
+
 	for (int32 LodIndex = 0; LodIndex < LodCount; ++LodIndex)
 	{
+		FText WarningMessage_InvalidSkeleton = FText::Format(NSLOCTEXT("InterchangeSkeletalMeshFactory", "BeginImportAsset_GameThread_InvalidSkeletonLOD", "Invalid Skeleton LOD {0} when importing SkeletalMesh asset {1}")
+			, FText::AsNumber(LodIndex)
+			, FText::FromString(Arguments.AssetName));
+		FText WarningMessage_InvalidRootJoint = FText::Format(NSLOCTEXT("InterchangeSkeletalMeshFactory", "BeginImportAsset_GameThread_InvalidSkeletonRootJoint", "Invalid Skeleton LOD {0} Root Joint when importing SkeletalMesh asset {1}")
+			, FText::AsNumber(LodIndex)
+			, FText::FromString(Arguments.AssetName));
+
 		FString LodUniqueId = LodDataUniqueIds[LodIndex];
 		const UInterchangeSkeletalMeshLodDataNode* LodDataNode = Cast<UInterchangeSkeletalMeshLodDataNode>(Arguments.NodeContainer->GetNode(LodUniqueId));
 		if (!LodDataNode)
 		{
-			UE_LOG(LogInterchangeImport, Warning, TEXT("Invalid LOD when importing SkeletalMesh asset %s"), *Arguments.AssetName);
+			UInterchangeResultWarning_Generic* Message = AddMessage<UInterchangeResultWarning_Generic>();
+			Message->Text = FText::Format(NSLOCTEXT("InterchangeSkeletalMeshFactory", "BeginImportAsset_GameThread_InvalidLOD", "Invalid LOD {0} when importing SkeletalMesh asset {1}")
+				, FText::AsNumber(LodIndex)
+				, FText::FromString(Arguments.AssetName));
 			continue;
 		}
 
 		FString SkeletonNodeUid;
 		if (!LodDataNode->GetCustomSkeletonUid(SkeletonNodeUid))
 		{
-			UE_LOG(LogInterchangeImport, Warning, TEXT("Invalid Skeleton LOD when importing SkeletalMesh asset %s"), *Arguments.AssetName);
+			UInterchangeResultWarning_Generic* Message = AddMessage<UInterchangeResultWarning_Generic>();
+			Message->Text = WarningMessage_InvalidSkeleton;
 			continue;
 		}
 		const UInterchangeSkeletonFactoryNode* SkeletonNode = Cast<UInterchangeSkeletonFactoryNode>(Arguments.NodeContainer->GetNode(SkeletonNodeUid));
 		if (!SkeletonNode)
 		{
-			UE_LOG(LogInterchangeImport, Warning, TEXT("Invalid Skeleton LOD when importing SkeletalMesh asset %s"), *Arguments.AssetName);
+			UInterchangeResultWarning_Generic* Message = AddMessage<UInterchangeResultWarning_Generic>();
+			Message->Text = WarningMessage_InvalidSkeleton;
 			continue;
 		}
 		FSoftObjectPath SkeletonNodeReferenceObject;
@@ -1065,7 +1077,8 @@ UInterchangeFactoryBase::FImportAssetResult UInterchangeSkeletalMeshFactory::Beg
 
 			if (!ensure(SkeletonReference))
 			{
-				UE_LOG(LogInterchangeImport, Warning, TEXT("Invalid Skeleton LOD when importing SkeletalMesh asset %s"), *Arguments.AssetName);
+				UInterchangeResultWarning_Generic* Message = AddMessage<UInterchangeResultWarning_Generic>();
+				Message->Text = WarningMessage_InvalidSkeleton;
 				break;
 			}
 			ImportAssetObjectData.SkeletonReference = SkeletonReference;
@@ -1074,28 +1087,29 @@ UInterchangeFactoryBase::FImportAssetResult UInterchangeSkeletalMeshFactory::Beg
 		FString RootJointNodeId;
 		if (!SkeletonNode->GetCustomRootJointUid(RootJointNodeId))
 		{
-			UE_LOG(LogInterchangeImport, Warning, TEXT("Invalid Skeleton LOD Root Joint when importing SkeletalMesh asset %s"), *Arguments.AssetName);
+			UInterchangeResultWarning_Generic* Message = AddMessage<UInterchangeResultWarning_Generic>();
+			Message->Text = WarningMessage_InvalidRootJoint;
 			continue;
 		}
 
 		const UInterchangeSceneNode* RootJointNode = Cast<UInterchangeSceneNode>(Arguments.NodeContainer->GetNode(RootJointNodeId));
 		if (!RootJointNode)
 		{
-			UE_LOG(LogInterchangeImport, Warning, TEXT("Invalid Skeleton RootJointNode."));
+			UInterchangeResultWarning_Generic* Message = AddMessage<UInterchangeResultWarning_Generic>();
+			Message->Text = WarningMessage_InvalidRootJoint;
 			continue;
 		}
 
 		FImportAssetObjectLODData& ImportAssetObjectLODData = ImportAssetObjectData.LodDatas.AddDefaulted_GetRef();
 		ImportAssetObjectLODData.LodIndex = CurrentLodIndex;
 
-		int32 SkeletonDepth = 0;
 		SkeletonNode->GetCustomUseTimeZeroForBindPose(ImportAssetObjectLODData.bUseTimeZeroAsBindPose);
 		
 		//Do not alter the skeletal mesh reference skeleton when importing geometry only
 		FReferenceSkeleton RefSkeleton;
-		UE::Interchange::Private::FSkeletonHelper::ProcessImportMeshSkeleton(SkeletonReference
+		UE::Interchange::Private::FSkeletonHelper::ProcessImportMeshSkeleton( Results
+			, SkeletonReference
 			, ContentInfo.bApplyGeometryOnly ? RefSkeleton : SkeletalMesh->GetRefSkeleton()
-			, SkeletonDepth
 			, Arguments.NodeContainer
 			, RootJointNodeId
 			, ImportAssetObjectLODData.RefBonesBinary
@@ -1104,7 +1118,11 @@ UInterchangeFactoryBase::FImportAssetResult UInterchangeSkeletalMeshFactory::Beg
 
 		if (bSpecifiedSkeleton && !SkeletonReference->IsCompatibleMesh(SkeletalMesh))
 		{
-			UE_LOG(LogInterchangeImport, Warning, TEXT("The skeleton %s is incompatible with the imported skeletalmesh asset %s"), *SkeletonReference->GetName(), *Arguments.AssetName);
+			UInterchangeResultWarning_Generic* Message = AddMessage<UInterchangeResultWarning_Generic>();
+			Message->Text = FText::Format(NSLOCTEXT("InterchangeSkeletalMeshFactory", "BeginImportAsset_GameThread_IncompatibleSkeleton", "The skeleton {0} is incompatible with the imported LOD {1} skeletalmesh asset {2}.")
+				, FText::FromString(SkeletonReference->GetName())
+				, FText::AsNumber(LodIndex)
+				, FText::FromString(Arguments.AssetName));
 		}
 
 		CurrentLodIndex++;
@@ -1141,22 +1159,30 @@ UInterchangeFactoryBase::FImportAssetResult UInterchangeSkeletalMeshFactory::Imp
 	const IInterchangeMeshPayloadInterface* MeshTranslatorPayloadInterface = Cast<IInterchangeMeshPayloadInterface>(Arguments.Translator);
 	if (!MeshTranslatorPayloadInterface)
 	{
-		UE_LOG(LogInterchangeImport, Error, TEXT("Cannot import skeletalMesh, the translator do not implement the IInterchangeSkeletalMeshPayloadInterface."));
+		UInterchangeResultError_Generic* Message = AddMessage<UInterchangeResultError_Generic>();
+		Message->Text = FText::Format(NSLOCTEXT("InterchangeSkeletalMeshFactory", "ImportAsset_Async_TranslatorInterfaceMissing", "Cannot import skeletalMesh {0}, the translator {1} do not implement the IInterchangeSkeletalMeshPayloadInterface.")
+			, FText::FromString(Arguments.AssetName)
+			, FText::FromString(Arguments.Translator->GetName()));
 		return ImportAssetResult;
 	}
 
+	FText ErrorMessage_SkeletalMeshDontExist = FText::Format(NSLOCTEXT("InterchangeSkeletalMeshFactory", "ImportAsset_Async_AssetDontExist", "Could not import the SkeletalMesh asset {0}, because the asset do not exist.")
+		, FText::FromString(Arguments.AssetName));
+
 	UObject* SkeletalMeshObject = UE::Interchange::FFactoryCommon::AsyncFindObject(SkeletalMeshFactoryNode, GetFactoryClass(), Arguments.Parent, Arguments.AssetName);
 
 	if (!SkeletalMeshObject)
 	{
-		UE_LOG(LogInterchangeImport, Error, TEXT("Could not import the SkeletalMesh asset %s, because the asset do not exist."), *Arguments.AssetName);
+		UInterchangeResultError_Generic* Message = AddMessage<UInterchangeResultError_Generic>();
+		Message->Text = ErrorMessage_SkeletalMeshDontExist;
 		return ImportAssetResult;
 	}
 
 	USkeletalMesh* SkeletalMesh = Cast<USkeletalMesh>(SkeletalMeshObject);
 	if (!ensure(SkeletalMesh))
 	{
-		UE_LOG(LogInterchangeImport, Error, TEXT("Could not cast to SkeletalMesh asset %s"), *Arguments.AssetName);
+		UInterchangeResultError_Generic* Message = AddMessage<UInterchangeResultError_Generic>();
+		Message->Text = ErrorMessage_SkeletalMeshDontExist;
 		return ImportAssetResult;
 	}
 
@@ -1258,6 +1284,14 @@ UInterchangeFactoryBase::FImportAssetResult UInterchangeSkeletalMeshFactory::Imp
 
 	for (int32 LodIndex = 0; LodIndex < LodCount; ++LodIndex)
 	{
+		FText WarningMessage_InvalidSkeleton = FText::Format(NSLOCTEXT("InterchangeSkeletalMeshFactory", "BeginImportAsset_GameThread_InvalidSkeletonLOD", "Invalid Skeleton LOD {0} when importing SkeletalMesh asset {1}")
+			, FText::AsNumber(LodIndex)
+			, FText::FromString(Arguments.AssetName));
+		FText WarningMessage_InvalidRootJoint = FText::Format(NSLOCTEXT("InterchangeSkeletalMeshFactory", "BeginImportAsset_GameThread_InvalidSkeletonRootJoint", "Invalid Skeleton LOD {0} Root Joint when importing SkeletalMesh asset {1}")
+			, FText::AsNumber(LodIndex)
+			, FText::FromString(Arguments.AssetName));
+
+
 		TRACE_CPUPROFILER_EVENT_SCOPE(UInterchangeSkeletalMeshFactory::CreateAsset_LOD)
 		ESkeletalMeshGeoImportVersions GeoImportVersion = ESkeletalMeshGeoImportVersions::LatestVersion;
 		ESkeletalMeshSkinningImportVersions SkinningImportVersion = ESkeletalMeshSkinningImportVersions::LatestVersion;
@@ -1270,20 +1304,25 @@ UInterchangeFactoryBase::FImportAssetResult UInterchangeSkeletalMeshFactory::Imp
 		const UInterchangeSkeletalMeshLodDataNode* LodDataNode = Cast<UInterchangeSkeletalMeshLodDataNode>(Arguments.NodeContainer->GetNode(LodUniqueId));
 		if (!LodDataNode)
 		{
-			UE_LOG(LogInterchangeImport, Warning, TEXT("Invalid LOD when importing SkeletalMesh asset %s"), *Arguments.AssetName);
+			UInterchangeResultWarning_Generic* Message = AddMessage<UInterchangeResultWarning_Generic>();
+			Message->Text = FText::Format(NSLOCTEXT("InterchangeSkeletalMeshFactory", "ImportAsset_Async_InvalidLOD", "Invalid LOD {0} when importing SkeletalMesh asset {1}")
+				, FText::AsNumber(LodIndex)
+				, FText::FromString(Arguments.AssetName));
 			continue;
 		}
 
 		FString SkeletonNodeUid;
 		if (!LodDataNode->GetCustomSkeletonUid(SkeletonNodeUid))
 		{
-			UE_LOG(LogInterchangeImport, Warning, TEXT("Invalid Skeleton LOD when importing SkeletalMesh asset %s"), *Arguments.AssetName);
+			UInterchangeResultWarning_Generic* Message = AddMessage<UInterchangeResultWarning_Generic>();
+			Message->Text = WarningMessage_InvalidSkeleton;
 			continue;
 		}
 		const UInterchangeSkeletonFactoryNode* SkeletonNode = Cast<UInterchangeSkeletonFactoryNode>(Arguments.NodeContainer->GetNode(SkeletonNodeUid));
 		if (!SkeletonNode)
 		{
-			UE_LOG(LogInterchangeImport, Warning, TEXT("Invalid Skeleton LOD when importing SkeletalMesh asset %s"), *Arguments.AssetName);
+			UInterchangeResultWarning_Generic* Message = AddMessage<UInterchangeResultWarning_Generic>();
+			Message->Text = WarningMessage_InvalidSkeleton;
 			continue;
 		}
 
@@ -1298,14 +1337,16 @@ UInterchangeFactoryBase::FImportAssetResult UInterchangeSkeletalMeshFactory::Imp
 		FString RootJointNodeId;
 		if (!SkeletonNode->GetCustomRootJointUid(RootJointNodeId))
 		{
-			UE_LOG(LogInterchangeImport, Warning, TEXT("Invalid Skeleton LOD Root Joint when importing SkeletalMesh asset %s"), *Arguments.AssetName);
+			UInterchangeResultWarning_Generic* Message = AddMessage<UInterchangeResultWarning_Generic>();
+			Message->Text = WarningMessage_InvalidRootJoint;
 			continue;
 		}
 		
 		const UInterchangeSceneNode* RootJointNode = Cast<UInterchangeSceneNode>(Arguments.NodeContainer->GetNode(RootJointNodeId));
 		if (!RootJointNode)
 		{
-			UE_LOG(LogInterchangeImport, Warning, TEXT("Invalid Skeleton RootJointNode."));
+			UInterchangeResultWarning_Generic* Message = AddMessage<UInterchangeResultWarning_Generic>();
+			Message->Text = WarningMessage_InvalidRootJoint;
 			continue;
 		}
 		FImportAssetObjectLODData& ImportAssetObjectLODData = ImportAssetObjectData.LodDatas[CurrentLodIndex];
@@ -1323,6 +1364,11 @@ UInterchangeFactoryBase::FImportAssetResult UInterchangeSkeletalMeshFactory::Imp
 			TArray<FString> MeshUids;
 			LodDataNode->GetMeshUids(MeshUids);
 			MeshReferences.Reserve(MeshUids.Num());
+
+			FText WarningMessage_InvalidLODMeshReference = FText::Format(NSLOCTEXT("InterchangeSkeletalMeshFactory", "ImportAsset_Async_InvalidLODMeshReference", "Invalid LOD {0} mesh reference when importing SkeletalMesh asset {1}")
+				, FText::AsNumber(LodIndex)
+				, FText::FromString(Arguments.AssetName));
+
 			for (const FString& MeshUid : MeshUids)
 			{
 				UE::Interchange::Private::FMeshNodeContext MeshReference;
@@ -1333,7 +1379,8 @@ UInterchangeFactoryBase::FImportAssetResult UInterchangeSkeletalMeshFactory::Imp
 					MeshReference.SceneNode = Cast<UInterchangeSceneNode>(Arguments.NodeContainer->GetNode(MeshUid));
 					if (!ensure(MeshReference.SceneNode != nullptr))
 					{
-						UE_LOG(LogInterchangeImport, Warning, TEXT("Invalid LOD mesh reference when importing SkeletalMesh asset %s"), *Arguments.AssetName);
+						UInterchangeResultWarning_Generic* Message = AddMessage<UInterchangeResultWarning_Generic>();
+						Message->Text = WarningMessage_InvalidLODMeshReference;
 						continue;
 					}
 					FString MeshDependencyUid;
@@ -1363,7 +1410,8 @@ UInterchangeFactoryBase::FImportAssetResult UInterchangeSkeletalMeshFactory::Imp
 
 				if (!ensure(MeshReference.MeshNode != nullptr))
 				{
-					UE_LOG(LogInterchangeImport, Warning, TEXT("Invalid LOD mesh reference when importing SkeletalMesh asset %s"), *Arguments.AssetName);
+					UInterchangeResultWarning_Generic* Message = AddMessage<UInterchangeResultWarning_Generic>();
+					Message->Text = WarningMessage_InvalidLODMeshReference;
 					continue;
 				}
 
@@ -1374,7 +1422,10 @@ UInterchangeFactoryBase::FImportAssetResult UInterchangeSkeletalMeshFactory::Imp
 				}
 				else
 				{
-					UE_LOG(LogInterchangeImport, Warning, TEXT("Empty LOD mesh reference payload when importing SkeletalMesh asset %s"), *Arguments.AssetName);
+					UInterchangeResultWarning_Generic* Message = AddMessage<UInterchangeResultWarning_Generic>();
+					Message->Text = FText::Format(NSLOCTEXT("InterchangeSkeletalMeshFactory", "ImportAsset_Async_EmptyLODMeshReference", "Empty LOD {0} mesh reference payload when importing SkeletalMesh asset {1}")
+						, FText::AsNumber(LodIndex)
+						, FText::FromString(Arguments.AssetName));
 					continue;
 				}
 				MeshReferences.Add(MeshReference);
@@ -1541,8 +1592,10 @@ UInterchangeFactoryBase::FImportAssetResult UInterchangeSkeletalMeshFactory::Imp
 			const FVector3f BoundingBoxSize = BoundingBox.GetSize();
 			if (SkeletalMeshImportData.Points.Num() > 2 && BoundingBoxSize.X < UE_THRESH_POINTS_ARE_SAME && BoundingBoxSize.Y < UE_THRESH_POINTS_ARE_SAME && BoundingBoxSize.Z < UE_THRESH_POINTS_ARE_SAME)
 			{
-				//TODO log a user error
-				//AddTokenizedErrorMessage(FTokenizedMessage::Create(EMessageSeverity::Error, FText::Format(LOCTEXT("FbxSkeletaLMeshimport_ErrorMeshTooSmall", "Cannot import this mesh, the bounding box of this mesh is smaller than the supported threshold[{0}]."), FText::FromString(FString::Printf(TEXT("%f"), THRESH_POINTS_ARE_SAME)))), FFbxErrors::SkeletalMesh_FillImportDataFailed);
+				UInterchangeResultError_Generic* Message = AddMessage<UInterchangeResultError_Generic>();
+				Message->Text = FText::Format(NSLOCTEXT("InterchangeSkeletalMeshFactory", "ImportAsset_Async_ErrorMeshTooSmall", "The mesh {0} bounding box is smaller than the supported threshold[{1}]. All Vertices will be merge into one vertex.")
+					, FText::FromString(Arguments.AssetName)
+					, FText::AsNumber(UE_THRESH_POINTS_ARE_SAME));
 			}
 			FBoxSphereBounds BoxSphereBound((FBox)BoundingBox);
 			SkeletalMesh->SetImportedBounds(FBoxSphereBounds((FBox)BoundingBox));
diff --git a/Engine/Plugins/Interchange/Runtime/Source/Import/Private/Mesh/InterchangeSkeletonHelper.cpp b/Engine/Plugins/Interchange/Runtime/Source/Import/Private/Mesh/InterchangeSkeletonHelper.cpp
--- a/Engine/Plugins/Interchange/Runtime/Source/Import/Private/Mesh/InterchangeSkeletonHelper.cpp
+++ b/Engine/Plugins/Interchange/Runtime/Source/Import/Private/Mesh/InterchangeSkeletonHelper.cpp
@@ -5,6 +5,7 @@
 #include "Animation/Skeleton.h"
 #include "InterchangeCommonPipelineDataFactoryNode.h"
 #include "InterchangeImportLog.h"
+#include "InterchangeResultsContainer.h"
 #include "InterchangeSceneNode.h"
 #include "ReferenceSkeleton.h"
 #include "Rendering/SkeletalMeshLODImporterData.h"
@@ -15,7 +16,14 @@
 namespace UE::Interchange::Private
 {
 
-	bool FSkeletonHelper::ProcessImportMeshSkeleton(const USkeleton* SkeletonAsset, FReferenceSkeleton& RefSkeleton, int32& SkeletalDepth, const UInterchangeBaseNodeContainer* NodeContainer, const FString& RootJointNodeId, TArray<SkeletalMeshImportData::FBone>& RefBonesBinary, const bool bUseTimeZeroAsBindPose, bool& bOutDiffPose)
+	bool FSkeletonHelper::ProcessImportMeshSkeleton(TObjectPtr<UInterchangeResultsContainer> Results
+		, const USkeleton* SkeletonAsset
+		, FReferenceSkeleton& RefSkeleton
+		, const UInterchangeBaseNodeContainer* NodeContainer
+		, const FString& RootJointNodeId
+		, TArray<SkeletalMeshImportData::FBone>& RefBonesBinary
+		, const bool bUseTimeZeroAsBindPose
+		, bool& bOutDiffPose)
 	{
 		RefBonesBinary.Empty();
 		// Setup skeletal hierarchy + names structure.
@@ -23,12 +31,28 @@ namespace UE::Interchange::Private
 
 		FReferenceSkeletonModifier RefSkelModifier(RefSkeleton, SkeletonAsset);
 		TArray <FJointInfo> JointInfos;
-		RecursiveAddBones(NodeContainer, RootJointNodeId, JointInfos, INDEX_NONE, RefBonesBinary, bUseTimeZeroAsBindPose, bOutDiffPose);
+		TArray<FString> BoneNotBindNames;
+		RecursiveAddBones(NodeContainer, RootJointNodeId, JointInfos, INDEX_NONE, RefBonesBinary, bUseTimeZeroAsBindPose, bOutDiffPose, BoneNotBindNames);
 		if (bOutDiffPose)
 		{
 			//bOutDiffPose can only be true if the user ask to bind with time zero transform.
 			ensure(bUseTimeZeroAsBindPose);
 		}
+		//Do not output this warning in automation testing
+		if (!GIsAutomationTesting && BoneNotBindNames.Num() > 0 && !bUseTimeZeroAsBindPose)
+		{
+			FString BonesWithoutBindPoses;
+			for (const FString& BoneName : BoneNotBindNames)
+			{
+				BonesWithoutBindPoses += BoneName;
+				BonesWithoutBindPoses += TEXT("  \n");
+			}
+
+			FText MissingBindPoseMessage = FText::Format(NSLOCTEXT("FSkeletonHelper", "ProcessImportMeshSkeleton__BonesAreMissingFromBindPose", "The following bones are missing from the bind pose:\n{0}\nThis can happen for bones that are not vert weighted. If they are not in the correct orientation after importing,\nplease set the \"Use T0 as ref pose\" option or add them to the bind pose and reimport the skeletal mesh.")
+				, FText::FromString(BonesWithoutBindPoses));
+			UInterchangeResultWarning_Generic* Message = Results->Add<UInterchangeResultWarning_Generic>();
+			Message->Text = MissingBindPoseMessage;
+		}
 		// Digest bones to the serializable format.
 		for (int32 b = 0; b < JointInfos.Num(); b++)
 		{
@@ -39,35 +63,13 @@ namespace UE::Interchange::Private
 			const FTransform BoneTransform(BinaryBone.LocalTransform);
 			if (RefSkeleton.FindRawBoneIndex(BoneInfo.Name) != INDEX_NONE)
 			{
-				UE_LOG(LogInterchangeImport, Error, TEXT("Invalid Skeleton because of non-unique bone names [%s]"), *BoneInfo.Name.ToString());
+				UInterchangeResultError_Generic* Message = Results->Add<UInterchangeResultError_Generic>();
+				Message->Text = FText::Format(NSLOCTEXT("FSkeletonHelper", "ProcessImportMeshSkeleton_InvalidSkeletonUniqueNames", "Invalid Skeleton because of non - unique bone names [{0}].")
+					, FText::FromName(BoneInfo.Name));
 				return false;
 			}
 			RefSkelModifier.Add(BoneInfo, BoneTransform);
 		}
-
-		// Add hierarchy index to each bone and detect max depth.
-		SkeletalDepth = 0;
-
-		TArray<int32> SkeletalDepths;
-		SkeletalDepths.Empty(JointInfos.Num());
-		SkeletalDepths.AddZeroed(JointInfos.Num());
-		for (int32 BoneIndex = 0; BoneIndex < RefSkeleton.GetRawBoneNum(); ++BoneIndex)
-		{
-			int32 Parent = RefSkeleton.GetRawParentIndex(BoneIndex);
-			int32 Depth = 1.0f;
-
-			SkeletalDepths[BoneIndex] = 1.0f;
-			if (Parent != INDEX_NONE)
-			{
-				Depth += SkeletalDepths[Parent];
-			}
-			if (SkeletalDepth < Depth)
-			{
-				SkeletalDepth = Depth;
-			}
-			SkeletalDepths[BoneIndex] = Depth;
-		}
-
 		return true;
 	}
 
@@ -167,7 +169,14 @@ namespace UE::Interchange::Private
 		}
 	}
 
-	void FSkeletonHelper::RecursiveAddBones(const UInterchangeBaseNodeContainer* NodeContainer, const FString& JointNodeId, TArray <FJointInfo>& JointInfos, int32 ParentIndex, TArray<SkeletalMeshImportData::FBone>& RefBonesBinary, const bool bUseTimeZeroAsBindPose, bool& bOutDiffPose)
+	void FSkeletonHelper::RecursiveAddBones(const UInterchangeBaseNodeContainer* NodeContainer
+		, const FString& JointNodeId
+		, TArray <FJointInfo>& JointInfos
+		, int32 ParentIndex
+		, TArray<SkeletalMeshImportData::FBone>& RefBonesBinary
+		, const bool bUseTimeZeroAsBindPose
+		, bool& bOutDiffPose
+		, TArray<FString>& OutBoneNotBindNames)
 	{
 		const UInterchangeSceneNode* JointNode = Cast<UInterchangeSceneNode>(NodeContainer->GetNode(JointNodeId));
 		if (!JointNode)
@@ -222,9 +231,9 @@ namespace UE::Interchange::Private
 			}
 			Info.LocalTransform = TimeZeroLocalTransform;
 		}
-		else if (bHasBindPoseTransform)
+		else if (!GIsAutomationTesting && !bHasBindPoseTransform && !bUseTimeZeroAsBindPose)
 		{
-			Info.LocalTransform = BindPoseLocalTransform;
+			OutBoneNotBindNames.Add(Info.Name);
 		}
 
 		Info.ParentIndex = ParentIndex;
@@ -243,7 +252,7 @@ namespace UE::Interchange::Private
 		Bone.NumChildren = ChildrenIds.Num();
 		for (int32 ChildIndex = 0; ChildIndex < ChildrenIds.Num(); ++ChildIndex)
 		{
-			RecursiveAddBones(NodeContainer, ChildrenIds[ChildIndex], JointInfos, JointInfoIndex, RefBonesBinary, bUseTimeZeroAsBindPose, bOutDiffPose);
+			RecursiveAddBones(NodeContainer, ChildrenIds[ChildIndex], JointInfos, JointInfoIndex, RefBonesBinary, bUseTimeZeroAsBindPose, bOutDiffPose, OutBoneNotBindNames);
 		}
 	}
 
@@ -323,7 +332,12 @@ namespace UE::Interchange::Private
 	void FSkeletonHelper::RecursiveBuildSkeletalSkeleton(const FString JoinToAddUid, const int32 ParentIndex, const UInterchangeBaseNodeContainer* BaseNodeContainer, TArray<FMeshBoneInfo>& SkeletalLodRawInfos, bool bConvertStaticToSkeletalActive)
 	{
 		const UInterchangeSceneNode* SceneNode = Cast<UInterchangeSceneNode>(BaseNodeContainer->GetNode(JoinToAddUid));
-		if (!bConvertStaticToSkeletalActive && (!SceneNode || !SceneNode->IsSpecializedTypeContains(FSceneNodeStaticData::GetJointSpecializeTypeString())))
+		if (!SceneNode)
+		{
+			return;
+		}
+
+		if (!bConvertStaticToSkeletalActive && !SceneNode->IsSpecializedTypeContains(FSceneNodeStaticData::GetJointSpecializeTypeString()))
 		{
 			return;
 		}
diff --git a/Engine/Plugins/Interchange/Runtime/Source/Import/Public/InterchangeSkeletonHelper.h b/Engine/Plugins/Interchange/Runtime/Source/Import/Public/InterchangeSkeletonHelper.h
--- a/Engine/Plugins/Interchange/Runtime/Source/Import/Public/InterchangeSkeletonHelper.h
+++ b/Engine/Plugins/Interchange/Runtime/Source/Import/Public/InterchangeSkeletonHelper.h
@@ -11,6 +11,7 @@
 struct FReferenceSkeleton;
 class UInterchangeBaseNode;
 class UInterchangeBaseNodeContainer;
+class UInterchangeResultsContainer;
 class USkeleton;
 struct FMeshBoneInfo;
 
@@ -27,16 +28,37 @@ namespace UE::Interchange::Private
 	struct INTERCHANGEIMPORT_API FSkeletonHelper
 	{
 	public:
-		static bool ProcessImportMeshSkeleton(const USkeleton* SkeletonAsset, FReferenceSkeleton& RefSkeleton, int32& SkeletalDepth, const UInterchangeBaseNodeContainer* NodeContainer, const FString& RootJointNodeId, TArray<SkeletalMeshImportData::FBone>& RefBonesBinary, const bool bUseTimeZeroAsBindPose, bool& bOutDiffPose);
+		static bool ProcessImportMeshSkeleton(TObjectPtr<UInterchangeResultsContainer> Results
+			, const USkeleton* SkeletonAsset
+			, FReferenceSkeleton& RefSkeleton
+			, const UInterchangeBaseNodeContainer* NodeContainer
+			, const FString& RootJointNodeId
+			, TArray<SkeletalMeshImportData::FBone>& RefBonesBinary
+			, const bool bUseTimeZeroAsBindPose
+			, bool& bOutDiffPose);
+
 		static bool IsCompatibleSkeleton(const USkeleton* Skeleton, const FString RootJoinUid, const UInterchangeBaseNodeContainer* BaseNodeContainer, bool bConvertStaticToSkeletalActive);
 		static void RecursiveAddSkeletonMetaDataValues(UInterchangeBaseNodeContainer* NodeContainer, UInterchangeBaseNode* DestinationNode, const FString& JointUid);
 	private:
-		static void RecursiveAddBones(const UInterchangeBaseNodeContainer* NodeContainer, const FString& JointNodeId, TArray <FJointInfo>& JointInfos, int32 ParentIndex, TArray<SkeletalMeshImportData::FBone>& RefBonesBinary, const bool bUseTimeZeroAsBindPose, bool& bOutDiffPose);
+		static void RecursiveAddBones(const UInterchangeBaseNodeContainer* NodeContainer
+			, const FString& JointNodeId
+			, TArray <FJointInfo>& JointInfos
+			, int32 ParentIndex
+			, TArray<SkeletalMeshImportData::FBone>& RefBonesBinary
+			, const bool bUseTimeZeroAsBindPose
+			, bool& bOutDiffPose
+			, TArray<FString>& OutBoneNotBindNames);
+
 		static FName SkeletalLodGetBoneName(const TArray<FMeshBoneInfo>& SkeletalLodRawInfos, int32 BoneIndex);
 		static int32 SkeletalLodFindBoneIndex(const TArray<FMeshBoneInfo>& SkeletalLodRawInfos, FName BoneName);
 		static int32 SkeletalLodGetParentIndex(const TArray<FMeshBoneInfo>& SkeletalLodRawInfos, int32 BoneIndex);
 		static bool DoesParentChainMatch(int32 StartBoneIndex, const FReferenceSkeleton& SkeletonRef, const TArray<FMeshBoneInfo>& SkeletalLodRawInfos);
-		static void RecursiveBuildSkeletalSkeleton(const FString JoinToAddUid, const int32 ParentIndex, const UInterchangeBaseNodeContainer* BaseNodeContainer, TArray<FMeshBoneInfo>& SkeletalLodRawInfos, bool bConvertStaticToSkeletalActive = false);
+
+		static void RecursiveBuildSkeletalSkeleton(const FString JoinToAddUid
+			, const int32 ParentIndex
+			, const UInterchangeBaseNodeContainer* BaseNodeContainer
+			, TArray<FMeshBoneInfo>& SkeletalLodRawInfos
+			, bool bConvertStaticToSkeletalActive = false);
 	};
 } //namespace UE::Interchange::Private
 
