MaterialX:
- Create Scalar/Vector parameter to directly use the input name of the Mtlx node, instead of counting of the Label in Interchange.

diff --git a/Engine/Plugins/Interchange/Runtime/Source/Import/Private/MaterialX/MaterialXUtils/MaterialXSurfaceShaderAbstract.cpp b/Engine/Plugins/Interchange/Runtime/Source/Import/Private/MaterialX/MaterialXUtils/MaterialXSurfaceShaderAbstract.cpp
--- a/Engine/Plugins/Interchange/Runtime/Source/Import/Private/MaterialX/MaterialXUtils/MaterialXSurfaceShaderAbstract.cpp
+++ b/Engine/Plugins/Interchange/Runtime/Source/Import/Private/MaterialX/MaterialXUtils/MaterialXSurfaceShaderAbstract.cpp
@@ -28,27 +28,22 @@ FMaterialXSurfaceShaderAbstract::FMaterialXSurfaceShaderAbstract(UInterchangeBas
 	, bTangentSpaceInput{ false }
 {}
 
-bool FMaterialXSurfaceShaderAbstract::AddAttribute(MaterialX::InputPtr Input, const FString& InputChannelName, UInterchangeShaderNode* ShaderNode) const
+bool FMaterialXSurfaceShaderAbstract::AddAttribute(MaterialX::InputPtr Input, const FString& InputChannelName, UInterchangeShaderNode* ShaderNode)
 {
+	using namespace UE::Interchange::Materials::Standard::Nodes;
 	if(Input)
 	{
-		// The parent is either a node, or it's an interfacename and we just take the name of the input
-		mx::NodePtr Node = Input->getParent()->asA<mx::Node>();
-		FString NodeName = Node ? Node->getName().c_str() + FString{ TEXT("_") } : FString{};
-		ShaderNode->SetDisplayLabel(NodeName + Input->getName().c_str());
-
 		if(Input->getType() == mx::Type::Float)
 		{
-			return ShaderNode->AddFloatInput(InputChannelName, mx::fromValueString<float>(Input->getValueString()), true);
+			return AddFloatAttribute(Input, InputChannelName, ShaderNode);
 		}
-		else if(Input->getType() == mx::Type::Integer) //Let's add it a Float attribute, because Interchange doesn't create a scalar if it's an int
+		else if(Input->getType() == mx::Type::Integer) //Let's add Float attribute, because Interchange doesn't create a scalar if it's an int
 		{
 			return ShaderNode->AddFloatAttribute(UInterchangeShaderPortsAPI::MakeInputValueKey(InputChannelName), mx::fromValueString<int32>(Input->getValueString()));
 		}
 		else if(Input->getType() == mx::Type::Color3 || Input->getType() == mx::Type::Color4)
 		{
-			FLinearColor LinearColor = GetLinearColor(Input);
-			return ShaderNode->AddLinearColorInput(InputChannelName, LinearColor, true);
+			return AddLinearColorAttribute(Input, InputChannelName, ShaderNode);
 		}
 		else if(Input->getType() == mx::Type::Vector2)
 		{
@@ -57,8 +52,7 @@ bool FMaterialXSurfaceShaderAbstract::AddAttribute(MaterialX::InputPtr Input, co
 		}
 		else if(Input->getType() == mx::Type::Vector3 || Input->getType() == mx::Type::Vector4)
 		{
-			FLinearColor Vector = GetVector(Input);
-			return ShaderNode->AddLinearColorInput(InputChannelName, Vector, true);
+			return AddVectorAttribute(Input, InputChannelName, ShaderNode);
 		}
 	}
 
@@ -103,45 +97,69 @@ bool FMaterialXSurfaceShaderAbstract::AddAttributeFromValueOrInterface(MaterialX
 	return bAttribute;
 }
 
-bool FMaterialXSurfaceShaderAbstract::AddFloatAttribute(MaterialX::InputPtr Input, const FString& InputChannelName, UInterchangeShaderNode* ShaderNode, float DefaultValue) const
+bool FMaterialXSurfaceShaderAbstract::AddFloatAttribute(MaterialX::InputPtr Input, const FString& InputChannelName, UInterchangeShaderNode* ShaderNode, float DefaultValue)
 {
+	using namespace UE::Interchange::Materials::Standard::Nodes;
 	if(Input && Input->hasValue())
 	{
 		float Value = mx::fromValueString<float>(Input->getValueString());
 
 		if(!FMath::IsNearlyEqual(Value, DefaultValue))
 		{
-			return ShaderNode->AddFloatInput(InputChannelName, Value, true);
+			// The parent is either a node, or it's an interfacename and we just take the name of the input
+			mx::NodePtr Node = Input->getParent()->asA<mx::Node>();
+			FString NodeName = Node ? Node->getName().c_str() + FString{ TEXT("_") } : FString{};
+			NodeName += Input->getName().c_str();
+
+			UInterchangeShaderNode* ScalarParameterNode = CreateShaderNode(NodeName, ScalarParameter::Name.ToString());
+			ScalarParameterNode->AddFloatAttribute(UInterchangeShaderPortsAPI::MakeInputParameterKey(ScalarParameter::Attributes::DefaultValue.ToString()), mx::fromValueString<float>(Input->getValueString()));
+			return UInterchangeShaderPortsAPI::ConnectDefaultOuputToInput(ShaderNode, InputChannelName, ScalarParameterNode->GetUniqueID());
 		}
 	}
 
 	return false;
 }
 
-bool FMaterialXSurfaceShaderAbstract::AddLinearColorAttribute(MaterialX::InputPtr Input, const FString& InputChannelName, UInterchangeShaderNode* ShaderNode, const FLinearColor& DefaultValue) const
+bool FMaterialXSurfaceShaderAbstract::AddLinearColorAttribute(MaterialX::InputPtr Input, const FString& InputChannelName, UInterchangeShaderNode* ShaderNode, const FLinearColor& DefaultValue)
 {
+	using namespace UE::Interchange::Materials::Standard::Nodes;
 	if(Input && Input->hasValue())
 	{
 		const FLinearColor Value = GetLinearColor(Input);
 
 		if(!Value.Equals(DefaultValue))
 		{
-			return ShaderNode->AddLinearColorInput(InputChannelName, Value, true);
+			// The parent is either a node, or it's an interfacename and we just take the name of the input
+			mx::NodePtr Node = Input->getParent()->asA<mx::Node>();
+			FString NodeName = Node ? Node->getName().c_str() + FString{ TEXT("_") } : FString{};
+			NodeName += Input->getName().c_str();
+
+			UInterchangeShaderNode* VectorParameterNode = CreateShaderNode(NodeName, VectorParameter::Name.ToString());
+			VectorParameterNode->AddLinearColorAttribute(UInterchangeShaderPortsAPI::MakeInputParameterKey(VectorParameter::Attributes::DefaultValue.ToString()), Value);
+			return UInterchangeShaderPortsAPI::ConnectDefaultOuputToInput(ShaderNode, InputChannelName, VectorParameterNode->GetUniqueID());
 		}
 	}
 
 	return false;
 }
 
-bool FMaterialXSurfaceShaderAbstract::AddVectorAttribute(MaterialX::InputPtr Input, const FString& InputChannelName, UInterchangeShaderNode* ShaderNode, const FVector4f& DefaultValue) const
+bool FMaterialXSurfaceShaderAbstract::AddVectorAttribute(MaterialX::InputPtr Input, const FString& InputChannelName, UInterchangeShaderNode* ShaderNode, const FVector4f& DefaultValue)
 {
+	using namespace UE::Interchange::Materials::Standard::Nodes;
 	if(Input && Input->hasValue())
 	{
 		const FLinearColor Value = GetVector(Input);
 
 		if(!Value.Equals(DefaultValue))
 		{
-			return ShaderNode->AddLinearColorInput(InputChannelName, Value, true);
+			// The parent is either a node, or it's an interfacename and we just take the name of the input
+			mx::NodePtr Node = Input->getParent()->asA<mx::Node>();
+			FString NodeName = Node ? Node->getName().c_str() + FString{ TEXT("_") } : FString{};
+			NodeName += Input->getName().c_str();
+
+			UInterchangeShaderNode* VectorParameterNode = CreateShaderNode(NodeName, VectorParameter::Name.ToString());
+			VectorParameterNode->AddLinearColorAttribute(UInterchangeShaderPortsAPI::MakeInputParameterKey(VectorParameter::Attributes::DefaultValue.ToString()), Value);
+			return UInterchangeShaderPortsAPI::ConnectDefaultOuputToInput(ShaderNode, InputChannelName, VectorParameterNode->GetUniqueID());
 		}
 	}
 
diff --git a/Engine/Plugins/Interchange/Runtime/Source/Import/Public/MaterialX/MaterialXUtils/MaterialXSurfaceShaderAbstract.h b/Engine/Plugins/Interchange/Runtime/Source/Import/Public/MaterialX/MaterialXUtils/MaterialXSurfaceShaderAbstract.h
--- a/Engine/Plugins/Interchange/Runtime/Source/Import/Public/MaterialX/MaterialXUtils/MaterialXSurfaceShaderAbstract.h
+++ b/Engine/Plugins/Interchange/Runtime/Source/Import/Public/MaterialX/MaterialXUtils/MaterialXSurfaceShaderAbstract.h
@@ -49,7 +49,7 @@ protected:
 	 *
 	 * @return true if the attribute was successfully added
 	 */
-	bool AddAttribute(MaterialX::InputPtr Input, const FString& InputChannelName, UInterchangeShaderNode* ShaderNode) const;
+	bool AddAttribute(MaterialX::InputPtr Input, const FString& InputChannelName, UInterchangeShaderNode* ShaderNode);
 
 	/**
 	 * Add an attribute to a shader node from the given MaterialX input if that input has either a value or an interface name
@@ -72,7 +72,7 @@ protected:
 	 *
 	 * @return true if the attribute was successfully added
 	 */
-	bool AddFloatAttribute(MaterialX::InputPtr Input, const FString& InputChannelName, UInterchangeShaderNode* ShaderNode, float DefaultValue) const;
+	bool AddFloatAttribute(MaterialX::InputPtr Input, const FString& InputChannelName, UInterchangeShaderNode* ShaderNode, float DefaultValue = std::numeric_limits<float>::max());
 
 	/**
 	 * Add a FLinearColor attribute to a shader node only if its value taken from the input is not equal to its default value. Return false if the attribute does not exist or if we cannot add it
@@ -84,7 +84,7 @@ protected:
 	 *
 	 * @return true if the attribute was successfully added
 	 */
-	bool AddLinearColorAttribute(MaterialX::InputPtr Input, const FString& InputChannelName, UInterchangeShaderNode* ShaderNode, const FLinearColor& DefaultValue) const;
+	bool AddLinearColorAttribute(MaterialX::InputPtr Input, const FString& InputChannelName, UInterchangeShaderNode* ShaderNode, const FLinearColor& DefaultValue = FLinearColor{std::numeric_limits<float>::max(), std::numeric_limits<float>::max(), std::numeric_limits<float>::max()});
 
 	/**
 	 * Add a FLinearColor attribute to a shader node only if its value taken from the input is not equal to its default value. Return false if the attribute does not exist or if we cannot add it
@@ -96,7 +96,7 @@ protected:
 	 *
 	 * @return true if the attribute was successfully added
 	 */
-	bool AddVectorAttribute(MaterialX::InputPtr Input, const FString& InputChannelName, UInterchangeShaderNode* ShaderNode, const FVector4f& DefaultValue) const;
+	bool AddVectorAttribute(MaterialX::InputPtr Input, const FString& InputChannelName, UInterchangeShaderNode* ShaderNode, const FVector4f& DefaultValue = FVector4f{ std::numeric_limits<float>::max(), std::numeric_limits<float>::max(), std::numeric_limits<float>::max() });
 
 	/**
 	 * Connect an output either from a node name or a node graph from a MaterialX input to the ShaderNode
