Interchange: Fix fbx material normal map connection that was broken.

#rnx

diff --git a/Engine/Plugins/Interchange/Runtime/Source/Parsers/Fbx/Private/FbxMaterial.cpp b/Engine/Plugins/Interchange/Runtime/Source/Parsers/Fbx/Private/FbxMaterial.cpp
--- a/Engine/Plugins/Interchange/Runtime/Source/Parsers/Fbx/Private/FbxMaterial.cpp
+++ b/Engine/Plugins/Interchange/Runtime/Source/Parsers/Fbx/Private/FbxMaterial.cpp
@@ -357,7 +357,7 @@ namespace UE
 							if (FbxFileTexture* FbxTexture = MaterialProperty.GetSrcObject<FbxFileTexture>(0))
 							{
 								const FString TexturePath = FbxTexture->GetFileName();
-								if (TexturePath.IsEmpty() || !FPaths::FileExists(TexturePath))
+								if (!TexturePath.IsEmpty() && FPaths::FileExists(TexturePath))
 								{
 									if (const UInterchangeTexture2DNode* TextureNode = CreateTexture2DNode(NodeContainer, TexturePath))
 									{
