[Interchange - glTF] Adding support for KHR_materials_iridescence.

#rnx

#virtualized

diff --git a/Engine/Plugins/Enterprise/GLTFExporter/Source/GLTFExporter/Private/Converters/GLTFMaterialUtilities.cpp b/Engine/Plugins/Enterprise/GLTFExporter/Source/GLTFExporter/Private/Converters/GLTFMaterialUtilities.cpp
--- a/Engine/Plugins/Enterprise/GLTFExporter/Source/GLTFExporter/Private/Converters/GLTFMaterialUtilities.cpp
+++ b/Engine/Plugins/Enterprise/GLTFExporter/Source/GLTFExporter/Private/Converters/GLTFMaterialUtilities.cpp
@@ -993,6 +993,14 @@ void FGLTFImportMaterialMatchMakingHelper::Process()
 			GetValue(Inputs::SpecularTexture, JsonMaterial.Specular.Texture);
 
 			GetValue(Inputs::IOR, JsonMaterial.IOR.Value);
+
+			//Iridescence
+			GetValue(Inputs::IridescenceFactor, JsonMaterial.Iridescence.IridescenceFactor);
+			GetValue(Inputs::IridescenceTexture, JsonMaterial.Iridescence.IridescenceTexture);
+			GetValue(Inputs::IridescenceIOR, JsonMaterial.Iridescence.IridescenceIOR);
+			GetValue(Inputs::IridescenceThicknessMinimum, JsonMaterial.Iridescence.IridescenceThicknessMinimum);
+			GetValue(Inputs::IridescenceThicknessMaximum, JsonMaterial.Iridescence.IridescenceThicknessMaximum);
+			GetValue(Inputs::IridescenceThicknessTexture, JsonMaterial.Iridescence.IridescenceThicknessTexture);
 		}
 
 		//Importer does not support Emissive for Transmission at the moment
diff --git a/Engine/Plugins/Enterprise/GLTFExporter/Source/GLTFExporter/Private/Json/GLTFJsonMaterial.cpp b/Engine/Plugins/Enterprise/GLTFExporter/Source/GLTFExporter/Private/Json/GLTFJsonMaterial.cpp
--- a/Engine/Plugins/Enterprise/GLTFExporter/Source/GLTFExporter/Private/Json/GLTFJsonMaterial.cpp
+++ b/Engine/Plugins/Enterprise/GLTFExporter/Source/GLTFExporter/Private/Json/GLTFJsonMaterial.cpp
@@ -205,6 +205,40 @@ void FGLTFJsonSpecularGlossinessExtension::WriteObject(IGLTFJsonWriter& Writer)
 	}
 }
 
+void FGLTFJsonIridescenceExtension::WriteObject(IGLTFJsonWriter& Writer) const
+{
+	if (!FMath::IsNearlyEqual(IridescenceFactor, 0.f))
+	{
+		Writer.Write(TEXT("iridescenceFactor"), IridescenceFactor);
+	}
+
+	if (IridescenceTexture.Index != nullptr)
+	{
+		Writer.Write(TEXT("iridescenceTexture"), IridescenceTexture);
+	}
+
+	if (!FMath::IsNearlyEqual(IridescenceIOR, 1.3f))
+	{
+		Writer.Write(TEXT("iridescenceIor"), IridescenceIOR);
+	}
+
+
+	if (!FMath::IsNearlyEqual(IridescenceThicknessMinimum, 100.f))
+	{
+		Writer.Write(TEXT("iridescenceThicknessMinimum"), IridescenceThicknessMinimum);
+	}
+	
+	if (!FMath::IsNearlyEqual(IridescenceThicknessMaximum, 100.f))
+	{
+		Writer.Write(TEXT("iridescenceThicknessMaximum"), IridescenceThicknessMaximum);
+	}
+
+	if (IridescenceThicknessTexture.Index != nullptr)
+	{
+		Writer.Write(TEXT("iridescenceThicknessTexture"), IridescenceThicknessTexture);
+	}
+}
+
 void FGLTFJsonMaterial::WriteObject(IGLTFJsonWriter& Writer) const
 {
 	if (!Name.IsEmpty())
@@ -261,25 +295,31 @@ void FGLTFJsonMaterial::WriteObject(IGLTFJsonWriter& Writer) const
 		HasEmissiveStrength || 
 		Specular.HasValue() || 
 		IOR.HasValue() ||
-		Sheen.HasValue())
+		Sheen.HasValue() ||
+		Iridescence.HasValue())
 	{
 		Writer.StartExtensions();
 
-		if (ShadingModel != EGLTFJsonShadingModel::Unlit /*&& !KHR_materials_pbrSpecularGlossiness*/ && Specular.HasValue())
+		if (ShadingModel != EGLTFJsonShadingModel::Unlit && ShadingModel != EGLTFJsonShadingModel::SpecularGlossiness && Specular.HasValue())
 		{
 			Writer.Write(EGLTFJsonExtension::KHR_MaterialsSpecular, Specular);
 		}
 
-		if (ShadingModel != EGLTFJsonShadingModel::Unlit /*&& !KHR_materials_pbrSpecularGlossiness*/ && IOR.HasValue())
+		if (ShadingModel != EGLTFJsonShadingModel::Unlit && ShadingModel != EGLTFJsonShadingModel::SpecularGlossiness && IOR.HasValue())
 		{
 			Writer.Write(EGLTFJsonExtension::KHR_MaterialsIOR, IOR);
 		}
 
-		if (ShadingModel != EGLTFJsonShadingModel::Unlit /*&& !KHR_materials_pbrSpecularGlossiness*/ && Sheen.HasValue())
+		if (ShadingModel != EGLTFJsonShadingModel::Unlit && ShadingModel != EGLTFJsonShadingModel::SpecularGlossiness && Sheen.HasValue())
 		{
 			Writer.Write(EGLTFJsonExtension::KHR_MaterialsSheen, Sheen);
 		}
 
+		if (ShadingModel != EGLTFJsonShadingModel::Unlit && ShadingModel != EGLTFJsonShadingModel::SpecularGlossiness && Iridescence.HasValue())
+		{
+			Writer.Write(EGLTFJsonExtension::KHR_MaterialsIridescence, Iridescence);
+		}
+
 		if (ShadingModel == EGLTFJsonShadingModel::Unlit)
 		{
 			// Write empty object
diff --git a/Engine/Plugins/Enterprise/GLTFExporter/Source/GLTFExporter/Private/Utilities/GLTFJsonUtilities.cpp b/Engine/Plugins/Enterprise/GLTFExporter/Source/GLTFExporter/Private/Utilities/GLTFJsonUtilities.cpp
--- a/Engine/Plugins/Enterprise/GLTFExporter/Source/GLTFExporter/Private/Utilities/GLTFJsonUtilities.cpp
+++ b/Engine/Plugins/Enterprise/GLTFExporter/Source/GLTFExporter/Private/Utilities/GLTFJsonUtilities.cpp
@@ -15,6 +15,7 @@ const TCHAR* FGLTFJsonUtilities::GetValue(EGLTFJsonExtension Enum)
 		case EGLTFJsonExtension::KHR_MaterialsSheen:                  return TEXT("KHR_materials_sheen");
 		case EGLTFJsonExtension::KHR_MaterialsTransmission:           return TEXT("KHR_materials_transmission");
 		case EGLTFJsonExtension::KHR_MaterialsSpecularGlossiness:     return TEXT("KHR_materials_pbrSpecularGlossiness");
+		case EGLTFJsonExtension::KHR_MaterialsIridescence:            return TEXT("KHR_materials_iridescence");
 		case EGLTFJsonExtension::KHR_MeshQuantization:                return TEXT("KHR_mesh_quantization");
 		case EGLTFJsonExtension::KHR_TextureTransform:                return TEXT("KHR_texture_transform");
 		case EGLTFJsonExtension::KHR_MaterialsSpecular:	              return TEXT("KHR_materials_specular");
diff --git a/Engine/Plugins/Enterprise/GLTFExporter/Source/GLTFExporter/Public/Json/GLTFJsonEnums.h b/Engine/Plugins/Enterprise/GLTFExporter/Source/GLTFExporter/Public/Json/GLTFJsonEnums.h
--- a/Engine/Plugins/Enterprise/GLTFExporter/Source/GLTFExporter/Public/Json/GLTFJsonEnums.h
+++ b/Engine/Plugins/Enterprise/GLTFExporter/Source/GLTFExporter/Public/Json/GLTFJsonEnums.h
@@ -13,6 +13,7 @@ enum class EGLTFJsonExtension
 	KHR_MaterialsSheen,
 	KHR_MaterialsTransmission,
 	KHR_MaterialsSpecularGlossiness,
+	KHR_MaterialsIridescence,
 	KHR_MeshQuantization,
 	KHR_TextureTransform,
 	KHR_MaterialsSpecular
diff --git a/Engine/Plugins/Enterprise/GLTFExporter/Source/GLTFExporter/Public/Json/GLTFJsonMaterial.h b/Engine/Plugins/Enterprise/GLTFExporter/Source/GLTFExporter/Public/Json/GLTFJsonMaterial.h
--- a/Engine/Plugins/Enterprise/GLTFExporter/Source/GLTFExporter/Public/Json/GLTFJsonMaterial.h
+++ b/Engine/Plugins/Enterprise/GLTFExporter/Source/GLTFExporter/Public/Json/GLTFJsonMaterial.h
@@ -171,6 +171,31 @@ struct GLTFEXPORTER_API FGLTFJsonSpecularGlossinessExtension : IGLTFJsonObject
 	virtual void WriteObject(IGLTFJsonWriter& Writer) const override;
 };
 
+struct GLTFEXPORTER_API FGLTFJsonIridescenceExtension : IGLTFJsonObject
+{
+	float                   IridescenceFactor;
+	FGLTFJsonTextureInfo    IridescenceTexture;
+	float                   IridescenceIOR;
+
+	float                   IridescenceThicknessMinimum;
+	float                   IridescenceThicknessMaximum;
+	FGLTFJsonTextureInfo    IridescenceThicknessTexture;
+
+
+	FGLTFJsonIridescenceExtension()
+		: IridescenceFactor(0.f)
+		, IridescenceIOR(1.3f)
+		, IridescenceThicknessMinimum(100.f)
+		, IridescenceThicknessMaximum(400.f)
+	{
+	}
+
+	bool HasValue() const { return !FMath::IsNearlyEqual(IridescenceFactor, 0.f) || IridescenceTexture.Index != nullptr || !FMath::IsNearlyEqual(IridescenceIOR, 1.3f) || 
+		!FMath::IsNearlyEqual(IridescenceThicknessMinimum, 100.f) || !FMath::IsNearlyEqual(IridescenceThicknessMaximum, 400.f) || IridescenceThicknessTexture.Index != nullptr; }
+
+	virtual void WriteObject(IGLTFJsonWriter& Writer) const override;
+};
+
 struct GLTFEXPORTER_API FGLTFJsonMaterial : IGLTFJsonIndexedObject
 {
 	FString Name;
@@ -197,6 +222,7 @@ struct GLTFEXPORTER_API FGLTFJsonMaterial : IGLTFJsonIndexedObject
 	FGLTFJsonIORExtension          IOR;
 	FGLTFJsonSheenExtension        Sheen;
 	FGLTFJsonTransmissionExtension Transmission;
+	FGLTFJsonIridescenceExtension  Iridescence;
 
 	virtual void WriteObject(IGLTFJsonWriter& Writer) const override;
 
diff --git a/Engine/Plugins/Interchange/Runtime/Source/Import/Private/Gltf/InterchangeGLTFMaterial.cpp b/Engine/Plugins/Interchange/Runtime/Source/Import/Private/Gltf/InterchangeGLTFMaterial.cpp
--- a/Engine/Plugins/Interchange/Runtime/Source/Import/Private/Gltf/InterchangeGLTFMaterial.cpp
+++ b/Engine/Plugins/Interchange/Runtime/Source/Import/Private/Gltf/InterchangeGLTFMaterial.cpp
@@ -412,8 +412,9 @@ namespace UE::Interchange::GLTFMaterials
 			{
 				//BaseColorTexture
 				//BaseColorTexture_OffsetScale
-				//BaseColorTExture_Rotation
-				//BaseColorTExture_TexCoord
+				//BaseColorTexture_Rotation
+				//BaseColorTexture_TexCoord
+				//BaseColorTexture_TilingMethod
 				SetMap(Inputs::BaseColorTexture, GltfMaterial.BaseColor);
 
 				//BaseColorFactor
@@ -426,6 +427,7 @@ namespace UE::Interchange::GLTFMaterials
 				//NormalTexture_OffsetScale
 				//NormalTexture_Rotation
 				//NormalTexture_TexCoord
+				//NormalTexture_TilingMethod
 				SetMap(Inputs::NormalTexture, GltfMaterial.Normal);
 
 				//NormalScale
@@ -438,6 +440,7 @@ namespace UE::Interchange::GLTFMaterials
 				//EmissiveTexture_OffsetScale
 				//EmissiveTexture_Rotation
 				//EmissiveTexture_TexCoord
+				//EmissiveTexture_TilingMethod
 				SetMap(Inputs::EmissiveTexture, GltfMaterial.Emissive);
 
 				//EmissiveFactor
@@ -452,6 +455,7 @@ namespace UE::Interchange::GLTFMaterials
 				//OcclusionTexture_OffsetScale
 				//OcclusionTexture_Rotation
 				//OcclusionTexture_TexCoord
+				//OcclusionTexture_TilingMethod
 				SetMap(Inputs::OcclusionTexture, GltfMaterial.Occlusion);
 
 				//OcclusionStrength
@@ -469,6 +473,7 @@ namespace UE::Interchange::GLTFMaterials
 					//DiffuseTexture_OffsetScale
 					//DiffuseTexture_Rotation
 					//DiffuseTexture_TexCoord
+					//DiffuseTexture_TilingMethod
 					SetMap(Inputs::DiffuseTexture, GltfMaterial.BaseColor);
 
 					//DiffuseFactor
@@ -480,6 +485,7 @@ namespace UE::Interchange::GLTFMaterials
 					//SpecularGlossinessTexture_OffsetScale
 					//SpecularGlossinessTexture_Rotation
 					//SpecularGlossinessTexture_TexCoord
+					//SpecularGlossinessTexture_TilingMethod
 					SetMap(Inputs::SpecularGlossinessTexture, GltfMaterial.SpecularGlossiness.Map);
 
 					//SpecFactor
@@ -499,6 +505,7 @@ namespace UE::Interchange::GLTFMaterials
 					//BaseColorTexture_OffsetScale
 					//BaseColorTexture_Rotation
 					//BaseColorTexture_TexCoord
+					//BaseColorTexture_TilingMethod
 					SetMap(Inputs::BaseColorTexture, GltfMaterial.BaseColor);
 
 					//BaseColorFactor
@@ -510,6 +517,7 @@ namespace UE::Interchange::GLTFMaterials
 					//MetallicRoughnessTexture_OffsetScale
 					//MetallicRoughnessTexture_Rotation
 					//MetallicRoughnessTexture_TexCoord
+					//MetallicRoughnessTexture_TilingMethod
 					SetMap(Inputs::MetallicRoughnessTexture, GltfMaterial.MetallicRoughness.Map);
 
 					//MetallicFactor
@@ -525,11 +533,40 @@ namespace UE::Interchange::GLTFMaterials
 					//SpecularTexture_OffsetScale
 					//SpecularTexture_Rotation
 					//SpecularTexture_TexCoord
+					//SpecularTexture_TilingMethod
 					SetMap(Inputs::SpecularTexture, GltfMaterial.Specular.SpecularMap);
 
 					//SpecularFactor
 					SetScalar(Inputs::SpecularFactor, GltfMaterial.Specular.SpecularFactor, 0.5); //(UE Specular default is 0.5).
 				}
+
+				if (GltfMaterial.Iridescence.bHasIridescence)
+				{
+					//IridescenceIOR
+					SetScalar(Inputs::IridescenceIOR, GltfMaterial.Iridescence.IOR, 1.3f);
+
+					//IridescenceTexture
+					//IridescenceTexture_OffsetScale
+					//IridescenceTexture_Rotation
+					//IridescenceTexture_TexCoord
+					//IridescenceTexture_TilingMethod
+					SetMap(Inputs::IridescenceTexture, GltfMaterial.Iridescence.Texture);
+
+					//IridescenceFactor
+					SetScalar(Inputs::IridescenceFactor, GltfMaterial.Iridescence.Factor, 0.0f);
+
+					//IridescenceThicknessTexture
+					//IridescenceThicknessTexture_OffsetScale
+					//IridescenceThicknessTexture_Rotation
+					//IridescenceThicknessTexture_TexCoord
+					//IridescenceThicknessTexture_TilingMethod
+					SetMap(Inputs::IridescenceThicknessTexture, GltfMaterial.Iridescence.Thickness.Texture);
+
+					//IridescenceThicknessMinimum
+					SetScalar(Inputs::IridescenceThicknessMinimum, GltfMaterial.Iridescence.Thickness.Minimum, 100.0f);
+					//IridescenceThicknessMaximum
+					SetScalar(Inputs::IridescenceThicknessMaximum, GltfMaterial.Iridescence.Thickness.Maximum, 400.0f);
+				}
 			}
 			else
 			{
@@ -549,6 +586,7 @@ namespace UE::Interchange::GLTFMaterials
 					//TransmissionTexture_OffsetScale
 					//TransmissionTexture_Rotation
 					//TransmissionTexture_TexCoord
+					//TransmissionTexture_TilingMethod
 					SetMap(Inputs::TransmissionTexture, GltfMaterial.Transmission.TransmissionMap);
 
 					//TransmissionFactor
@@ -565,6 +603,7 @@ namespace UE::Interchange::GLTFMaterials
 					//ClearCoatTexture_OffsetScale
 					//ClearCoatTexture_Rotation
 					//ClearCoatTexture_TexCoord
+					//ClearCoatTexture_TilingMethod
 					SetMap(Inputs::ClearCoatTexture, GltfMaterial.ClearCoat.ClearCoatMap);
 
 					//ClearCoatFactor
@@ -576,6 +615,7 @@ namespace UE::Interchange::GLTFMaterials
 					//ClearCoatRoughnessTexture_OffsetScale
 					//ClearCoatRoughnessTexture_Rotation
 					//ClearCoatRoughnessTexture_TexCoord
+					//ClearCoatRoughnessTexture_TilingMethod
 					SetMap(Inputs::ClearCoatRoughnessTexture, GltfMaterial.ClearCoat.RoughnessMap);
 
 					//ClearCoatRoughnessFactor
@@ -587,6 +627,7 @@ namespace UE::Interchange::GLTFMaterials
 					//ClearCoatNormalTexture_OffsetScale
 					//ClearCoatNormalTexture_Rotation
 					//ClearCoatNormalTexture_TexCoord
+					//ClearCoatNormalTexture_TilingMethod
 					SetMap(Inputs::ClearCoatNormalTexture, GltfMaterial.ClearCoat.NormalMap);
 
 					//ClearCoatNormalFactor
@@ -600,6 +641,7 @@ namespace UE::Interchange::GLTFMaterials
 					//SheenColorTexture_OffsetScale
 					//SheenColorTexture_Rotation
 					//SheenColorTexture_TexCoord
+					//SheenColorTexture_TilingMethod
 					SetMap(Inputs::SheenColorTexture, GltfMaterial.Sheen.SheenColorMap);
 
 					//SheenColorFactor
@@ -610,6 +652,7 @@ namespace UE::Interchange::GLTFMaterials
 					//SheenRoughnessTexture_OffsetScale
 					//SheenRoughnessTexture_Rotation
 					//SheenRoughnessTexture_TexCoord
+					//SheenRoughnessTexture_TilingMethod
 					SetMap(Inputs::SheenRoughnessTexture, GltfMaterial.Sheen.SheenRoughnessMap);
 
 					//SheenRoughnessFactor
diff --git a/Engine/Plugins/Interchange/Runtime/Source/Import/Private/Gltf/InterchangeGltfTranslator.cpp b/Engine/Plugins/Interchange/Runtime/Source/Import/Private/Gltf/InterchangeGltfTranslator.cpp
--- a/Engine/Plugins/Interchange/Runtime/Source/Import/Private/Gltf/InterchangeGltfTranslator.cpp
+++ b/Engine/Plugins/Interchange/Runtime/Source/Import/Private/Gltf/InterchangeGltfTranslator.cpp
@@ -63,6 +63,7 @@ static const TArray<FString> ImporterSupportedExtensions = {
 	GLTF::ToString(GLTF::EExtension::KHR_MaterialsSpecular),
 	GLTF::ToString(GLTF::EExtension::KHR_MaterialsPbrSpecularGlossiness),
 	GLTF::ToString(GLTF::EExtension::KHR_MaterialsEmissiveStrength),
+	GLTF::ToString(GLTF::EExtension::KHR_MaterialsIridescence),
 	GLTF::ToString(GLTF::EExtension::MSFT_PackingOcclusionRoughnessMetallic),
 	GLTF::ToString(GLTF::EExtension::MSFT_PackingNormalRoughnessMetallic),
 	/* Textures */
diff --git a/Engine/Plugins/Interchange/Runtime/Source/Import/Public/Gltf/InterchangeGLTFMaterial.h b/Engine/Plugins/Interchange/Runtime/Source/Import/Public/Gltf/InterchangeGLTFMaterial.h
--- a/Engine/Plugins/Interchange/Runtime/Source/Import/Public/Gltf/InterchangeGLTFMaterial.h
+++ b/Engine/Plugins/Interchange/Runtime/Source/Import/Public/Gltf/InterchangeGLTFMaterial.h
@@ -113,6 +113,17 @@ namespace UE::Interchange::GLTFMaterials
 		//Transmission specific:
 		DECLARE_INTERCHANGE_GLTF_MI_MAP(Transmission)
 		const FString TransmissionFactor = TEXT("TransmissionFactor");
+
+
+		//Iridescence specific:
+		const FString IridescenceIOR = TEXT("IridescenceIOR");
+		
+		DECLARE_INTERCHANGE_GLTF_MI_MAP(Iridescence)
+		const FString IridescenceFactor = TEXT("IridescenceFactor");
+
+		DECLARE_INTERCHANGE_GLTF_MI_MAP(IridescenceThickness)
+		const FString IridescenceThicknessMinimum = TEXT("IridescenceThicknessMinimum");
+		const FString IridescenceThicknessMaximum = TEXT("IridescenceThicknessMaximum");
 	}
 
 	enum EShadingModel : uint8
diff --git a/Engine/Plugins/Interchange/Runtime/Source/Parsers/GLTFCore/Private/GLTF/ExtensionsHandler.cpp b/Engine/Plugins/Interchange/Runtime/Source/Parsers/GLTFCore/Private/GLTF/ExtensionsHandler.cpp
--- a/Engine/Plugins/Interchange/Runtime/Source/Parsers/GLTFCore/Private/GLTF/ExtensionsHandler.cpp
+++ b/Engine/Plugins/Interchange/Runtime/Source/Parsers/GLTFCore/Private/GLTF/ExtensionsHandler.cpp
@@ -218,6 +218,7 @@ namespace GLTF
 			EExtension::KHR_MaterialsIOR,
 			EExtension::KHR_MaterialsSpecular,
 			EExtension::KHR_MaterialsEmissiveStrength,
+			EExtension::KHR_MaterialsIridescence,
 			EExtension::MSFT_PackingOcclusionRoughnessMetallic,
 			EExtension::MSFT_PackingNormalRoughnessMetallic 
 		};
@@ -369,6 +370,21 @@ namespace GLTF
 					}
 				}
 				break;
+				case EExtension::KHR_MaterialsIridescence:
+				{
+					const FJsonObject& Iridescence = ExtObj;
+
+					Material.Iridescence.bHasIridescence = true;
+
+					Material.Iridescence.Factor = GetScalar(Iridescence, TEXT("iridescenceFactor"), 0.0f);
+					GLTF::SetTextureMap(Iridescence, TEXT("iridescenceTexture"), nullptr, Asset->Textures, Material.Iridescence.Texture, Messages);
+					Material.Iridescence.IOR = GetScalar(Iridescence, TEXT("iridescenceIor"), 1.3f);
+
+					Material.Iridescence.Thickness.Minimum = GetScalar(Iridescence, TEXT("iridescenceThicknessMinimum"), 100.0f);
+					Material.Iridescence.Thickness.Maximum = GetScalar(Iridescence, TEXT("iridescenceThicknessMaximum"), 400.0f);
+					GLTF::SetTextureMap(Iridescence, TEXT("iridescenceThicknessTexture"), nullptr, Asset->Textures, Material.Iridescence.Thickness.Texture, Messages);
+				}
+				break;
 				default:
 					if (!ensure(false))
 					{
diff --git a/Engine/Plugins/Interchange/Runtime/Source/Parsers/GLTFCore/Private/GLTF/GLTFAsset.cpp b/Engine/Plugins/Interchange/Runtime/Source/Parsers/GLTFCore/Private/GLTF/GLTFAsset.cpp
--- a/Engine/Plugins/Interchange/Runtime/Source/Parsers/GLTFCore/Private/GLTF/GLTFAsset.cpp
+++ b/Engine/Plugins/Interchange/Runtime/Source/Parsers/GLTFCore/Private/GLTF/GLTFAsset.cpp
@@ -396,6 +396,8 @@ namespace GLTF
 				return TEXT("KHR_materials_specular");
 			case GLTF::EExtension::KHR_MaterialsEmissiveStrength:
 				return TEXT("KHR_materials_emissive_strength");
+			case GLTF::EExtension::KHR_MaterialsIridescence:
+				return TEXT("KHR_materials_iridescence");
 			case GLTF::EExtension::KHR_TextureTransform:
 				return TEXT("KHR_texture_transform");
 			case GLTF::EExtension::KHR_DracoMeshCompression:
diff --git a/Engine/Plugins/Interchange/Runtime/Source/Parsers/GLTFCore/Public/GLTFAsset.h b/Engine/Plugins/Interchange/Runtime/Source/Parsers/GLTFCore/Public/GLTFAsset.h
--- a/Engine/Plugins/Interchange/Runtime/Source/Parsers/GLTFCore/Public/GLTFAsset.h
+++ b/Engine/Plugins/Interchange/Runtime/Source/Parsers/GLTFCore/Public/GLTFAsset.h
@@ -37,6 +37,7 @@ namespace GLTF
 		KHR_MaterialsIOR,
 		KHR_MaterialsSpecular,
 		KHR_MaterialsEmissiveStrength,
+		KHR_MaterialsIridescence,
 		KHR_TextureTransform,
 		KHR_DracoMeshCompression,
 		KHR_LightsPunctual,
diff --git a/Engine/Plugins/Interchange/Runtime/Source/Parsers/GLTFCore/Public/GLTFMaterial.h b/Engine/Plugins/Interchange/Runtime/Source/Parsers/GLTFCore/Public/GLTFMaterial.h
--- a/Engine/Plugins/Interchange/Runtime/Source/Parsers/GLTFCore/Public/GLTFMaterial.h
+++ b/Engine/Plugins/Interchange/Runtime/Source/Parsers/GLTFCore/Public/GLTFMaterial.h
@@ -128,6 +128,24 @@ namespace GLTF
 			FTextureMap NormalMap;
 		};
 
+		struct FIridescence
+		{
+			struct FThickness
+			{
+				float       Minimum = 100.0f;
+				float       Maximum = 400.0f;
+				FTextureMap	Texture;
+			};
+
+			float Factor = 0.0f;
+			FTextureMap Texture;
+			float IOR = 1.3f;
+
+			FThickness Thickness;
+
+			bool bHasIridescence = false;
+		};
+
 		FString Name;
 
 		// PBR properties
@@ -157,14 +175,16 @@ namespace GLTF
 		float      AlphaCutoff;  // only used when AlphaMode == Mask
 
 		// extension properties
-		FPacking Packing;
-		bool     bIsUnlitShadingModel;
-		bool     bHasClearCoat;
-		bool     bHasSheen;
-		bool     bHasTransmission;
-		bool     bHasIOR;
-		bool     bHasSpecular;
-		bool     bHasEmissiveStrength;
+		FPacking        Packing;
+		bool            bIsUnlitShadingModel;
+		bool            bHasClearCoat;
+		bool            bHasSheen;
+		bool            bHasTransmission;
+		bool            bHasIOR;
+		bool            bHasSpecular;
+		bool            bHasEmissiveStrength;
+		FIridescence    Iridescence;
+
 
 		FString UniqueId; //will be generated in FAsset::GenerateNames
 
