MaterialX: Disabling import into level, as it's only used by the lights which are non-standard for the moment and materials reimport is not implemented yet.

diff --git a/Engine/Plugins/Interchange/Runtime/Source/Import/Private/MaterialX/InterchangeMaterialXTranslator.cpp b/Engine/Plugins/Interchange/Runtime/Source/Import/Private/MaterialX/InterchangeMaterialXTranslator.cpp
--- a/Engine/Plugins/Interchange/Runtime/Source/Import/Private/MaterialX/InterchangeMaterialXTranslator.cpp
+++ b/Engine/Plugins/Interchange/Runtime/Source/Import/Private/MaterialX/InterchangeMaterialXTranslator.cpp
@@ -25,7 +25,7 @@ static FAutoConsoleVariableRef CCvarInterchangeEnableMaterialXImport(
 
 EInterchangeTranslatorType UInterchangeMaterialXTranslator::GetTranslatorType() const
 {
-	return EInterchangeTranslatorType::Scenes;
+	return EInterchangeTranslatorType::Assets;
 }
 
 EInterchangeTranslatorAssetType UInterchangeMaterialXTranslator::GetSupportedAssetTypes() const
