Fixing errors when compiling TextureSamplerParameterBlur

diff --git a/Engine/Plugins/Interchange/Runtime/Source/Import/Private/MaterialX/MaterialExpressions/MaterialExpressionTextureSampleParameterBlur.cpp b/Engine/Plugins/Interchange/Runtime/Source/Import/Private/MaterialX/MaterialExpressions/MaterialExpressionTextureSampleParameterBlur.cpp
--- a/Engine/Plugins/Interchange/Runtime/Source/Import/Private/MaterialX/MaterialExpressions/MaterialExpressionTextureSampleParameterBlur.cpp
+++ b/Engine/Plugins/Interchange/Runtime/Source/Import/Private/MaterialX/MaterialExpressions/MaterialExpressionTextureSampleParameterBlur.cpp
@@ -82,23 +82,6 @@ namespace Box
 	}
 }
 
-namespace
-{
-	constexpr const TCHAR* ComputeSampleSizeUV = TEXT(
-		R"(float2 derivUVx = ddx(uv) * 0.5f;
-   float2 derivUVy = ddy(uv) * 0.5f;
-   float derivX = abs(derivUVx.x) + abs(derivUVy.x);
-   float derivY = abs(derivUVx.y) + abs(derivUVy.y);
-   float sampleSizeU = 2.0f * filterSize * derivX + filterOffset;
-   if (sampleSizeU < 1.0E-05f)
-       sampleSizeU = 1.0E-05f;
-   float sampleSizeV = 2.0f * filterSize * derivY + filterOffset;
-   if (sampleSizeV < 1.0E-05f)
-       sampleSizeV = 1.0E-05f;
-   return float2(sampleSizeU, sampleSizeV);)"
-	);
-}
-
 UMaterialExpressionMaterialXTextureSampleParameterBlur::UMaterialExpressionMaterialXTextureSampleParameterBlur(const FObjectInitializer& ObjectInitializer)
 	: Super(ObjectInitializer)
 {
@@ -119,91 +102,87 @@ UMaterialExpressionMaterialXTextureSampleParameterBlur::UMaterialExpressionMater
 }
 
 #if WITH_EDITOR
+
+/**
+ * The function is defined in MaterialExpressions.cpp, easier to declare it extern it here rather than redefining it.
+ */
+extern int32 CompileTextureSample(
+	FMaterialCompiler* Compiler,
+	UTexture* Texture,
+	int32 TexCoordCodeIndex,
+	const EMaterialSamplerType SamplerType,
+	const TOptional<FName> ParameterName = TOptional<FName>(),
+	const int32 MipValue0Index = INDEX_NONE,
+	const int32 MipValue1Index = INDEX_NONE,
+	const ETextureMipValueMode MipValueMode = TMVM_None,
+	const ESamplerSourceMode SamplerSource = SSM_FromTextureAsset,
+	const bool AutomaticViewMipBias = false
+);
+
 int32 UMaterialExpressionMaterialXTextureSampleParameterBlur::Compile(FMaterialCompiler* Compiler, int32 OutputIndex)
 {
-	if(KernelSize != EMAterialXTextureSampleBlurKernel ::Kernel1)
+	if(KernelSize == EMAterialXTextureSampleBlurKernel::Kernel1)
 	{
-		FExpressionInput TexCoordInput;
-		if(Coordinates.GetTracedInput().Expression)
-		{
-			TexCoordInput = Coordinates;
-		}
-		else
-		{
-			UMaterialExpressionTextureCoordinate* DefaultTexCoord = NewObject<UMaterialExpressionTextureCoordinate>();
-			DefaultTexCoord->CoordinateIndex = ConstCoordinate;
-			TexCoordInput.Connect(0, DefaultTexCoord);
-		}
+		return Super::Compile(Compiler, OutputIndex);
+	}
 
-		UMaterialExpressionCustom* ComputeSampleSize = NewObject<UMaterialExpressionCustom>();
+	if(FString ErrorMessage; !TextureIsValid(Texture, ErrorMessage))
+	{
+		return CompilerError(Compiler, *ErrorMessage);
+	}
 
-		ComputeSampleSize->Inputs[0].InputName = TEXT("uv");
-		ComputeSampleSize->Inputs[0].Input = TexCoordInput;
+	if(FString SamplerTypeError; !VerifySamplerType(Compiler->GetShaderPlatform(), Compiler->GetTargetPlatform(), Texture, SamplerType, SamplerTypeError))
+	{
+		return Compiler->Errorf(TEXT("%s"), *SamplerTypeError);
+	}
 
-		auto ComputeSampleSizeInput = [&](float Value, FName InputName)
-		{
-			UMaterialExpressionConstant* FilterSizeConstant = NewObject<UMaterialExpressionConstant>();
-			FilterSizeConstant->R = Value;
-			FCustomInput & Input = ComputeSampleSize->Inputs.Add_GetRef({});
-			Input.InputName = InputName;
-			Input.Input.Expression = FilterSizeConstant;
-		};
+	int32 IndexCoordinates = Coordinates.GetTracedInput().Expression ? Coordinates.Compile(Compiler) : Compiler->TextureCoordinate(ConstCoordinate, false, false);
+	int32 IndexHalf = Compiler->Constant(0.5f);
 
-		ComputeSampleSizeInput(FilterSize, TEXT("filterSize"));
-		ComputeSampleSizeInput(FilterOffset, TEXT("filterOffset"));
+	int32 IndexDerivUVx = Compiler->Mul(Compiler->DDX(IndexCoordinates), IndexHalf);
+	int32 IndexDerivUVy = Compiler->Mul(Compiler->DDY(IndexCoordinates), IndexHalf);
 
-		ComputeSampleSize->OutputType = ECustomMaterialOutputType::CMOT_Float2;
-		ComputeSampleSize->Code = ComputeSampleSizeUV;
-				
-		int32 FilterWidth;
-		const float* Kernel = GetKernel(FilterWidth);
+	int32 IndexDerivX = Compiler->Add(Compiler->Abs(Compiler->ComponentMask(IndexDerivUVx, true, false, false, false)),
+									 Compiler->Abs(Compiler->ComponentMask(IndexDerivUVy, true, false, false, false)));
 
-		UMaterialExpressionCustom* Convolution = NewObject<UMaterialExpressionCustom>();
-		Convolution->OutputType = ECustomMaterialOutputType::CMOT_Float4;
-		Convolution->Inputs.Empty();
-		Convolution->Inputs.Reserve(25);
-		Convolution->Code = TEXT("float4 result = float4(0,0,0,0);\n");
+	int32 IndexDerivY = Compiler->Add(Compiler->Abs(Compiler->ComponentMask(IndexDerivUVx, false, true, false, false)),
+									  Compiler->Abs(Compiler->ComponentMask(IndexDerivUVy, false, true, false, false)));
 
-		int32 Index = 0;
+	int32 IndexEpsilon = Compiler->Constant(UE_SMALL_NUMBER);
+	int32 IndexTwo = Compiler->Constant(2.f);
+	int32 IndexFilterSize = Compiler->Constant(FilterSize);
+	int32 IndexFilterOffset = Compiler->Constant(FilterOffset);
 
-		for(int Row = -FilterWidth; Row <= FilterWidth; ++Row)
-		{
-			for(int Col = -FilterWidth; Col <= FilterWidth; ++Col)
-			{
-				UMaterialExpressionConstant2Vector* Constant = NewObject<UMaterialExpressionConstant2Vector>();
-				Constant->R = Col;
-				Constant->G = Row;
-				UMaterialExpressionMultiply* MultiplyCoord = NewObject<UMaterialExpressionMultiply>();
-				MultiplyCoord->A.Expression = ComputeSampleSize;
-				MultiplyCoord->B.Expression = Constant;
-				UMaterialExpressionAdd* AddCoord = NewObject<UMaterialExpressionAdd>();
-				AddCoord->A = TexCoordInput;
-				AddCoord->B.Expression = MultiplyCoord;
-
-				UMaterialExpressionTextureSampleParameter2D* TextureSample = NewObject<UMaterialExpressionTextureSampleParameter2D>();
-				TextureSample->ConstCoordinate = ConstCoordinate;
-				TextureSample->Coordinates.Expression = AddCoord;
-				TextureSample->TextureObject = TextureObject;
-				TextureSample->Texture = Texture;
-				TextureSample->SamplerSource = SamplerSource;
-				TextureSample->SamplerType = SamplerType;
-
-				Convolution->Inputs.Add({});
-				FString InputName{ TEXT("S") + FString::FromInt(Index) };
-				Convolution->Inputs[Index].InputName = FName(InputName);// FName(TEXT("S") + Index);
-				Convolution->Inputs[Index].Input.Expression = TextureSample;
-				Convolution->Code += (TEXT("result += ") + InputName + TEXT("*") + FString::SanitizeFloat(Kernel[Index]) + TEXT(";\n"));
-				++Index;
-			}
-		}
-		Convolution->Code += TEXT("return result;");
+	int32 IndexComputeSampleSizeUV = Compiler->AppendVector(Compiler->Max(Compiler->Add(Compiler->Mul(Compiler->Mul(IndexTwo, IndexFilterSize), IndexDerivX), IndexFilterOffset), IndexEpsilon),
+															Compiler->Max(Compiler->Add(Compiler->Mul(Compiler->Mul(IndexTwo, IndexFilterSize), IndexDerivY), IndexFilterOffset), IndexEpsilon));
 
-		return Convolution->Compile(Compiler, 0);
-	}
-	else
+	int32 FilterWidth, Index = 0;
+	const float* Kernel = GetKernel(FilterWidth);
+
+	int32 IndexResult = Compiler->Constant4(0.f, 0.f, 0.f, 0.f);
+
+	for(int32 Row = -FilterWidth; Row <= FilterWidth; ++Row)
 	{
-		return Super::Compile(Compiler, OutputIndex);
+		for(int32 Col = -FilterWidth; Col <= FilterWidth; ++Col)
+		{
+			IndexResult = Compiler->Add(IndexResult,
+										Compiler->Mul(
+											Compiler->Constant(Kernel[Index++]),
+											CompileTextureSample(
+												Compiler,
+												Texture,
+												Compiler->Add(Compiler->Mul(IndexComputeSampleSizeUV, Compiler->Constant2(Col, Row)), IndexCoordinates),
+												SamplerType,
+												ParameterName,
+												CompileMipValue0(Compiler),
+												CompileMipValue1(Compiler),
+												MipValueMode,
+												SamplerSource,
+												AutomaticViewMipBias)));
+		}
 	}
+
+	return IndexResult;
 }
 
 void UMaterialExpressionMaterialXTextureSampleParameterBlur::GetCaption(TArray<FString>& OutCaptions) const
diff --git a/Engine/Source/Runtime/Engine/Private/Materials/MaterialExpressions.cpp b/Engine/Source/Runtime/Engine/Private/Materials/MaterialExpressions.cpp
--- a/Engine/Source/Runtime/Engine/Private/Materials/MaterialExpressions.cpp
+++ b/Engine/Source/Runtime/Engine/Private/Materials/MaterialExpressions.cpp
@@ -609,7 +609,7 @@ int32 CompileExternalTextureCoordinates(FMaterialCompiler* Compiler, const int32
  * @param SamplerSource           (Optional) Sampler source override
  * @return Index to a new code chunk that samples the texture
  */
-int32 CompileTextureSample(
+ENGINE_API int32 CompileTextureSample(
 	FMaterialCompiler* Compiler,
 	UTexture* Texture,
 	int32 TexCoordCodeIndex,
