Add support of animation property tracks other than just Visibility.
For the time being only: Float, Double, Integer, Bool and Byte (Enum) property tracks are supported.
Support for Cameras and Lights.

diff --git a/Engine/Plugins/Experimental/Enterprise/DatasmithInterchange/Source/DatasmithInterchange/Private/InterchangeDatasmithAnimationUtils.cpp b/Engine/Plugins/Experimental/Enterprise/DatasmithInterchange/Source/DatasmithInterchange/Private/InterchangeDatasmithAnimationUtils.cpp
--- a/Engine/Plugins/Experimental/Enterprise/DatasmithInterchange/Source/DatasmithInterchange/Private/InterchangeDatasmithAnimationUtils.cpp
+++ b/Engine/Plugins/Experimental/Enterprise/DatasmithInterchange/Source/DatasmithInterchange/Private/InterchangeDatasmithAnimationUtils.cpp
@@ -8,6 +8,7 @@
 #include "InterchangeCommonAnimationPayload.h"
 #include "InterchangeAnimationTrackSetNode.h"
 #include "InterchangeAnimSequenceFactoryNode.h"
+#include "InterchangeAnimationDefinitions.h"
 
 #include "Curves/RichCurve.h"
 #include "Misc/FrameRate.h"
@@ -137,7 +138,7 @@ namespace UE::DatasmithInterchange::AnimUtils
 		const FString ActorNodeUid = NodeUtils::GetActorUid(VisibilityAnimation.GetName());
 		TrackNode->SetCustomActorDependencyUid(ActorNodeUid);
 
-		TrackNode->SetCustomTargetedProperty((int32)EInterchangeAnimatedProperty::Visibility);
+		TrackNode->SetCustomPropertyTrack(UE::Interchange::Animation::PropertyTracks::Visibility);
 		TrackNode->SetCustomAnimationPayloadKey(TrackNode->GetUniqueID(), EInterchangeAnimationPayLoadType::STEPCURVE);
 
 		LevelSequenceNode->AddCustomAnimationTrackUid(TrackNode->GetUniqueID());
diff --git a/Engine/Plugins/Interchange/Runtime/Source/Import/Private/Animation/InterchangeLevelSequenceFactory.cpp b/Engine/Plugins/Interchange/Runtime/Source/Import/Private/Animation/InterchangeLevelSequenceFactory.cpp
--- a/Engine/Plugins/Interchange/Runtime/Source/Import/Private/Animation/InterchangeLevelSequenceFactory.cpp
+++ b/Engine/Plugins/Interchange/Runtime/Source/Import/Private/Animation/InterchangeLevelSequenceFactory.cpp
@@ -3,31 +3,31 @@
 #include "Animation/InterchangeLevelSequenceFactory.h"
 
 #include "Animation/InterchangeAnimationPayloadInterface.h"
-#include "InterchangeLevelSequenceFactoryNode.h"
+#include "Animation/InterchangeLevelSequenceHelper.h"
 #include "InterchangeAnimationTrackSetNode.h"
+#include "InterchangeAnimationDefinitions.h"
 #include "InterchangeAnimSequenceFactoryNode.h"
 #include "InterchangeImportCommon.h"
-#include "InterchangeImportLog.h"
+#include "InterchangeLevelSequenceFactoryNode.h"
 #include "InterchangeResult.h"
 #include "InterchangeSourceData.h"
 #include "InterchangeTranslatorBase.h"
 #include "Nodes/InterchangeBaseNode.h"
 #include "Nodes/InterchangeBaseNodeContainer.h"
 
-//#include "Algo/Transform.h"
 #include "Channels/MovieSceneBoolChannel.h"
 #include "Channels/MovieSceneChannelProxy.h"
 #include "Channels/MovieSceneDoubleChannel.h"
+#include "Channels/MovieSceneIntegerChannel.h"
+#include "Channels/MovieSceneByteChannel.h"
 #include "LevelSequence.h"
 #include "MovieScene.h"
 #include "MovieSceneSection.h"
 #include "MovieSceneSequence.h"
 #include "Sections/MovieScene3DTransformSection.h"
-#include "Sections/MovieSceneBoolSection.h"
 #include "Sections/MovieSceneSubSection.h"
 #include "Tracks/MovieScene3DTransformTrack.h"
 #include "Tracks/MovieSceneSubTrack.h"
-#include "Tracks/MovieSceneVisibilityTrack.h"
 
 #if WITH_EDITORONLY_DATA
 
@@ -108,21 +108,7 @@ namespace UE::Interchange::Private
 				}
 				else if (const UInterchangeAnimationTrackNode* TrackNode = Cast<UInterchangeAnimationTrackNode>(TranslatedNode))
 				{
-					int32 TargetedProperty;
-					if (!TrackNode->GetCustomTargetedProperty(TargetedProperty))
-					{
-						continue;
-					}
-
-					// Only visibility is supported for the time being
-					if (TargetedProperty != (int32)EInterchangeAnimatedProperty::Visibility)
-					{
-						continue;
-					}
-
-					// Get targeted actor exists
-					AActor* Actor = GetActor(NodeContainer, TrackNode);
-					if (Actor)
+					if (GetActor(NodeContainer, TrackNode))
 					{
 						return true;
 					}
@@ -148,12 +134,13 @@ namespace UE::Interchange::Private
 		void PopulateLevelSequence();
 
 	private:
-		void PopulateTransformTrack(const UInterchangeTransformAnimationTrackNode& TransformTrackNode, int32 TrackIndex);
+		void PopulateTransformTrack(const UInterchangeTransformAnimationTrackNode& TransformTrackNode);
 		void PopulateSubsequenceTrack(const UInterchangeAnimationTrackSetInstanceNode& InstanceNode);
-		void PopulateAnimationTrack(const UInterchangeAnimationTrackNode& AnimationTrackNode, int32 TrackIndex);
+		void PopulateAnimationTrack(const UInterchangeAnimationTrackNode& AnimationTrackNode);
 
 		AActor* GetActor(const UInterchangeAnimationTrackNode& TrackNode);
-		FGuid BindActorToLevelSequence(AActor& Actor, int32 TrackIndex);
+		FGuid BindActorToLevelSequence(AActor* Actor);
+		FGuid BindComponentToLevelSequence(AActor* Actor);
 		void UpdateTransformChannels(TArrayView<FMovieSceneDoubleChannel*>& Channels, int32 IndexOffset, const TArray<FRichCurve>& Curves);
 		
 		template<typename ChannelType, typename ValueType> 
@@ -221,14 +208,13 @@ namespace UE::Interchange::Private
 		TArray<FString> AnimationTrackUids;
 		FactoryNode.GetCustomAnimationTrackUids(AnimationTrackUids);
 
-		int32 TrackIndex = 0;
 		for (const FString& AnimationTrackUid : AnimationTrackUids)
 		{
 			if (const UInterchangeBaseNode* TranslatedNode = NodeContainer.GetNode(AnimationTrackUid))
 			{
 				if (const UInterchangeTransformAnimationTrackNode* TransformTrackNode = Cast<UInterchangeTransformAnimationTrackNode>(TranslatedNode))
 				{
-					PopulateTransformTrack(*TransformTrackNode, TrackIndex);
+					PopulateTransformTrack(*TransformTrackNode);
 				}
 				else if (const UInterchangeAnimationTrackSetInstanceNode* InstanceTrackNode = Cast<UInterchangeAnimationTrackSetInstanceNode>(TranslatedNode))
 				{
@@ -236,10 +222,8 @@ namespace UE::Interchange::Private
 				}
 				else if (const UInterchangeAnimationTrackNode* TrackNode = Cast<UInterchangeAnimationTrackNode>(TranslatedNode))
 				{
-					PopulateAnimationTrack(*TrackNode, TrackIndex);
+					PopulateAnimationTrack(*TrackNode);
 				}
-
-				++TrackIndex;
 			}
 		}
 
@@ -247,7 +231,7 @@ namespace UE::Interchange::Private
 		LevelSequence.MovieScene->SetEvaluationType(EMovieSceneEvaluationType::FrameLocked);
 	}
 
-	void FLevelSequenceHelper::PopulateTransformTrack(const UInterchangeTransformAnimationTrackNode& TransformTrackNode, int32 TrackIndex)
+	void FLevelSequenceHelper::PopulateTransformTrack(const UInterchangeTransformAnimationTrackNode& TransformTrackNode)
 	{
 		// Get targeted actor exists
 		AActor* Actor = GetActor(TransformTrackNode);
@@ -273,7 +257,7 @@ namespace UE::Interchange::Private
 			return;
 		}
 
-		FGuid ObjectBinding = BindActorToLevelSequence(*Actor, TrackIndex);
+		FGuid ObjectBinding = BindActorToLevelSequence(Actor);
 
 		UMovieScene3DTransformTrack* TransformTrack = MovieScene->FindTrack<UMovieScene3DTransformTrack>(ObjectBinding);
 		if (!TransformTrack)
@@ -444,16 +428,10 @@ namespace UE::Interchange::Private
 		ClearSubsequenceTrack = false;
 	}
 
-	void FLevelSequenceHelper::PopulateAnimationTrack(const UInterchangeAnimationTrackNode& AnimationTrackNode, int32 TrackIndex)
+	void FLevelSequenceHelper::PopulateAnimationTrack(const UInterchangeAnimationTrackNode& AnimationTrackNode)
 	{
-		int32 TargetedProperty;
-		if (!AnimationTrackNode.GetCustomTargetedProperty(TargetedProperty))
-		{
-			return;
-		}
-
-		// Only visibility is supported for the time being
-		if (TargetedProperty != (int32)EInterchangeAnimatedProperty::Visibility)
+		FName PropertyTrack;
+		if(!AnimationTrackNode.GetCustomPropertyTrack(PropertyTrack))
 		{
 			return;
 		}
@@ -476,65 +454,119 @@ namespace UE::Interchange::Private
 
 		TFuture<TOptional<UE::Interchange::FAnimationPayloadData>> Result = PayloadInterface.GetAnimationPayloadData(PayloadKey);
 		const TOptional<UE::Interchange::FAnimationPayloadData>& PayloadData = Result.Get();
-		if (!PayloadData.IsSet() || PayloadData->StepCurves.Num() != 1)
+		if (!PayloadData.IsSet() || (PayloadData->StepCurves.IsEmpty() && PayloadData->Curves.IsEmpty()))
 		{
 			UE_LOG(LogInterchangeImport, Warning, TEXT("No payload for animation track %s on actor %s"), *AnimationTrackNode.GetDisplayLabel(), *Actor->GetActorLabel());
 			return;
 		}
 
-		FGuid ObjectBinding = BindActorToLevelSequence(*Actor, TrackIndex);
-
-		UMovieSceneVisibilityTrack* VisibilityTrack = MovieScene->FindTrack<UMovieSceneVisibilityTrack>(ObjectBinding);
-		if (!VisibilityTrack)
+		FGuid ObjectBinding;
+		if(PropertyTrack == UE::Interchange::Animation::PropertyTracks::Visibility)
 		{
-			VisibilityTrack = MovieScene->AddTrack<UMovieSceneVisibilityTrack>(ObjectBinding);
+			ObjectBinding = BindActorToLevelSequence(Actor);
 		}
 		else
 		{
-			VisibilityTrack->RemoveAllAnimationData();
+			ObjectBinding = BindComponentToLevelSequence(Actor);
 		}
 
-		if (!VisibilityTrack)
-		{
-			return;
-		}
+		UMovieSceneSection * Section = FInterchangePropertyTracksHelper::GetInstance().GetSection(MovieScene, AnimationTrackNode, ObjectBinding, PropertyTrack);
 
-		bool bSectionAdded = false;
-		UMovieSceneBoolSection* BoolSection = Cast<UMovieSceneBoolSection>(VisibilityTrack->FindOrAddSection(0, bSectionAdded));
-		if (!BoolSection)
+		if(!Section)
 		{
 			return;
 		}
 
-		if (bSectionAdded)
+		const FName DoubleChannelTypeName = FMovieSceneDoubleChannel::StaticStruct()->GetFName();
+		const FName FloatChannelTypeName = FMovieSceneFloatChannel::StaticStruct()->GetFName();
+		const FName IntegerChannelTypeName = FMovieSceneIntegerChannel::StaticStruct()->GetFName();
+		const FName BoolChannelTypeName = FMovieSceneBoolChannel::StaticStruct()->GetFName();
+		const FName EnumChannelTypeName = FMovieSceneByteChannel::StaticStruct()->GetFName();
+
+		auto CopyToChannel = [this](auto Channel, const FRichCurve& Curve)
 		{
-			int32 CompletionMode;
-			if (AnimationTrackNode.GetCustomCompletionMode(CompletionMode))
+			const FFrameRate& FrameRate = this->MovieScene->GetTickResolution();
+
+			const TArray<FRichCurveKey>& CurveKeys = Curve.GetConstRefOfKeys();
+
+			TArray<FFrameNumber> FrameNumbers;
+			FrameNumbers.Reserve(CurveKeys.Num());
+
+			using FMovieSceneValue = typename std::remove_pointer_t<decltype(Channel)>::ChannelValueType;
+			TArray<FMovieSceneValue> MovieSceneValues;
+			MovieSceneValues.Reserve(CurveKeys.Num());
+
+			for(int32 KeyIndex = 0; KeyIndex < CurveKeys.Num(); ++KeyIndex)
 			{
-				// Make sure EMovieSceneCompletionMode enum value are still between 0 and 2
-				static_assert(0 == (uint32)EMovieSceneCompletionMode::KeepState, "ENUM_VALUE_HAS_CHANGED");
-				static_assert(1 == (uint32)EMovieSceneCompletionMode::RestoreState, "ENUM_VALUE_HAS_CHANGED");
-				static_assert(2 == (uint32)EMovieSceneCompletionMode::ProjectDefault, "ENUM_VALUE_HAS_CHANGED");
+				const FRichCurveKey& CurveKey = CurveKeys[KeyIndex];
 
-				BoolSection->EvalOptions.CompletionMode = (EMovieSceneCompletionMode)CompletionMode;
+				FFrameNumber& FrameNumber = FrameNumbers.Add_GetRef(FrameRate.AsFrameNumber(CurveKey.Time));
+
+				if(FrameNumber < this->MinFrameNumber)
+				{
+					this->MinFrameNumber = FrameNumber;
+				}
+
+				if(FrameNumber > this->MaxFrameNumber)
+				{
+					this->MaxFrameNumber = FrameNumber;
+				}
+
+				FMovieSceneValue& SceneValue = MovieSceneValues.AddDefaulted_GetRef();
+				SceneValue.InterpMode = CurveKey.InterpMode;
+				SceneValue.Value = CurveKey.Value;
+			}
+
+			if(!MovieSceneValues.IsEmpty())
+			{
+				Channel->Set(FrameNumbers, MovieSceneValues);
 			}
-			// By default the completion mode is EMovieSceneCompletionMode::ProjectDefault
 			else
 			{
-				BoolSection->EvalOptions.CompletionMode = EMovieSceneCompletionMode::ProjectDefault;
+				Channel->RemoveDefault();
 			}
+		};
 
-			BoolSection->SetRange(TRange<FFrameNumber>::All());
-		}
+		FMovieSceneChannelProxy& ChannelProxy = Section->GetChannelProxy();
+		TArrayView<const FMovieSceneChannelEntry> ChannelEntries = ChannelProxy.GetAllEntries();
 
-		TArrayView<FMovieSceneBoolChannel*> Channels = BoolSection->GetChannelProxy().GetChannels<FMovieSceneBoolChannel>();
-		if (ensure(Channels[0]))
+		for(const FMovieSceneChannelEntry& ChannelEntry : ChannelEntries)
 		{
-			const FInterchangeStepCurve& Curve = PayloadData->StepCurves[0];
-			const TArray<float>& KeyTimes = Curve.KeyTimes;
-			const TArray<bool>& Values = Curve.BooleanKeyValues.GetValue();
+			const FName ChannelTypeName = ChannelEntry.GetChannelTypeName();
+			if(ChannelTypeName != DoubleChannelTypeName &&
+			   ChannelTypeName != FloatChannelTypeName &&
+			   ChannelTypeName != IntegerChannelTypeName &&
+			   ChannelTypeName != BoolChannelTypeName &&
+			   ChannelTypeName != EnumChannelTypeName)
+			{
+				continue;
+			}
 
-			UpdateStepChannel(*(Channels[0]), KeyTimes, Values);
+			TArrayView<FMovieSceneChannel* const> Channels = ChannelEntry.GetChannels();
+			for(int32 Index = 0; Index < Channels.Num(); ++Index)
+			{
+				FMovieSceneChannelHandle Channel = ChannelProxy.MakeHandle(ChannelTypeName, Index);
+				if(ChannelTypeName == FMovieSceneBoolChannel::StaticStruct()->GetFName())
+				{
+					UpdateStepChannel(*(Channel.Cast<FMovieSceneBoolChannel>().Get()), PayloadData->StepCurves[Index].KeyTimes, PayloadData->StepCurves[0].BooleanKeyValues.GetValue());
+				}
+				else if(ChannelTypeName == FMovieSceneByteChannel::StaticStruct()->GetFName())
+				{
+					UpdateStepChannel(*(Channel.Cast<FMovieSceneByteChannel>().Get()), PayloadData->StepCurves[Index].KeyTimes, PayloadData->StepCurves[0].ByteKeyValues.GetValue());
+				}
+				else if(ChannelTypeName == FMovieSceneIntegerChannel::StaticStruct()->GetFName())
+				{
+					UpdateStepChannel(*(Channel.Cast<FMovieSceneBoolChannel>().Get()), PayloadData->StepCurves[Index].KeyTimes, PayloadData->StepCurves[0].BooleanKeyValues.GetValue());
+				}
+				else if(ChannelTypeName == FMovieSceneFloatChannel::StaticStruct()->GetFName())
+				{
+					CopyToChannel(Channel.Cast<FMovieSceneFloatChannel>().Get(), PayloadData->Curves[Index]);
+				}
+				else if(ChannelTypeName == FMovieSceneDoubleChannel::StaticStruct()->GetFName())
+				{
+					CopyToChannel(Channel.Cast<FMovieSceneDoubleChannel>().Get(), PayloadData->Curves[Index]);
+				}
+			}
 		}
 	}
 
@@ -562,44 +594,40 @@ namespace UE::Interchange::Private
 		return Actor;
 	}
 
-	FGuid FLevelSequenceHelper::BindActorToLevelSequence(AActor& Actor, int32 TrackIndex)
+	FGuid FLevelSequenceHelper::BindActorToLevelSequence(AActor* Actor)
 	{
-		// Bind the actor to the level sequence
-		// But first, check if there's already a possessable at the current index
-		FMovieScenePossessable* OldPossessable = nullptr;
-		if (TrackIndex < MovieScene->GetPossessableCount())
+		FGuid ActorBinding = LevelSequence.FindBindingFromObject(Actor, Actor->GetWorld());
+		if(!ActorBinding.IsValid())
 		{
-			OldPossessable = &MovieScene->GetPossessable(TrackIndex);
+			ActorBinding = MovieScene->AddPossessable(Actor->GetActorLabel(), Actor->GetClass());
+			LevelSequence.BindPossessableObject(ActorBinding, *Actor, Actor->GetWorld());
 		}
 
-		FString ActorLabel = Actor.GetActorLabel();
-		FGuid ObjectBinding;
-		if (OldPossessable)
+		return ActorBinding;
+	}
+
+	FGuid FLevelSequenceHelper::BindComponentToLevelSequence(AActor* Actor)
+	{
+		FGuid ActorBinding = BindActorToLevelSequence(Actor);
+		USceneComponent* Component = Actor->GetDefaultAttachComponent();
+		FGuid ComponentBinding = LevelSequence.FindBindingFromObject(Component, Actor);
+		if(!ComponentBinding.IsValid())
 		{
-			// If so, remove any track associated with it and unbind it from the level sequence
-			FGuid OldGuid = OldPossessable->GetGuid();
-			UMovieScene3DTransformTrack* TransformTrack = MovieScene->FindTrack<UMovieScene3DTransformTrack>(OldGuid);
-			if (TransformTrack)
+			ComponentBinding = MovieScene->AddPossessable(Component->GetReadableName(), Component->GetClass());
+
+			if(ActorBinding.IsValid() && ComponentBinding.IsValid())
 			{
-				MovieScene->RemoveTrack(*TransformTrack);
+				if(FMovieScenePossessable* ComponentPossessable = MovieScene->FindPossessable(ComponentBinding))
+				{
+					ComponentPossessable->SetParent(ActorBinding, MovieScene);
+				}
 			}
-			LevelSequence.UnbindPossessableObjects(OldGuid);
 
-			// Replace the old possessable with the new one in the MovieScene
-			FMovieScenePossessable NewPossessable(ActorLabel, Actor.GetClass());
-			MovieScene->ReplacePossessable(OldGuid, NewPossessable);
-			ObjectBinding = NewPossessable.GetGuid();
+			// Bind component
+			LevelSequence.BindPossessableObject(ComponentBinding, *Component, Actor);
 		}
-		else
-		{
-			// No previous possessable, so simply add it to the MovieScene
-			ObjectBinding = MovieScene->AddPossessable(ActorLabel, Actor.GetClass());
-		}
-
-		// Finally, associate the guid with the actor
-		LevelSequence.BindPossessableObject(ObjectBinding, Actor, Actor.GetWorld());
 
-		return ObjectBinding;
+		return ComponentBinding;
 	}
 
 	void FLevelSequenceHelper::UpdateTransformChannels(TArrayView<FMovieSceneDoubleChannel*>& Channels, int32 IndexOffset, const TArray<FRichCurve>& Curves)
diff --git a/Engine/Plugins/Interchange/Runtime/Source/Import/Private/Animation/InterchangeLevelSequenceHelper.cpp b/Engine/Plugins/Interchange/Runtime/Source/Import/Private/Animation/InterchangeLevelSequenceHelper.cpp
new file mode 100644
--- /dev/null
+++ b/Engine/Plugins/Interchange/Runtime/Source/Import/Private/Animation/InterchangeLevelSequenceHelper.cpp
@@ -0,0 +1,120 @@
+// Copyright Epic Games, Inc. All Rights Reserved.
+
+#include "InterchangeLevelSequenceHelper.h"
+#include "InterchangeAnimationDefinitions.h"
+
+#include "Camera/CameraComponent.h"
+#include "CineCameraComponent.h"
+#include "Components/LocalLightComponent.h"
+
+#include "MovieScene.h"
+#include "MovieSceneSection.h"
+#include "Tracks/MovieSceneBoolTrack.h"
+#include "Tracks/MovieSceneByteTrack.h"
+#include "Tracks/MovieSceneColorTrack.h"
+#include "Tracks/MovieSceneFloatTrack.h"
+#include "Tracks/MovieSceneVisibilityTrack.h"
+
+namespace UE::Interchange::Private
+{
+	FInterchangePropertyTracksHelper::FInterchangePropertyTracksHelper()
+		: PropertyTracks{
+			// Light
+			{UE::Interchange::Animation::PropertyTracks::Light::Color, {UMovieSceneColorTrack::StaticClass()->GetName(), GET_MEMBER_NAME_STRING_CHECKED(ULocalLightComponent, LightColor), TEXT("Light Color")}},
+			{UE::Interchange::Animation::PropertyTracks::Light::Intensity, {UMovieSceneFloatTrack::StaticClass()->GetName(), GET_MEMBER_NAME_STRING_CHECKED(ULightComponent, Intensity), TEXT("Intensity")}},
+			{UE::Interchange::Animation::PropertyTracks::Light::IntensityUnits, {UMovieSceneByteTrack::StaticClass()->GetName(), GET_MEMBER_NAME_STRING_CHECKED(ULocalLightComponent, IntensityUnits), TEXT("Intensity Units"), StaticEnum<ELightUnits>()}},
+			{UE::Interchange::Animation::PropertyTracks::Light::Temperature, {UMovieSceneFloatTrack::StaticClass()->GetName(), GET_MEMBER_NAME_STRING_CHECKED(ULocalLightComponent, Temperature), TEXT("Temperature")}},
+			{UE::Interchange::Animation::PropertyTracks::Light::UseTemperature, {UMovieSceneBoolTrack::StaticClass()->GetName(), GET_MEMBER_NAME_STRING_CHECKED(ULocalLightComponent, bUseTemperature), TEXT("Use Temperature")}},
+			{UE::Interchange::Animation::PropertyTracks::Visibility, {UMovieSceneVisibilityTrack::StaticClass()->GetName(), AActor::GetHiddenPropertyName().ToString(), TEXT("Visibility")}},
+
+			// Camera
+			{UE::Interchange::Animation::PropertyTracks::Camera::AspectRatioAxisConstraint, {UMovieSceneByteTrack::StaticClass()->GetName(), GET_MEMBER_NAME_STRING_CHECKED(UCameraComponent, AspectRatioAxisConstraint), TEXT("Aspect Ratio Axis Constraint"), StaticEnum<EAspectRatioAxisConstraint>()}},
+			{UE::Interchange::Animation::PropertyTracks::Camera::ConstrainAspectRatio, {UMovieSceneBoolTrack::StaticClass()->GetName(), GET_MEMBER_NAME_STRING_CHECKED(UCameraComponent, bConstrainAspectRatio), TEXT("Constrain Aspect Ratio")}},
+			{UE::Interchange::Animation::PropertyTracks::Camera::CurrentAperture, {UMovieSceneFloatTrack::StaticClass()->GetName(), GET_MEMBER_NAME_STRING_CHECKED(UCineCameraComponent, CurrentAperture), TEXT("Current Aperture")}},
+			{UE::Interchange::Animation::PropertyTracks::Camera::CurrentFocalLength, {UMovieSceneFloatTrack::StaticClass()->GetName(), GET_MEMBER_NAME_STRING_CHECKED(UCineCameraComponent, CurrentFocalLength), TEXT("Current Focal Length")}},
+			{UE::Interchange::Animation::PropertyTracks::Camera::CustomNearClippingPlane, {UMovieSceneFloatTrack::StaticClass()->GetName(), GET_MEMBER_NAME_STRING_CHECKED(UCineCameraComponent, CustomNearClippingPlane), TEXT("Custom Near Clipping Plane")}},
+			{UE::Interchange::Animation::PropertyTracks::Camera::FieldOfView, {UMovieSceneFloatTrack::StaticClass()->GetName(), GET_MEMBER_NAME_STRING_CHECKED(UCameraComponent, FieldOfView), TEXT("Field of View")}},
+			{UE::Interchange::Animation::PropertyTracks::Camera::OrthoFarClipPlane, {UMovieSceneFloatTrack::StaticClass()->GetName(), GET_MEMBER_NAME_STRING_CHECKED(UCameraComponent, OrthoFarClipPlane), TEXT("Ortho Far Clip Plane")}},
+			{UE::Interchange::Animation::PropertyTracks::Camera::OrthoWidth, {UMovieSceneFloatTrack::StaticClass()->GetName(), GET_MEMBER_NAME_STRING_CHECKED(UCameraComponent, OrthoWidth), TEXT("Ortho Width")}},
+			{UE::Interchange::Animation::PropertyTracks::Camera::PostProcessBlendWeight, {UMovieSceneFloatTrack::StaticClass()->GetName(), GET_MEMBER_NAME_STRING_CHECKED(UCameraComponent, PostProcessBlendWeight), TEXT("PostProcess Blend Weight")}},
+			{UE::Interchange::Animation::PropertyTracks::Camera::ProjectionMode, {UMovieSceneByteTrack::StaticClass()->GetName(), GET_MEMBER_NAME_STRING_CHECKED(UCameraComponent, ProjectionMode), TEXT("Projection Mode"), StaticEnum<ECameraProjectionMode::Type>()}},
+			{UE::Interchange::Animation::PropertyTracks::Camera::UseFieldOfViewForLOD, {UMovieSceneBoolTrack::StaticClass()->GetName(), GET_MEMBER_NAME_STRING_CHECKED(UCameraComponent, bUseFieldOfViewForLOD), TEXT("Use Field Of View For LOD")}},
+		}
+	{}
+
+	FInterchangePropertyTracksHelper& FInterchangePropertyTracksHelper::GetInstance()
+	{
+		static FInterchangePropertyTracksHelper Instance;
+		return Instance;
+	}
+
+	UMovieSceneSection* FInterchangePropertyTracksHelper::GetSection(UMovieScene* MovieScene, const UInterchangeAnimationTrackNode& AnimationTrackNode, const FGuid& ObjectBinding, const FName& Property) const
+	{
+		const FInterchangeProperty* InterchangePropertyTrack = PropertyTracks.Find(Property);
+
+		if(!InterchangePropertyTrack)
+		{
+			return nullptr;
+		}
+
+		TSubclassOf<UMovieSceneTrack> TrackClass = FindObjectClass<UMovieSceneTrack>(*InterchangePropertyTrack->ClassType);
+
+		UMovieScenePropertyTrack* PropertyTrack = Cast<UMovieScenePropertyTrack>(MovieScene->FindTrack(TrackClass, ObjectBinding, InterchangePropertyTrack->Name));
+
+		if(!PropertyTrack)
+		{
+			PropertyTrack = Cast<UMovieScenePropertyTrack>(MovieScene->AddTrack(TrackClass, ObjectBinding));
+
+			if(!PropertyTrack)
+			{
+				return nullptr;
+			}
+
+			PropertyTrack->SetPropertyNameAndPath(InterchangePropertyTrack->Name, InterchangePropertyTrack->Path);
+		}
+		else
+		{
+			PropertyTrack->RemoveAllAnimationData();
+		}
+
+		if(InterchangePropertyTrack->EnumClass)
+		{
+			if(UMovieSceneByteTrack* ByteTrack = Cast<UMovieSceneByteTrack>(PropertyTrack))
+			{
+				ByteTrack->SetEnum(InterchangePropertyTrack->EnumClass);
+			}
+		}
+
+		bool bSectionAdded = false;
+
+		UMovieSceneSection* Section = PropertyTrack->FindOrAddSection(0, bSectionAdded);
+
+		if(!Section)
+		{
+			return nullptr;
+		}
+
+		if(bSectionAdded)
+		{
+			int32 CompletionMode;
+			if(AnimationTrackNode.GetCustomCompletionMode(CompletionMode))
+			{
+				// Make sure EMovieSceneCompletionMode enum value are still between 0 and 2
+				static_assert(0 == (uint32)EMovieSceneCompletionMode::KeepState, "ENUM_VALUE_HAS_CHANGED");
+				static_assert(1 == (uint32)EMovieSceneCompletionMode::RestoreState, "ENUM_VALUE_HAS_CHANGED");
+				static_assert(2 == (uint32)EMovieSceneCompletionMode::ProjectDefault, "ENUM_VALUE_HAS_CHANGED");
+
+				Section->EvalOptions.CompletionMode = (EMovieSceneCompletionMode)CompletionMode;
+			}
+			// By default the completion mode is EMovieSceneCompletionMode::ProjectDefault
+			else
+			{
+				Section->EvalOptions.CompletionMode = EMovieSceneCompletionMode::ProjectDefault;
+			}
+
+			Section->SetRange(TRange<FFrameNumber>::All());
+		}
+
+		return Section;
+	}
+}
\ No newline at end of file
diff --git a/Engine/Plugins/Interchange/Runtime/Source/Import/Private/Animation/InterchangeLevelSequenceHelper.h b/Engine/Plugins/Interchange/Runtime/Source/Import/Private/Animation/InterchangeLevelSequenceHelper.h
new file mode 100644
--- /dev/null
+++ b/Engine/Plugins/Interchange/Runtime/Source/Import/Private/Animation/InterchangeLevelSequenceHelper.h
@@ -0,0 +1,68 @@
+// Copyright Epic Games, Inc. All Rights Reserved.
+
+#pragma once
+
+#include "InterchangeAnimationTrackSetNode.h"
+#include "InterchangeImportLog.h"
+
+#include "UObject/ObjectRedirector.h"
+
+class UMovieScene;
+class UMovieSceneSection;
+
+namespace UE::Interchange::Private
+{
+	/**
+	 * Finds a UObject class by name.
+	 * @param ClassName		The name of the class to look for (ie:UClass*->GetName()).
+	 * @return				A sub class of UObject or nullptr.
+	 */
+	template<typename T>
+	TSubclassOf<T> FindObjectClass(const TCHAR* ClassName)
+	{
+		if(!ensure(ClassName))
+		{
+			UE_LOG(LogInterchangeImport, Error, TEXT("Classname is null to find an appropriate animation property track."))
+				return nullptr;
+		}
+
+		UClass* ExpressionClass = FindFirstObject<UClass>(ClassName, EFindFirstObjectOptions::EnsureIfAmbiguous);
+
+		if(!ExpressionClass)
+		{
+			if(UObjectRedirector* RenamedClassRedirector = FindFirstObject<UObjectRedirector>(ClassName, EFindFirstObjectOptions::EnsureIfAmbiguous))
+			{
+				ExpressionClass = CastChecked<UClass>(RenamedClassRedirector->DestinationObject);
+			}
+		}
+
+		if(ExpressionClass && ExpressionClass->IsChildOf<T>())
+		{
+			return ExpressionClass;
+		}
+		else
+		{
+			return nullptr;
+		}
+	}
+
+	struct FInterchangePropertyTracksHelper
+	{
+		static FInterchangePropertyTracksHelper& GetInstance();
+
+		UMovieSceneSection* GetSection(UMovieScene* MovieScene, const UInterchangeAnimationTrackNode& AnimationTrackNode, const FGuid& ObjectBinding, const FName& Property) const;
+
+	private:
+		FInterchangePropertyTracksHelper();
+
+		struct FInterchangeProperty
+		{
+			FString ClassType; // Float, Double, Byte, etc. Basically the class name of the UMovieSceneTrack
+			FString Path;
+			FName Name;
+			UEnum* EnumClass = nullptr; // Only used for Enum property tracks
+		};
+
+		TMap<FName, FInterchangeProperty> PropertyTracks;
+	};
+}
\ No newline at end of file
diff --git a/Engine/Plugins/Interchange/Runtime/Source/Nodes/Private/InterchangeAnimationTrackSetNode.cpp b/Engine/Plugins/Interchange/Runtime/Source/Nodes/Private/InterchangeAnimationTrackSetNode.cpp
--- a/Engine/Plugins/Interchange/Runtime/Source/Nodes/Private/InterchangeAnimationTrackSetNode.cpp
+++ b/Engine/Plugins/Interchange/Runtime/Source/Nodes/Private/InterchangeAnimationTrackSetNode.cpp
@@ -185,11 +185,21 @@ bool UInterchangeAnimationTrackNode::GetCustomTargetedProperty(int32& AttributeV
 	IMPLEMENT_NODE_ATTRIBUTE_GETTER(TargetedProperty, int32);
 }
 
+bool UInterchangeAnimationTrackNode::GetCustomPropertyTrack(FName& AttributeValue) const
+{
+	IMPLEMENT_NODE_ATTRIBUTE_GETTER(PropertyTrack, FName);
+}
+
 bool UInterchangeAnimationTrackNode::SetCustomTargetedProperty(const int32& AttributeValue)
 {
 	IMPLEMENT_NODE_ATTRIBUTE_SETTER_NODELEGATE(TargetedProperty, int32);
 }
 
+bool UInterchangeAnimationTrackNode::SetCustomPropertyTrack(const FName& AttributeValue)
+{
+	IMPLEMENT_NODE_ATTRIBUTE_SETTER_NODELEGATE(PropertyTrack, FName);
+}
+
 // UInterchangeAnimationTrackSetInstanceNode
 
 bool UInterchangeAnimationTrackSetInstanceNode::SetCustomTimeScale(const float& AttributeValue)
diff --git a/Engine/Plugins/Interchange/Runtime/Source/Nodes/Public/InterchangeAnimationDefinitions.h b/Engine/Plugins/Interchange/Runtime/Source/Nodes/Public/InterchangeAnimationDefinitions.h
new file mode 100644
--- /dev/null
+++ b/Engine/Plugins/Interchange/Runtime/Source/Nodes/Public/InterchangeAnimationDefinitions.h
@@ -0,0 +1,46 @@
+// Copyright Epic Games, Inc. All Rights Reserved.
+
+#pragma once
+
+#include "UObject/NameTypes.h"
+
+namespace UE
+{
+	namespace Interchange
+	{
+		namespace Animation
+		{
+			namespace PropertyTracks
+			{
+				const FName Visibility = TEXT("Actor.bHidden");
+
+				namespace Light
+				{
+					const FName Color = TEXT("Light.LightColor");
+					const FName Intensity = TEXT("Light.Intensity");
+					const FName IntensityUnits = TEXT("Light.IntensityUnits");
+					const FName Temperature = TEXT("Light.Temperature");
+					const FName UseTemperature = TEXT("Light.bUseTemperature");
+				}
+
+				namespace Camera
+				{
+					const FName AutoActivate = TEXT("Camera.bAutoActivate");
+					const FName AspectRatioAxisConstraint = TEXT("Camera.AspectRatioAxisConstraint");
+					const FName ConstrainAspectRatio = TEXT("Camera.bConstrainAspectRatio");
+					const FName CurrentAperture = TEXT("Camera.CurrentAperture");
+					const FName CurrentFocalLength = TEXT("Camera.CurrentFocalLength");
+					const FName CustomNearClippingPlane = TEXT("Camera.CustomNearClippingPlane");
+					const FName FieldOfView = TEXT("Camera.FieldOfView");
+					const FName Mobility = TEXT("Camera.Mobility");
+					const FName OrthoFarClipPlane = TEXT("Camera.OrthoFarClipPlane");
+					const FName OrthoWidth = TEXT("Camera.OrthoWidth");
+					const FName PostProcessBlendWeight = TEXT("Camera.PostProcessBlendWeight");
+					const FName ProjectionMode = TEXT("Camera.ProjectionMode");
+					const FName ShouldUpdatePhysicsVolume = TEXT("Camera.bShouldUpdatePhysicsVolume");
+					const FName UseFieldOfViewForLOD = TEXT("Camera.bUseFieldOfViewForLOD");
+				}
+			}
+		}
+	}
+}
\ No newline at end of file
diff --git a/Engine/Plugins/Interchange/Runtime/Source/Nodes/Public/InterchangeAnimationTrackSetNode.h b/Engine/Plugins/Interchange/Runtime/Source/Nodes/Public/InterchangeAnimationTrackSetNode.h
--- a/Engine/Plugins/Interchange/Runtime/Source/Nodes/Public/InterchangeAnimationTrackSetNode.h
+++ b/Engine/Plugins/Interchange/Runtime/Source/Nodes/Public/InterchangeAnimationTrackSetNode.h
@@ -313,15 +313,27 @@ public:
 	/**
 	 * Set the property animated by this track.
 	 */
+	UE_DEPRECATED(5.4, "SetCustomTargetedProperty has been deprecated, please use SetCustomPropertyTrack instead.")
 	UFUNCTION(BlueprintCallable, Category = "Interchange | Node | AnimationTrack")
 	bool SetCustomTargetedProperty(const int32& TargetedProperty);
 
+	/**
+	 * Set the property animated by this track. Usually the name of a UMovieSceneTrack, e.g for UMovieSceneColorTrack -> Color
+	 */
+	UFUNCTION(BlueprintCallable, Category = "Interchange | Node | AnimationTrack")
+	bool SetCustomPropertyTrack(const FName& PropertyTrack);
+
 	/**
 	 * Get the property animated by this track.
 	 */
+	UE_DEPRECATED(5.4, "SetCustomTargetedProperty has been deprecated, please use SetCustomPropertyTrack instead.")
 	UFUNCTION(BlueprintCallable, Category = "Interchange | Node | AnimationTrack")
 	bool GetCustomTargetedProperty(int32& TargetedProperty) const;
 
+
+	UFUNCTION(BlueprintCallable, Category = "Interchange | Node | AnimationTrack")
+	bool GetCustomPropertyTrack(FName& PropertyTrack) const;
+
 	/**
 	 * Set the payload key needed to retrieve the animation for this track.
 	 */
@@ -352,6 +364,7 @@ private:
 	const UE::Interchange::FAttributeKey Macro_CustomAnimationPayloadTypeKey = UE::Interchange::FAttributeKey(TEXT("AnimationPayloadType"));
 	const UE::Interchange::FAttributeKey Macro_CustomFrameCountKey = UE::Interchange::FAttributeKey(TEXT("FrameCount"));
 	const UE::Interchange::FAttributeKey Macro_CustomTargetedPropertyKey = UE::Interchange::FAttributeKey(TEXT("TargetedProperty"));
+	const UE::Interchange::FAttributeKey Macro_CustomPropertyTrackKey = UE::Interchange::FAttributeKey(TEXT("PropertyTrack"));
 };
 
 /**
diff --git a/Engine/Plugins/Interchange/Runtime/Source/Parsers/CommonParser/Private/InterchangeCommonAnimationPayload.cpp b/Engine/Plugins/Interchange/Runtime/Source/Parsers/CommonParser/Private/InterchangeCommonAnimationPayload.cpp
--- a/Engine/Plugins/Interchange/Runtime/Source/Parsers/CommonParser/Private/InterchangeCommonAnimationPayload.cpp
+++ b/Engine/Plugins/Interchange/Runtime/Source/Parsers/CommonParser/Private/InterchangeCommonAnimationPayload.cpp
@@ -236,12 +236,28 @@ void FInterchangeStepCurve::RemoveRedundantKeys(float Threshold)
 				return ValueA == ValueB;
 			});
 	}
+	else if(BooleanKeyValues.IsSet())
+	{
+		InternalRemoveRedundantKey<bool>(BooleanKeyValues.GetValue(), [](bool ValueA, bool ValueB)
+		{
+			return ValueA == ValueB;
+		});
+	}
+	else if(ByteKeyValues.IsSet())
+	{
+		InternalRemoveRedundantKey<uint8>(ByteKeyValues.GetValue(), [](uint8 ValueA, uint8 ValueB)
+		{
+			return ValueA == ValueB;
+		});
+	}
 }
 
 void FInterchangeStepCurve::Serialize(FArchive& Ar)
 {
 	Ar << KeyTimes;
 	Ar << FloatKeyValues;
+	Ar << ByteKeyValues;
+	Ar << BooleanKeyValues;
 	Ar << IntegerKeyValues;
 	Ar << StringKeyValues;
 }
diff --git a/Engine/Plugins/Interchange/Runtime/Source/Parsers/CommonParser/Public/InterchangeCommonAnimationPayload.h b/Engine/Plugins/Interchange/Runtime/Source/Parsers/CommonParser/Public/InterchangeCommonAnimationPayload.h
--- a/Engine/Plugins/Interchange/Runtime/Source/Parsers/CommonParser/Public/InterchangeCommonAnimationPayload.h
+++ b/Engine/Plugins/Interchange/Runtime/Source/Parsers/CommonParser/Public/InterchangeCommonAnimationPayload.h
@@ -195,6 +195,7 @@ struct INTERCHANGECOMMONPARSER_API FInterchangeStepCurve
 	TArray<float> KeyTimes;
 	TOptional<TArray<float>> FloatKeyValues;
 	TOptional<TArray<bool>> BooleanKeyValues;
+	TOptional<TArray<uint8>> ByteKeyValues;
 	TOptional<TArray<int32>> IntegerKeyValues;
 	TOptional<TArray<FString>> StringKeyValues;
 
