Interchange: Add "keep section separate" feature

#rnx

diff --git a/Engine/Plugins/Interchange/Editor/Source/InterchangeEditor/Private/InterchangeFbxAssetImportDataConverter.cpp b/Engine/Plugins/Interchange/Editor/Source/InterchangeEditor/Private/InterchangeFbxAssetImportDataConverter.cpp
--- a/Engine/Plugins/Interchange/Editor/Source/InterchangeEditor/Private/InterchangeFbxAssetImportDataConverter.cpp
+++ b/Engine/Plugins/Interchange/Editor/Source/InterchangeEditor/Private/InterchangeFbxAssetImportDataConverter.cpp
@@ -157,12 +157,9 @@ namespace UE::Interchange::Private
 		FillInterchangeGenericAssetsPipelineFromFbxMeshImportData(GenericAssetPipeline, SkeletalMeshImportData);
 
 		GenericAssetPipeline->CommonSkeletalMeshesAndAnimationsProperties->bImportMeshesInBoneHierarchy = SkeletalMeshImportData->bImportMeshesInBoneHierarchy;
+		GenericAssetPipeline->CommonMeshesProperties->bKeepSectionsSeparate = SkeletalMeshImportData->bKeepSectionsSeparate;
 		GenericAssetPipeline->MeshPipeline->bImportMorphTargets = SkeletalMeshImportData->bImportMorphTargets;
 		GenericAssetPipeline->MeshPipeline->bImportVertexAttributes = SkeletalMeshImportData->bImportVertexAttributes;
-
-		//TODO integrate this new feature to interchange
-		//SkeletalMeshImportData->bKeepSectionsSeparate;
-
 		GenericAssetPipeline->MeshPipeline->bUpdateSkeletonReferencePose = SkeletalMeshImportData->bUpdateSkeletonReferencePose;
 		GenericAssetPipeline->CommonSkeletalMeshesAndAnimationsProperties->bUseT0AsRefPose = SkeletalMeshImportData->bUseT0AsRefPose;
 		if (SkeletalMeshImportData->ImportContentType == EFBXImportContentType::FBXICT_All)
@@ -347,10 +344,7 @@ namespace UE::Interchange::Private
 				DestinationSkeletalMeshImportData->bImportMeshesInBoneHierarchy = GenericAssetPipeline->CommonSkeletalMeshesAndAnimationsProperties->bImportMeshesInBoneHierarchy;
 				DestinationSkeletalMeshImportData->bImportMorphTargets = GenericAssetPipeline->MeshPipeline->bImportMorphTargets;
 				DestinationSkeletalMeshImportData->bImportVertexAttributes = GenericAssetPipeline->MeshPipeline->bImportVertexAttributes;
-
-				//TODO integrate this new feature to interchange
-				//DestinationSkeletalMeshImportData->bKeepSectionsSeparate = ;
-
+				DestinationSkeletalMeshImportData->bKeepSectionsSeparate = GenericAssetPipeline->CommonMeshesProperties->bKeepSectionsSeparate;
 				DestinationSkeletalMeshImportData->bPreserveSmoothingGroups = true;
 				DestinationSkeletalMeshImportData->bUpdateSkeletonReferencePose = GenericAssetPipeline->MeshPipeline->bUpdateSkeletonReferencePose;
 				DestinationSkeletalMeshImportData->bUseT0AsRefPose = GenericAssetPipeline->CommonSkeletalMeshesAndAnimationsProperties->bUseT0AsRefPose;
diff --git a/Engine/Plugins/Interchange/Runtime/Source/FactoryNodes/Private/InterchangeMeshFactoryNode.cpp b/Engine/Plugins/Interchange/Runtime/Source/FactoryNodes/Private/InterchangeMeshFactoryNode.cpp
--- a/Engine/Plugins/Interchange/Runtime/Source/FactoryNodes/Private/InterchangeMeshFactoryNode.cpp
+++ b/Engine/Plugins/Interchange/Runtime/Source/FactoryNodes/Private/InterchangeMeshFactoryNode.cpp
@@ -215,9 +215,39 @@ bool UInterchangeMeshFactoryNode::SetCustomVertexColorOverride(const FColor& Att
 	IMPLEMENT_NODE_ATTRIBUTE_SETTER_NODELEGATE(VertexColorOverride, FColor)
 }
 
+bool UInterchangeMeshFactoryNode::GetCustomKeepSectionsSeparate(bool& AttributeValue) const
+{
+	IMPLEMENT_NODE_ATTRIBUTE_GETTER(KeepSectionsSeparate, bool)
+}
+
+bool UInterchangeMeshFactoryNode::SetCustomKeepSectionsSeparate(const bool& AttributeValue)
+{
+	IMPLEMENT_NODE_ATTRIBUTE_SETTER_NODELEGATE(KeepSectionsSeparate, bool)
+}
+
 void UInterchangeMeshFactoryNode::GetSlotMaterialDependencies(TMap<FString, FString>& OutMaterialDependencies) const
 {
-	OutMaterialDependencies = SlotMaterialDependencies.ToMap();
+	bool bKeepSectionsSeparate = false;
+	GetCustomKeepSectionsSeparate(bKeepSectionsSeparate);
+	if (bKeepSectionsSeparate)
+	{
+		OutMaterialDependencies = SlotMaterialDependencies.ToMap();
+	}
+	else
+	{
+		TMap<FString, FString> MaterialDependencies = SlotMaterialDependencies.ToMap();
+		TArray<FString> MaterialUids;
+		MaterialUids.Reserve(MaterialDependencies.Num());
+		OutMaterialDependencies.Empty(MaterialDependencies.Num());
+		for (TPair<FString, FString>& MaterialDependency : MaterialDependencies)
+		{
+			if (MaterialUids.Find(MaterialDependency.Value) == INDEX_NONE)
+			{
+				OutMaterialDependencies.Add(MaterialDependency);
+				MaterialUids.Add(MaterialDependency.Value);
+			}
+		}
+	}
 }
 
 bool UInterchangeMeshFactoryNode::GetSlotMaterialDependencyUid(const FString& SlotName, FString& OutMaterialDependency) const
@@ -240,6 +270,13 @@ bool UInterchangeMeshFactoryNode::RemoveSlotMaterialDependencyUid(const FString&
 	return false;
 }
 
+bool UInterchangeMeshFactoryNode::ResetSlotMaterialDependencies()
+{
+	SlotMaterialDependencies.Empty();
+	SlotMaterialDependencies.RebuildCache();
+	return true;
+}
+
 bool UInterchangeMeshFactoryNode::GetCustomLODGroup(FName& AttributeValue) const
 {
 	IMPLEMENT_NODE_ATTRIBUTE_GETTER(LODGroup, FName)
diff --git a/Engine/Plugins/Interchange/Runtime/Source/FactoryNodes/Public/InterchangeMeshActorFactoryNode.h b/Engine/Plugins/Interchange/Runtime/Source/FactoryNodes/Public/InterchangeMeshActorFactoryNode.h
--- a/Engine/Plugins/Interchange/Runtime/Source/FactoryNodes/Public/InterchangeMeshActorFactoryNode.h
+++ b/Engine/Plugins/Interchange/Runtime/Source/FactoryNodes/Public/InterchangeMeshActorFactoryNode.h
@@ -66,6 +66,13 @@ public:
 	UFUNCTION(BlueprintCallable, Category = "Interchange | Node | StaticMesh")
 	bool SetCustomGeometricTransform(const FTransform& AttributeValue);
 
+	bool GetCustomKeepSectionsSeparate(bool& AttributeValue) const
+	{
+		//Scene import do not support this options so we set the value to false and return false
+		AttributeValue = false;
+		return false;
+	}
+
 private:
 	UE::Interchange::TMapAttributeHelper<FString, FString> SlotMaterialDependencies;
 
diff --git a/Engine/Plugins/Interchange/Runtime/Source/FactoryNodes/Public/InterchangeMeshFactoryNode.h b/Engine/Plugins/Interchange/Runtime/Source/FactoryNodes/Public/InterchangeMeshFactoryNode.h
--- a/Engine/Plugins/Interchange/Runtime/Source/FactoryNodes/Public/InterchangeMeshFactoryNode.h
+++ b/Engine/Plugins/Interchange/Runtime/Source/FactoryNodes/Public/InterchangeMeshFactoryNode.h
@@ -87,6 +87,14 @@ public:
 	UFUNCTION(BlueprintCallable, Category = "Interchange | FactoryNode | Mesh")
 	bool SetCustomVertexColorOverride(const FColor& AttributeValue);
 
+	/** Query whether sections with matching materials are kept separate and will not get combined. */
+	UFUNCTION(BlueprintCallable, Category = "Interchange | FactoryNode | Mesh")
+	bool GetCustomKeepSectionsSeparate(bool& AttributeValue) const;
+
+	/** Set whether sections with matching materials are kept separate and will not get combined. */
+	UFUNCTION(BlueprintCallable, Category = "Interchange | FactoryNode | Mesh")
+	bool SetCustomKeepSectionsSeparate(const bool& AttributeValue);
+
 	/** Allow to retrieve the correspondence table between slot names and assigned materials for this object. */
 	UFUNCTION(BlueprintCallable, Category = "Interchange | FactoryNode | Mesh")
 	void GetSlotMaterialDependencies(TMap<FString, FString>& OutMaterialDependencies) const;
@@ -103,6 +111,10 @@ public:
 	UFUNCTION(BlueprintCallable, Category = "Interchange | FactoryNode | Mesh")
 	bool RemoveSlotMaterialDependencyUid(const FString& SlotName);
 
+	/** Reset all the material dependencies. */
+	UFUNCTION(BlueprintCallable, Category = "Interchange | FactoryNode | Mesh")
+	bool ResetSlotMaterialDependencies();
+
 	/** Query whether normals in the imported mesh are ignored and recomputed. When normals are recomputed the tangents are also recomputed. */
 	UFUNCTION(BlueprintCallable, Category = "Interchange | FactoryNode | Mesh")
 	bool GetCustomLODGroup(FName& AttributeValue) const;
@@ -182,6 +194,7 @@ private:
 	const UE::Interchange::FAttributeKey Macro_CustomVertexColorReplaceKey = UE::Interchange::FAttributeKey(TEXT("VertexColorReplace"));
 	const UE::Interchange::FAttributeKey Macro_CustomVertexColorIgnoreKey = UE::Interchange::FAttributeKey(TEXT("VertexColorIgnore"));
 	const UE::Interchange::FAttributeKey Macro_CustomVertexColorOverrideKey = UE::Interchange::FAttributeKey(TEXT("VertexColorOverride"));
+	const UE::Interchange::FAttributeKey Macro_CustomKeepSectionsSeparateKey = UE::Interchange::FAttributeKey(TEXT("KeepSectionsSeparate"));
 	const UE::Interchange::FAttributeKey Macro_CustomLODGroupKey = UE::Interchange::FAttributeKey(TEXT("LODGroup"));
 	const UE::Interchange::FAttributeKey Macro_CustomRecomputeNormalsKey = UE::Interchange::FAttributeKey(TEXT("RecomputeNormals"));
 	const UE::Interchange::FAttributeKey Macro_CustomRecomputeTangentsKey = UE::Interchange::FAttributeKey(TEXT("RecomputeTangents"));
diff --git a/Engine/Plugins/Interchange/Runtime/Source/Import/Private/Mesh/InterchangeMeshHelper.cpp b/Engine/Plugins/Interchange/Runtime/Source/Import/Private/Mesh/InterchangeMeshHelper.cpp
new file mode 100644
--- /dev/null
+++ b/Engine/Plugins/Interchange/Runtime/Source/Import/Private/Mesh/InterchangeMeshHelper.cpp
@@ -0,0 +1,45 @@
+// Copyright Epic Games, Inc. All Rights Reserved. 
+#include "Mesh/InterchangeMeshHelper.h"
+
+#include "MeshDescription.h"
+#include "StaticMeshOperations.h"
+#include "GenericPlatform/GenericPlatformMisc.h"
+#include "StaticMeshAttributes.h"
+
+namespace UE::Interchange::Private::MeshHelper
+{
+	void RemapPolygonGroups(const FMeshDescription& SourceMesh, FMeshDescription& TargetMesh, PolygonGroupMap& RemapPolygonGroup)
+	{
+		FStaticMeshConstAttributes SourceAttributes(SourceMesh);
+		TPolygonGroupAttributesConstRef<FName> SourceImportedMaterialSlotNames = SourceAttributes.GetPolygonGroupMaterialSlotNames();
+
+		FStaticMeshAttributes TargetAttributes(TargetMesh);
+		TPolygonGroupAttributesRef<FName> TargetImportedMaterialSlotNames = TargetAttributes.GetPolygonGroupMaterialSlotNames();
+
+		for (FPolygonGroupID SourcePolygonGroupID : SourceMesh.PolygonGroups().GetElementIDs())
+		{
+			FPolygonGroupID TargetMatchingID = INDEX_NONE;
+			for (FPolygonGroupID TargetPolygonGroupID : TargetMesh.PolygonGroups().GetElementIDs())
+			{
+				if (SourceImportedMaterialSlotNames[SourcePolygonGroupID] == TargetImportedMaterialSlotNames[TargetPolygonGroupID])
+				{
+					TargetMatchingID = TargetPolygonGroupID;
+					break;
+				}
+			}
+			if (TargetMatchingID == INDEX_NONE)
+			{
+				TargetMatchingID = TargetMesh.CreatePolygonGroup();
+				TargetImportedMaterialSlotNames[TargetMatchingID] = SourceImportedMaterialSlotNames[SourcePolygonGroupID];
+			}
+			else
+			{
+				//Since we want to keep the sections separate we need to create a new polygongroup
+				TargetMatchingID = TargetMesh.CreatePolygonGroup();
+				FString NewSlotName = SourceImportedMaterialSlotNames[SourcePolygonGroupID].ToString() + TEXT("_Section") + FString::FromInt(TargetMatchingID.GetValue());
+				TargetImportedMaterialSlotNames[TargetMatchingID] = FName(NewSlotName);
+			}
+			RemapPolygonGroup.Add(SourcePolygonGroupID, TargetMatchingID);
+		}
+	}
+} //ns UE::Interchange::Private::MeshHelper
\ No newline at end of file
diff --git a/Engine/Plugins/Interchange/Runtime/Source/Import/Private/Mesh/InterchangeSkeletalMeshFactory.cpp b/Engine/Plugins/Interchange/Runtime/Source/Import/Private/Mesh/InterchangeSkeletalMeshFactory.cpp
--- a/Engine/Plugins/Interchange/Runtime/Source/Import/Private/Mesh/InterchangeSkeletalMeshFactory.cpp
+++ b/Engine/Plugins/Interchange/Runtime/Source/Import/Private/Mesh/InterchangeSkeletalMeshFactory.cpp
@@ -28,6 +28,7 @@
 #include "MaterialDomain.h"
 #include "Materials/Material.h"
 #include "Math/GenericOctree.h"
+#include "Mesh/InterchangeMeshHelper.h"
 #include "Mesh/InterchangeMeshPayload.h"
 #include "Mesh/InterchangeMeshPayloadInterface.h"
 #include "Misc/MessageDialog.h"
@@ -436,6 +437,9 @@ namespace UE
 				{
 					AppendSettings.bMergeUVChannels[ChannelIdx] = true;
 				}
+				
+				bool bKeepSectionsSeparate = false;
+				SkeletalMeshFactoryNode->GetCustomKeepSectionsSeparate(bKeepSectionsSeparate);
 
 				bool bImportMorphTarget = true;
 				SkeletalMeshFactoryNode->GetCustomImportMorphTarget(bImportMorphTarget);
@@ -574,6 +578,13 @@ namespace UE
 					//The Mesh node parent bake transform can be pass to the payload request or not, it depend on the count of instance and the scale of the transform.
 					const FInternalInstanceData & InstanceData = MeshInstancesDatas.FindChecked(MeshNodeContext.TranslatorPayloadKey.UniqueId);
 					AppendSettings.MeshTransform = InstanceData.ShouldFetchWithTransform() ? FTransform::Identity : MeshNodeContext.SceneGlobalTransform.Get(FTransform::Identity);
+					if (bKeepSectionsSeparate)
+					{
+						AppendSettings.PolygonGroupsDelegate = FAppendPolygonGroupsDelegate::CreateLambda([](const FMeshDescription& SourceMesh, FMeshDescription& TargetMesh, PolygonGroupMap& RemapPolygonGroup)
+							{
+								UE::Interchange::Private::MeshHelper::RemapPolygonGroups(SourceMesh, TargetMesh, RemapPolygonGroup);
+							});
+					}
 
 					FStaticMeshOperations::AppendMeshDescription(LodMeshPayload->MeshDescription, LodMeshDescription, AppendSettings);
 					if (MeshNodeContext.MeshNode->IsSkinnedMesh() || bIsRigidMesh)
diff --git a/Engine/Plugins/Interchange/Runtime/Source/Import/Private/Mesh/InterchangeStaticMeshFactory.cpp b/Engine/Plugins/Interchange/Runtime/Source/Import/Private/Mesh/InterchangeStaticMeshFactory.cpp
--- a/Engine/Plugins/Interchange/Runtime/Source/Import/Private/Mesh/InterchangeStaticMeshFactory.cpp
+++ b/Engine/Plugins/Interchange/Runtime/Source/Import/Private/Mesh/InterchangeStaticMeshFactory.cpp
@@ -24,6 +24,7 @@
 #include "Materials/Material.h"
 #include "Materials/MaterialInstance.h"
 #include "Materials/MaterialInterface.h"
+#include "Mesh/InterchangeMeshHelper.h"
 #include "Mesh/InterchangeMeshPayloadInterface.h"
 #include "MeshBudgetProjectSettings.h"
 #include "Model.h"
@@ -261,6 +262,9 @@ UInterchangeFactoryBase::FImportAssetResult UInterchangeStaticMeshFactory::Impor
 		StaticMesh->GetVertexColorData(ExisitingVertexColorData);
 	}
 
+	bool bKeepSectionsSeparate = false;
+	StaticMeshFactoryNode->GetCustomKeepSectionsSeparate(bKeepSectionsSeparate);
+
 	// Set material slots from imported materials
 	auto UpdateOrAddStaticMaterial = [&StaticMesh, bReimport](const FName& MaterialSlotName, UMaterialInterface* MaterialInterface)
 	{
@@ -374,6 +378,13 @@ UInterchangeFactoryBase::FImportAssetResult UInterchangeStaticMeshFactory::Impor
 				{
 					continue;
 				}
+				if (bKeepSectionsSeparate)
+				{
+					AppendSettings.PolygonGroupsDelegate = FAppendPolygonGroupsDelegate::CreateLambda([](const FMeshDescription& SourceMesh, FMeshDescription& TargetMesh, PolygonGroupMap& RemapPolygonGroup)
+						{
+							UE::Interchange::Private::MeshHelper::RemapPolygonGroups(SourceMesh, TargetMesh, RemapPolygonGroup);
+						});
+				}
 				FStaticMeshOperations::AppendMeshDescription(LodMeshPayload->MeshDescription, LodMeshDescription, AppendSettings);
 			}
 		}
diff --git a/Engine/Plugins/Interchange/Runtime/Source/Import/Public/Mesh/InterchangeMeshHelper.h b/Engine/Plugins/Interchange/Runtime/Source/Import/Public/Mesh/InterchangeMeshHelper.h
new file mode 100644
--- /dev/null
+++ b/Engine/Plugins/Interchange/Runtime/Source/Import/Public/Mesh/InterchangeMeshHelper.h
@@ -0,0 +1,8 @@
+// Copyright Epic Games, Inc. All Rights Reserved. 
+#include "MeshDescription.h"
+#include "StaticMeshOperations.h"
+
+namespace UE::Interchange::Private::MeshHelper
+{
+	INTERCHANGEIMPORT_API void RemapPolygonGroups(const FMeshDescription& SourceMesh, FMeshDescription& TargetMesh, PolygonGroupMap& RemapPolygonGroup);
+} //ns UE::Interchange::Private::MeshHelper
\ No newline at end of file
diff --git a/Engine/Plugins/Interchange/Runtime/Source/Parsers/Fbx/Private/FbxAPI.cpp b/Engine/Plugins/Interchange/Runtime/Source/Parsers/Fbx/Private/FbxAPI.cpp
--- a/Engine/Plugins/Interchange/Runtime/Source/Parsers/Fbx/Private/FbxAPI.cpp
+++ b/Engine/Plugins/Interchange/Runtime/Source/Parsers/Fbx/Private/FbxAPI.cpp
@@ -232,7 +232,7 @@ namespace UE
 				SDKScene->GetFbxManager()->CreateMissingBindPoses(SDKScene);
 				//if we created missing bind poses, update the number of bind poses
 				const int32 NbPoses = SDKScene->GetFbxManager()->GetBindPoseCount(SDKScene);
-				if (NbPoses != Default_NbPoses)
+				if (NbPoses != Default_NbPoses && !GIsAutomationTesting)
 				{
 					UInterchangeResultWarning_Generic* Message = AddMessage<UInterchangeResultWarning_Generic>();
 					Message->Text = LOCTEXT("MissingBindPose", "Missing bind pose - the FBX SDK has created one.");
diff --git a/Engine/Plugins/Interchange/Runtime/Source/Parsers/Fbx/Private/FbxMaterial.cpp b/Engine/Plugins/Interchange/Runtime/Source/Parsers/Fbx/Private/FbxMaterial.cpp
--- a/Engine/Plugins/Interchange/Runtime/Source/Parsers/Fbx/Private/FbxMaterial.cpp
+++ b/Engine/Plugins/Interchange/Runtime/Source/Parsers/Fbx/Private/FbxMaterial.cpp
@@ -92,9 +92,12 @@ namespace UE
 				if (TextureFilename.IsEmpty() || !FPaths::FileExists(TextureFilename))
 				{
 					const UInterchangeBaseNode* ShaderGraphNode = NodeContainer.GetNode(ShaderUniqueID);
-					UInterchangeResultTextureWarning_TextureFileDoNotExist* Message = Parser.AddMessage<UInterchangeResultTextureWarning_TextureFileDoNotExist>();
-					Message->TextureName = TextureFilename;
-					Message->MaterialName = ShaderGraphNode ? ShaderGraphNode->GetDisplayLabel() : TEXT("Unknown");
+					if(!GIsAutomationTesting)
+					{
+						UInterchangeResultTextureWarning_TextureFileDoNotExist* Message = Parser.AddMessage<UInterchangeResultTextureWarning_TextureFileDoNotExist>();
+						Message->TextureName = TextureFilename;
+						Message->MaterialName = ShaderGraphNode ? ShaderGraphNode->GetDisplayLabel() : TEXT("Unknown");
+					}
 
 					return TextureSampleShader;
 				}
@@ -201,9 +204,12 @@ namespace UE
 				}
 				else
 				{
-					UInterchangeResultTextureWarning_TextureFileDoNotExist* Message = Parser.AddMessage<UInterchangeResultTextureWarning_TextureFileDoNotExist>();
-					Message->TextureName = FbxTexture ? UTF8_TO_TCHAR(FbxTexture->GetFileName()) : TEXT("Undefined");
-					Message->MaterialName = ShaderGraphNode->GetDisplayLabel();
+					if (!GIsAutomationTesting)
+					{
+						UInterchangeResultTextureWarning_TextureFileDoNotExist* Message = Parser.AddMessage<UInterchangeResultTextureWarning_TextureFileDoNotExist>();
+						Message->TextureName = FbxTexture ? UTF8_TO_TCHAR(FbxTexture->GetFileName()) : TEXT("Undefined");
+						Message->MaterialName = ShaderGraphNode->GetDisplayLabel();
+					}
 
 					return false;
 				}
@@ -454,9 +460,12 @@ namespace UE
 					//Only import texture that exist on disk
 					if (!FPaths::FileExists(TextureFilename))
 					{
-						UInterchangeResultTextureWarning_TextureFileDoNotExist* Message = Parser.AddMessage<UInterchangeResultTextureWarning_TextureFileDoNotExist>();
-						Message->TextureName = TextureFilename;
-						Message->MaterialName.Empty();
+						if (!GIsAutomationTesting)
+						{
+							UInterchangeResultTextureWarning_TextureFileDoNotExist* Message = Parser.AddMessage<UInterchangeResultTextureWarning_TextureFileDoNotExist>();
+							Message->TextureName = TextureFilename;
+							Message->MaterialName.Empty();
+						}
 						continue;
 					}
 					//Create a texture node and make it child of the material node
@@ -472,12 +481,22 @@ namespace UE
 			void FFbxMaterial::AddAllNodeMaterials(UInterchangeSceneNode* SceneNode, FbxNode* ParentFbxNode, UInterchangeBaseNodeContainer& NodeContainer)
 			{
 				int32 MaterialCount = ParentFbxNode->GetMaterialCount();
+				TMap<FbxSurfaceMaterial*, int32> UniqueSlotNames;
+				UniqueSlotNames.Reserve(MaterialCount);
 				for (int32 MaterialIndex = 0; MaterialIndex < MaterialCount; ++MaterialIndex)
 				{
 					if (FbxSurfaceMaterial* SurfaceMaterial = ParentFbxNode->GetMaterial(MaterialIndex))
 					{
 						const UInterchangeShaderGraphNode* ShaderGraphNode = AddShaderGraphNode(SurfaceMaterial, NodeContainer);
-						SceneNode->SetSlotMaterialDependencyUid(Parser.GetFbxHelper()->GetFbxObjectName(SurfaceMaterial), ShaderGraphNode->GetUniqueID());
+						
+						int32& SlotMaterialCount = UniqueSlotNames.FindOrAdd(SurfaceMaterial);
+						FString MaterialSlotName = Parser.GetFbxHelper()->GetFbxObjectName(SurfaceMaterial);
+						if (SlotMaterialCount > 0)
+						{
+							MaterialSlotName += TEXT("_Section") + FString::FromInt(SlotMaterialCount);
+						}
+						SceneNode->SetSlotMaterialDependencyUid(MaterialSlotName, ShaderGraphNode->GetUniqueID());
+						SlotMaterialCount++;
 					}
 				}
 			}
diff --git a/Engine/Plugins/Interchange/Runtime/Source/Parsers/Fbx/Private/FbxMesh.cpp b/Engine/Plugins/Interchange/Runtime/Source/Parsers/Fbx/Private/FbxMesh.cpp
--- a/Engine/Plugins/Interchange/Runtime/Source/Parsers/Fbx/Private/FbxMesh.cpp
+++ b/Engine/Plugins/Interchange/Runtime/Source/Parsers/Fbx/Private/FbxMesh.cpp
@@ -254,6 +254,67 @@ namespace UE::Interchange::Private
 		}
 	}
 
+	void ExtractMeshMaterials(FFbxParser& Parser, FbxMesh* Mesh, FbxNode* MeshNode, TFunction<void(const FString& MaterialName, const FString& MaterialUid, const int32 MeshMaterialIndex)> CollectMaterial)
+	{
+		if (!Mesh || !MeshNode)
+		{
+			return;
+		}
+
+		//Grab all Material indexes use by the mesh
+		TArray<int32> MaterialIndexes;
+		int32 PolygonCount = Mesh->GetPolygonCount();
+		if (FbxGeometryElementMaterial* GeometryElementMaterial = Mesh->GetElementMaterial())
+		{
+			FbxLayerElementArrayTemplate<int32>& IndexArray = GeometryElementMaterial->GetIndexArray();
+			switch (GeometryElementMaterial->GetMappingMode())
+			{
+			case FbxGeometryElement::eByPolygon:
+			{
+				if (IndexArray.GetCount() == PolygonCount)
+				{
+					for (int32 PolygonIndex = 0; PolygonIndex < PolygonCount; ++PolygonIndex)
+					{
+						MaterialIndexes.AddUnique(IndexArray.GetAt(PolygonIndex));
+					}
+				}
+			}
+			break;
+
+			case FbxGeometryElement::eAllSame:
+			{
+				if (IndexArray.GetCount() > 0)
+				{
+					MaterialIndexes.AddUnique(IndexArray.GetAt(0));
+				}
+			}
+			break;
+			}
+		}
+		const int32 MaterialCount = MeshNode->GetMaterialCount();
+		TMap<FbxSurfaceMaterial*, int32> UniqueSlotNames;
+		UniqueSlotNames.Reserve(MaterialCount);
+		bool bAddAllNodeMaterials = (MaterialIndexes.Num() == 0);
+		for (int32 MaterialIndex = 0; MaterialIndex < MaterialCount; ++MaterialIndex)
+		{
+			if (FbxSurfaceMaterial* FbxMaterial = MeshNode->GetMaterial(MaterialIndex))
+			{
+				int32& SlotMaterialCount = UniqueSlotNames.FindOrAdd(FbxMaterial);
+				FString MaterialName = Parser.GetFbxHelper()->GetFbxObjectName(FbxMaterial);
+				FString MaterialUid = TEXT("\\Material\\") + MaterialName;
+				if (bAddAllNodeMaterials || MaterialIndexes.Contains(MaterialIndex))
+				{
+					if (SlotMaterialCount > 0)
+					{
+						MaterialName += TEXT("_Section") + FString::FromInt(SlotMaterialCount);
+					}
+					SlotMaterialCount++;
+					CollectMaterial(MaterialName, MaterialUid, MaterialIndex);
+				}
+			}
+		}
+	}
+
 // Wraps some common code useful for multiple fbx import code path
 struct FFBXUVs
 {
@@ -549,21 +610,27 @@ bool FMeshDescriptionImporter::FillMeshDescriptionFromFbxMesh(FbxMesh* Mesh, con
 	//
 	//Create a material name array in the node order, also fill the Meshdescription PolygonGroup
 	TArray<FName> MaterialNames;
-	int32 MaterialCount = (MeshNode != nullptr) ? MeshNode->GetMaterialCount() : Mesh->GetElementMaterialCount();
+	TArray<int32> MaterialRemap;
+	int32 MaterialCount = (MeshNode != nullptr) ? MeshNode->GetMaterialCount() : 1;
 	if (MeshNode)
 	{
-		MaterialCount = MeshNode->GetMaterialCount();
 		MaterialNames.Reserve(MaterialCount);
+		MaterialRemap.Reserve(MaterialCount);
 		for (int32 MaterialIndex = 0; MaterialIndex < MaterialCount; ++MaterialIndex)
 		{
-			FbxSurfaceMaterial* FbxMaterial = MeshNode->GetMaterial(MaterialIndex);
-			MaterialNames.Add(*Parser.GetFbxHelper()->GetFbxObjectName(FbxMaterial));
+			MaterialRemap.Add(MaterialIndex);
 		}
+
+		ExtractMeshMaterials(Parser, Mesh, MeshNode, [&MaterialNames, &MaterialRemap](const FString& MaterialName, const FString& MaterialUid, const int32 MeshMaterialIndex)
+			{
+				MaterialRemap[MeshMaterialIndex] = MaterialNames.Add(*MaterialName);
+			});
+		MaterialCount = MaterialNames.Num();
 	}
 
 	// Must do this before triangulating the mesh due to an FBX bug in TriangulateMeshAdvance
 	int32 LayerSmoothingCount = Mesh->GetLayerCount(FbxLayerElement::eSmoothing);
-	if (LayerSmoothingCount == 0)
+	if (LayerSmoothingCount == 0 && !GIsAutomationTesting)
 	{
 		UInterchangeResultMeshWarning_Generic* Message = AddMessage<UInterchangeResultMeshWarning_Generic>(Mesh);
 		Message->Text = LOCTEXT("MissingSmoothGroup", "No smoothing group information was found for this mesh '{MeshName}' in the FBX file. Please make sure to enable the 'Export Smoothing Groups' option in the FBX Exporter before exporting the file.");
@@ -844,7 +911,7 @@ bool FMeshDescriptionImporter::FillMeshDescriptionFromFbxMesh(FbxMesh* Mesh, con
 				if (FBXUVs.LayerElementUV[UVLayerIndex] != nullptr)
 				{
 					int32 UVCount = FBXUVs.LayerElementUV[UVLayerIndex]->GetDirectArray().GetCount();
-					if (UVCount == 0)
+					if (UVCount == 0 && !GIsAutomationTesting)
 					{
 						UInterchangeResultMeshWarning_Generic* Message = AddMessage<UInterchangeResultMeshWarning_Generic>(Mesh);
 						Message->Text = LOCTEXT("CreateUVs_UVCorrupted", "Found invalid UVs value when importing mesh '{MeshName}'.");
@@ -1099,12 +1166,12 @@ bool FMeshDescriptionImporter::FillMeshDescriptionFromFbxMesh(FbxMesh* Mesh, con
 							// material index is stored in the IndexArray, not the DirectArray (which is irrelevant with 2009.1)
 							case FbxLayerElement::eAllSame:
 							{
-								MaterialIndex = LayerElementMaterial->GetIndexArray().GetAt(0);
+								MaterialIndex = MaterialRemap[LayerElementMaterial->GetIndexArray().GetAt(0)];
 							}
 							break;
 							case FbxLayerElement::eByPolygon:
 							{
-								MaterialIndex = LayerElementMaterial->GetIndexArray().GetAt(PolygonIndex);
+								MaterialIndex = MaterialRemap[LayerElementMaterial->GetIndexArray().GetAt(PolygonIndex)];
 							}
 							break;
 						}
@@ -1719,55 +1786,12 @@ void FFbxMesh::AddAllMeshes(FbxScene* SDKScene, FbxGeometryConverter* SDKGeometr
 		MeshNode->SetCustomHasVertexColor(bMeshHasVertexColor);
 		const int32 MeshUVCount = Mesh->GetElementUVCount();
 		MeshNode->SetCustomUVCount(MeshUVCount);
-						
-		//Add Material dependencies, we use always the first fbx node that instance the fbx geometry to grab the fbx surface materials.
-		{
-			//Grab all Material indexes use by the mesh
-			TArray<int32> MaterialIndexes;
-			int32 PolygonCount = Mesh->GetPolygonCount();
-			if (FbxGeometryElementMaterial* GeometryElementMaterial = Mesh->GetElementMaterial())
-			{
-				FbxLayerElementArrayTemplate<int32>& IndexArray = GeometryElementMaterial->GetIndexArray();
-				switch (GeometryElementMaterial->GetMappingMode())
-				{
-					case FbxGeometryElement::eByPolygon:
-					{
-						if (IndexArray.GetCount() == PolygonCount)
-						{
-							for (int32 PolygonIndex = 0; PolygonIndex < PolygonCount; ++PolygonIndex)
-							{
-								MaterialIndexes.AddUnique(IndexArray.GetAt(PolygonIndex));
-							}
-						}
-					}
-					break;
 
-					case FbxGeometryElement::eAllSame:
-					{
-						if (IndexArray.GetCount() > 0)
-						{
-							MaterialIndexes.AddUnique(IndexArray.GetAt(0));
-						}
-					}
-					break;
-				}
-			}
-			if (FbxNode* FbxMeshNode = Mesh->GetNode())
+		//Add Material dependencies, we use always the first fbx node that instance the fbx geometry to grab the fbx surface materials.
+		ExtractMeshMaterials(Parser, Mesh, Mesh->GetNode(), [&MeshNode](const FString& MaterialName, const FString& MaterialUid, const int32 MeshMaterialIndex)
 			{
-				bool bAddAllNodeMaterials = (MaterialIndexes.Num() == 0);
-				const int32 MaterialCount = FbxMeshNode->GetMaterialCount();
-				for (int32 MaterialIndex = 0; MaterialIndex < MaterialCount; ++MaterialIndex)
-				{
-					FbxSurfaceMaterial* FbxMaterial = FbxMeshNode->GetMaterial(MaterialIndex);
-					FString MaterialName = Parser.GetFbxHelper()->GetFbxObjectName(FbxMaterial);
-					FString MaterialUid = TEXT("\\Material\\") + MaterialName;
-					if (bAddAllNodeMaterials || MaterialIndexes.Contains(MaterialIndex))
-					{
-						MeshNode->SetSlotMaterialDependencyUid(MaterialName, MaterialUid);
-					}
-				}
-			}
-		}
+				MeshNode->SetSlotMaterialDependencyUid(MaterialName, MaterialUid);
+			});
 
 		FString PayLoadKey = MeshUniqueID;
 		if (ensure(!PayloadContexts.Contains(PayLoadKey)))
diff --git a/Engine/Plugins/Interchange/Runtime/Source/Pipelines/Private/InterchangeGenericSkeletalMeshPipeline.cpp b/Engine/Plugins/Interchange/Runtime/Source/Pipelines/Private/InterchangeGenericSkeletalMeshPipeline.cpp
--- a/Engine/Plugins/Interchange/Runtime/Source/Pipelines/Private/InterchangeGenericSkeletalMeshPipeline.cpp
+++ b/Engine/Plugins/Interchange/Runtime/Source/Pipelines/Private/InterchangeGenericSkeletalMeshPipeline.cpp
@@ -349,6 +349,11 @@ UInterchangeSkeletalMeshFactoryNode* UInterchangeGenericMeshPipeline::CreateSkel
 	SkeletalMeshFactoryNode->AddFactoryDependencyUid(SkeletonUid);
 	BaseNodeContainer->AddNode(SkeletalMeshFactoryNode);
 
+	if (CommonMeshesProperties->bKeepSectionsSeparate)
+	{
+		SkeletalMeshFactoryNode->SetCustomKeepSectionsSeparate(CommonMeshesProperties->bKeepSectionsSeparate);
+	}
+
 	SkeletonFactoryNode->SetCustomSkeletalMeshFactoryNodeUid(SkeletalMeshFactoryNode->GetUniqueID());
 
 	AddLodDataToSkeletalMesh(SkeletonFactoryNode, SkeletalMeshFactoryNode, MeshUidsPerLodIndex);
diff --git a/Engine/Plugins/Interchange/Runtime/Source/Pipelines/Private/InterchangeGenericStaticMeshPipeline.cpp b/Engine/Plugins/Interchange/Runtime/Source/Pipelines/Private/InterchangeGenericStaticMeshPipeline.cpp
--- a/Engine/Plugins/Interchange/Runtime/Source/Pipelines/Private/InterchangeGenericStaticMeshPipeline.cpp
+++ b/Engine/Plugins/Interchange/Runtime/Source/Pipelines/Private/InterchangeGenericStaticMeshPipeline.cpp
@@ -408,6 +408,11 @@ UInterchangeStaticMeshFactoryNode* UInterchangeGenericMeshPipeline::CreateStatic
 	StaticMeshFactoryNode->InitializeStaticMeshNode(StaticMeshUid, DisplayLabel, UStaticMesh::StaticClass()->GetName());
 	BaseNodeContainer->AddNode(StaticMeshFactoryNode);
 
+	if (CommonMeshesProperties->bKeepSectionsSeparate)
+	{
+		StaticMeshFactoryNode->SetCustomKeepSectionsSeparate(CommonMeshesProperties->bKeepSectionsSeparate);
+	}
+
 	AddLodDataToStaticMesh(StaticMeshFactoryNode, MeshUidsPerLodIndex);
 
 	switch (CommonMeshesProperties->VertexColorImportOption)
diff --git a/Engine/Plugins/Interchange/Runtime/Source/Pipelines/Public/InterchangeGenericAssetsPipelineSharedSettings.h b/Engine/Plugins/Interchange/Runtime/Source/Pipelines/Public/InterchangeGenericAssetsPipelineSharedSettings.h
--- a/Engine/Plugins/Interchange/Runtime/Source/Pipelines/Public/InterchangeGenericAssetsPipelineSharedSettings.h
+++ b/Engine/Plugins/Interchange/Runtime/Source/Pipelines/Public/InterchangeGenericAssetsPipelineSharedSettings.h
@@ -71,7 +71,11 @@ public:
 	/** If enable, meshes will be baked with the scene instance hierarchy transform. */
 	UPROPERTY(EditAnywhere, BlueprintReadWrite, Category = "Common Meshes")
 	bool bBakeMeshes = true;
-
+	
+	/** If checked, sections with matching materials are kept separate and will not get combined. */
+	UPROPERTY(EditAnywhere, BlueprintReadWrite, Category = "Common Meshes")
+	bool bKeepSectionsSeparate = false;
+	
 	/** Specify how vertex colors should be imported */
 	UPROPERTY(EditAnywhere, BlueprintReadWrite, Category = "Common Meshes")
 	EInterchangeVertexColorImportOption VertexColorImportOption = EInterchangeVertexColorImportOption::IVCIO_Replace;
diff --git a/Engine/Plugins/Interchange/Runtime/Source/Pipelines/Public/InterchangePipelineMeshesUtilities.h b/Engine/Plugins/Interchange/Runtime/Source/Pipelines/Public/InterchangePipelineMeshesUtilities.h
--- a/Engine/Plugins/Interchange/Runtime/Source/Pipelines/Public/InterchangePipelineMeshesUtilities.h
+++ b/Engine/Plugins/Interchange/Runtime/Source/Pipelines/Public/InterchangePipelineMeshesUtilities.h
@@ -339,10 +339,25 @@ namespace UE::Interchange::MeshesUtilities
 	template<class T>
 	void ApplySlotMaterialDependencies(T& FactoryNode, const TMap<FString, FString>& SlotMaterialDependencies, const UInterchangeBaseNodeContainer& NodeContainer)
 	{
+		bool bKeepSectionsSeparate = false;
+		FactoryNode.GetCustomKeepSectionsSeparate(bKeepSectionsSeparate);
+		
+		TMap<FString, FString> ExistingSlotMaterialDependencies;
+		if (bKeepSectionsSeparate)
+		{
+			FactoryNode.GetSlotMaterialDependencies(ExistingSlotMaterialDependencies);
+		}
+
 		for (const TPair<FString, FString>& SlotMaterialDependency : SlotMaterialDependencies)
 		{
+			FString NewSlotName = SlotMaterialDependency.Key;
+			if (ExistingSlotMaterialDependencies.Contains(NewSlotName))
+			{
+				NewSlotName += TEXT("_Section") + FString::FromInt(ExistingSlotMaterialDependencies.Num());
+				ExistingSlotMaterialDependencies.Add(NewSlotName, SlotMaterialDependency.Value);
+			}
 			const FString MaterialFactoryNodeUid = UInterchangeBaseMaterialFactoryNode::GetMaterialFactoryNodeUidFromMaterialNodeUid(SlotMaterialDependency.Value);
-			FactoryNode.SetSlotMaterialDependencyUid(SlotMaterialDependency.Key, MaterialFactoryNodeUid);
+			FactoryNode.SetSlotMaterialDependencyUid(NewSlotName, MaterialFactoryNodeUid);
 			if (UInterchangeBaseMaterialFactoryNode* MaterialFactoryNode = Cast<UInterchangeBaseMaterialFactoryNode>(NodeContainer.GetFactoryNode(MaterialFactoryNodeUid)))
 			{
 				bool IsMaterialImportEnabled = true;
@@ -367,11 +382,8 @@ namespace UE::Interchange::MeshesUtilities
 		TMap<FString, FString> SlotMaterialDependencies;
 		FactoryNode.GetSlotMaterialDependencies(SlotMaterialDependencies);
 
-		for (const TPair<FString, FString>& SlotMaterialDependency : SlotMaterialDependencies)
-		{
-			const FString& MaterialName = SlotMaterialDependency.Key;
-			FactoryNode.RemoveSlotMaterialDependencyUid(MaterialName);
-		}
+		//Empty all slot dependencies, they will be added backin the correct order
+		FactoryNode.ResetSlotMaterialDependencies();
 
 		TArray<FString> KeyReorder;
 		KeyReorder.Reserve(SlotMaterialDependencies.Num());
diff --git a/Engine/Source/Runtime/Interchange/Core/Public/Nodes/InterchangeBaseNodeUtilities.h b/Engine/Source/Runtime/Interchange/Core/Public/Nodes/InterchangeBaseNodeUtilities.h
--- a/Engine/Source/Runtime/Interchange/Core/Public/Nodes/InterchangeBaseNodeUtilities.h
+++ b/Engine/Source/Runtime/Interchange/Core/Public/Nodes/InterchangeBaseNodeUtilities.h
@@ -503,13 +503,19 @@ public:
 				for (int32 Index = 0; Index < KeyCount; Index++)
 				{
 					TAttributeHandle<KeyType> KeyAttribute = AttributesPtr->GetAttributeHandle<KeyType>(GetKeyAttribute(Index));
-					ensure(KeyAttribute.IsValid());
+					if (!ensure(KeyAttribute.IsValid()))
+					{
+						continue;
+					}
 
 					KeyType Key;
 					KeyAttribute.Get(Key);
 
 					TAttributeHandle<ValueType> ValueAttribute = AttributesPtr->GetAttributeHandle<ValueType>(GetValueAttribute(Key));
-					ensure(ValueAttribute.IsValid());
+					if (!ensure(ValueAttribute.IsValid()))
+					{
+						continue;
+					}
 
 					CachedKeysAndValues.Add(Key, TPair<TAttributeHandle<KeyType>, TAttributeHandle<ValueType>>(KeyAttribute, ValueAttribute));
 				}
@@ -587,8 +593,8 @@ private:
 			const FAttributeKey& ValueAttribute = Pair.Value.Value.GetKey();
 			AttributesPtr->UnregisterAttribute(ValueAttribute);
 		}
-
 		CachedKeysAndValues.Empty(NumOfExpectedElements);
+		KeyCountHandle.Set(0);
 	}
 
 	TMap<KeyType, TPair<TAttributeHandle<KeyType>, TAttributeHandle<ValueType>>> CachedKeysAndValues;
