Improve Interchange  import performance
Took the opportunity:
- Renamed FTaskPipelinePostImport FTaskPostImport as it now includes the call to FinalizeObject_GameThread on factories
- Fixed issues when Interchange is used at runtime
- Only select InterchangeSceneImport asset in content browser when importing a scene

diff --git a/Engine/Plugins/Interchange/Runtime/Source/Import/Private/Mesh/InterchangeStaticMeshFactory.cpp b/Engine/Plugins/Interchange/Runtime/Source/Import/Private/Mesh/InterchangeStaticMeshFactory.cpp
--- a/Engine/Plugins/Interchange/Runtime/Source/Import/Private/Mesh/InterchangeStaticMeshFactory.cpp
+++ b/Engine/Plugins/Interchange/Runtime/Source/Import/Private/Mesh/InterchangeStaticMeshFactory.cpp
@@ -202,7 +202,10 @@ UInterchangeFactoryBase::FImportAssetResult UInterchangeStaticMeshFactory::Begin
 	StaticMesh->CreateBodySetup();
 	
 #if WITH_EDITOR
-	StaticMesh->PreEditChange(nullptr);	
+	if (!FApp::IsGame())
+	{
+		StaticMesh->PreEditChange(nullptr);
+	}
 #endif // WITH_EDITOR
 
 	ImportAssetResult.ImportedObject = StaticMesh;
@@ -482,6 +485,8 @@ UInterchangeFactoryBase::FImportAssetResult UInterchangeStaticMeshFactory::Impor
 	// Getting the file Hash will cache it into the source data
 	Arguments.SourceData->GetFileContentHash();
 
+	BuildFromMeshDescriptions(*StaticMesh);
+
 	ImportAssetResult.ImportedObject = StaticMeshObject;
 	return ImportAssetResult;
 }
@@ -517,6 +522,18 @@ UInterchangeFactoryBase::FImportAssetResult UInterchangeStaticMeshFactory::EndIm
 		return ImportAssetResult;
 	}
 
+	if (FApp::IsGame())
+	{
+		if (!Arguments.ReimportObject)
+		{
+			// Apply all StaticMeshFactoryNode custom attributes to the static mesh asset
+			StaticMeshFactoryNode->ApplyAllCustomAttributeToObject(StaticMesh);
+		}
+
+		ImportAssetResult.ImportedObject = StaticMesh;
+		return ImportAssetResult;
+	}
+
 	for (int32 LodIndex = 0; LodIndex < ImportAssetObjectData.LodMeshDescriptions.Num(); ++LodIndex)
 	{
 		// Add the lod mesh data to the static mesh
@@ -620,7 +637,7 @@ UInterchangeFactoryBase::FImportAssetResult UInterchangeStaticMeshFactory::EndIm
 		}
 	}
 
-	CommitMeshDescriptions(*StaticMesh, MoveTemp(ImportAssetObjectData.LodMeshDescriptions));
+	CommitMeshDescriptions(*StaticMesh);
 
 	ImportSockets(Arguments, StaticMesh, StaticMeshFactoryNode);
 
@@ -665,7 +682,7 @@ UInterchangeFactoryBase::FImportAssetResult UInterchangeStaticMeshFactory::EndIm
 	{
 		if (!ImportAssetObjectData.bImportedCustomCollision)
 		{
-			GenerateKDopCollision(Arguments, StaticMesh);
+			GenerateKDopCollision(StaticMesh);
 		}
 #if WITH_EDITORONLY_DATA
 		else
@@ -701,20 +718,36 @@ UInterchangeFactoryBase::FImportAssetResult UInterchangeStaticMeshFactory::EndIm
 	return ImportAssetResult;
 }
 
-void UInterchangeStaticMeshFactory::CommitMeshDescriptions(UStaticMesh& StaticMesh, TArray<FMeshDescription>&& LodMeshDescriptions)
+void UInterchangeStaticMeshFactory::CommitMeshDescriptions(UStaticMesh& StaticMesh)
 {
-#if WITH_EDITOR
+	if (FApp::IsGame())
+	{
+		return;
+	}
+
+	TArray<FMeshDescription> LodMeshDescriptions = MoveTemp(ImportAssetObjectData.LodMeshDescriptions);
+
+	UStaticMesh::FCommitMeshDescriptionParams CommitMeshDescriptionParams;
+	CommitMeshDescriptionParams.bMarkPackageDirty = false; // Marking packages dirty isn't thread-safe
+
 	for (int32 LodIndex = 0; LodIndex < LodMeshDescriptions.Num(); ++LodIndex)
 	{
 		FMeshDescription* StaticMeshDescription = StaticMesh.CreateMeshDescription(LodIndex);
 		check(StaticMeshDescription);
 		*StaticMeshDescription = MoveTemp(LodMeshDescriptions[LodIndex]);
 
-		UStaticMesh::FCommitMeshDescriptionParams CommitMeshDescriptionParams;
-		CommitMeshDescriptionParams.bMarkPackageDirty = false; // Marking packages dirty isn't threadsafe
 		StaticMesh.CommitMeshDescription(LodIndex, CommitMeshDescriptionParams);
 	}
-#else // WITH_EDITOR
+}
+
+void UInterchangeStaticMeshFactory::BuildFromMeshDescriptions(UStaticMesh& StaticMesh)
+{
+	if (!FApp::IsGame())
+	{
+		return;
+	}
+
+	TArray<FMeshDescription> LodMeshDescriptions = MoveTemp(ImportAssetObjectData.LodMeshDescriptions);
 	TArray<const FMeshDescription*> MeshDescriptionPointers;
 	MeshDescriptionPointers.Reserve(LodMeshDescriptions.Num());
 
@@ -725,15 +758,31 @@ void UInterchangeStaticMeshFactory::CommitMeshDescriptions(UStaticMesh& StaticMe
 
 	UStaticMesh::FBuildMeshDescriptionsParams BuildMeshDescriptionsParams;
 	BuildMeshDescriptionsParams.bUseHashAsGuid = true;
-	// Do not mark the package dirty since MarkPackageDirty is not thread safe
 	BuildMeshDescriptionsParams.bMarkPackageDirty = false;
 	BuildMeshDescriptionsParams.bBuildSimpleCollision = false;
 	// Do not commit since we only need the render data and commit is slow
 	BuildMeshDescriptionsParams.bCommitMeshDescription = false;
 	BuildMeshDescriptionsParams.bFastBuild = true;
+	// For the time being at runtime collision is set to complex one
+	// TODO: Revisit pipeline options for collision. bImportCollision is not enough.
+	BuildMeshDescriptionsParams.bAllowCpuAccess = ImportAssetObjectData.bImportCollision;
+	StaticMesh.bAllowCPUAccess = BuildMeshDescriptionsParams.bAllowCpuAccess;
 
 	StaticMesh.BuildFromMeshDescriptions(MeshDescriptionPointers, BuildMeshDescriptionsParams);
-#endif // !WITH_EDITOR
+	
+	// TODO: Expand support for different collision types
+	if (ensure(StaticMesh.GetRenderData()))
+	{
+		if (ImportAssetObjectData.bImportCollision && !ImportAssetObjectData.bImportedCustomCollision)
+		{
+			if (StaticMesh.GetBodySetup() == nullptr)
+			{
+				StaticMesh.CreateBodySetup();
+			}
+
+			StaticMesh.GetBodySetup()->CollisionTraceFlag = ECollisionTraceFlag::CTF_UseComplexAsSimple;
+		}
+	}	
 }
 
 #if WITH_EDITORONLY_DATA
@@ -1031,7 +1080,7 @@ static FVector3f GetTriangleNormal(const FTransform& Transform, TVertexAttribute
 }
 
 
-bool UInterchangeStaticMeshFactory::AddBoxGeomFromTris(const FImportAssetObjectParams& Arguments, const FMeshDescription& MeshDescription, const FTransform& Transform, FKAggregateGeom& AggGeom)
+bool UInterchangeStaticMeshFactory::AddBoxGeomFromTris(const FMeshDescription& MeshDescription, const FTransform& Transform, FKAggregateGeom& AggGeom)
 {
 	// Maintain an array of the planes we have encountered so far.
 	// We are expecting two instances of three unique plane orientations, one for each side of the box.
@@ -1336,7 +1385,7 @@ bool UInterchangeStaticMeshFactory::ImportBoxCollision(const FImportAssetObjectP
 			continue;
 		}
 
-		if (AddBoxGeomFromTris(Arguments, PayloadData->MeshDescription, Transform, AggGeo))
+		if (AddBoxGeomFromTris(PayloadData->MeshDescription, Transform, AggGeo))
 		{
 			bResult = true;
 			FKBoxElem& NewElem = AggGeo.BoxElems.Last();
@@ -1548,7 +1597,7 @@ bool UInterchangeStaticMeshFactory::ImportConvexCollision(const FImportAssetObje
 }
 
 
-bool UInterchangeStaticMeshFactory::GenerateKDopCollision(const FImportAssetObjectParams& Arguments, UStaticMesh* StaticMesh)
+bool UInterchangeStaticMeshFactory::GenerateKDopCollision(UStaticMesh* StaticMesh)
 {
 #if WITH_EDITOR
 
diff --git a/Engine/Plugins/Interchange/Runtime/Source/Import/Private/Scene/InterchangeActorFactory.cpp b/Engine/Plugins/Interchange/Runtime/Source/Import/Private/Scene/InterchangeActorFactory.cpp
--- a/Engine/Plugins/Interchange/Runtime/Source/Import/Private/Scene/InterchangeActorFactory.cpp
+++ b/Engine/Plugins/Interchange/Runtime/Source/Import/Private/Scene/InterchangeActorFactory.cpp
@@ -31,7 +31,20 @@ UObject* UInterchangeActorFactory::ImportSceneObject_GameThread(const UInterchan
 	{
 		if (UObject* ObjectToUpdate = ProcessActor(*SpawnedActor, *FactoryNode, *CreateSceneObjectsParams.NodeContainer))
 		{
-			ActorHelper::ApplyAllCustomAttributes(CreateSceneObjectsParams, *ObjectToUpdate);
+			if (USceneComponent* RootComponent = SpawnedActor->GetRootComponent())
+			{
+				// Cache mobility value to allow application of transform
+				EComponentMobility::Type CachedMobility = RootComponent->Mobility;
+				RootComponent->SetMobility(EComponentMobility::Type::Movable);
+
+				ActorHelper::ApplyAllCustomAttributes(CreateSceneObjectsParams, *ObjectToUpdate);
+
+				// Restore mobility value
+				if (CachedMobility != EComponentMobility::Type::Movable)
+				{
+					RootComponent->SetMobility(CachedMobility);
+				}
+			}
 		}
 	}
 
diff --git a/Engine/Plugins/Interchange/Runtime/Source/Import/Private/Scene/InterchangeStaticMeshActorFactory.cpp b/Engine/Plugins/Interchange/Runtime/Source/Import/Private/Scene/InterchangeStaticMeshActorFactory.cpp
--- a/Engine/Plugins/Interchange/Runtime/Source/Import/Private/Scene/InterchangeStaticMeshActorFactory.cpp
+++ b/Engine/Plugins/Interchange/Runtime/Source/Import/Private/Scene/InterchangeStaticMeshActorFactory.cpp
@@ -87,7 +87,7 @@ UObject* UInterchangeStaticMeshActorFactory::ProcessActor(AActor& SpawnedActor,
 
 				GeometricTransformMeshComponent->SetupAttachment(StaticMeshComponent);
 
-				StaticMeshActor->ReregisterAllComponents();
+				StaticMeshActor->RegisterAllComponents();
 
 				EComponentMobility::Type MobilityToSet = StaticMeshComponent->Mobility;
 				GeometricTransformMeshComponent->SetMobility(EComponentMobility::Type::Movable); //so that RelativeTransform can be set
@@ -98,10 +98,10 @@ UObject* UInterchangeStaticMeshActorFactory::ProcessActor(AActor& SpawnedActor,
 			}
 		}
 
-		StaticMeshComponent->UnregisterComponent();		
-
 		if (const UInterchangeFactoryBaseNode* MeshNode = ActorHelper::FindAssetInstanceFactoryNode(&NodeContainer, &FactoryNode))
 		{
+			StaticMeshComponent->UnregisterComponent();
+
 			FSoftObjectPath ReferenceObject;
 			MeshNode->GetCustomReferenceObject(ReferenceObject);
 			if (UStaticMesh* StaticMesh = Cast<UStaticMesh>(ReferenceObject.TryLoad()))
@@ -115,6 +115,8 @@ UObject* UInterchangeStaticMeshActorFactory::ProcessActor(AActor& SpawnedActor,
 					UE::Interchange::ActorHelper::ApplySlotMaterialDependencies(NodeContainer, *MeshActorFactoryNode, *StaticMeshComponent);
 				}
 			}
+
+			StaticMeshComponent->RegisterComponent();
 		}
 		else
 		{
diff --git a/Engine/Plugins/Interchange/Runtime/Source/Import/Public/Mesh/InterchangeStaticMeshFactory.h b/Engine/Plugins/Interchange/Runtime/Source/Import/Public/Mesh/InterchangeStaticMeshFactory.h
--- a/Engine/Plugins/Interchange/Runtime/Source/Import/Public/Mesh/InterchangeStaticMeshFactory.h
+++ b/Engine/Plugins/Interchange/Runtime/Source/Import/Public/Mesh/InterchangeStaticMeshFactory.h
@@ -53,7 +53,7 @@ private:
 
 	bool AddConvexGeomFromVertices(const FImportAssetObjectParams& Arguments, const FMeshDescription& MeshDescription, const FTransform& Transform, FKAggregateGeom& AggGeom);
 	bool DecomposeConvexMesh(const FImportAssetObjectParams& Arguments, const FMeshDescription& MeshDescription, const FTransform& Transform, UBodySetup* BodySetup);
-	bool AddBoxGeomFromTris(const FImportAssetObjectParams& Arguments, const FMeshDescription& MeshDescription, const FTransform& Transform, FKAggregateGeom& AggGeom);
+	bool AddBoxGeomFromTris(const FMeshDescription& MeshDescription, const FTransform& Transform, FKAggregateGeom& AggGeom);
 	bool AddSphereGeomFromVertices(const FImportAssetObjectParams& Arguments, const FMeshDescription& MeshDescription, const FTransform& Transform, FKAggregateGeom& AggGeom);
 	bool AddCapsuleGeomFromVertices(const FImportAssetObjectParams& Arguments, const FMeshDescription& MeshDescription, const FTransform& Transform, FKAggregateGeom& AggGeom);
 
@@ -61,10 +61,11 @@ private:
 	bool ImportCapsuleCollision(const FImportAssetObjectParams& Arguments, UStaticMesh* StaticMesh, const UInterchangeStaticMeshLodDataNode* LodDataNode);
 	bool ImportSphereCollision(const FImportAssetObjectParams& Arguments, UStaticMesh* StaticMesh, const UInterchangeStaticMeshLodDataNode* LodDataNode);
 	bool ImportConvexCollision(const FImportAssetObjectParams& Arguments, UStaticMesh* StaticMesh, const UInterchangeStaticMeshLodDataNode* LodDataNode);
-	bool GenerateKDopCollision(const FImportAssetObjectParams& Arguments, UStaticMesh* StaticMesh);
+	bool GenerateKDopCollision(UStaticMesh* StaticMesh);
 	bool ImportSockets(const FImportAssetObjectParams& Arguments, UStaticMesh* StaticMesh, const UInterchangeStaticMeshFactoryNode* FactoryNode);
 
-	void CommitMeshDescriptions(UStaticMesh& StaticMesh, TArray<FMeshDescription>&& LodMeshDescriptions);
+	void CommitMeshDescriptions(UStaticMesh& StaticMesh);
+	void BuildFromMeshDescriptions(UStaticMesh& StaticMesh);
 
 #if WITH_EDITORONLY_DATA
 	void SetupSourceModelsSettings(UStaticMesh& StaticMesh, const TArray<FMeshDescription>& LodMeshDescriptions, int32 PreviousLodCount, int32 FinalLodCount, bool bIsAReimport);
diff --git a/Engine/Source/Developer/AssetTools/Private/AssetTools.cpp b/Engine/Source/Developer/AssetTools/Private/AssetTools.cpp
--- a/Engine/Source/Developer/AssetTools/Private/AssetTools.cpp
+++ b/Engine/Source/Developer/AssetTools/Private/AssetTools.cpp
@@ -73,6 +73,7 @@
 #include "Misc/ConfigCacheIni.h"
 #include "Misc/NamePermissionList.h"
 #include "InterchangeManager.h"
+#include "InterchangeSceneImportAsset.h"
 #include "InterchangeProjectSettings.h"
 #include "Engine/World.h"
 #include "Engine/Level.h"
@@ -3520,7 +3521,7 @@ TArray<UObject*> UAssetToolsImpl::ImportAssetsInternal(const TArray<FString>& Fi
 				
 				TFunction<void(UE::Interchange::FImportResult&)> AppendAndBroadcastImportResultIfNeeded =
 					// Note: ImportStatus captured by value so that the lambda keeps the shared ptr alive
-					[ImportStatus, AppendImportResult, bSyncToBrowser, bForceContentBrowserSyncIfOnlyOneMainAsset](UE::Interchange::FImportResult& Result)
+					[bSceneImport = Params.bSceneImport, ImportStatus, AppendImportResult, bSyncToBrowser, bForceContentBrowserSyncIfOnlyOneMainAsset](UE::Interchange::FImportResult& Result)
 					{
 						AppendImportResult(Result);
 
@@ -3530,28 +3531,33 @@ TArray<UObject*> UAssetToolsImpl::ImportAssetsInternal(const TArray<FString>& Fi
 							InterchangeManager.OnBatchImportComplete.Broadcast(ImportStatus->InterchangeResultsContainer);
 
 							TArray<UObject*> MainAssets;
-							for (const TWeakObjectPtr<UObject>& WeakObject : ImportStatus->ImportedObjects)
+							if (bSceneImport)
 							{
-								if (WeakObject.IsValid()
-									&& (WeakObject->IsA<UStaticMesh>()
-										|| WeakObject->IsA<USkeletalMesh>()
-										|| WeakObject->IsA<UAnimSequence>()))
+								for (const TWeakObjectPtr<UObject>& WeakObject : ImportStatus->ImportedObjects)
 								{
-									MainAssets.Add(WeakObject.Get());
+									if (WeakObject.IsValid() && WeakObject->IsA<UInterchangeSceneImportAsset>())
+									{
+										MainAssets.Add(WeakObject.Get());
+										// There should be only one anyway
+										break;
+									}
 								}
 							}
-
-							//Force browser to sync to the import asset if there is only one asset imported
-							if (bSyncToBrowser || (bForceContentBrowserSyncIfOnlyOneMainAsset && MainAssets.Num() == 1))
+							else
 							{
-								TArray<UObject*> ImportedObjects;
-								ImportedObjects.Reserve(ImportStatus->ImportedObjects.Num());
 								for (const TWeakObjectPtr<UObject>& WeakObject : ImportStatus->ImportedObjects)
 								{
-									ImportedObjects.Add(WeakObject.Get());
+									if (WeakObject.IsValid() && WeakObject->IsAsset())
+									{
+										MainAssets.Add(WeakObject.Get());
+									}
 								}
+							}
 
-								UAssetToolsImpl::Get().SyncBrowserToAssets(bSyncToBrowser ? ImportedObjects : MainAssets);
+							//Force browser to sync to the import asset if there is only one asset imported
+							if (bSyncToBrowser || (bForceContentBrowserSyncIfOnlyOneMainAsset && MainAssets.Num() == 1))
+							{
+								UAssetToolsImpl::Get().SyncBrowserToAssets(MainAssets);
 							}
 						}
 					};
diff --git a/Engine/Source/Runtime/Engine/Private/StaticMeshBuild.cpp b/Engine/Source/Runtime/Engine/Private/StaticMeshBuild.cpp
--- a/Engine/Source/Runtime/Engine/Private/StaticMeshBuild.cpp
+++ b/Engine/Source/Runtime/Engine/Private/StaticMeshBuild.cpp
@@ -96,6 +96,15 @@ bool UStaticMesh::CanBuild() const
 		return false;
 	}
 
+	// Mesh descriptions are unlikely at runtime but UStaticMesh::Build can still be called
+	if (FApp::IsGame())
+	{
+		if (!IsMeshDescriptionValid(0))
+		{
+			return false;
+		}
+	}
+
 	return true;
 }
 
diff --git a/Engine/Source/Runtime/Interchange/Core/Public/InterchangePipelineBase.h b/Engine/Source/Runtime/Interchange/Core/Public/InterchangePipelineBase.h
--- a/Engine/Source/Runtime/Interchange/Core/Public/InterchangePipelineBase.h
+++ b/Engine/Source/Runtime/Interchange/Core/Public/InterchangePipelineBase.h
@@ -173,7 +173,7 @@ public:
 	 * ScriptedExecutePostImportPipeline is call after we completely import an asset, PostEditChange is already called.
 	 * This can be useful if you need an asset build data to finish the setup of another asset.
 	 * @example - PhysicsAsset need skeletal mesh render data to be build properly.
-	 * @note - the FTaskPipelinePostImport is calling this function not the virtual one that is call by the default implementation.
+	 * @note - the FTaskPostImport is calling this function not the virtual one that is call by the default implementation.
 	 */
 	UFUNCTION(BlueprintNativeEvent, BlueprintCallable, Category = "Interchange | Pipeline")
 	INTERCHANGECORE_API void ScriptedExecutePostImportPipeline(const UInterchangeBaseNodeContainer* BaseNodeContainer, const FString& FactoryNodeKey, UObject* CreatedAsset, bool bIsAReimport);
diff --git a/Engine/Source/Runtime/Interchange/Engine/Private/InterchangeManager.cpp b/Engine/Source/Runtime/Interchange/Engine/Private/InterchangeManager.cpp
--- a/Engine/Source/Runtime/Interchange/Engine/Private/InterchangeManager.cpp
+++ b/Engine/Source/Runtime/Interchange/Engine/Private/InterchangeManager.cpp
@@ -534,7 +534,7 @@ FGraphEventArray UE::Interchange::FImportAsyncHelper::GetCompletionTaskGraphEven
 	{
 		TasksToComplete.Add(WaitAssetCompilationTask);
 	}
-	TasksToComplete.Append(PipelinePostImportTasks);
+	TasksToComplete.Append(PostImportTasks);
 
 	if (PreCompletionTask.GetReference())
 	{
diff --git a/Engine/Source/Runtime/Interchange/Engine/Private/Tasks/InterchangeTaskCompletion.cpp b/Engine/Source/Runtime/Interchange/Engine/Private/Tasks/InterchangeTaskCompletion.cpp
--- a/Engine/Source/Runtime/Interchange/Engine/Private/Tasks/InterchangeTaskCompletion.cpp
+++ b/Engine/Source/Runtime/Interchange/Engine/Private/Tasks/InterchangeTaskCompletion.cpp
@@ -56,6 +56,12 @@ void UE::Interchange::FTaskPreCompletion::DoTask(ENamedThreads::Type CurrentThre
 
 		const bool bCallPostImportGameThreadCallback = ensure(AsyncHelper->SourceDatas.IsValidIndex(SourceIndex));
 
+		UInterchangeFactoryBase::FSetupObjectParams Arguments;
+		Arguments.SourceData = AsyncHelper->SourceDatas[SourceIndex];
+		Arguments.NodeContainer = AsyncHelper->BaseNodeContainers[SourceIndex].Get();
+		Arguments.Pipelines = AsyncHelper->Pipelines;
+		Arguments.OriginalPipelines = AsyncHelper->OriginalPipelines;
+
 		//First iteration to call SetupObject_GameThread and pipeline ExecutePostFactoryPipeline
 		for (const FImportAsyncHelper::FImportedObjectInfo& ObjectInfo : ImportedObjects)
 		{
@@ -63,16 +69,10 @@ void UE::Interchange::FTaskPreCompletion::DoTask(ENamedThreads::Type CurrentThre
 			//In case Some factory code cannot run outside of the main thread we offer this callback to finish the work before calling post edit change (building the asset)
 			if (bCallPostImportGameThreadCallback && ObjectInfo.Factory)
 			{
-				UInterchangeFactoryBase::FSetupObjectParams Arguments;
 				Arguments.ImportedObject = ImportedObject;
-				Arguments.SourceData = AsyncHelper->SourceDatas[SourceIndex];
-				Arguments.FactoryNode = ObjectInfo.FactoryNode;
 				// Should we assert if there is no factory node?
+				Arguments.FactoryNode = ObjectInfo.FactoryNode;
 				Arguments.NodeUniqueID = ObjectInfo.FactoryNode ? ObjectInfo.FactoryNode->GetUniqueID() : FString();
-				Arguments.NodeContainer = AsyncHelper->BaseNodeContainers[SourceIndex].Get();
-				Arguments.Pipelines = ObjectPtrDecay(AsyncHelper->Pipelines);
-				Arguments.Pipelines = AsyncHelper->Pipelines;
-				Arguments.OriginalPipelines = AsyncHelper->OriginalPipelines;
 				Arguments.bIsReimport = ObjectInfo.bIsReimport;
 				ObjectInfo.Factory->SetupObject_GameThread(Arguments);
 			}
@@ -159,23 +159,6 @@ void UE::Interchange::FTaskPreCompletion::DoTask(ENamedThreads::Type CurrentThre
 			{
 				AsyncHelper->SceneImportResult->AddImportedObject(ImportedObject);
 			}
-
-			//In case Some factory code cannot run outside of the main thread we offer this callback to finish the work after calling post edit change (building the asset)
-			//Its possible the build of the asset to be asynchronous, the factory must handle is own asset correctly
-			if (bCallPostImportGameThreadCallback && ObjectInfo.Factory)
-			{
-				UInterchangeFactoryBase::FSetupObjectParams Arguments;
-				Arguments.ImportedObject = ImportedObject;
-				Arguments.SourceData = AsyncHelper->SourceDatas[SourceIndex];
-				Arguments.FactoryNode = ObjectInfo.FactoryNode;
-				Arguments.NodeUniqueID = ObjectInfo.FactoryNode ? ObjectInfo.FactoryNode->GetUniqueID() : FString();
-				Arguments.NodeContainer = AsyncHelper->BaseNodeContainers[SourceIndex].Get();
-				Arguments.Pipelines = ObjectPtrDecay(AsyncHelper->Pipelines);
-				Arguments.Pipelines = AsyncHelper->Pipelines;
-				Arguments.OriginalPipelines = AsyncHelper->OriginalPipelines;
-				Arguments.bIsReimport = ObjectInfo.bIsReimport;
-				ObjectInfo.Factory->FinalizeObject_GameThread(Arguments);
-			}
 		}
 	};
 
@@ -220,6 +203,8 @@ void UE::Interchange::FTaskCompletion::DoTask(ENamedThreads::Type CurrentThread,
 					//We broadcast this event for both import and reimport.
 					UInterchangeManager::GetInterchangeManager().OnAssetPostImport.Broadcast(Asset);
 				}
+
+				UE_LOG(LogInterchangeEngine, Display, TEXT("Interchange import completed [%s]"), *AsyncHelper->SourceDatas[SourceIndex]->ToDisplayString());
 			});
 	}
 	else
@@ -242,7 +227,7 @@ void UE::Interchange::FTaskCompletion::DoTask(ENamedThreads::Type CurrentThread,
 			});
 
 		//If task is canceled, remove all actors from their world
-		AsyncHelper->IterateImportedSceneObjectsPerSourceIndex([](int32 SourceIndex, const TArray<FImportAsyncHelper::FImportedObjectInfo>& AssetInfos)
+		AsyncHelper->IterateImportedSceneObjectsPerSourceIndex([AsyncHelper](int32 SourceIndex, const TArray<FImportAsyncHelper::FImportedObjectInfo>& AssetInfos)
 			{
 				for (const FImportAsyncHelper::FImportedObjectInfo& SceneObjectInfo : AssetInfos)
 				{
@@ -255,6 +240,8 @@ void UE::Interchange::FTaskCompletion::DoTask(ENamedThreads::Type CurrentThread,
 						}
 					}
 				}
+
+				UE_LOG(LogInterchangeEngine, Display, TEXT("Interchange import cancelled [%s]"), *AsyncHelper->SourceDatas[SourceIndex]->ToDisplayString());
 			});
 	}
 
diff --git a/Engine/Source/Runtime/Interchange/Engine/Private/Tasks/InterchangeTaskParsing.cpp b/Engine/Source/Runtime/Interchange/Engine/Private/Tasks/InterchangeTaskParsing.cpp
--- a/Engine/Source/Runtime/Interchange/Engine/Private/Tasks/InterchangeTaskParsing.cpp
+++ b/Engine/Source/Runtime/Interchange/Engine/Private/Tasks/InterchangeTaskParsing.cpp
@@ -425,15 +425,12 @@ void UE::Interchange::FTaskParsing::DoTask(ENamedThreads::Type CurrentThread, co
 	//Start the Post pipeline task
 	for (int32 SourceIndex = 0; SourceIndex < AsyncHelper->SourceDatas.Num(); ++SourceIndex)
 	{
-		for (int32 GraphPipelineIndex = 0; GraphPipelineIndex < AsyncHelper->Pipelines.Num(); ++GraphPipelineIndex)
-		{
-			int32 GraphPipelineTaskIndex = AsyncHelper->PipelinePostImportTasks.Add(
-				TGraphTask<FTaskPipelinePostImport>::CreateTask(&(AssetCompilationPrerequistes)).ConstructAndDispatchWhenReady(SourceIndex, GraphPipelineIndex, WeakAsyncHelper)
-			);
-			//Ensure we run the pipeline in the same order we create the task, since the pipeline modifies the node container, its important that its not processed in parallel, Adding the one we start to the prerequisites
-			//is the way to go here
-			AssetCompilationPrerequistes.Add(AsyncHelper->PipelinePostImportTasks[GraphPipelineTaskIndex]);
-		}
+		int32 GraphTaskIndex = AsyncHelper->PostImportTasks.Add(
+			TGraphTask<FTaskPostImport>::CreateTask(&(AssetCompilationPrerequistes)).ConstructAndDispatchWhenReady(SourceIndex, WeakAsyncHelper)
+		);
+		//Ensure we run the pipeline in the same order we create the task, since the pipeline modifies the node container, its important that its not processed in parallel, Adding the one we start to the prerequisites
+		//is the way to go here
+		AssetCompilationPrerequistes.Add(AsyncHelper->PostImportTasks[GraphTaskIndex]);
 	}
 
 	AsyncHelper->CompletionTask = TGraphTask<FTaskCompletion>::CreateTask(&AssetCompilationPrerequistes).ConstructAndDispatchWhenReady(InterchangeManager, WeakAsyncHelper);
diff --git a/Engine/Source/Runtime/Interchange/Engine/Private/Tasks/InterchangeTaskPipeline.cpp b/Engine/Source/Runtime/Interchange/Engine/Private/Tasks/InterchangeTaskPipeline.cpp
--- a/Engine/Source/Runtime/Interchange/Engine/Private/Tasks/InterchangeTaskPipeline.cpp
+++ b/Engine/Source/Runtime/Interchange/Engine/Private/Tasks/InterchangeTaskPipeline.cpp
@@ -9,6 +9,7 @@
 #include "InterchangeEngineLogPrivate.h"
 #include "InterchangeManager.h"
 #include "InterchangePipelineBase.h"
+#include "InterchangeFactoryBase.h"
 #include "Interfaces/Interface_AsyncCompilation.h"
 #include "Materials/MaterialInterface.h"
 #include "Nodes/InterchangeBaseNodeContainer.h"
@@ -64,7 +65,7 @@ void UE::Interchange::FTaskWaitAssetCompilation::DoTask(ENamedThreads::Type Curr
 #if INTERCHANGE_TRACE_ASYNCHRONOUS_TASK_ENABLED
 		INTERCHANGE_TRACE_ASYNCHRONOUS_TASK(WaitAssetCompilation)
 #endif
-	
+
 	LLM_SCOPE_BYNAME(TEXT("Interchange"));
 
 #if WITH_EDITOR
@@ -132,13 +133,12 @@ void UE::Interchange::FTaskWaitAssetCompilation::DoTask(ENamedThreads::Type Curr
 				}).Wait();
 			return bCompilationFinish;
 		}, 0.05f);
-
 #endif //WITH_EDITOR
 }
 
-void UE::Interchange::FTaskPipelinePostImport::DoTask(ENamedThreads::Type CurrentThread, const FGraphEventRef& MyCompletionGraphEvent)
+void UE::Interchange::FTaskPostImport::DoTask(ENamedThreads::Type CurrentThread, const FGraphEventRef& MyCompletionGraphEvent)
 {
-	TRACE_CPUPROFILER_EVENT_SCOPE(UE::Interchange::FTaskPipelinePostImport::DoTask)
+	TRACE_CPUPROFILER_EVENT_SCOPE(UE::Interchange::FTaskPostImport::DoTask)
 #if INTERCHANGE_TRACE_ASYNCHRONOUS_TASK_ENABLED
 	INTERCHANGE_TRACE_ASYNCHRONOUS_TASK(PipelinePostImport)
 #endif
@@ -157,26 +157,61 @@ void UE::Interchange::FTaskPipelinePostImport::DoTask(ENamedThreads::Type Curren
 		return;
 	}
 
-	if (!ensure(AsyncHelper->Pipelines.IsValidIndex(PipelineIndex)) || !ensure(AsyncHelper->BaseNodeContainers.IsValidIndex(SourceIndex)))
+	if (!ensure(AsyncHelper->BaseNodeContainers.IsValidIndex(SourceIndex)))
+	{
+		return;
+	}
+
+	//Get the Container from the async helper
+	UInterchangeBaseNodeContainer* NodeContainer = AsyncHelper->BaseNodeContainers[SourceIndex].Get();
+	if (!ensure(NodeContainer))
 	{
 		return;
 	}
-	UInterchangePipelineBase* PipelineBase = AsyncHelper->Pipelines[PipelineIndex];
+
 	TArray<FString> NodeUniqueIDs;
 	TArray<UObject*> ImportedObjects;
 	TArray<bool> IsAssetsReimported;
 
 	auto FillImportedObjectsFromSource =
-		[&NodeUniqueIDs, &ImportedObjects, &IsAssetsReimported, this](const TArray<UE::Interchange::FImportAsyncHelper::FImportedObjectInfo>& ImportedInfos)
+		[&AsyncHelper, &NodeUniqueIDs, &ImportedObjects, &IsAssetsReimported, this](const TArray<UE::Interchange::FImportAsyncHelper::FImportedObjectInfo>& ImportedInfos)
 		{
 			NodeUniqueIDs.Reserve(NodeUniqueIDs.Num() + ImportedInfos.Num());
 			ImportedObjects.Reserve(ImportedObjects.Num() + ImportedInfos.Num());
 			IsAssetsReimported.Reserve(IsAssetsReimported.Num() + ImportedInfos.Num());
-			for (const UE::Interchange::FImportAsyncHelper::FImportedObjectInfo& ImportedInfo : ImportedInfos)
+
+			const bool bCallPostImportGameThreadCallback = ensure(AsyncHelper->SourceDatas.IsValidIndex(SourceIndex));
+
+			//Call the pipeline for each asset created by this import
+			UInterchangeFactoryBase::FSetupObjectParams Arguments;
+			Arguments.SourceData = bCallPostImportGameThreadCallback ? AsyncHelper->SourceDatas[SourceIndex] : nullptr;
+			Arguments.NodeContainer = AsyncHelper->BaseNodeContainers[SourceIndex].Get();
+			Arguments.Pipelines = AsyncHelper->Pipelines;
+			Arguments.OriginalPipelines = AsyncHelper->OriginalPipelines;
+
+			for (const UE::Interchange::FImportAsyncHelper::FImportedObjectInfo& ObjectInfo : ImportedInfos)
 			{
-				NodeUniqueIDs.Add(ImportedInfo.FactoryNode->GetUniqueID());
-				ImportedObjects.Add(ImportedInfo.ImportedObject);
-				IsAssetsReimported.Add(ImportedInfo.bIsReimport);
+				UObject* ImportedObject = ObjectInfo.ImportedObject;
+
+				//In case Some factory code cannot run outside of the main thread we offer this callback to finish the work after calling post edit change (building the asset)
+				//Its possible the build of the asset to be asynchronous, the factory must handle is own asset correctly
+				if (bCallPostImportGameThreadCallback && ObjectInfo.Factory)
+				{
+					Arguments.ImportedObject = ImportedObject;
+					Arguments.FactoryNode = ObjectInfo.FactoryNode;
+					Arguments.NodeUniqueID = ObjectInfo.FactoryNode ? ObjectInfo.FactoryNode->GetUniqueID() : FString();
+					Arguments.bIsReimport = ObjectInfo.bIsReimport;
+					ObjectInfo.Factory->FinalizeObject_GameThread(Arguments);
+				}
+
+				if (!IsValid(ImportedObject))
+				{
+					continue;
+				}
+
+				NodeUniqueIDs.Add(ObjectInfo.FactoryNode->GetUniqueID());
+				ImportedObjects.Add(ObjectInfo.ImportedObject);
+				IsAssetsReimported.Add(ObjectInfo.bIsReimport);
 			}
 		};
 
@@ -189,17 +224,13 @@ void UE::Interchange::FTaskPipelinePostImport::DoTask(ENamedThreads::Type Curren
 		return;
 	}
 
-	//Get the Container from the async helper
-	UInterchangeBaseNodeContainer* NodeContainer = AsyncHelper->BaseNodeContainers[SourceIndex].Get();
-	if (!ensure(NodeContainer))
-	{
-		return;
-	}
-	UInterchangePipelineBase* Pipeline = AsyncHelper->Pipelines[PipelineIndex];
-
-	//Call the pipeline for each asset created by this import
-	for (int32 ObjectIndex = 0; ObjectIndex < ImportedObjects.Num(); ++ObjectIndex)
+	// Execute post-import script on all imported objects for all pipelines
+	for (int32 PipelineIndex = 0; PipelineIndex < AsyncHelper->Pipelines.Num(); ++PipelineIndex)
 	{
-		Pipeline->ScriptedExecutePostImportPipeline(NodeContainer, NodeUniqueIDs[ObjectIndex], ImportedObjects[ObjectIndex], IsAssetsReimported[ObjectIndex]);
+		for (int32 ObjectIndex = 0; ObjectIndex < ImportedObjects.Num(); ++ObjectIndex)
+		{
+			UInterchangePipelineBase* Pipeline = AsyncHelper->Pipelines[PipelineIndex];
+			Pipeline->ScriptedExecutePostImportPipeline(NodeContainer, NodeUniqueIDs[ObjectIndex], ImportedObjects[ObjectIndex], IsAssetsReimported[ObjectIndex]);
+		}
 	}
 }
diff --git a/Engine/Source/Runtime/Interchange/Engine/Private/Tasks/InterchangeTaskPipeline.h b/Engine/Source/Runtime/Interchange/Engine/Private/Tasks/InterchangeTaskPipeline.h
--- a/Engine/Source/Runtime/Interchange/Engine/Private/Tasks/InterchangeTaskPipeline.h
+++ b/Engine/Source/Runtime/Interchange/Engine/Private/Tasks/InterchangeTaskPipeline.h
@@ -55,7 +55,7 @@ namespace UE
 			void DoTask(ENamedThreads::Type CurrentThread, const FGraphEventRef& MyCompletionGraphEvent);
 		};
 
-		//We want to be sure any asset compilation is finish before calling FTaskPipelinePostImport, we use a async task to wait until they are done
+		//We want to be sure any asset compilation is finish before calling FTaskPostImport, we use a async task to wait until they are done
 		class FTaskWaitAssetCompilation
 		{
 		private:
@@ -87,46 +87,21 @@ namespace UE
 			void DoTask(ENamedThreads::Type CurrentThread, const FGraphEventRef& MyCompletionGraphEvent);
 		};
 
-		class FTaskPipelinePostImport
+		class FTaskPostImport
 		{
 		private:
 			int32 SourceIndex;
-			int32 PipelineIndex;
 			TWeakPtr<FImportAsyncHelper, ESPMode::ThreadSafe> WeakAsyncHelper;
 
 		public:
-			FTaskPipelinePostImport(int32 InSourceIndex, int32 InPipelineIndex, TWeakPtr<FImportAsyncHelper, ESPMode::ThreadSafe> InAsyncHelper)
+			FTaskPostImport(int32 InSourceIndex, TWeakPtr<FImportAsyncHelper, ESPMode::ThreadSafe> InAsyncHelper)
 				: SourceIndex(InSourceIndex)
-				, PipelineIndex(InPipelineIndex)
 				, WeakAsyncHelper(InAsyncHelper)
 			{
 			}
 
 			ENamedThreads::Type GetDesiredThread()
 			{
-				TSharedPtr<FImportAsyncHelper, ESPMode::ThreadSafe> AsyncHelper = WeakAsyncHelper.Pin();
-				//Always use the game thread if there is some error
-				if (!ensure(AsyncHelper.IsValid()))
-				{
-					return ENamedThreads::GameThread;
-				}
-				if (!AsyncHelper->Pipelines.IsValidIndex(PipelineIndex))
-				{
-					return ENamedThreads::GameThread;
-				}
-
-				//Scripted (python) cannot run outside of the game thread, it will lock forever if we do this
-				if (AsyncHelper->Pipelines[PipelineIndex]->IsScripted())
-				{
-					return ENamedThreads::GameThread;
-				}
-
-				//Ask the pipeline implementation
-				if (AsyncHelper->Pipelines[PipelineIndex]->CanExecuteOnAnyThread(EInterchangePipelineTask::PostImport))
-				{
-					return ENamedThreads::AnyBackgroundThreadNormalTask;
-				}
-
 				return ENamedThreads::GameThread;
 			}
 
@@ -137,12 +112,10 @@ namespace UE
 
 			FORCEINLINE TStatId GetStatId() const
 			{
-				RETURN_QUICK_DECLARE_CYCLE_STAT(FTaskPipelinePostImport, STATGROUP_TaskGraphTasks);
+				RETURN_QUICK_DECLARE_CYCLE_STAT(FTaskPostImport, STATGROUP_TaskGraphTasks);
 			}
 
 			void DoTask(ENamedThreads::Type CurrentThread, const FGraphEventRef& MyCompletionGraphEvent);
 		};
-
-
 	} //ns Interchange
 }//ns UE
diff --git a/Engine/Source/Runtime/Interchange/Engine/Public/InterchangeManager.h b/Engine/Source/Runtime/Interchange/Engine/Public/InterchangeManager.h
--- a/Engine/Source/Runtime/Interchange/Engine/Public/InterchangeManager.h
+++ b/Engine/Source/Runtime/Interchange/Engine/Public/InterchangeManager.h
@@ -229,7 +229,7 @@ namespace UE
 			TArray<FGraphEventRef> TranslatorTasks;
 			TArray<FGraphEventRef> PipelineTasks;
 			FGraphEventRef WaitAssetCompilationTask;
-			TArray<FGraphEventRef> PipelinePostImportTasks;
+			TArray<FGraphEventRef> PostImportTasks;
 			FGraphEventRef ParsingTask;
 			TArray<FGraphEventRef> BeginImportObjectTasks;
 			TArray<FGraphEventRef> ImportObjectTasks;
