Fixed CIS
#rnx

diff --git a/Engine/Plugins/Interchange/Runtime/Source/Import/Private/Mesh/InterchangeStaticMeshFactory.cpp b/Engine/Plugins/Interchange/Runtime/Source/Import/Private/Mesh/InterchangeStaticMeshFactory.cpp
--- a/Engine/Plugins/Interchange/Runtime/Source/Import/Private/Mesh/InterchangeStaticMeshFactory.cpp
+++ b/Engine/Plugins/Interchange/Runtime/Source/Import/Private/Mesh/InterchangeStaticMeshFactory.cpp
@@ -720,6 +720,7 @@ UInterchangeFactoryBase::FImportAssetResult UInterchangeStaticMeshFactory::EndIm
 
 void UInterchangeStaticMeshFactory::CommitMeshDescriptions(UStaticMesh& StaticMesh)
 {
+#if WITH_EDITOR
 	if (FApp::IsGame())
 	{
 		return;
@@ -738,6 +739,7 @@ void UInterchangeStaticMeshFactory::CommitMeshDescriptions(UStaticMesh& StaticMe
 
 		StaticMesh.CommitMeshDescription(LodIndex, CommitMeshDescriptionParams);
 	}
+#endif
 }
 
 void UInterchangeStaticMeshFactory::BuildFromMeshDescriptions(UStaticMesh& StaticMesh)
