Issue 410235: NonUnity errors in OpenXRInput.cpp, InterchangeGenericScenesPipeline.cpp
-Missing includes

diff --git a/Engine/Plugins/Interchange/Runtime/Source/Pipelines/Private/InterchangeGenericScenesPipeline.cpp b/Engine/Plugins/Interchange/Runtime/Source/Pipelines/Private/InterchangeGenericScenesPipeline.cpp
--- a/Engine/Plugins/Interchange/Runtime/Source/Pipelines/Private/InterchangeGenericScenesPipeline.cpp
+++ b/Engine/Plugins/Interchange/Runtime/Source/Pipelines/Private/InterchangeGenericScenesPipeline.cpp
@@ -33,6 +33,7 @@
 #include "Engine/SpotLight.h"
 #include "Engine/StaticMeshActor.h"
 #include "Engine/World.h"
+#include "Misc/PackageName.h"
 
 #if WITH_EDITOR
 #include "ObjectTools.h"
diff --git a/Engine/Plugins/Runtime/OpenXR/Source/OpenXRInput/Private/OpenXRInput.cpp b/Engine/Plugins/Runtime/OpenXR/Source/OpenXRInput/Private/OpenXRInput.cpp
--- a/Engine/Plugins/Runtime/OpenXR/Source/OpenXRInput/Private/OpenXRInput.cpp
+++ b/Engine/Plugins/Runtime/OpenXR/Source/OpenXRInput/Private/OpenXRInput.cpp
@@ -18,6 +18,7 @@
 #include "InputMappingContext.h"
 #include "PlayerMappableInputConfig.h"
 #include "GenericPlatform/GenericPlatformInputDeviceMapper.h"
+#include "Engine/Engine.h"
 
 #if WITH_EDITOR
 #include "Editor/EditorEngine.h"
