Interchange:
- Move animation curve settings to a specific category.
- Set edit condition for the normals and tangents settings.
- Create a cvar "Interchange.FeatureFlags.Import.DefaultBasicLayoutView" to control the default layout value when the user open for the first time(the value is save in the local per project settings) the Interchange import dialog.
- Set default value for animation settings curves category.

#rnx

#virtualized

diff --git a/Engine/Plugins/Interchange/Editor/Source/Pipelines/Private/SInterchangePipelineConfigurationDialog.cpp b/Engine/Plugins/Interchange/Editor/Source/Pipelines/Private/SInterchangePipelineConfigurationDialog.cpp
--- a/Engine/Plugins/Interchange/Editor/Source/Pipelines/Private/SInterchangePipelineConfigurationDialog.cpp
+++ b/Engine/Plugins/Interchange/Editor/Source/Pipelines/Private/SInterchangePipelineConfigurationDialog.cpp
@@ -27,6 +27,13 @@
 
 #define LOCTEXT_NAMESPACE "InterchangePipelineConfiguration"
 
+static bool GInterchangeDefaultBasicLayoutView = false;
+static FAutoConsoleVariableRef CCvarInterchangeEnableFBXImport(
+	TEXT("Interchange.FeatureFlags.Import.DefaultBasicLayoutView"),
+	GInterchangeDefaultBasicLayoutView,
+	TEXT("Whether the import dialog start by default in basic layout."),
+	ECVF_Default);
+
 const FName ReimportStackName = TEXT("ReimportPipeline");
 const FString ReimportPipelinePrefix = TEXT("reimport_");
 
@@ -297,6 +304,10 @@ TSharedRef<SBox> SInterchangePipelineConfigurationDialog::SpawnPipelineConfigura
 	{
 
 		StackTextComboBox = SNew(SHorizontalBox)
+		.Visibility_Lambda([this]()
+			{
+				return bBasicLayout ? EVisibility::Collapsed : EVisibility::All;
+			})
 		+ SHorizontalBox::Slot()
 		.VAlign(VAlign_Center)
 		.AutoWidth()
@@ -401,6 +412,9 @@ void SInterchangePipelineConfigurationDialog::Construct(const FArguments& InArgs
 		OwnerWindowPinned->GetOnWindowClosedEvent().AddRaw(this, &SInterchangePipelineConfigurationDialog::OnWindowClosed);
 	}
 
+	//Get the default layout when the user open the import dialog for the first time.
+	bBasicLayout = GInterchangeDefaultBasicLayoutView;
+
 	if (bReimport)
 	{
 		bFilterOptions = false;
@@ -473,7 +487,7 @@ void SInterchangePipelineConfigurationDialog::Construct(const FArguments& InArgs
 					.Padding(4.f, 0.f)
 					[
 						SNew(STextBlock)
-						.Text(LOCTEXT("SInterchangePipelineConfigurationDialog_FilterPipelineOptions", "Filter Options"))
+						.Text(LOCTEXT("SInterchangePipelineConfigurationDialog_FilterPipelineOptions", "Filter on Contents"))
 					]
 					+ SHorizontalBox::Slot()
 					.Padding(4.f, 0.f)
diff --git a/Engine/Plugins/Interchange/Runtime/Source/Pipelines/Private/InterchangeGenericAnimationPipeline.cpp b/Engine/Plugins/Interchange/Runtime/Source/Pipelines/Private/InterchangeGenericAnimationPipeline.cpp
--- a/Engine/Plugins/Interchange/Runtime/Source/Pipelines/Private/InterchangeGenericAnimationPipeline.cpp
+++ b/Engine/Plugins/Interchange/Runtime/Source/Pipelines/Private/InterchangeGenericAnimationPipeline.cpp
@@ -21,6 +21,21 @@
 
 #include UE_INLINE_GENERATED_CPP_BY_NAME(InterchangeGenericAnimationPipeline)
 
+#if WITH_EDITOR
+bool UInterchangeGenericAnimationPipeline::CanEditChange(const FProperty* InProperty) const
+{
+	// If other logic prevents editing, we want to respect that
+	const bool ParentVal = Super::CanEditChange(InProperty);
+
+	// Can we edit flower color?
+	if (InProperty->GetFName() == GET_MEMBER_NAME_CHECKED(UInterchangeGenericAnimationPipeline, FrameImportRange))
+	{
+		return ParentVal && bImportAnimations && bImportBoneTracks && AnimationRange == EInterchangeAnimationRange::SetRange;
+	}
+	return ParentVal;
+}
+#endif
+
 void UInterchangeGenericAnimationPipeline::AdjustSettingsForContext(EInterchangePipelineContext ImportType, TObjectPtr<UObject> ReimportAsset)
 {
 	Super::AdjustSettingsForContext(ImportType, ReimportAsset);
diff --git a/Engine/Plugins/Interchange/Runtime/Source/Pipelines/Public/InterchangeGenericAnimationPipeline.h b/Engine/Plugins/Interchange/Runtime/Source/Pipelines/Public/InterchangeGenericAnimationPipeline.h
--- a/Engine/Plugins/Interchange/Runtime/Source/Pipelines/Public/InterchangeGenericAnimationPipeline.h
+++ b/Engine/Plugins/Interchange/Runtime/Source/Pipelines/Public/InterchangeGenericAnimationPipeline.h
@@ -37,6 +37,10 @@ class INTERCHANGEPIPELINES_API UInterchangeGenericAnimationPipeline : public UIn
 
 public:
 
+#if WITH_EDITOR
+	virtual bool CanEditChange(const FProperty* InProperty) const override;
+#endif
+
 	//Common SkeletalMeshes And Animations Properties Settings Pointer
 	UPROPERTY(Transient)
 	TWeakObjectPtr<UInterchangeGenericCommonSkeletalMeshesAndAnimationsProperties> CommonSkeletalMeshesAndAnimationsProperties;
@@ -51,63 +55,63 @@ public:
 	bool bImportAnimations = true;
 
 	/** Import bone transform tracks. If false, this will discard any bone transform tracks.*/
-	UPROPERTY(EditAnywhere, BlueprintReadWrite, Category = "Animations")
+	UPROPERTY(EditAnywhere, BlueprintReadWrite, Category = "Animations", meta = (EditCondition = "bImportAnimations"))
 	bool bImportBoneTracks = true;
 
 	/** Which animation range to import. The one defined at Exported, at Animated time or define a range manually */
-	UPROPERTY(EditAnywhere, BlueprintReadWrite, Category = "Animations", meta = (EditCondition = "bImportBoneTracks", DisplayName = "Animation Length"))
+	UPROPERTY(EditAnywhere, BlueprintReadWrite, Category = "Animations", meta = (EditCondition = "bImportAnimations && bImportBoneTracks", DisplayName = "Animation Length"))
 	EInterchangeAnimationRange AnimationRange = EInterchangeAnimationRange::Timeline;
 
 	/** Frame range used when Set Range is used in Animation Length */
-	UPROPERTY(EditAnywhere, BlueprintReadWrite, Category = "Animations", meta = (EditCondition = "bImportBoneTracks", UIMin = 0, ClampMin = 0))
+	UPROPERTY(EditAnywhere, BlueprintReadWrite, Category = "Animations", meta = (UIMin = 0, ClampMin = 0))
 	FInt32Interval FrameImportRange = FInt32Interval(0, 0);
 
 	/** Enable this option to use default sample rate for the imported animation at 30 frames per second */
-	UPROPERTY(EditAnywhere, BlueprintReadWrite, Category = "Animations", meta = (EditCondition = "bImportBoneTracks", ToolTip = "If enabled, samples all animation curves to 30 FPS"))
+	UPROPERTY(EditAnywhere, BlueprintReadWrite, Category = "Animations", meta = (EditCondition = "bImportAnimations && bImportBoneTracks", ToolTip = "If enabled, samples all animation curves to 30 FPS"))
 	bool bUse30HzToBakeBoneAnimation = false;
 
 	/** Use this option to specify a sample rate for the imported animation, a value of 0 use the best matching sample rate. */
-	UPROPERTY(EditAnywhere, BlueprintReadWrite, Category = "Animations", meta = (EditCondition = "bImportBoneTracks && !bUse30HzToBakeBoneAnimation", ToolTip = "Sample fbx animation data at the specified sample rate, 0 find automaticaly the best sample rate", ClampMin = 0, UIMin = 0, ClampMax = 48000, UIMax = 60))
+	UPROPERTY(EditAnywhere, BlueprintReadWrite, Category = "Animations", meta = (EditCondition = "bImportAnimations && bImportBoneTracks && !bUse30HzToBakeBoneAnimation", ToolTip = "Sample fbx animation data at the specified sample rate, 0 find automaticaly the best sample rate", ClampMin = 0, UIMin = 0, ClampMax = 48000, UIMax = 60))
 	int32 CustomBoneAnimationSampleRate = 0;
 
 	/** If enabled, snaps the animation to the closest frame boundary using the import sampling rate */
-	UPROPERTY(EditAnywhere, BlueprintReadWrite, Category = "Animations")
+	UPROPERTY(EditAnywhere, BlueprintReadWrite, Category = "Animations", meta = (EditCondition = "bImportAnimations && bImportBoneTracks"))
 	bool bSnapToClosestFrameBoundary = false;
 
 	/** If true, import node attributes as either Animation Curves or Animation Attributes */
-	UPROPERTY(EditAnywhere, BlueprintReadWrite, Category = "Animations", meta = (DisplayName = "Import Attributes as Curves or Animation Attributes"))
+	UPROPERTY(EditAnywhere, BlueprintReadWrite, Category = "Animations", meta = (SubCategory = "Curves", EditCondition = "bImportAnimations", DisplayName = "Import Attributes as Curves or Animation Attributes"))
 	bool bImportCustomAttribute = true;
 
 	/** Whether to automatically add curve metadata to an animation's skeleton. If this is disabled, curve metadata will be added to skeletal meshes for morph targets, but no metadata entry will be created for general curves. */
-	UPROPERTY(EditAnywhere, BlueprintReadWrite, Category = "Animations", meta = (EditCondition = "bImportCustomAttribute"))
+	UPROPERTY(EditAnywhere, BlueprintReadWrite, Category = "Animations", meta = (SubCategory = "Curves", EditCondition = "bImportAnimations && bImportCustomAttribute"))
 	bool bAddCurveMetadataToSkeleton = true;
 
 	/** Set Material Curve Type for all custom attributes that exists */
-	UPROPERTY(EditAnywhere, BlueprintReadWrite, Category = "Animations", meta = (EditCondition = "bImportCustomAttribute", DisplayName = "Set Material Curve Type"))
+	UPROPERTY(EditAnywhere, BlueprintReadWrite, Category = "Animations", meta = (SubCategory = "Curves", EditCondition = "bImportAnimations && bImportCustomAttribute", DisplayName = "Set Material Curve Type"))
 	bool bSetMaterialDriveParameterOnCustomAttribute = false;
 
 	/** Set Material Curve Type for the custom attribute with the following suffixes. This doesn't matter if Set Material Curve Type is true  */
-	UPROPERTY(EditAnywhere, BlueprintReadWrite, Category = "Animations", meta = (EditCondition = "bImportCustomAttribute", DisplayName = "Material Curve Suffixes"))
+	UPROPERTY(EditAnywhere, BlueprintReadWrite, Category = "Animations", meta = (SubCategory = "Curves", EditCondition = "bImportAnimations && bImportCustomAttribute && bSetMaterialDriveParameterOnCustomAttribute", DisplayName = "Material Curve Suffixes"))
 	TArray<FString> MaterialCurveSuffixes = {TEXT("_mat")};
 
 	/** When importing custom attribute as curve, remove redundant keys */
-	UPROPERTY(EditAnywhere, BlueprintReadWrite, Category = "Animations", meta = (EditCondition = "bImportCustomAttribute", DisplayName = "Remove Redundant Keys"))
+	UPROPERTY(EditAnywhere, BlueprintReadWrite, Category = "Animations", meta = (SubCategory = "Curves", EditCondition = "bImportAnimations && bImportCustomAttribute", DisplayName = "Remove Redundant Keys"))
 	bool bRemoveCurveRedundantKeys = false;
 
 	/** When importing custom attribute or morphtarget as curve, do not import if it doesn't have any value other than zero. This is to avoid adding extra curves to evaluate */
-	UPROPERTY(EditAnywhere, BlueprintReadWrite, Category = "Animations", meta = (EditCondition = "bImportCustomAttribute", DisplayName = "Do not import curves with only 0 values"))
+	UPROPERTY(EditAnywhere, BlueprintReadWrite, Category = "Animations", meta = (SubCategory = "Curves", EditCondition = "bImportAnimations && bImportCustomAttribute", DisplayName = "Do not import curves with only 0 values"))
 	bool bDoNotImportCurveWithZero = false;
 
 	/** If true, all previous node attributes imported as Animation Attributes will be deleted when doing a re-import. */
-	UPROPERTY(EditAnywhere, BlueprintReadWrite, Category = "Animations", meta = (DisplayName = "Delete existing Animation Attributes"))
+	UPROPERTY(EditAnywhere, BlueprintReadWrite, Category = "Animations", meta = (SubCategory = "Curves", EditCondition = "bImportAnimations && bImportCustomAttribute", DisplayName = "Delete existing Animation Attributes"))
 	bool bDeleteExistingNonCurveCustomAttributes = false;
 
 	/** If true, all previous node attributes imported as Animation Curves will be deleted when doing a re-import. */
-	UPROPERTY(EditAnywhere, BlueprintReadWrite, Category = "Animations", meta = (DisplayName = "Delete existing Animation Curves"))
+	UPROPERTY(EditAnywhere, BlueprintReadWrite, Category = "Animations", meta = (SubCategory = "Curves", EditCondition = "bImportAnimations && bImportCustomAttribute", DisplayName = "Delete existing Animation Curves"))
 	bool bDeleteExistingCustomAttributeCurves = false;
 
 	/** If true, all previous morph target curves will be deleted when doing a re-import */
-	UPROPERTY(EditAnywhere, BlueprintReadWrite, Category = "Animations")
+	UPROPERTY(EditAnywhere, BlueprintReadWrite, Category = "Animations", meta = (SubCategory = "Curves", EditCondition = "bImportAnimations && bImportCustomAttribute", DisplayName = "Delete existing Morph Target Curves"))
 	bool bDeleteExistingMorphTargetCurves = false;
 
 	/** Name of source animation that was imported, used to reimport correct animation from the translated source */
diff --git a/Engine/Plugins/Interchange/Runtime/Source/Pipelines/Public/InterchangeGenericAssetsPipelineSharedSettings.h b/Engine/Plugins/Interchange/Runtime/Source/Pipelines/Public/InterchangeGenericAssetsPipelineSharedSettings.h
--- a/Engine/Plugins/Interchange/Runtime/Source/Pipelines/Public/InterchangeGenericAssetsPipelineSharedSettings.h
+++ b/Engine/Plugins/Interchange/Runtime/Source/Pipelines/Public/InterchangeGenericAssetsPipelineSharedSettings.h
@@ -93,11 +93,11 @@ public:
 	bool bRecomputeTangents = true;
 
 	/** If true, recompute tangents will use mikkt space. */
-	UPROPERTY(EditAnywhere, BlueprintReadWrite, Category = "Common Meshes", meta = (SubCategory = "Build"))
+	UPROPERTY(EditAnywhere, BlueprintReadWrite, Category = "Common Meshes", meta = (SubCategory = "Build", editcondition = "bRecomputeTangents"))
 	bool bUseMikkTSpace = true;
 
 	/** If true, we will use the surface area and the corner angle of the triangle as a ratio when computing the normals. */
-	UPROPERTY(EditAnywhere, BlueprintReadWrite, Category = "Common Meshes", meta = (SubCategory = "Build"))
+	UPROPERTY(EditAnywhere, BlueprintReadWrite, Category = "Common Meshes", meta = (SubCategory = "Build", editcondition = "bRecomputeNormals"))
 	bool bComputeWeightedNormals = false;
 
 	/** If true, Tangents will be stored at 16 bit vs 8 bit precision. */
