Fixed crash because of long names when importing a fbx file through DatasmithInterchange.
Two issues: DatasmithInterchange should not import FBx files yet and long names issue valid for Interchange
 1 - Removed support for FBX files through Datasmith Interchange
 2 - New logic to generate a unique name for actors based on label

#rnx

diff --git a/Engine/Plugins/Experimental/Enterprise/DatasmithInterchange/Source/DatasmithInterchange/Private/InterchangeDatasmithTranslator.cpp b/Engine/Plugins/Experimental/Enterprise/DatasmithInterchange/Source/DatasmithInterchange/Private/InterchangeDatasmithTranslator.cpp
--- a/Engine/Plugins/Experimental/Enterprise/DatasmithInterchange/Source/DatasmithInterchange/Private/InterchangeDatasmithTranslator.cpp
+++ b/Engine/Plugins/Experimental/Enterprise/DatasmithInterchange/Source/DatasmithInterchange/Private/InterchangeDatasmithTranslator.cpp
@@ -122,7 +122,7 @@ TArray<FString> UInterchangeDatasmithTranslator::GetSupportedFormats() const
 
 	for (const FString& Format : DatasmithFormats)
 	{
-		if (Format.Contains(TEXT("gltf")) || Format.Contains(TEXT("glb")))
+		if (Format.Contains(TEXT("gltf")) || Format.Contains(TEXT("glb")) || Format.Contains(TEXT("fbx")))
 		{
 			continue;
 		}
@@ -139,7 +139,7 @@ bool UInterchangeDatasmithTranslator::CanImportSourceData(const UInterchangeSour
 
 	const FString FilePath = InSourceData->GetFilename();
 	const FString FileExtension = FPaths::GetExtension(FilePath);
-	if (FileExtension.Equals(TEXT("gltf"), ESearchCase::IgnoreCase) || FileExtension.Equals(TEXT("glb"), ESearchCase::IgnoreCase))
+	if (FileExtension.Equals(TEXT("gltf"), ESearchCase::IgnoreCase) || FileExtension.Equals(TEXT("glb"), ESearchCase::IgnoreCase) || FileExtension.Equals(TEXT("fbx"), ESearchCase::IgnoreCase))
 	{
 		// Do not translate gltf since there is already a native gltf interchange translator. 
 		return false;
diff --git a/Engine/Plugins/Interchange/Runtime/Source/Import/Private/Scene/InterchangeActorHelper.cpp b/Engine/Plugins/Interchange/Runtime/Source/Import/Private/Scene/InterchangeActorHelper.cpp
--- a/Engine/Plugins/Interchange/Runtime/Source/Import/Private/Scene/InterchangeActorHelper.cpp
+++ b/Engine/Plugins/Interchange/Runtime/Source/Import/Private/Scene/InterchangeActorHelper.cpp
@@ -17,6 +17,7 @@
 #include "Engine/World.h"
 #include "MaterialDomain.h"
 #include "Materials/Material.h"
+#include "Misc/SecureHash.h"
 
 #if WITH_EDITOR
 #include "Editor/EditorEngine.h"
@@ -46,8 +47,20 @@ AActor* UE::Interchange::ActorHelper::SpawnFactoryActor(const UInterchangeFactor
 		return nullptr;
 	}
 
+	//To avoid too long actor's name, compute a 128-bit hash based on the string
+	// and use that as a unique name in the form of a GUID
+	FString UniqueID;
+	{
+		FTCHARToUTF8 Converted(*CreateSceneObjectsParams.ObjectName);
+		FMD5 MD5Gen;
+		MD5Gen.Update((const uint8*)Converted.Get(), Converted.Length());
+		uint32 Digest[4];
+		MD5Gen.Final((uint8*)Digest);
+		UniqueID = FGuid(Digest[0], Digest[1], Digest[2], Digest[3]).ToString(EGuidFormats::Base36Encoded);
+	}
+
 	FActorSpawnParameters SpawnParameters;
-	SpawnParameters.Name = FName(*CreateSceneObjectsParams.ObjectName);
+	SpawnParameters.Name = FName(*UniqueID);
 	SpawnParameters.NameMode = FActorSpawnParameters::ESpawnActorNameMode::Requested;
 	SpawnParameters.OverrideLevel = CreateSceneObjectsParams.Level;
 	SpawnParameters.SpawnCollisionHandlingOverride = ESpawnActorCollisionHandlingMethod::AlwaysSpawn;
@@ -112,7 +125,7 @@ AActor* UE::Interchange::ActorHelper::SpawnFactoryActor(const UInterchangeFactor
 	if (SpawnedActor)
 	{
 #if WITH_EDITOR
-		SpawnedActor->SetActorLabel(SpawnParameters.Name.ToString());
+		SpawnedActor->SetActorLabel(CreateSceneObjectsParams.ObjectName);
 #endif
 		if (!SpawnedActor->GetRootComponent())
 		{
