Interchange fix the bKeepSectionSeparate feature integration.

#rnx

#virtualized

diff --git a/Engine/Plugins/Interchange/Runtime/Source/Pipelines/Private/InterchangeGenericScenesPipeline.cpp b/Engine/Plugins/Interchange/Runtime/Source/Pipelines/Private/InterchangeGenericScenesPipeline.cpp
--- a/Engine/Plugins/Interchange/Runtime/Source/Pipelines/Private/InterchangeGenericScenesPipeline.cpp
+++ b/Engine/Plugins/Interchange/Runtime/Source/Pipelines/Private/InterchangeGenericScenesPipeline.cpp
@@ -631,7 +631,7 @@ void UInterchangeGenericLevelPipeline::SetUpFactoryNode(UInterchangeActorFactory
 			TMap<FString, FString> SlotMaterialDependencies;
 			SceneNode->GetSlotMaterialDependencies(SlotMaterialDependencies);
 
-			UE::Interchange::MeshesUtilities::ApplySlotMaterialDependencies(*MeshActorFactoryNode, SlotMaterialDependencies, *BaseNodeContainer);
+			UE::Interchange::MeshesUtilities::ApplySlotMaterialDependencies(*MeshActorFactoryNode, SlotMaterialDependencies, *BaseNodeContainer, nullptr);
 
 			FString AnimationAssetUidToPlay;
 			if (SceneNode->GetCustomAnimationAssetUidToPlay(AnimationAssetUidToPlay))
diff --git a/Engine/Plugins/Interchange/Runtime/Source/Pipelines/Private/InterchangeGenericSkeletalMeshPipeline.cpp b/Engine/Plugins/Interchange/Runtime/Source/Pipelines/Private/InterchangeGenericSkeletalMeshPipeline.cpp
--- a/Engine/Plugins/Interchange/Runtime/Source/Pipelines/Private/InterchangeGenericSkeletalMeshPipeline.cpp
+++ b/Engine/Plugins/Interchange/Runtime/Source/Pipelines/Private/InterchangeGenericSkeletalMeshPipeline.cpp
@@ -506,6 +506,7 @@ void UInterchangeGenericMeshPipeline::AddLodDataToSkeletalMesh(const UInterchang
 			LodDataNode->SetCustomSkeletonUid(SkeletonUid);
 			SkeletalMeshFactoryNode->AddLodDataUniqueId(SkeletalMeshLodDataUniqueID);
 		}
+		TMap<FString, FString> ExistingLodSlotMaterialDependencies;
 		constexpr bool bAddSourceNodeName = true;
 		for (const FString& NodeUid : NodeUids)
 		{
@@ -538,8 +539,7 @@ void UInterchangeGenericMeshPipeline::AddLodDataToSkeletalMesh(const UInterchang
 				MeshNode->GetSlotMaterialDependencies(SlotMaterialDependencies);
 			}
 
-			UE::Interchange::MeshesUtilities::ApplySlotMaterialDependencies(*SkeletalMeshFactoryNode, SlotMaterialDependencies, *BaseNodeContainer);
-
+			UE::Interchange::MeshesUtilities::ApplySlotMaterialDependencies(*SkeletalMeshFactoryNode, SlotMaterialDependencies, *BaseNodeContainer, &ExistingLodSlotMaterialDependencies);
 			LodDataNode->AddMeshUid(NodeUid);
 		}
 		UE::Interchange::MeshesUtilities::ReorderSlotMaterialDependencies(*SkeletalMeshFactoryNode, *BaseNodeContainer);
diff --git a/Engine/Plugins/Interchange/Runtime/Source/Pipelines/Private/InterchangeGenericStaticMeshPipeline.cpp b/Engine/Plugins/Interchange/Runtime/Source/Pipelines/Private/InterchangeGenericStaticMeshPipeline.cpp
--- a/Engine/Plugins/Interchange/Runtime/Source/Pipelines/Private/InterchangeGenericStaticMeshPipeline.cpp
+++ b/Engine/Plugins/Interchange/Runtime/Source/Pipelines/Private/InterchangeGenericStaticMeshPipeline.cpp
@@ -511,6 +511,7 @@ void UInterchangeGenericMeshPipeline::AddLodDataToStaticMesh(UInterchangeStaticM
 			BaseNodeContainer->SetNodeParentUid(StaticMeshLodDataUniqueID, StaticMeshFactoryUid);
 			StaticMeshFactoryNode->AddLodDataUniqueId(StaticMeshLodDataUniqueID);
 		}
+		TMap<FString, FString> ExistingLodSlotMaterialDependencies;
 		constexpr bool bAddSourceNodeName = true;
 		for (const FString& NodeUid : NodeUids)
 		{
@@ -546,7 +547,7 @@ void UInterchangeGenericMeshPipeline::AddLodDataToStaticMesh(UInterchangeStaticM
 				MeshNode->GetSlotMaterialDependencies(SlotMaterialDependencies);
 			}
 
-			UE::Interchange::MeshesUtilities::ApplySlotMaterialDependencies(*StaticMeshFactoryNode, SlotMaterialDependencies, *BaseNodeContainer);
+			UE::Interchange::MeshesUtilities::ApplySlotMaterialDependencies(*StaticMeshFactoryNode, SlotMaterialDependencies, *BaseNodeContainer, &ExistingLodSlotMaterialDependencies);
 
 			if (bImportCollisionAccordingToMeshName)
 			{
diff --git a/Engine/Plugins/Interchange/Runtime/Source/Pipelines/Public/InterchangePipelineMeshesUtilities.h b/Engine/Plugins/Interchange/Runtime/Source/Pipelines/Public/InterchangePipelineMeshesUtilities.h
--- a/Engine/Plugins/Interchange/Runtime/Source/Pipelines/Public/InterchangePipelineMeshesUtilities.h
+++ b/Engine/Plugins/Interchange/Runtime/Source/Pipelines/Public/InterchangePipelineMeshesUtilities.h
@@ -335,26 +335,29 @@ namespace UE::Interchange::MeshesUtilities
 
 	/**
 	 * Applies material slot dependencies stored in SlotMaterialDependencies to FactoryNode.
+	 * If the caller want to support bKeepSectionSeparate feature it must provide a valid ExistingSlotMaterialDependenciesPtr.
 	 */
 	template<class T>
-	void ApplySlotMaterialDependencies(T& FactoryNode, const TMap<FString, FString>& SlotMaterialDependencies, const UInterchangeBaseNodeContainer& NodeContainer)
+	void ApplySlotMaterialDependencies(T& FactoryNode
+		, const TMap<FString, FString>& SlotMaterialDependencies
+		, const UInterchangeBaseNodeContainer& NodeContainer
+		, TMap<FString, FString> *ExistingSlotMaterialDependenciesPtr)
 	{
 		bool bKeepSectionsSeparate = false;
-		FactoryNode.GetCustomKeepSectionsSeparate(bKeepSectionsSeparate);
-		
-		TMap<FString, FString> ExistingSlotMaterialDependencies;
-		if (bKeepSectionsSeparate)
+		if (ExistingSlotMaterialDependenciesPtr)
 		{
-			FactoryNode.GetSlotMaterialDependencies(ExistingSlotMaterialDependencies);
+			FactoryNode.GetCustomKeepSectionsSeparate(bKeepSectionsSeparate);
 		}
 
 		for (const TPair<FString, FString>& SlotMaterialDependency : SlotMaterialDependencies)
 		{
 			FString NewSlotName = SlotMaterialDependency.Key;
-			if (ExistingSlotMaterialDependencies.Contains(NewSlotName))
+			if (bKeepSectionsSeparate && ExistingSlotMaterialDependenciesPtr && ExistingSlotMaterialDependenciesPtr->Contains(NewSlotName))
 			{
-				NewSlotName += TEXT("_Section") + FString::FromInt(ExistingSlotMaterialDependencies.Num());
-				ExistingSlotMaterialDependencies.Add(NewSlotName, SlotMaterialDependency.Value);
+				TMap<FString, FString> NodeMaterialDependencies;
+				FactoryNode.GetSlotMaterialDependencies(NodeMaterialDependencies);
+				NewSlotName += TEXT("_Section") + FString::FromInt(NodeMaterialDependencies.Num());
+				ExistingSlotMaterialDependenciesPtr->Add(NewSlotName, SlotMaterialDependency.Value);
 			}
 			const FString MaterialFactoryNodeUid = UInterchangeBaseMaterialFactoryNode::GetMaterialFactoryNodeUidFromMaterialNodeUid(SlotMaterialDependency.Value);
 			FactoryNode.SetSlotMaterialDependencyUid(NewSlotName, MaterialFactoryNodeUid);
@@ -374,6 +377,10 @@ namespace UE::Interchange::MeshesUtilities
 			}
 		}
 
+		if (ExistingSlotMaterialDependenciesPtr)
+		{
+			ExistingSlotMaterialDependenciesPtr->Append(SlotMaterialDependencies);
+		}
 	}
 
 	template<class T>
