Interchange: Add error message if a fbx mesh cannot be triangulate.

#rnx

diff --git a/Engine/Plugins/Interchange/Runtime/Source/Parsers/Fbx/Private/FbxMesh.cpp b/Engine/Plugins/Interchange/Runtime/Source/Parsers/Fbx/Private/FbxMesh.cpp
--- a/Engine/Plugins/Interchange/Runtime/Source/Parsers/Fbx/Private/FbxMesh.cpp
+++ b/Engine/Plugins/Interchange/Runtime/Source/Parsers/Fbx/Private/FbxMesh.cpp
@@ -1742,13 +1742,21 @@ void FFbxMesh::AddAllMeshes(FbxScene* SDKScene, FbxGeometryConverter* SDKGeometr
 			continue;
 		}
 		FbxMesh* Mesh = static_cast<FbxMesh*>(Geometry);
-		if (!Mesh || !Mesh->IsTriangleMesh())
+		if (!Mesh)
 		{
 			continue;
 		}
 
 		FString MeshName = Parser.GetFbxHelper()->GetMeshName(Mesh);
 		FString MeshUniqueID = Parser.GetFbxHelper()->GetMeshUniqueID(Mesh);
+		if (!Mesh->IsTriangleMesh())
+		{
+			//Unable to triangulate this mesh skipping it
+			UInterchangeResultError_Generic* Message = Parser.AddMessage<UInterchangeResultError_Generic>();
+			Message->Text = FText::Format(LOCTEXT("InterchangeFbxSdkMeshTriangulationError", "Fbx sdk is unable to triangulate mesh '{0}': it will be omitted."), FText::FromString(MeshName));
+			continue;
+		}
+
 		const UInterchangeMeshNode* ExistingMeshNode = Cast<UInterchangeMeshNode>(NodeContainer.GetNode(MeshUniqueID));
 		UInterchangeMeshNode* MeshNode = nullptr;
 		if (ExistingMeshNode)
