Interchange: Avoid mesh name clash in fbx translator

#rnx

diff --git a/Engine/Plugins/Interchange/Runtime/Source/Parsers/Fbx/Private/FbxAPI.cpp b/Engine/Plugins/Interchange/Runtime/Source/Parsers/Fbx/Private/FbxAPI.cpp
--- a/Engine/Plugins/Interchange/Runtime/Source/Parsers/Fbx/Private/FbxAPI.cpp
+++ b/Engine/Plugins/Interchange/Runtime/Source/Parsers/Fbx/Private/FbxAPI.cpp
@@ -247,6 +247,21 @@ namespace UE
 					Message->Text = LOCTEXT("MissingBindPose", "Missing bind pose - the FBX SDK has created one.");
 				}
 
+				auto MakeFbxObjectNameUnique = [](FbxObject* Object, TMap<FString, int32>& Names)
+					{
+						FString ObjectName = UTF8_TO_TCHAR(Object->GetName());
+						if (int32* Count = Names.Find(ObjectName))
+						{
+							(*Count)++;
+							ObjectName += TEXT("_ncl_") + FString::FromInt(*Count);
+							Object->SetName(TCHAR_TO_UTF8(*ObjectName));
+						}
+						else
+						{
+							Names.Add(ObjectName, 0);
+						}
+					};
+
 				//////////////////////////////////////////////////////////////////////////
 				// Ensure Node Name Validity (uniqueness)
 				// Name clash must be global because unreal bones do not support name conflict (they are stored in an array, no hierarchy)
@@ -254,17 +269,26 @@ namespace UE
 				for (int32 NodeIndex = 0; NodeIndex < SDKScene->GetNodeCount(); ++NodeIndex)
 				{
 					FbxNode* Node = SDKScene->GetNode(NodeIndex);
-					FString NodeName = Node->GetName();
-					if (int32* Count = NodeNames.Find(NodeName))
+					MakeFbxObjectNameUnique(Node, NodeNames);
+				}
+
+				//////////////////////////////////////////////////////////////////////////
+				// Ensure Mesh Name Validity (uniqueness)
+				// Name clash must be global because we will build Unique ID from the mesh name
+				TMap<FString, int32> MeshNames;
+				for (int32 GeometryIndex = 0; GeometryIndex < SDKScene->GetGeometryCount(); ++GeometryIndex)
+				{
+					FbxGeometry* Geometry = SDKScene->GetGeometry(GeometryIndex);
+					if (Geometry->GetAttributeType() != FbxNodeAttribute::eMesh)
 					{
-						(*Count)++;
-						NodeName += TEXT("_ncl_") + FString::FromInt(*Count);
-						Node->SetName(TCHAR_TO_UTF8(*NodeName));
+						continue;
 					}
-					else
+					FbxMesh* Mesh = static_cast<FbxMesh*>(Geometry);
+					if (!Mesh)
 					{
-						NodeNames.Add(NodeName, 0);
+						continue;
 					}
+					MakeFbxObjectNameUnique(Mesh, MeshNames);
 				}
 			}
 		} //ns Private
diff --git a/Engine/Plugins/Interchange/Runtime/Source/Parsers/Fbx/Private/FbxConvert.cpp b/Engine/Plugins/Interchange/Runtime/Source/Parsers/Fbx/Private/FbxConvert.cpp
--- a/Engine/Plugins/Interchange/Runtime/Source/Parsers/Fbx/Private/FbxConvert.cpp
+++ b/Engine/Plugins/Interchange/Runtime/Source/Parsers/Fbx/Private/FbxConvert.cpp
@@ -108,11 +108,11 @@ namespace UE
 			}
 
 			/**
-			 * Convert ANSI char to a FString using ANSI_TO_TCHAR macro
+			 * Convert UTF8 char to a FString using ANSI_TO_TCHAR macro
 			 */
 			FString FFbxConvert::MakeString(const ANSICHAR* Name)
 			{
-				return FString(ANSI_TO_TCHAR(Name));
+				return FString(UTF8_TO_TCHAR(Name));
 			}
 
 			void FFbxConvert::ApplyUnroll(FbxNode* Node, FbxAnimLayer* Layer, FbxAnimCurveFilterUnroll* UnrollFilter)
diff --git a/Engine/Plugins/Interchange/Runtime/Source/Parsers/Fbx/Private/FbxConvert.h b/Engine/Plugins/Interchange/Runtime/Source/Parsers/Fbx/Private/FbxConvert.h
--- a/Engine/Plugins/Interchange/Runtime/Source/Parsers/Fbx/Private/FbxConvert.h
+++ b/Engine/Plugins/Interchange/Runtime/Source/Parsers/Fbx/Private/FbxConvert.h
@@ -166,7 +166,7 @@ namespace UE
 				/** String Conversion API Begin */
 
 				/**
-				 * Convert ANSI char to a FString using ANSI_TO_TCHAR macro
+				 * Convert UTF8 char to a FString using ANSI_TO_TCHAR macro
 				 */
 				static FString MakeString(const ANSICHAR* Name);
 
