Texture fix some logic treating long-lat cubemaps as nonpow2 when they are in fact always pow2 after building even if source was nonpow2
Also make the default import settings of using no mip maps for nonpow2 textures only apply to 2d textures

#rnx

diff --git a/Engine/Plugins/Interchange/Runtime/Source/Import/Private/Texture/InterchangeTextureFactory.cpp b/Engine/Plugins/Interchange/Runtime/Source/Import/Private/Texture/InterchangeTextureFactory.cpp
--- a/Engine/Plugins/Interchange/Runtime/Source/Import/Private/Texture/InterchangeTextureFactory.cpp
+++ b/Engine/Plugins/Interchange/Runtime/Source/Import/Private/Texture/InterchangeTextureFactory.cpp
@@ -1514,21 +1514,24 @@ namespace UE::Interchange::Private::InterchangeTextureFactory
 			// if Texture is not pow2 , change to TMGS_NoMipMaps (if it was default)
 			//   this used to be done by Texture2d.cpp ; it is now optional
 			//	 you can set it back to having mips if you want
-			if (Texture->MipGenSettings == TMGS_FromTextureGroup)
+			if (Texture->MipGenSettings == TMGS_FromTextureGroup && Texture->GetTextureClass() == ETextureClass::TwoD )
 			{
 				Texture->MipGenSettings = TMGS_NoMipmaps;
 			}
 
-			// if Texture is not multiple of 4, change TC to EditorIcon ("UserInterface2D")
-			//	if you do not do this, you might see "Texture forced to uncompressed because size is not a multiple of 4"
-			//  this needs to match the logic in Texture.cpp : GetDefaultTextureFormatName
-			const int32 SizeX = Texture->Source.GetSizeX();
-			const int32 SizeY = Texture->Source.GetSizeY();
-			if ((SizeX & 3) != 0 || (SizeY & 3) != 0)
+			if ( ! Texture->Source.IsLongLatCubemap() )
 			{
-				if (Texture->CompressionSettings == TC_Default) // AutoDXT/BC1
+				// if Texture is not multiple of 4, change TC to EditorIcon ("UserInterface2D")
+				//	if you do not do this, you might see "Texture forced to uncompressed because size is not a multiple of 4"
+				//  this needs to match the logic in Texture.cpp : GetDefaultTextureFormatName
+				const int32 SizeX = Texture->Source.GetSizeX();
+				const int32 SizeY = Texture->Source.GetSizeY();
+				if ((SizeX & 3) != 0 || (SizeY & 3) != 0)
 				{
-					Texture->CompressionSettings = TC_EditorIcon; // "UserInterface2D"
+					if (Texture->CompressionSettings == TC_Default) // AutoDXT/BC1
+					{
+						Texture->CompressionSettings = TC_EditorIcon; // "UserInterface2D"
+					}
 				}
 			}
 		}
diff --git a/Engine/Source/Developer/TextureCompressor/Private/TextureCompressorModule.cpp b/Engine/Source/Developer/TextureCompressor/Private/TextureCompressorModule.cpp
--- a/Engine/Source/Developer/TextureCompressor/Private/TextureCompressorModule.cpp
+++ b/Engine/Source/Developer/TextureCompressor/Private/TextureCompressorModule.cpp
@@ -2179,7 +2179,23 @@ static inline FVector ComputeWSCubeDirectionAtTexelCenter(uint32 CubemapFace, ui
 
 static uint32 ComputeLongLatCubemapExtents(int32 SrcImageSizeX, const uint32 MaxCubemapTextureResolution)
 {
-	return FMath::Clamp(1U << FMath::FloorLog2(SrcImageSizeX / 2), 32U, MaxCubemapTextureResolution);
+	// MaxCubemapTextureResolution when not set is 0xFFFFFFFF
+
+	uint32 Out = 1U << FMath::FloorLog2(SrcImageSizeX / 2);
+
+	if ( Out <= 32 || MaxCubemapTextureResolution <= 32 )
+	{
+		return 32;
+	}
+	else if ( Out > MaxCubemapTextureResolution )
+	{
+		// RoundDownToPowerOfTwo
+		return 1U << FMath::FloorLog2(MaxCubemapTextureResolution);
+	}
+	else
+	{
+		return Out;
+	}
 }
 
 void ITextureCompressorModule::GenerateBaseCubeMipFromLongitudeLatitude2D(FImage* OutMip, const FImage& SrcImage, const uint32 MaxCubemapTextureResolution, uint8 SourceEncodingOverride)
@@ -4144,6 +4160,7 @@ private:
 
 				// Max Texture Size resizing happens here :
 				// note we do not check for TMGS_Angular here
+				// note that TMGS_NoMipMaps *can* use this path; in that case it generates mips using 2x2 simple average
 				const FImage& BaseMip = bSuitableFormat ? BaseImage : Temp;
 				GenerateMipChain(BuildSettings, BaseMip, BuildSourceImageMips, 1);
 
diff --git a/Engine/Source/Editor/UnrealEd/Private/Factories/EditorFactories.cpp b/Engine/Source/Editor/UnrealEd/Private/Factories/EditorFactories.cpp
--- a/Engine/Source/Editor/UnrealEd/Private/Factories/EditorFactories.cpp
+++ b/Engine/Source/Editor/UnrealEd/Private/Factories/EditorFactories.cpp
@@ -4174,21 +4174,24 @@ UObject* UTextureFactory::FactoryCreateBinary
 		// if Texture is not pow2 , change to TMGS_NoMipMaps (if it was default)
 		//   this used to be done by Texture2d.cpp ; it is now optional
 		//	 you can set it back to having mips if you want
-		if ( Texture->MipGenSettings == TMGS_FromTextureGroup )
+		if ( Texture->MipGenSettings == TMGS_FromTextureGroup && Texture->GetTextureClass() == ETextureClass::TwoD )
 		{
 			Texture->MipGenSettings = TMGS_NoMipmaps;		
 		}
 		
-		// if Texture is not multiple of 4, change TC to EditorIcon ("UserInterface2D")
-		//	if you do not do this, you might see "Texture forced to uncompressed because size is not a multiple of 4"
-		//  this needs to match the logic in Texture.cpp : GetDefaultTextureFormatName
-		int32 SizeX = Texture->Source.GetSizeX();
-		int32 SizeY = Texture->Source.GetSizeY();
-		if ( (SizeX&3) != 0 || (SizeY&3) != 0 )
-		{
-			if ( Texture->CompressionSettings == TC_Default ) // AutoDXT/BC1
+		if ( ! Texture->Source.IsLongLatCubemap() )
+		{
+			// if Texture is not multiple of 4, change TC to EditorIcon ("UserInterface2D")
+			//	if you do not do this, you might see "Texture forced to uncompressed because size is not a multiple of 4"
+			//  this needs to match the logic in Texture.cpp : GetDefaultTextureFormatName
+			int32 SizeX = Texture->Source.GetSizeX();
+			int32 SizeY = Texture->Source.GetSizeY();
+			if ( (SizeX&3) != 0 || (SizeY&3) != 0 )
 			{
-				Texture->CompressionSettings = TC_EditorIcon; // "UserInterface2D"
+				if ( Texture->CompressionSettings == TC_Default ) // AutoDXT/BC1
+				{
+					Texture->CompressionSettings = TC_EditorIcon; // "UserInterface2D"
+				}
 			}
 		}
 	}
diff --git a/Engine/Source/Runtime/Engine/Private/Texture.cpp b/Engine/Source/Runtime/Engine/Private/Texture.cpp
--- a/Engine/Source/Runtime/Engine/Private/Texture.cpp
+++ b/Engine/Source/Runtime/Engine/Private/Texture.cpp
@@ -593,6 +593,11 @@ void UTexture::ValidateSettingsAfterImportOrEdit(bool * pRequiresNotifyMaterials
 		{
 			bIsPowerOfTwo = true;
 		}
+		if ( Source.IsLongLatCubemap() )
+		{
+			// longlat cube always generates pow2 output
+			bIsPowerOfTwo = true;
+		}
 
 		// Downscale can violate IsPow2, but it only acts when NoMipMaps, so it's moot
 
@@ -1661,6 +1666,11 @@ bool UTexture::IsPossibleToStream() const
 		{
 			bIsPowerOfTwo = true;
 		}
+		if ( Source.IsLongLatCubemap() )
+		{
+			// longlat cube always generates pow2 output
+			bIsPowerOfTwo = true;
+		}
 		if ( ! bIsPowerOfTwo )
 		{
 			// NeverStream should have been set
