Interchange: Fix skeleton conflict issue where node where duplicated.

#rnx

diff --git a/Engine/Plugins/Interchange/Runtime/Source/Pipelines/Private/InterchangeGenericAssetsPipeline.cpp b/Engine/Plugins/Interchange/Runtime/Source/Pipelines/Private/InterchangeGenericAssetsPipeline.cpp
--- a/Engine/Plugins/Interchange/Runtime/Source/Pipelines/Private/InterchangeGenericAssetsPipeline.cpp
+++ b/Engine/Plugins/Interchange/Runtime/Source/Pipelines/Private/InterchangeGenericAssetsPipeline.cpp
@@ -553,6 +553,16 @@ void UInterchangeGenericAssetsPipeline::InternalRecursiveFillJointsFromNodeConta
 
 namespace UE::Interchange::Private
 {
+	void SetParentChildConflict(TSharedPtr<UInterchangeGenericAssetsPipeline::FSkeletonJoint> ParentJoint)
+		{
+			TSharedPtr<UInterchangeGenericAssetsPipeline::FSkeletonJoint> ParentJointIter = ParentJoint;
+			while (ParentJointIter.IsValid() && !ParentJointIter->bChildConflict)
+			{
+				ParentJointIter->bChildConflict = true;
+				ParentJointIter = ParentJointIter->Parent;
+			}
+		};
+
 	void RecursivelyFillJointRemoved(TSharedPtr<UInterchangeGenericAssetsPipeline::FSkeletonJoint> AssetJoint
 		, TSharedPtr<UInterchangeGenericAssetsPipeline::FSkeletonJoint> ParentJoint
 		, TMap<FString, TSharedPtr<UInterchangeGenericAssetsPipeline::FSkeletonJoint>>& ConflictDataJoints)
@@ -580,6 +590,7 @@ namespace UE::Interchange::Private
 			ConflictDataJoints.Add(NewJoint->JointName, NewJoint);
 		}
 		SetRemove(NewJoint);
+		SetParentChildConflict(ParentJoint);
 		const int32 ChildCount = AssetJoint->Children.Num();
 		for (int32 ChildIndex = 0; ChildIndex < ChildCount; ++ChildIndex)
 		{
@@ -597,25 +608,19 @@ namespace UE::Interchange::Private
 		{
 			return;
 		}
+
 		auto SetAdded = [](TSharedPtr<UInterchangeGenericAssetsPipeline::FSkeletonJoint> Joint)
 			{
 				Joint->bMatch = false;
 				Joint->bAdded = true;
 				Joint->bRemoved = false;
 			};
-		auto SetConflict = [&ParentJoint](TSharedPtr<UInterchangeGenericAssetsPipeline::FSkeletonJoint> Joint )
+		auto SetConflict = [](TSharedPtr<UInterchangeGenericAssetsPipeline::FSkeletonJoint> Joint )
 			{
 				Joint->bMatch = false;
 				Joint->bAdded = true;
 				Joint->bRemoved = false;
 				Joint->bConflict = true;
-				//Add child conflict to all parent chain
-				TSharedPtr<UInterchangeGenericAssetsPipeline::FSkeletonJoint> ParentJointIter = ParentJoint;
-				while (ParentJointIter.IsValid() && !ParentJointIter->bChildConflict)
-				{
-					ParentJointIter->bChildConflict = true;
-					ParentJointIter = ParentJointIter->Parent;
-				}
 			};
 		TSharedPtr<UInterchangeGenericAssetsPipeline::FSkeletonJoint> NewJoint = MakeShared<UInterchangeGenericAssetsPipeline::FSkeletonJoint>();
 		NewJoint->JointName = ImportJoint->JointName;
@@ -639,6 +644,9 @@ namespace UE::Interchange::Private
 		{
 			SetAdded(NewJoint);
 		}
+		
+		SetParentChildConflict(ParentJoint);
+
 		const int32 ChildCount = ImportJoint->Children.Num();
 		for (int32 ChildIndex = 0; ChildIndex < ChildCount; ++ChildIndex)
 		{
@@ -680,37 +688,73 @@ namespace UE::Interchange::Private
 			}
 			SetMatch(NewJoint);
 			
-
-			const int32 ChildCount = FMath::Max(ImportJoint->Children.Num(), AssetJoint->Children.Num());
-			for (int32 ChildIndex = 0; ChildIndex < ChildCount; ++ChildIndex)
+			TArray<TPair<int32, int32>> ChildrenMatched;
+			TArray<int32> ChildrenRemoved;
+			TArray<int32> ChildrenAdded;
+			for (int32 AssetChildIndex = 0; AssetChildIndex < AssetJoint->Children.Num(); ++AssetChildIndex)
 			{
-				TSharedPtr<UInterchangeGenericAssetsPipeline::FSkeletonJoint> ImportJointChild = ImportJoint->Children.IsValidIndex(ChildIndex) ? ImportJoint->Children[ChildIndex] : nullptr;
-				TSharedPtr<UInterchangeGenericAssetsPipeline::FSkeletonJoint> AssetJointChild = AssetJoint->Children.IsValidIndex(ChildIndex) ? AssetJoint->Children[ChildIndex] : nullptr;
-				if (ImportJointChild.IsValid())
+				TSharedPtr<UInterchangeGenericAssetsPipeline::FSkeletonJoint> AssetJointChild = AssetJoint->Children[AssetChildIndex];
+				int32 ImportMatchIndex = INDEX_NONE;
+				for (int32 ImportChildIndex = 0; ImportChildIndex < ImportJoint->Children.Num(); ++ImportChildIndex)
 				{
-					if (AssetJointChild.IsValid())
+					TSharedPtr<UInterchangeGenericAssetsPipeline::FSkeletonJoint> ImportJointChild = ImportJoint->Children[ImportChildIndex];
+					if (ImportJointChild->JointName.Equals(AssetJointChild->JointName, ESearchCase::IgnoreCase))
 					{
-						//We have both children mean test match
-						RecursivelyFillJointMatch(AssetJointChild, ImportJointChild, NewJoint, AssetJoints, ConflictDataJoints);
+						ImportMatchIndex = ImportChildIndex;
+						ChildrenMatched.Add(TPair<int32, int32>(AssetChildIndex, ImportChildIndex));
+						break;
 					}
-					else
+				}
+				if (ImportMatchIndex == INDEX_NONE)
+				{
+					ChildrenRemoved.Add(AssetChildIndex);
+				}
+			}
+
+			for (int32 ImportChildIndex = 0; ImportChildIndex < ImportJoint->Children.Num(); ++ImportChildIndex)
+			{
+				TSharedPtr<UInterchangeGenericAssetsPipeline::FSkeletonJoint> ImportJointChild = ImportJoint->Children[ImportChildIndex];
+				int32 ImportMatchIndex = INDEX_NONE;
+				for (int32 AssetChildIndex = 0; AssetChildIndex < AssetJoint->Children.Num(); ++AssetChildIndex)
+				{
+					TSharedPtr<UInterchangeGenericAssetsPipeline::FSkeletonJoint> AssetJointChild = AssetJoint->Children[AssetChildIndex];
+					if (ImportJointChild->JointName.Equals(AssetJointChild->JointName, ESearchCase::IgnoreCase))
 					{
-						//Import only joint mean added
-						RecursivelyFillJointAdded(ImportJoint, NewJoint, AssetJoints, ConflictDataJoints);
+						ImportMatchIndex = ImportChildIndex;
+						break;
 					}
 				}
-				else
+				if (ImportMatchIndex == INDEX_NONE)
 				{
-					//Asset only joint mean removed
-					RecursivelyFillJointRemoved(AssetJoint, NewJoint, ConflictDataJoints);
+					ChildrenAdded.Add(ImportChildIndex);
 				}
 			}
+			
+			//build the matched nodes
+			for (TPair<int32, int32> MatchIndices : ChildrenMatched)
+			{
+				TSharedPtr<UInterchangeGenericAssetsPipeline::FSkeletonJoint> AssetJointChild = AssetJoint->Children[MatchIndices.Key];
+				TSharedPtr<UInterchangeGenericAssetsPipeline::FSkeletonJoint> ImportJointChild = ImportJoint->Children[MatchIndices.Value];
+				RecursivelyFillJointMatch(AssetJointChild, ImportJointChild, NewJoint, AssetJoints, ConflictDataJoints);
+			}
+			//build the removed nodes
+			for (int32 AssetChildIndex : ChildrenRemoved)
+			{
+				TSharedPtr<UInterchangeGenericAssetsPipeline::FSkeletonJoint> AssetJointChild = AssetJoint->Children[AssetChildIndex];
+				RecursivelyFillJointRemoved(AssetJointChild, NewJoint, ConflictDataJoints);
+			}
+			//build the added nodes
+			for (int32 ImportChildIndex : ChildrenAdded)
+			{
+				TSharedPtr<UInterchangeGenericAssetsPipeline::FSkeletonJoint> ImportJointChild = ImportJoint->Children[ImportChildIndex];
+				RecursivelyFillJointAdded(ImportJointChild, NewJoint, AssetJoints, ConflictDataJoints);
+			}
 		}
 		else
 		{
 			//when match fail we have two separate branch to display
-			RecursivelyFillJointAdded(ImportJoint, ParentJoint, AssetJoints, ConflictDataJoints);
 			RecursivelyFillJointRemoved(AssetJoint, ParentJoint, ConflictDataJoints);
+			RecursivelyFillJointAdded(ImportJoint, ParentJoint, AssetJoints, ConflictDataJoints);
 		}
 	}
 }
@@ -1908,8 +1952,15 @@ public:
 			JointIcon = FAppStyle::GetBrush("FBXIcon.ReimportCompareRemoved");
 			Tooltip = NSLOCTEXT("SInterchangeCompareSkeletonTreeViewItem", "Construct_RemoveJoint_tooltip", "Re-import remove this joint.").ToString();
 		}
-		const bool bConflictRoot = SkeletonCompareData->bConflict && (SkeletonCompareData->Parent.IsValid() && !SkeletonCompareData->Parent->bConflict);
-		FSlateColor ForegroundTextColor = SkeletonCompareData->bMatch ? FSlateColor::UseForeground() : (!bConflictRoot ? SInterchangeGenericAssetMaterialConflictWidget::SlateColorSubConflict : SInterchangeGenericAssetMaterialConflictWidget::SlateColorFullConflict);
+		FSlateColor ForegroundTextColor = FSlateColor::UseForeground();
+		if (SkeletonCompareData->bMatch && SkeletonCompareData->bChildConflict)
+		{
+			ForegroundTextColor = SInterchangeGenericAssetMaterialConflictWidget::SlateColorSubConflict;
+		}
+		else if (!SkeletonCompareData->bMatch)
+		{
+			ForegroundTextColor = SInterchangeGenericAssetMaterialConflictWidget::SlateColorFullConflict;
+		}
 
 		this->ChildSlot
 			[
