Interchange: Fixing values not overridden when creating Material Instances

diff --git a/Engine/Plugins/Interchange/Runtime/Source/Import/Private/Material/InterchangeMaterialFactory.cpp b/Engine/Plugins/Interchange/Runtime/Source/Import/Private/Material/InterchangeMaterialFactory.cpp
--- a/Engine/Plugins/Interchange/Runtime/Source/Import/Private/Material/InterchangeMaterialFactory.cpp
+++ b/Engine/Plugins/Interchange/Runtime/Source/Import/Private/Material/InterchangeMaterialFactory.cpp
@@ -412,7 +412,7 @@ namespace UE::Interchange::MaterialFactory::Internal
 	}
 
 
-	void UpdateParameterBool(UMaterialInstance& MaterialInstance, const FString& InputName, const UInterchangeMaterialInstanceFactoryNode& FactoryNode)
+	void UpdateParameterBool(UMaterialInstance& MaterialInstance, const FString& InputName, const UInterchangeMaterialInstanceFactoryNode& FactoryNode, bool bIsAParameter)
 	{
 #if WITH_EDITORONLY_DATA
 		const FName ParameterName = *InputName;
@@ -421,7 +421,8 @@ namespace UE::Interchange::MaterialFactory::Internal
 		if (MaterialInstance.GetStaticSwitchParameterValue(ParameterName, bInstanceValue, Uid))
 		{
 			bool bInputValue = false;
-			FactoryNode.GetBooleanAttribute(UInterchangeShaderPortsAPI::MakeInputValueKey(InputName), bInputValue);
+			const FString InputKey = bIsAParameter ? UInterchangeShaderPortsAPI::MakeInputParameterKey(InputName) : UInterchangeShaderPortsAPI::MakeInputValueKey(InputName);
+			FactoryNode.GetBooleanAttribute(InputKey, bInputValue);
 
 			if (bInputValue != bInstanceValue)
 			{
@@ -434,7 +435,7 @@ namespace UE::Interchange::MaterialFactory::Internal
 #endif // #if WITH_EDITORONLY_DATA
 	}
 
-	void UpdateParameterFloat(UMaterialInstance& MaterialInstance, const FString& InputName, const UInterchangeMaterialInstanceFactoryNode& FactoryNode)
+	void UpdateParameterFloat(UMaterialInstance& MaterialInstance, const FString& InputName, const UInterchangeMaterialInstanceFactoryNode& FactoryNode, bool bIsAParameter)
 	{
 		const FName ParameterName = *InputName;
 		float InstanceValue;
@@ -442,7 +443,8 @@ namespace UE::Interchange::MaterialFactory::Internal
 		if (MaterialInstance.GetScalarParameterValue(ParameterName, InstanceValue))
 		{
 			float InputValue = 0.f;
-			FactoryNode.GetFloatAttribute(UInterchangeShaderPortsAPI::MakeInputValueKey(InputName), InputValue);
+			const FString InputKey = bIsAParameter ? UInterchangeShaderPortsAPI::MakeInputParameterKey(InputName) : UInterchangeShaderPortsAPI::MakeInputValueKey(InputName);
+			FactoryNode.GetFloatAttribute(InputKey, InputValue);
 
 			if (!FMath::IsNearlyEqual(InputValue, InstanceValue))
 			{
@@ -461,7 +463,7 @@ namespace UE::Interchange::MaterialFactory::Internal
 		}
 	}
 
-	void UpdateParameterLinearColor(UMaterialInstance& MaterialInstance, const FString& InputName, const UInterchangeMaterialInstanceFactoryNode& FactoryNode)
+	void UpdateParameterLinearColor(UMaterialInstance& MaterialInstance, const FString& InputName, const UInterchangeMaterialInstanceFactoryNode& FactoryNode, bool bIsAParameter)
 	{
 		const FName ParameterName = *InputName;
 		FLinearColor InstanceValue;
@@ -469,7 +471,8 @@ namespace UE::Interchange::MaterialFactory::Internal
 		if (MaterialInstance.GetVectorParameterValue(ParameterName, InstanceValue))
 		{
 			FLinearColor InputValue;
-			if (FactoryNode.GetLinearColorAttribute(UInterchangeShaderPortsAPI::MakeInputValueKey(InputName), InputValue))
+			const FString InputKey = bIsAParameter ? UInterchangeShaderPortsAPI::MakeInputParameterKey(InputName) : UInterchangeShaderPortsAPI::MakeInputValueKey(InputName);
+			if (FactoryNode.GetLinearColorAttribute(InputKey, InputValue))
 			{
 				if (!InputValue.Equals(InstanceValue))
 				{
@@ -489,7 +492,7 @@ namespace UE::Interchange::MaterialFactory::Internal
 		}
 	}
 
-	void UpdateParameterTexture(UMaterialInstance& MaterialInstance, const FString& InputName, const UInterchangeMaterialInstanceFactoryNode& FactoryNode, const UInterchangeBaseNodeContainer& NodeContainer)
+	void UpdateParameterTexture(UMaterialInstance& MaterialInstance, const FString& InputName, const UInterchangeMaterialInstanceFactoryNode& FactoryNode, const UInterchangeBaseNodeContainer& NodeContainer, bool bIsAParameter)
 	{
 		const FName ParameterName = *InputName;
 		UTexture* InstanceValue;
@@ -497,7 +500,8 @@ namespace UE::Interchange::MaterialFactory::Internal
 		if (MaterialInstance.GetTextureParameterValue(ParameterName, InstanceValue))
 		{
 			FString InputValue;
-			if (FactoryNode.GetStringAttribute(UInterchangeShaderPortsAPI::MakeInputValueKey(InputName), InputValue))
+			const FString InputKey = bIsAParameter ? UInterchangeShaderPortsAPI::MakeInputParameterKey(InputName) : UInterchangeShaderPortsAPI::MakeInputValueKey(InputName);
+			if (FactoryNode.GetStringAttribute(InputKey, InputValue))
 			{
 				if (const UInterchangeTextureFactoryNode* TextureFactoryNode = Cast<UInterchangeTextureFactoryNode>(NodeContainer.GetNode(InputValue)))
 				{
@@ -1702,20 +1706,21 @@ void UInterchangeMaterialFactory::SetupMaterialInstance(UMaterialInstance& Mater
 	for (const FString& InputName : Inputs)
 	{
 		const FName ParameterName = *InputName;
+		const bool bIsAParameter = UInterchangeShaderPortsAPI::HasParameter(&FactoryNode, ParameterName);
 
-		switch (UInterchangeShaderPortsAPI::GetInputType(&FactoryNode, InputName))
+		switch (UInterchangeShaderPortsAPI::GetInputType(&FactoryNode, InputName, bIsAParameter))
 		{
 		case UE::Interchange::EAttributeTypes::Bool:
-			UpdateParameterBool(MaterialInstance, InputName, FactoryNode);
+			UpdateParameterBool(MaterialInstance, InputName, FactoryNode, bIsAParameter);
 			break;
 		case UE::Interchange::EAttributeTypes::Float:
-			UpdateParameterFloat(MaterialInstance, InputName, FactoryNode);
+			UpdateParameterFloat(MaterialInstance, InputName, FactoryNode, bIsAParameter);
 			break;
 		case UE::Interchange::EAttributeTypes::LinearColor:
-			UpdateParameterLinearColor(MaterialInstance, InputName, FactoryNode);
+			UpdateParameterLinearColor(MaterialInstance, InputName, FactoryNode, bIsAParameter);
 			break;
 		case UE::Interchange::EAttributeTypes::String:
-			UpdateParameterTexture(MaterialInstance, InputName, FactoryNode, NodeContainer);
+			UpdateParameterTexture(MaterialInstance, InputName, FactoryNode, NodeContainer, bIsAParameter);
 			break;
 		}
 	}
@@ -1771,8 +1776,8 @@ void UInterchangeMaterialFactory::SetupReimportedMaterialInstance(UMaterialInsta
 
 	for (const FString& InputName : Inputs)
 	{
-		const FString ParameterName = *InputName;
-		const FString AttributKey = UInterchangeShaderPortsAPI::MakeInputValueKey(InputName);
+		bool bIsAParameter = UInterchangeShaderPortsAPI::HasParameter(&FactoryNode, *InputName);
+		const FString AttributKey = bIsAParameter ? UInterchangeShaderPortsAPI::MakeInputParameterKey(InputName) : UInterchangeShaderPortsAPI::MakeInputValueKey(InputName);
 
 		FGuid Uid;
 		switch (UInterchangeShaderPortsAPI::GetInputType(&FactoryNode, InputName))
@@ -1789,12 +1794,9 @@ void UInterchangeMaterialFactory::SetupReimportedMaterialInstance(UMaterialInsta
 						bPreviousInputValue = bInstanceValue;
 					}
 
-					bool bInputValue = false;
-					FactoryNode.GetBooleanAttribute(AttributKey, bInputValue);
-
 					if (bInstanceValue == bPreviousInputValue)
 					{
-						UpdateParameterBool(MaterialInstance, InputName, FactoryNode);
+						UpdateParameterBool(MaterialInstance, InputName, FactoryNode, bIsAParameter);
 					}
 				}
 			}
@@ -1815,7 +1817,7 @@ void UInterchangeMaterialFactory::SetupReimportedMaterialInstance(UMaterialInsta
 
 				if(bUpdateParameter)
 				{
-					UpdateParameterFloat(MaterialInstance, InputName, FactoryNode);
+					UpdateParameterFloat(MaterialInstance, InputName, FactoryNode, bIsAParameter);
 				}
 			}
 			break;
@@ -1834,7 +1836,7 @@ void UInterchangeMaterialFactory::SetupReimportedMaterialInstance(UMaterialInsta
 
 				if (bUpdateParameter)
 				{
-					UpdateParameterLinearColor(MaterialInstance, InputName, FactoryNode);
+					UpdateParameterLinearColor(MaterialInstance, InputName, FactoryNode, bIsAParameter);
 				}
 			}
 			break;
@@ -1854,7 +1856,7 @@ void UInterchangeMaterialFactory::SetupReimportedMaterialInstance(UMaterialInsta
 
 				if (bUpdateParameter)
 				{
-					UpdateParameterTexture(MaterialInstance, InputName, FactoryNode, NodeContainer);
+					UpdateParameterTexture(MaterialInstance, InputName, FactoryNode, NodeContainer, bIsAParameter);
 				}
 			}
 			break;
diff --git a/Engine/Plugins/Interchange/Runtime/Source/Pipelines/Private/InterchangeGenericMaterialPipeline.cpp b/Engine/Plugins/Interchange/Runtime/Source/Pipelines/Private/InterchangeGenericMaterialPipeline.cpp
--- a/Engine/Plugins/Interchange/Runtime/Source/Pipelines/Private/InterchangeGenericMaterialPipeline.cpp
+++ b/Engine/Plugins/Interchange/Runtime/Source/Pipelines/Private/InterchangeGenericMaterialPipeline.cpp
@@ -2091,24 +2091,20 @@ void UInterchangeGenericMaterialPipeline::HandleTextureCoordinateNode(const UInt
 
 	// U tiling
 	{
-		TVariant<FString, FLinearColor, float> UTilingValue = VisitShaderInput(ShaderNode, Nodes::TextureCoordinate::Inputs::UTiling.ToString());
-
-		if (UTilingValue.IsType<float>())
+		if (float UTilingValue; ShaderNode->GetFloatAttribute(UInterchangeShaderPortsAPI::MakeInputValueKey(Nodes::TextureCoordinate::Inputs::UTiling.ToString()), UTilingValue))
 		{
 			const FName UTilingMemberName = GET_MEMBER_NAME_CHECKED(UMaterialExpressionTextureCoordinate, UTiling);
-			TexCoordFactoryNode->AddFloatAttribute(UTilingMemberName.ToString(), UTilingValue.Get<float>());
+			TexCoordFactoryNode->AddFloatAttribute(UTilingMemberName.ToString(), UTilingValue);
 			TexCoordFactoryNode->AddApplyAndFillDelegates<float>(UTilingMemberName.ToString(), UMaterialExpressionTextureCoordinate::StaticClass(), UTilingMemberName);
 		}
 	}
 
 	// V tiling
 	{
-		TVariant<FString, FLinearColor, float> VTilingValue = VisitShaderInput(ShaderNode, Nodes::TextureCoordinate::Inputs::VTiling.ToString());
-
-		if (VTilingValue.IsType<float>())
+		if(float VTilingValue; ShaderNode->GetFloatAttribute(UInterchangeShaderPortsAPI::MakeInputValueKey(Nodes::TextureCoordinate::Inputs::UTiling.ToString()), VTilingValue))
 		{
 			const FName VTilingMemberName = GET_MEMBER_NAME_CHECKED(UMaterialExpressionTextureCoordinate, VTiling);
-			TexCoordFactoryNode->AddFloatAttribute(VTilingMemberName.ToString(), VTilingValue.Get<float>());
+			TexCoordFactoryNode->AddFloatAttribute(VTilingMemberName.ToString(), VTilingValue);
 			TexCoordFactoryNode->AddApplyAndFillDelegates<float>(VTilingMemberName.ToString(), UMaterialExpressionTextureCoordinate::StaticClass(), VTilingMemberName);
 		}
 	}
@@ -3132,91 +3128,41 @@ UInterchangeMaterialInstanceFactoryNode* UInterchangeGenericMaterialPipeline::Cr
 	MaterialInstanceFactoryNode->SetCustomInstanceClassName(UMaterialInstanceDynamic::StaticClass()->GetPathName());
 #endif
 
+	VisitShaderGraphNode(ShaderGraphNode, MaterialInstanceFactoryNode);
+
+	return MaterialInstanceFactoryNode;
+}
+
+void UInterchangeGenericMaterialPipeline::VisitShaderGraphNode(const UInterchangeShaderGraphNode* ShaderGraphNode, UInterchangeMaterialInstanceFactoryNode* MaterialInstanceFactoryNode) const
+{
 	TArray<FString> Inputs;
 	UInterchangeShaderPortsAPI::GatherInputs(ShaderGraphNode, Inputs);
 
-	for (const FString& InputName : Inputs)
+	// We don't want to visit the whole shader graph for every input, for example with a StandardSurface with 31 inputs, the MaterialFunction is connected to all inputs of the Material but should be visited only once
+	TSet<const UInterchangeShaderNode*> VisitedNodes;
+	for(const FString& InputName : Inputs)
 	{
-		const bool bIsAParameter = UInterchangeShaderPortsAPI::HasParameter(ShaderGraphNode, FName(InputName));
-
-		TVariant<FString, FLinearColor, float> InputValue;
-		FString ConnectedShaderNodeUid;
-		FString OutputName;
-		if (UInterchangeShaderPortsAPI::GetInputConnection(ShaderGraphNode, InputName, ConnectedShaderNodeUid, OutputName))
-		{
-			if (const UInterchangeShaderNode* ConnectedShaderNode = Cast<const UInterchangeShaderNode>(BaseNodeContainer->GetNode(ConnectedShaderNodeUid)))
-			{
-				InputValue = VisitShaderNode(ConnectedShaderNode);
-			}
-		}
-		else
-		{
-			switch(UInterchangeShaderPortsAPI::GetInputType(ShaderGraphNode, InputName,bIsAParameter))
-			{
-			case UE::Interchange::EAttributeTypes::Float:
-				{
-					float AttributeValue = 0.f;
-					ShaderGraphNode->GetFloatAttribute(CreateInputKey(InputName, bIsAParameter), AttributeValue);
-					InputValue.Set<float>(AttributeValue);
-				}
-				break;
-			case UE::Interchange::EAttributeTypes::LinearColor:
-				{
-					FLinearColor AttributeValue = FLinearColor::White;
-					ShaderGraphNode->GetLinearColorAttribute(CreateInputKey(InputName,bIsAParameter), AttributeValue);
-					InputValue.Set<FLinearColor>(AttributeValue);
-				}
-				break;
-			}
-		}
-
-		if (InputValue.IsType<float>())
-		{
-			MaterialInstanceFactoryNode->AddFloatAttribute(CreateInputKey(InputName, bIsAParameter), InputValue.Get<float>());
-		}
-		else if (InputValue.IsType<FLinearColor>())
-		{
-			MaterialInstanceFactoryNode->AddLinearColorAttribute(CreateInputKey(InputName, bIsAParameter), InputValue.Get<FLinearColor>());
-		}
-		else if (InputValue.IsType<FString>())
-		{
-			const FString MapName(InputName + TEXT("Map"));
-			MaterialInstanceFactoryNode->AddStringAttribute(CreateInputKey(MapName, bIsAParameter), InputValue.Get<FString>());
-
-			const FString MapWeightName(MapName + TEXT("Weight"));
-			MaterialInstanceFactoryNode->AddFloatAttribute(CreateInputKey(MapWeightName, bIsAParameter), 1.f);
-
-			MaterialInstanceFactoryNode->AddFactoryDependencyUid(InputValue.Get<FString>());
-		}
+		VisitShaderInput(ShaderGraphNode, MaterialInstanceFactoryNode, InputName, VisitedNodes);
 	}
-
-	return MaterialInstanceFactoryNode;
 }
 
-TVariant<FString, FLinearColor, float> UInterchangeGenericMaterialPipeline::VisitShaderNode(const UInterchangeShaderNode* ShaderNode) const
+void UInterchangeGenericMaterialPipeline::VisitShaderNode(const UInterchangeShaderNode* ShaderNode, UInterchangeMaterialInstanceFactoryNode* MaterialInstanceFactoryNode, TSet<const UInterchangeShaderNode*>& VisitedNodes) const
 {
 	using namespace UE::Interchange::Materials::Standard::Nodes;
 
-	TVariant<FString, FLinearColor, float> Result;
-
-	FString ShaderType;
-	if (ShaderNode->GetCustomShaderType(ShaderType))
+	if(FString ShaderType; ShaderNode->GetCustomShaderType(ShaderType))
 	{
-		if (*ShaderType == TextureSample::Name)
-		{
-			return VisitTextureSampleNode(ShaderNode);
-		}
-		else if (*ShaderType == Lerp::Name)
+		if(*ShaderType == ScalarParameter::Name)
 		{
-			return VisitLerpNode(ShaderNode);
+			return VisitScalarParameterNode(ShaderNode, MaterialInstanceFactoryNode);
 		}
-		else if (*ShaderType == Multiply::Name)
+		else if (*ShaderType == TextureSample::Name)
 		{
-			return VisitMultiplyNode(ShaderNode);
+			return VisitTextureSampleNode(ShaderNode, MaterialInstanceFactoryNode);
 		}
-		else if (*ShaderType == OneMinus::Name)
+		else if(*ShaderType == VectorParameter::Name)
 		{
-			return VisitOneMinusNode(ShaderNode);
+			return VisitVectorParameterNode(ShaderNode, MaterialInstanceFactoryNode);
 		}
 	}
 
@@ -3224,188 +3170,73 @@ TVariant<FString, FLinearColor, float> UInterchangeGenericMaterialPipeline::Visi
 		TArray<FString> Inputs;
 		UInterchangeShaderPortsAPI::GatherInputs(ShaderNode, Inputs);
 
-		if (Inputs.Num() > 0)
+		for(const FString & InputName: Inputs)
 		{
-			const FString& InputName = Inputs[0];
-			Result = VisitShaderInput(ShaderNode, InputName);
+			VisitShaderInput(ShaderNode, MaterialInstanceFactoryNode, InputName, VisitedNodes);
 		}
 	}
-
-	return Result;
 }
 
-TVariant<FString, FLinearColor, float> UInterchangeGenericMaterialPipeline::VisitShaderInput(const UInterchangeShaderNode* ShaderNode, const FString& InputName) const
+void UInterchangeGenericMaterialPipeline::VisitShaderInput(const UInterchangeShaderNode* ShaderNode, UInterchangeMaterialInstanceFactoryNode* MaterialInstanceFactoryNode, const FString& InputName, TSet<const UInterchangeShaderNode*>& VisitedNodes) const
 {
-	TVariant<FString, FLinearColor, float> Result;
-	
+	if(VisitedNodes.Find(ShaderNode))
+	{
+		return;
+	}
+
 	const bool bIsAParameter = UInterchangeShaderPortsAPI::HasParameter(ShaderNode, FName(InputName));
 
 	FString ConnectedShaderNodeUid;
 	FString OutputName;
 	if (UInterchangeShaderPortsAPI::GetInputConnection(ShaderNode, InputName, ConnectedShaderNodeUid, OutputName))
 	{
-		if (const UInterchangeShaderNode* ConnectedShaderNode = Cast<const UInterchangeShaderNode>(BaseNodeContainer->GetNode(ConnectedShaderNodeUid)))
+		const UInterchangeShaderNode* ConnectedShaderNode = Cast<const UInterchangeShaderNode>(BaseNodeContainer->GetNode(ConnectedShaderNodeUid));
+		if (ConnectedShaderNode && !VisitedNodes.Find(ConnectedShaderNode))
 		{
-			Result = VisitShaderNode(ConnectedShaderNode);
+			VisitShaderNode(ConnectedShaderNode, MaterialInstanceFactoryNode, VisitedNodes);
+			VisitedNodes.Emplace(ConnectedShaderNode);
 		}
 	}
 	else
 	{
-		switch(UInterchangeShaderPortsAPI::GetInputType(ShaderNode, InputName))
+		switch(UInterchangeShaderPortsAPI::GetInputType(ShaderNode, InputName, bIsAParameter))
 		{
 		case UE::Interchange::EAttributeTypes::Float:
+		{
+			if(float InputValue; ShaderNode->GetFloatAttribute(CreateInputKey(InputName, bIsAParameter), InputValue))
 			{
-				float InputValue = 0.f;
-				ShaderNode->GetFloatAttribute(CreateInputKey(InputName, bIsAParameter), InputValue);
-				Result.Set<float>(InputValue);
-			}
-			break;
-		case UE::Interchange::EAttributeTypes::LinearColor:
-			{
-				FLinearColor InputValue = FLinearColor::White;
-				ShaderNode->GetLinearColorAttribute(CreateInputKey(InputName, bIsAParameter), InputValue);
-				Result.Set<FLinearColor>(InputValue);
+				MaterialInstanceFactoryNode->AddFloatAttribute(CreateInputKey(InputName, bIsAParameter), InputValue);
 			}
-			break;
 		}
-	}
-
-	return Result;
-}
-
-TVariant<FString, FLinearColor, float> UInterchangeGenericMaterialPipeline::VisitLerpNode(const UInterchangeShaderNode* ShaderNode) const
-{
-	using namespace UE::Interchange::Materials::Standard::Nodes;
-
-	TVariant<FString, FLinearColor, float> ResultA = VisitShaderInput(ShaderNode, Lerp::Inputs::A.ToString());
-	TVariant<FString, FLinearColor, float> ResultB = VisitShaderInput(ShaderNode, Lerp::Inputs::B.ToString());
-
-	TVariant<FString, FLinearColor, float> ResultFactor = VisitShaderInput(ShaderNode, Lerp::Inputs::Factor.ToString());
-
-	bool bResultAIsStrongest = true;
-
-	if (ResultFactor.IsType<float>())
-	{
-		const float Factor = ResultFactor.Get<float>();
-		bResultAIsStrongest = (Factor <= 0.5f);
-
-		// Bake the lerp into a single value
-		if (!ResultA.IsType<FString>() && !ResultB.IsType<FString>())
+		break;
+		case UE::Interchange::EAttributeTypes::LinearColor:
 		{
-			if (ResultA.IsType<float>() && ResultB.IsType<float>())
+			if(FLinearColor InputValue;	ShaderNode->GetLinearColorAttribute(CreateInputKey(InputName, bIsAParameter), InputValue))
 			{
-				const float ValueA = ResultA.Get<float>();
-				const float ValueB = ResultB.Get<float>();
-
-				TVariant<FString, FLinearColor, float> Result;
-				Result.Set<float>(FMath::Lerp(ValueA, ValueB, Factor));
-				return Result;
-			}
-			else if (ResultA.IsType<FLinearColor>() && ResultB.IsType<FLinearColor>())
-			{
-				const FLinearColor ValueA = ResultA.Get<FLinearColor>();
-				const FLinearColor ValueB = ResultB.Get<FLinearColor>();
-
-				TVariant<FString, FLinearColor, float> Result;
-				Result.Set<FLinearColor>(FMath::Lerp(ValueA, ValueB, Factor));
-				return Result;
+				MaterialInstanceFactoryNode->AddLinearColorAttribute(CreateInputKey(InputName, bIsAParameter), InputValue);
 			}
 		}
-	}
-
-	if (bResultAIsStrongest)
-	{
-		return ResultA;
-	}
-	else
-	{
-		return ResultB;
-	}
-}
-
-TVariant<FString, FLinearColor, float> UInterchangeGenericMaterialPipeline::VisitMultiplyNode(const UInterchangeShaderNode* ShaderNode) const
-{
-	using namespace UE::Interchange::Materials::Standard::Nodes;
-
-	TVariant<FString, FLinearColor, float> ResultA = VisitShaderInput(ShaderNode, Lerp::Inputs::A.ToString());
-	TVariant<FString, FLinearColor, float> ResultB = VisitShaderInput(ShaderNode, Lerp::Inputs::B.ToString());
-
-	// Bake the multiply into a single value if possible
-	if (!ResultA.IsType<FString>() && !ResultB.IsType<FString>())
-	{
-		if (ResultA.IsType<float>() && ResultB.IsType<float>())
-		{
-			const float ValueA = ResultA.Get<float>();
-			const float ValueB = ResultB.Get<float>();
-
-			TVariant<FString, FLinearColor, float> Result;
-			Result.Set<float>(ValueA * ValueB);
-			return Result;
-		}
-		else if (ResultA.IsType<FLinearColor>() && ResultB.IsType<FLinearColor>())
-		{
-			const FLinearColor ValueA = ResultA.Get<FLinearColor>();
-			const FLinearColor ValueB = ResultB.Get<FLinearColor>();
-
-			TVariant<FString, FLinearColor, float> Result;
-			Result.Set<FLinearColor>(ValueA * ValueB);
-			return Result;
-		}
-		else if (ResultA.IsType<FLinearColor>() && ResultB.IsType<float>())
-		{
-			const FLinearColor ValueA = ResultA.Get<FLinearColor>();
-			const float ValueB = ResultB.Get<float>();
-
-			TVariant<FString, FLinearColor, float> Result;
-			Result.Set<FLinearColor>(ValueA * ValueB);
-			return Result;
-		}
-		else if (ResultA.IsType<float>() && ResultB.IsType<FLinearColor>())
-		{
-			const float ValueA = ResultA.Get<float>();
-			const FLinearColor ValueB = ResultB.Get<FLinearColor>();
-
-			TVariant<FString, FLinearColor, float> Result;
-			Result.Set<FLinearColor>(ValueA * ValueB);
-			return Result;
+		break;
 		}
 	}
-
-	return ResultA;
 }
 
-TVariant<FString, FLinearColor, float> UInterchangeGenericMaterialPipeline::VisitOneMinusNode(const UInterchangeShaderNode* ShaderNode) const
+void UInterchangeGenericMaterialPipeline::VisitScalarParameterNode(const UInterchangeShaderNode* ShaderNode, UInterchangeMaterialInstanceFactoryNode* MaterialInstanceFactoryNode) const
 {
 	using namespace UE::Interchange::Materials::Standard::Nodes;
 
-	TVariant<FString, FLinearColor, float> ResultInput = VisitShaderInput(ShaderNode, OneMinus::Inputs::Input.ToString());
-
-	if (ResultInput.IsType<FLinearColor>())
-	{
-		const FLinearColor Value = ResultInput.Get<FLinearColor>();
+	const bool bIsAParameter = UInterchangeShaderPortsAPI::HasParameter(ShaderNode, ScalarParameter::Attributes::DefaultValue);
 
-		TVariant<FString, FLinearColor, float> Result;
-		Result.Set<FLinearColor>(FLinearColor::White - Value);
-		return Result;
-	}
-	else if (ResultInput.IsType<float>())
+	if(float DefaultValue; ShaderNode->GetFloatAttribute(CreateInputKey(ScalarParameter::Attributes::DefaultValue.ToString(), bIsAParameter), DefaultValue))
 	{
-		const float Value = ResultInput.Get<float>();
-
-		TVariant<FString, FLinearColor, float> Result;
-		Result.Set<float>(1.f - Value);
-		return Result;
+		MaterialInstanceFactoryNode->AddFloatAttribute(CreateInputKey(ShaderNode->GetDisplayLabel(), bIsAParameter), DefaultValue);
 	}
-
-	return ResultInput;
 }
 
-TVariant<FString, FLinearColor, float> UInterchangeGenericMaterialPipeline::VisitTextureSampleNode(const UInterchangeShaderNode* ShaderNode) const
+void UInterchangeGenericMaterialPipeline::VisitTextureSampleNode(const UInterchangeShaderNode* ShaderNode, UInterchangeMaterialInstanceFactoryNode* MaterialInstanceFactoryNode) const
 {
 	using namespace UE::Interchange::Materials::Standard::Nodes;
 
-	TVariant<FString, FLinearColor, float> Result;
-
 	const bool bIsAParameter = UInterchangeShaderPortsAPI::HasParameter(ShaderNode, TextureSample::Inputs::Texture);
 
 	FString TextureUid;
@@ -3422,14 +3253,24 @@ TVariant<FString, FLinearColor, float> UInterchangeGenericMaterialPipeline::Visi
 				if (TextureTargetNodes.Num() > 0)
 				{
 					TextureFactoryUid = TextureTargetNodes[0];
+					MaterialInstanceFactoryNode->AddStringAttribute(CreateInputKey(ShaderNode->GetDisplayLabel(), bIsAParameter), TextureFactoryUid);
+					MaterialInstanceFactoryNode->AddFactoryDependencyUid(TextureFactoryUid);
 				}
 			}
-
-			Result.Set<FString>(TextureFactoryUid);
 		}
 	}
+}
 
-	return Result;
+void UInterchangeGenericMaterialPipeline::VisitVectorParameterNode(const UInterchangeShaderNode* ShaderNode, UInterchangeMaterialInstanceFactoryNode* MaterialInstanceFactoryNode) const
+{
+	using namespace UE::Interchange::Materials::Standard::Nodes;
+
+	const bool bIsAParameter = UInterchangeShaderPortsAPI::HasParameter(ShaderNode, VectorParameter::Attributes::DefaultValue);
+
+	if(FLinearColor DefaultValue; ShaderNode->GetLinearColorAttribute(CreateInputKey(VectorParameter::Attributes::DefaultValue.ToString(), bIsAParameter), DefaultValue))
+	{
+		MaterialInstanceFactoryNode->AddLinearColorAttribute(CreateInputKey(ShaderNode->GetDisplayLabel(), true), DefaultValue);
+	}
 }
 
 FString UInterchangeGenericMaterialPipeline::GetTextureUidAttributeFromShaderNode(const UInterchangeShaderNode* ShaderNode, FName ParameterName, bool& OutIsAParameter) const
diff --git a/Engine/Plugins/Interchange/Runtime/Source/Pipelines/Public/InterchangeGenericMaterialPipeline.h b/Engine/Plugins/Interchange/Runtime/Source/Pipelines/Public/InterchangeGenericMaterialPipeline.h
--- a/Engine/Plugins/Interchange/Runtime/Source/Pipelines/Public/InterchangeGenericMaterialPipeline.h
+++ b/Engine/Plugins/Interchange/Runtime/Source/Pipelines/Public/InterchangeGenericMaterialPipeline.h
@@ -227,22 +227,15 @@ private:
 
 	/**
 	 * Visits a given shader node and its connections to find its strongest value.
-	 * Only its first input is visited as it's assumed that it's the most impactful.
 	 * The goal is to simplify a branch of a node graph to a single value, to be used for material instancing.
 	 */
-	TVariant<FString, FLinearColor, float> VisitShaderNode(const UInterchangeShaderNode* ShaderNode) const;
-	TVariant<FString, FLinearColor, float> VisitShaderInput(const UInterchangeShaderNode* ShaderNode, const FString& InputName) const;
+	void VisitShaderGraphNode(const UInterchangeShaderGraphNode* ShaderGraphNode, UInterchangeMaterialInstanceFactoryNode* MaterialInstanceFactoryNode) const;
+	void VisitShaderNode(const UInterchangeShaderNode* ShaderNode, UInterchangeMaterialInstanceFactoryNode* MaterialInstanceFactoryNode, TSet<const UInterchangeShaderNode*> & VisitedNodes) const;
+	void VisitShaderInput(const UInterchangeShaderNode* ShaderNode, UInterchangeMaterialInstanceFactoryNode* MaterialInstanceFactoryNode, const FString& InputName, TSet<const UInterchangeShaderNode*> & VisitedNodes) const;
 
-	/** 
-	 * Returns the strongest value in a lerp.
-	 * If we're lerping between scalars or colors, the lerp result will get computed and returned.
-	 * If we're lerping between textures, the strongest one is returned based on the lerp factor.
-	 */
-	TVariant<FString, FLinearColor, float> VisitLerpNode(const UInterchangeShaderNode* ShaderNode) const;
-	TVariant<FString, FLinearColor, float> VisitMultiplyNode(const UInterchangeShaderNode* ShaderNode) const;
-	TVariant<FString, FLinearColor, float> VisitOneMinusNode(const UInterchangeShaderNode* ShaderNode) const;
-	TVariant<FString, FLinearColor, float> VisitTextureSampleNode(const UInterchangeShaderNode* ShaderNode) const;
-	
+	void VisitScalarParameterNode(const UInterchangeShaderNode* ShaderNode, UInterchangeMaterialInstanceFactoryNode* MaterialInstanceFactoryNode) const;
+	void VisitTextureSampleNode(const UInterchangeShaderNode* ShaderNode, UInterchangeMaterialInstanceFactoryNode* MaterialInstanceFactoryNode) const;
+	void VisitVectorParameterNode(const UInterchangeShaderNode* ShaderNode, UInterchangeMaterialInstanceFactoryNode* MaterialInstanceFactoryNode) const;
 
 	FString GetTextureUidAttributeFromShaderNode(const UInterchangeShaderNode* ShaderNode, FName ParameterName, bool& OutIsAParameter) const;
 	FString CreateInputKey(const FString& InputName, bool bIsAParameter) const;
