Fixed ASan error due to using pointer to dead scope variable.

diff --git a/Engine/Plugins/Interchange/Runtime/Source/Pipelines/Private/InterchangeMaterialXPipeline.cpp b/Engine/Plugins/Interchange/Runtime/Source/Pipelines/Private/InterchangeMaterialXPipeline.cpp
--- a/Engine/Plugins/Interchange/Runtime/Source/Pipelines/Private/InterchangeMaterialXPipeline.cpp
+++ b/Engine/Plugins/Interchange/Runtime/Source/Pipelines/Private/InterchangeMaterialXPipeline.cpp
@@ -794,33 +794,34 @@ EMaterialXSettings UMaterialXPipelineSettings::ToEnumKey(uint8 EnumType, uint8 E
 
 namespace
 {
-	static uint8* GetMaterialXSettingsIndexValue(EMaterialXSettings Enum, SIZE_T& Index)
+	static uint8 GetMaterialXSettingsIndexValue(const EMaterialXSettings Enum, SIZE_T& Index)
 	{
 		Index = Enum.GetIndex();
-		return
-			Index == 1 ? reinterpret_cast<uint8*>(Enum.TryGet<EInterchangeMaterialXBSDF>()) :
-			Index == 2 ? reinterpret_cast<uint8*>(Enum.TryGet<EInterchangeMaterialXEDF>()) :
-			Index == 3 ? reinterpret_cast<uint8*>(Enum.TryGet<EInterchangeMaterialXVDF>()) :
-			reinterpret_cast<uint8*>(Enum.TryGet<EInterchangeMaterialXShaders>());
+		const uint8* RawValuePointer = 
+			Index == 1 ? reinterpret_cast<const uint8*>(Enum.TryGet<EInterchangeMaterialXBSDF>()) :
+			Index == 2 ? reinterpret_cast<const uint8*>(Enum.TryGet<EInterchangeMaterialXEDF>()) :
+			Index == 3 ? reinterpret_cast<const uint8*>(Enum.TryGet<EInterchangeMaterialXVDF>()) :
+			reinterpret_cast<const uint8*>(Enum.TryGet<EInterchangeMaterialXShaders>());
+		return *RawValuePointer;
 	}
 }
 
 uint32 GetTypeHash(EMaterialXSettings Key)
 {
 	SIZE_T Index;
-	uint8* UnderlyingValue = GetMaterialXSettingsIndexValue(Key, Index);
-	return HashCombine(Index, *UnderlyingValue);
+	const uint8 UnderlyingValue = GetMaterialXSettingsIndexValue(Key, Index);
+	return HashCombine(Index, UnderlyingValue);
 }
 
 bool operator==(EMaterialXSettings Lhs, EMaterialXSettings Rhs)
 {
 	SIZE_T LhsIndex;
-	uint8* LhsUnderlyingValue = GetMaterialXSettingsIndexValue(Lhs, LhsIndex);
+	const uint8 LhsUnderlyingValue = GetMaterialXSettingsIndexValue(Lhs, LhsIndex);
 
 	SIZE_T RhsIndex;
-	uint8* RhsUnderlyingValue = GetMaterialXSettingsIndexValue(Rhs, RhsIndex);
+	const uint8 RhsUnderlyingValue = GetMaterialXSettingsIndexValue(Rhs, RhsIndex);
 
-	return LhsIndex==RhsIndex && *LhsUnderlyingValue == *RhsUnderlyingValue;
+	return LhsIndex == RhsIndex && LhsUnderlyingValue == RhsUnderlyingValue;
 }
 
 #undef MATERIALX_FUNCTIONS_SUBSTRATE_PATH
