Boundary checking on PCX import

diff --git a/Engine/Plugins/Interchange/Runtime/Source/Import/Private/Texture/InterchangePCXTranslator.cpp b/Engine/Plugins/Interchange/Runtime/Source/Import/Private/Texture/InterchangePCXTranslator.cpp
--- a/Engine/Plugins/Interchange/Runtime/Source/Import/Private/Texture/InterchangePCXTranslator.cpp
+++ b/Engine/Plugins/Interchange/Runtime/Source/Import/Private/Texture/InterchangePCXTranslator.cpp
@@ -131,7 +131,6 @@ TOptional<UE::Interchange::FImportImage> UInterchangePCXTranslator::GetTexturePa
 
 		if (PCX->NumPlanes == 1 && PCX->BitsPerPixel == 8)
 		{
-
 			// Set texture properties.
 			PayloadData.Init2DWithOneMip(
 				NewU,
@@ -139,6 +138,15 @@ TOptional<UE::Interchange::FImportImage> UInterchangePCXTranslator::GetTexturePa
 				TSF_BGRA8
 			);
 			FColor* DestPtr = (FColor*)PayloadData.RawData.GetData();
+			uint64 DestSize = PayloadData.RawData.GetSize();
+			FColor* DestEnd = DestPtr + DestSize;
+
+			if (Length < 256*3)
+			{
+				FText ErrorMessage = NSLOCTEXT("InterchangePCXTranslator", "InvalidFileSize", "Failed to import PCX, file size insufficient for image size.");
+				FTextureTranslatorUtilities::LogError(*this, MoveTemp(ErrorMessage));
+				return TOptional<UE::Interchange::FImportImage>();
+			}
 
 			// Import the palette.
 			uint8* PCXPalette = (uint8*)(Buffer + Length - 256 * 3);
@@ -149,22 +157,52 @@ TOptional<UE::Interchange::FImportImage> UInterchangePCXTranslator::GetTexturePa
 			}
 
 			// Import it.
-			FColor* DestEnd = DestPtr + NewU * NewV;
+			if (BufferEnd - Buffer < 128)
+			{
+				FText ErrorMessage = NSLOCTEXT("InterchangePCXTranslator", "InvalidFileSize", "Failed to import PCX, file size insufficient for image size.");
+				FTextureTranslatorUtilities::LogError(*this, MoveTemp(ErrorMessage));
+				return TOptional<UE::Interchange::FImportImage>();
+			}
+
 			Buffer += 128;
 			while (DestPtr < DestEnd)
 			{
+				if (BufferEnd == Buffer)
+				{
+					FText ErrorMessage = NSLOCTEXT("InterchangePCXTranslator", "InvalidFileSize", "Failed to import PCX, file size insufficient for image size.");
+					FTextureTranslatorUtilities::LogError(*this, MoveTemp(ErrorMessage));
+					return TOptional<UE::Interchange::FImportImage>();
+				}
 				uint8 Color = *Buffer++;
+
 				if ((Color & 0xc0) == 0xc0)
 				{
 					uint32 RunLength = Color & 0x3f;
+
+					if (BufferEnd == Buffer)
+					{
+						FText ErrorMessage = NSLOCTEXT("InterchangePCXTranslator", "InvalidFileSize", "Failed to import PCX, file size insufficient for image size.");
+						FTextureTranslatorUtilities::LogError(*this, MoveTemp(ErrorMessage));
+						return TOptional<UE::Interchange::FImportImage>();
+					}
 					Color = *Buffer++;
 
+					if (DestEnd - DestPtr < RunLength)
+					{
+						FText ErrorMessage = NSLOCTEXT("InterchangePCXTranslator", "InvalidRunLength", "Failed to import PCX, RLE length is invalid during decode");
+						FTextureTranslatorUtilities::LogError(*this, MoveTemp(ErrorMessage));
+						return TOptional<UE::Interchange::FImportImage>();
+					}
+
 					for (uint32 Index = 0; Index < RunLength; Index++)
 					{
 						*DestPtr++ = Palette[Color];
 					}
 				}
-				else *DestPtr++ = Palette[Color];
+				else
+				{
+					*DestPtr++ = Palette[Color];
+				}
 			}
 		}
 		else if (PCX->NumPlanes == 3 && PCX->BitsPerPixel == 8)
@@ -182,7 +220,14 @@ TOptional<UE::Interchange::FImportImage> UInterchangePCXTranslator::GetTexturePa
 			FMemory::Memset(Dest, 0xff, DestSize);
 
 			// Copy upside-down scanlines.
+			if (BufferEnd - Buffer < 128)
+			{
+				FText ErrorMessage = NSLOCTEXT("InterchangePCXTranslator", "InvalidFileSize", "Failed to import PCX, file size insufficient for image size.");
+				FTextureTranslatorUtilities::LogError(*this, MoveTemp(ErrorMessage));
+				return TOptional<UE::Interchange::FImportImage>();
+			}
 			Buffer += 128;
+
 			int32 CountU = FMath::Min<int32>(PCX->BytesPerLine, NewU);
 			for (int32 i = 0; i < NewV; i++)
 			{
@@ -195,11 +240,23 @@ TOptional<UE::Interchange::FImportImage> UInterchangePCXTranslator::GetTexturePa
 					{
 						if (!Overflow)
 						{
+							if (Buffer == BufferEnd)
+							{
+								FText ErrorMessage = NSLOCTEXT("InterchangePCXTranslator", "InvalidFileSize", "Failed to import PCX, file size insufficient for image size.");
+								FTextureTranslatorUtilities::LogError(*this, MoveTemp(ErrorMessage));
+								return TOptional<UE::Interchange::FImportImage>();
+							}
 							Color = *Buffer++;
 							if ((Color & 0xc0) == 0xc0)
 							{
 								RunLength = FMath::Min((Color & 0x3f), CountU - j);
 								Overflow = (Color & 0x3f) - RunLength;
+								if (Buffer == BufferEnd)
+								{
+									FText ErrorMessage = NSLOCTEXT("InterchangePCXTranslator", "InvalidFileSize", "Failed to import PCX, file size insufficient for image size.");
+									FTextureTranslatorUtilities::LogError(*this, MoveTemp(ErrorMessage));
+									return TOptional<UE::Interchange::FImportImage>();
+								}
 								Color = *Buffer++;
 							}
 							else
