[Interchange] Crash importing gltf while in PIE. Import into Level -> Import into asset browser crash fix.

#rnx

diff --git a/Engine/Plugins/Interchange/Runtime/Source/Import/Private/Animation/InterchangeAnimSequenceFactory.cpp b/Engine/Plugins/Interchange/Runtime/Source/Import/Private/Animation/InterchangeAnimSequenceFactory.cpp
--- a/Engine/Plugins/Interchange/Runtime/Source/Import/Private/Animation/InterchangeAnimSequenceFactory.cpp
+++ b/Engine/Plugins/Interchange/Runtime/Source/Import/Private/Animation/InterchangeAnimSequenceFactory.cpp
@@ -13,6 +13,7 @@
 #include "InterchangeMeshNode.h"
 #include "InterchangeSceneNode.h"
 #include "InterchangeSkeletonFactoryNode.h"
+#include "InterchangeSkeletalMeshFactoryNode.h"
 #include "InterchangeSourceData.h"
 #include "InterchangeTranslatorBase.h"
 #include "Nodes/InterchangeBaseNode.h"
@@ -913,62 +914,35 @@
 		return ImportAssetResult;
 	}
 
-	//Verify if the bone track animation is valid (sequence length versus framerate ...)
-	if(!IsBoneTrackAnimationValid(AnimSequenceFactoryNode, Arguments))
-	{
-		return ImportAssetResult;
-	}
-
-	UObject* ExistingAsset = Arguments.ReimportObject;
-	if (!ExistingAsset)
-	{
-		FSoftObjectPath ReferenceObject;
-		if (AnimSequenceFactoryNode->GetCustomReferenceObject(ReferenceObject))
-		{
-			ExistingAsset = ReferenceObject.TryLoad();
-		}
-	}
-
-	// create a new material or overwrite existing asset, if possible
-	if (!ExistingAsset)
-	{
-		NewAnimSequence = NewObject<UAnimSequence>(Arguments.Parent, *Arguments.AssetName, RF_Public | RF_Standalone);
-	}
-	else
-	{
-		//This is a reimport, we are just re-updating the source data
-		NewAnimSequence = Cast<UAnimSequence>(ExistingAsset);
-	}
-
-	if (!NewAnimSequence)
-	{
-		UE_LOG(LogInterchangeImport, Warning, TEXT("Could not create AnimSequence asset %s"), *Arguments.AssetName);
-		return ImportAssetResult;
-	}
-
-	AnimSequenceFactoryNode->SetCustomReferenceObject(FSoftObjectPath(NewAnimSequence));
-
-	NewAnimSequence->PreEditChange(nullptr);
-
-	//Verify if the bone track animation is valid (sequence length versus framerate ...)
-	if (!IsBoneTrackAnimationValid(AnimSequenceFactoryNode, Arguments))
-	{
-		return ImportAssetResult;
-	}
-
 	FString SkeletonUid;
 	if (!AnimSequenceFactoryNode->GetCustomSkeletonFactoryNodeUid(SkeletonUid))
 	{
-		UE_LOG(LogInterchangeImport, Warning, TEXT("Could not create AnimSequence asset %s, because there is no skeleton."), *Arguments.AssetName);
+		UE_LOG(LogInterchangeImport, Error, TEXT("Could not create AnimSequence asset %s, because there is no skeleton."), *Arguments.AssetName);
 		return ImportAssetResult;
 	}
 
 	const UInterchangeSkeletonFactoryNode* SkeletonFactoryNode = Cast<UInterchangeSkeletonFactoryNode>(Arguments.NodeContainer->GetNode(SkeletonUid));
 	if (!SkeletonFactoryNode)
 	{
-		UE_LOG(LogInterchangeImport, Warning, TEXT("Invalid skeleton factory node, the skeleton factory node is obligatory to import this animsequence [%s]!"), *Arguments.AssetName);
+		UE_LOG(LogInterchangeImport, Error, TEXT("Invalid skeleton factory node, the skeleton factory node is obligatory to import this animsequence [%s]!"), *Arguments.AssetName);
 		return ImportAssetResult;
 	}
+
+	FString SkeletalMeshFactoryNodeUid;
+	if (SkeletonFactoryNode->GetCustomSkeletalMeshFactoryNodeUid(SkeletalMeshFactoryNodeUid))
+	{
+		const UInterchangeSkeletalMeshFactoryNode* SkeletalMeshFactoryNode = Cast<const UInterchangeSkeletalMeshFactoryNode>(Arguments.NodeContainer->GetFactoryNode(SkeletalMeshFactoryNodeUid));
+
+		FSoftObjectPath SkeletalMesh;
+		SkeletalMeshFactoryNode->GetCustomReferenceObject(SkeletalMesh);
+
+		if (!SkeletalMesh.IsValid())
+		{
+			UE_LOG(LogInterchangeImport, Error, TEXT("SkeletalMesh does not exist, the skeleton and skeletal mesh is obligatory to import this animsequence [%s]!"), *Arguments.AssetName);
+			return ImportAssetResult;
+		}
+	}
+
 	FSoftObjectPath SkeletonFactoryNodeReferenceObject;
 	SkeletonFactoryNode->GetCustomReferenceObject(SkeletonFactoryNodeReferenceObject);
 
@@ -997,11 +971,48 @@
 
 		if (!ensure(Skeleton))
 		{
-			UE_LOG(LogInterchangeImport, Warning, TEXT("Invalid Skeleton when importing animation sequence asset %s"), *Arguments.AssetName);
+			UE_LOG(LogInterchangeImport, Error, TEXT("Invalid Skeleton when importing animation sequence asset %s"), *Arguments.AssetName);
 			return ImportAssetResult;
 		}
 	}
 
+	//Verify if the bone track animation is valid (sequence length versus framerate ...)
+	if(!IsBoneTrackAnimationValid(AnimSequenceFactoryNode, Arguments))
+	{
+		return ImportAssetResult;
+	}
+
+	UObject* ExistingAsset = Arguments.ReimportObject;
+	if (!ExistingAsset)
+	{
+		FSoftObjectPath ReferenceObject;
+		if (AnimSequenceFactoryNode->GetCustomReferenceObject(ReferenceObject))
+		{
+			ExistingAsset = ReferenceObject.TryLoad();
+		}
+	}
+
+	// create a new material or overwrite existing asset, if possible
+	if (!ExistingAsset)
+	{
+		NewAnimSequence = NewObject<UAnimSequence>(Arguments.Parent, *Arguments.AssetName, RF_Public | RF_Standalone);
+	}
+	else
+	{
+		//This is a reimport, we are just re-updating the source data
+		NewAnimSequence = Cast<UAnimSequence>(ExistingAsset);
+	}
+
+	if (!NewAnimSequence)
+	{
+		UE_LOG(LogInterchangeImport, Warning, TEXT("Could not create AnimSequence asset %s"), *Arguments.AssetName);
+		return ImportAssetResult;
+	}
+
+	AnimSequenceFactoryNode->SetCustomReferenceObject(FSoftObjectPath(NewAnimSequence));
+
+	NewAnimSequence->PreEditChange(nullptr);
+
 	NewAnimSequence->SetSkeleton(Skeleton);
 
 	AnimSequence = NewAnimSequence;
diff --git a/Engine/Plugins/Interchange/Runtime/Source/Import/Private/Mesh/InterchangeSkeletalMeshFactory.cpp b/Engine/Plugins/Interchange/Runtime/Source/Import/Private/Mesh/InterchangeSkeletalMeshFactory.cpp
--- a/Engine/Plugins/Interchange/Runtime/Source/Import/Private/Mesh/InterchangeSkeletalMeshFactory.cpp
+++ b/Engine/Plugins/Interchange/Runtime/Source/Import/Private/Mesh/InterchangeSkeletalMeshFactory.cpp
@@ -1848,6 +1848,10 @@ UInterchangeFactoryBase::FImportAssetResult UInterchangeSkeletalMeshFactory::End
 		UE::Interchange::FFactoryCommon::ApplyReimportStrategyToAsset(SkeletalMesh, PreviousNode, CurrentNode, SkeletalMeshFactoryNode);
 	}
 
+	//For UAnimSequences we also have to check the existance of USkeletalMeshes not just USkeletons.
+	// (USkeletalMesh creation can fail while USkeletons can succeed still)
+	SkeletalMeshFactoryNode->SetCustomReferenceObject(FSoftObjectPath(SkeletalMesh));
+
 	ImportAssetResult.ImportedObject = SkeletalMesh;
 	return ImportAssetResult;
 
diff --git a/Engine/Plugins/Interchange/Runtime/Source/Import/Private/Mesh/InterchangeSkeletonHelper.cpp b/Engine/Plugins/Interchange/Runtime/Source/Import/Private/Mesh/InterchangeSkeletonHelper.cpp
--- a/Engine/Plugins/Interchange/Runtime/Source/Import/Private/Mesh/InterchangeSkeletonHelper.cpp
+++ b/Engine/Plugins/Interchange/Runtime/Source/Import/Private/Mesh/InterchangeSkeletonHelper.cpp
@@ -231,8 +231,9 @@ namespace UE::Interchange::Private
 			}
 			Info.LocalTransform = TimeZeroLocalTransform;
 		}
-		else if (!GIsAutomationTesting && !bHasBindPoseTransform && !bUseTimeZeroAsBindPose)
+		else if (!GIsAutomationTesting && !bHasBindPoseTransform && !bUseTimeZeroAsBindPose && JointNode->IsSpecializedTypeContains(FSceneNodeStaticData::GetJointSpecializeTypeString()))
 		{
+			//StaticMeshes converted to Skeletals are not expected to have BindPoses
 			OutBoneNotBindNames.Add(Info.Name);
 		}
 
