MaterialX:
- Adding support of OpenPBR
- Adding open_pbr.mtlx library in bxdf folder, waiting for it to be added officially in MaterialX libraries
- Adding a material function that falls back to StandardSurface when Substrate is not used
- Fixing a wrong use of function path based on Substrate activation in MaterialX Pipeline
- Use Substrate folder instead of Strata

#virtualized

diff --git a/Engine/Plugins/Interchange/Runtime/Source/Import/Private/MaterialX/InterchangeMaterialXDefinitions.h b/Engine/Plugins/Interchange/Runtime/Source/Import/Private/MaterialX/InterchangeMaterialXDefinitions.h
--- a/Engine/Plugins/Interchange/Runtime/Source/Import/Private/MaterialX/InterchangeMaterialXDefinitions.h
+++ b/Engine/Plugins/Interchange/Runtime/Source/Import/Private/MaterialX/InterchangeMaterialXDefinitions.h
@@ -23,6 +23,99 @@ THIRD_PARTY_INCLUDES_END
 
 MATERIALX_NAMESPACE_BEGIN
 
+	namespace OpenPBRSurface
+	{
+		namespace Input
+		{
+			static constexpr const char* BaseWeight = "base_weight";
+			static constexpr const char* BaseColor = "base_color";
+			static constexpr const char* BaseRoughness = "base_roughness";
+			static constexpr const char* BaseMetalness = "base_metalness";
+			static constexpr const char* SpecularWeight = "specular_weight";
+			static constexpr const char* SpecularColor = "specular_color";
+			static constexpr const char* SpecularRoughness = "specular_roughness";
+			static constexpr const char* SpecularIOR = "specular_ior";
+			static constexpr const char* SpecularIORLevel = "specular_ior_level";
+			static constexpr const char* SpecularAnisotropy = "specular_anisotropy";
+			static constexpr const char* SpecularRotation = "specular_rotation";
+			static constexpr const char* TransmissionWeight = "transmission_weight";
+			static constexpr const char* TransmissionColor = "transmission_color";
+			static constexpr const char* TransmissionDepth = "transmission_depth";
+			static constexpr const char* TransmissionScatter = "transmission_scatter";
+			static constexpr const char* TransmissionScatterAnisotropy= "transmission_scatter_anisotropy";
+			static constexpr const char* TransmissionDispersion = "transmission_dispersion";
+			static constexpr const char* SubsurfaceWeight = "subsurface_weight";
+			static constexpr const char* SubsurfaceColor = "subsurface_color";
+			static constexpr const char* SubsurfaceRadius = "subsurface_radius";
+			static constexpr const char* SubsurfaceRadiusScale = "subsurface_radius_scale";
+			static constexpr const char* SubsurfaceAnisotropy = "subsurface_anisotropy";
+			static constexpr const char* FuzzWeight = "fuzz_weight";
+			static constexpr const char* FuzzColor = "fuzz_color";
+			static constexpr const char* FuzzRoughness = "fuzz_roughness";
+			static constexpr const char* CoatWeight = "coat_weight";
+			static constexpr const char* CoatColor = "coat_color";
+			static constexpr const char* CoatRoughness = "coat_roughness";
+			static constexpr const char* CoatAnisotropy = "coat_anisotropy";
+			static constexpr const char* CoatRotation = "coat_rotation";
+			static constexpr const char* CoatIOR = "coat_ior";
+			static constexpr const char* CoatIORLevel = "coat_ior_level";
+			static constexpr const char* ThinFilmThickness = "thin_film_thickness";
+			static constexpr const char* ThinFilmIOR = "thin_film_ior";
+			static constexpr const char* EmissionLuminance = "emission_luminance";
+			static constexpr const char* EmissionColor = "emission_color";
+			static constexpr const char* GeometryOpacity = "geometry_opacity";
+			static constexpr const char* GeometryThinWalled = "geometry_thin_walled";
+			static constexpr const char* GeometryNormal = "geometry_normal";
+			static constexpr const char* GeometryCoatNormal = "geometry_coat_normal";
+			static constexpr const char* GeometryTangent = "geometry_tangent";
+		}
+
+		namespace DefaultValue
+		{
+			static constexpr float BaseWeight = 1.f;
+			static constexpr FLinearColor BaseColor{ 0.8, 0.8, 0.8 };
+			static constexpr float BaseRoughness = 0.f;
+			static constexpr float BaseMetalness = 0.f;
+			static constexpr float SpecularWeight = 1.f;
+			static constexpr FLinearColor SpecularColor{ 1, 1, 1 };
+			static constexpr float SpecularRoughness = 0.3f;
+			static constexpr float SpecularIOR = 1.5f;
+			static constexpr float SpecularIORLevel = 0.5f;
+			static constexpr float SpecularAnisotropy = 0.f;
+			static constexpr float SpecularRotation = 0.f;
+			static constexpr float TransmissionWeight = 0.f;
+			static constexpr FLinearColor TransmissionColor{ 1, 1, 1 };
+			static constexpr float TransmissionDepth = 0.f;
+			static constexpr FLinearColor TransmissionScatter{ 0, 0, 0 };
+			static constexpr float TransmissionScatterAnisotropy = 0.f;
+			static constexpr float TransmissionDispersion = 0.f;
+			static constexpr float SubsurfaceWeight = 0.f;
+			static constexpr FLinearColor SubsurfaceColor{ 0.8, 0.8, 0.8 };
+			static constexpr float SubsurfaceRadius = 1.f;
+			static constexpr FLinearColor SubsurfaceRadiusScale{1, 0.5, 0.25};
+			static constexpr float SubsurfaceAnisotropy = 0.f;
+			static constexpr float FuzzWeight = 0.f;
+			static constexpr FLinearColor FuzzColor{ 1, 1, 1 };
+			static constexpr float FuzzRoughness = 0.5f;
+			static constexpr float CoatWeight = 0.f;
+			static constexpr FLinearColor CoatColor{ 1, 1, 1 };
+			static constexpr float CoatRoughness = 0.f;
+			static constexpr float CoatAnisotropy = 0.f;
+			static constexpr float CoatRotation = 0.f;
+			static constexpr float CoatIOR = 1.6f;
+			static constexpr float CoatIORLevel = 0.5f;
+			static constexpr float ThinFilmThickness = 0.f;
+			static constexpr float ThinFilmIOR = 1.5f;
+			static constexpr float EmissionLuminance = 0.f;
+			static constexpr FLinearColor EmissionColor{ 1, 1, 1 };
+			static constexpr FLinearColor GeometryOpacity{ 1, 1, 1 };
+			static constexpr bool GeometryThinWalled = false;
+			static const FVector GeometryNormal{ 0, 0, 1 };
+			static const FVector GeometryCoatNormal{ 0, 0, 1 };
+			static const FVector GeometryTangent{ 0, 1, 0 };
+		}
+	}
+
 	namespace StandardSurface
 	{
 		namespace Input
@@ -403,6 +496,7 @@ MATERIALX_NAMESPACE_BEGIN
 		static constexpr const char* GltfPbr = "gltf_pbr";
 		static constexpr const char* DisneyBSDF2012 = "disney_brdf_2012";
 		static constexpr const char* DisneyBSDF2015 = "disney_bsdf_2015";
+		static constexpr const char* OpenPBRSurface = "open_pbr_surface";
 		static constexpr const char* StandardSurface = "standard_surface";
 		static constexpr const char* Surface = "surface";
 		static constexpr const char* UsdPreviewSurface = "UsdPreviewSurface";
@@ -415,6 +509,7 @@ MATERIALX_NAMESPACE_BEGIN
 
 	namespace NodeDefinition
 	{
+		static constexpr const char* OpenPBRSurface = "ND_open_pbr_surface_surfaceshader";
 		static constexpr const char* StandardSurface = "ND_standard_surface_surfaceshader";
 		static constexpr const char* SurfaceUnlit = "ND_surface_unlit";
 		static constexpr const char* UsdPreviewSurface = "ND_UsdPreviewSurface_surfaceshader";
diff --git a/Engine/Plugins/Interchange/Runtime/Source/Import/Private/MaterialX/MaterialXUtils/MaterialXBase.cpp b/Engine/Plugins/Interchange/Runtime/Source/Import/Private/MaterialX/MaterialXUtils/MaterialXBase.cpp
--- a/Engine/Plugins/Interchange/Runtime/Source/Import/Private/MaterialX/MaterialXUtils/MaterialXBase.cpp
+++ b/Engine/Plugins/Interchange/Runtime/Source/Import/Private/MaterialX/MaterialXUtils/MaterialXBase.cpp
@@ -19,6 +19,7 @@ void FMaterialXBase::UpdateDocumentRecursively(MaterialX::GraphElementPtr Graph)
 
 	static std::set<std::string> SkippedCategories{
 		mx::Category::Extract,
+		mx::Category::OpenPBRSurface,
 		mx::Category::Place2D,
 		mx::Category::Ramp4,
 		mx::Category::Saturate,
diff --git a/Engine/Plugins/Interchange/Runtime/Source/Import/Private/MaterialX/MaterialXUtils/MaterialXManager.cpp b/Engine/Plugins/Interchange/Runtime/Source/Import/Private/MaterialX/MaterialXUtils/MaterialXManager.cpp
--- a/Engine/Plugins/Interchange/Runtime/Source/Import/Private/MaterialX/MaterialXUtils/MaterialXManager.cpp
+++ b/Engine/Plugins/Interchange/Runtime/Source/Import/Private/MaterialX/MaterialXUtils/MaterialXManager.cpp
@@ -7,6 +7,7 @@
 
 #if WITH_EDITOR
 #include "MaterialX/InterchangeMaterialXDefinitions.h"
+#include "MaterialX/MaterialXUtils/MaterialXOpenPBRSurfaceShader.h"
 #include "MaterialX/MaterialXUtils/MaterialXSurfaceShader.h"
 #include "MaterialX/MaterialXUtils/MaterialXStandardSurfaceShader.h"
 #include "MaterialX/MaterialXUtils/MaterialXUsdPreviewSurfaceShader.h"
@@ -207,6 +208,7 @@ FMaterialXManager::FMaterialXManager()
 		TEXT("Y")
 	}
 	, MaterialXContainerDelegates{
+		{mx::Category::OpenPBRSurface, FMaterialXManager::FOnGetMaterialXInstance::CreateStatic(&FMaterialXOpenPBRSurfaceShader::MakeInstance)},
 		{mx::Category::Surface, FMaterialXManager::FOnGetMaterialXInstance::CreateStatic(&FMaterialXSurfaceShader::MakeInstance)},
 		{mx::Category::SurfaceUnlit, FMaterialXManager::FOnGetMaterialXInstance::CreateStatic(&FMaterialXSurfaceUnlitShader::MakeInstance)},
 		{mx::Category::StandardSurface, FMaterialXManager::FOnGetMaterialXInstance::CreateStatic(&FMaterialXStandardSurfaceShader::MakeInstance)},
diff --git a/Engine/Plugins/Interchange/Runtime/Source/Import/Private/MaterialX/MaterialXUtils/MaterialXOpenPBRSurfaceShader.cpp b/Engine/Plugins/Interchange/Runtime/Source/Import/Private/MaterialX/MaterialXUtils/MaterialXOpenPBRSurfaceShader.cpp
new file mode 100644
--- /dev/null
+++ b/Engine/Plugins/Interchange/Runtime/Source/Import/Private/MaterialX/MaterialXUtils/MaterialXOpenPBRSurfaceShader.cpp
@@ -0,0 +1,154 @@
+// Copyright Epic Games, Inc. All Rights Reserved.
+
+#if WITH_EDITOR
+#include "MaterialXOpenPBRSurfaceShader.h"
+
+#include "Engine/EngineTypes.h"
+
+namespace mx = MaterialX;
+
+FMaterialXOpenPBRSurfaceShader::FMaterialXOpenPBRSurfaceShader(UInterchangeBaseNodeContainer& BaseNodeContainer)
+	: FMaterialXSurfaceShaderAbstract(BaseNodeContainer)
+{
+	NodeDefinition = mx::NodeDefinition::OpenPBRSurface;
+}
+
+TSharedRef<FMaterialXBase> FMaterialXOpenPBRSurfaceShader::MakeInstance(UInterchangeBaseNodeContainer& BaseNodeContainer)
+{
+	TSharedRef<FMaterialXOpenPBRSurfaceShader> Result = MakeShared<FMaterialXOpenPBRSurfaceShader>(BaseNodeContainer);
+	Result->RegisterConnectNodeOutputToInputDelegates();
+	return Result;
+}
+
+void FMaterialXOpenPBRSurfaceShader::Translate(MaterialX::NodePtr OpenPBRSurfaceNode)
+{
+	using namespace UE::Interchange::Materials;
+	using namespace mx::OpenPBRSurface;
+
+	this->SurfaceShaderNode = OpenPBRSurfaceNode;
+	UInterchangeFunctionCallShaderNode* OpenPBRSurfaceShaderNode = CreateFunctionCallShaderNode(SurfaceShaderNode->getName().c_str(), 
+																								bIsSubstrateEnabled ?
+																								TEXT("/Engine/Functions/Substrate/MF_Substrate_OpenPBR_Opaque.MF_Substrate_OpenPBR_Opaque") :
+																								TEXT("/Interchange/Functions/MX_OpenPBR_Opaque.MX_OpenPBR_Opaque"));
+
+	// Inputs
+	//Base
+	ConnectNodeOutputToInput(Input::BaseWeight, OpenPBRSurfaceShaderNode, OpenPBRSurface::Parameters::BaseWeight.ToString(), DefaultValue::BaseWeight);
+	ConnectNodeOutputToInput(Input::BaseColor, OpenPBRSurfaceShaderNode, OpenPBRSurface::Parameters::BaseColor.ToString(), DefaultValue::BaseColor);
+	ConnectNodeOutputToInput(Input::BaseMetalness, OpenPBRSurfaceShaderNode, OpenPBRSurface::Parameters::BaseMetalness.ToString(), DefaultValue::BaseMetalness);
+	ConnectNodeOutputToInput(Input::BaseRoughness, OpenPBRSurfaceShaderNode, OpenPBRSurface::Parameters::BaseRoughness.ToString(), DefaultValue::BaseRoughness);
+
+	//Specular
+	ConnectNodeOutputToInput(Input::SpecularAnisotropy, OpenPBRSurfaceShaderNode, OpenPBRSurface::Parameters::SpecularAnisotropy.ToString(), DefaultValue::SpecularAnisotropy);
+	ConnectNodeOutputToInput(Input::SpecularColor, OpenPBRSurfaceShaderNode, OpenPBRSurface::Parameters::SpecularColor.ToString(), DefaultValue::SpecularColor);
+	ConnectNodeOutputToInput(Input::SpecularIOR, OpenPBRSurfaceShaderNode, OpenPBRSurface::Parameters::SpecularIOR.ToString(), DefaultValue::SpecularIOR);
+	ConnectNodeOutputToInput(Input::SpecularIORLevel, OpenPBRSurfaceShaderNode, OpenPBRSurface::Parameters::SpecularIORLevel.ToString(), DefaultValue::SpecularIORLevel);
+	ConnectNodeOutputToInput(Input::SpecularRotation, OpenPBRSurfaceShaderNode, OpenPBRSurface::Parameters::SpecularRotation.ToString(), DefaultValue::SpecularRotation);
+	ConnectNodeOutputToInput(Input::SpecularRoughness, OpenPBRSurfaceShaderNode, OpenPBRSurface::Parameters::SpecularRoughness.ToString(), DefaultValue::SpecularRoughness);
+	ConnectNodeOutputToInput(Input::SpecularWeight, OpenPBRSurfaceShaderNode, OpenPBRSurface::Parameters::SpecularWeight.ToString(), DefaultValue::SpecularWeight);
+
+	//Coat
+	ConnectNodeOutputToInput(Input::CoatAnisotropy, OpenPBRSurfaceShaderNode, OpenPBRSurface::Parameters::CoatAnisotropy.ToString(), DefaultValue::CoatAnisotropy);
+	ConnectNodeOutputToInput(Input::CoatColor, OpenPBRSurfaceShaderNode, OpenPBRSurface::Parameters::CoatColor.ToString(), DefaultValue::CoatColor);
+	ConnectNodeOutputToInput(Input::CoatIOR, OpenPBRSurfaceShaderNode, OpenPBRSurface::Parameters::CoatIOR.ToString(), DefaultValue::CoatIOR);
+	ConnectNodeOutputToInput(Input::CoatIORLevel, OpenPBRSurfaceShaderNode, OpenPBRSurface::Parameters::CoatIORLevel.ToString(), DefaultValue::CoatIORLevel);
+	ConnectNodeOutputToInput(Input::CoatRotation, OpenPBRSurfaceShaderNode, OpenPBRSurface::Parameters::CoatRotation.ToString(), DefaultValue::CoatRotation);
+	ConnectNodeOutputToInput(Input::CoatRoughness, OpenPBRSurfaceShaderNode, OpenPBRSurface::Parameters::CoatRoughness.ToString(), DefaultValue::CoatRoughness);
+	ConnectNodeOutputToInput(Input::CoatWeight, OpenPBRSurfaceShaderNode, OpenPBRSurface::Parameters::CoatWeight.ToString(), DefaultValue::CoatWeight);
+
+	//Emission
+	ConnectNodeOutputToInput(Input::EmissionColor, OpenPBRSurfaceShaderNode, OpenPBRSurface::Parameters::EmissionColor.ToString(), DefaultValue::EmissionColor);
+	ConnectNodeOutputToInput(Input::EmissionLuminance, OpenPBRSurfaceShaderNode, OpenPBRSurface::Parameters::EmissionLuminance.ToString(), DefaultValue::EmissionLuminance);
+
+	//ThinFilm
+	ConnectNodeOutputToInput(Input::ThinFilmIOR, OpenPBRSurfaceShaderNode, OpenPBRSurface::Parameters::ThinFilmIOR.ToString(), DefaultValue::ThinFilmIOR);
+	ConnectNodeOutputToInput(Input::ThinFilmThickness, OpenPBRSurfaceShaderNode, OpenPBRSurface::Parameters::ThinFilmThickness.ToString(), DefaultValue::ThinFilmThickness);
+
+	//Transmission
+	ConnectNodeOutputToInput(Input::TransmissionColor, OpenPBRSurfaceShaderNode, OpenPBRSurface::Parameters::TransmissionColor.ToString(), DefaultValue::TransmissionColor);
+	ConnectNodeOutputToInput(Input::TransmissionDepth, OpenPBRSurfaceShaderNode, OpenPBRSurface::Parameters::TransmissionDepth.ToString(), DefaultValue::TransmissionDepth);
+	ConnectNodeOutputToInput(Input::TransmissionDispersion, OpenPBRSurfaceShaderNode, OpenPBRSurface::Parameters::TransmissionDispersion.ToString(), DefaultValue::TransmissionDispersion);
+	ConnectNodeOutputToInput(Input::TransmissionScatter, OpenPBRSurfaceShaderNode, OpenPBRSurface::Parameters::TransmissionScatter.ToString(), DefaultValue::TransmissionScatter);
+	ConnectNodeOutputToInput(Input::TransmissionScatterAnisotropy, OpenPBRSurfaceShaderNode, OpenPBRSurface::Parameters::TransmissionScatterAnisotropy.ToString(), DefaultValue::TransmissionScatterAnisotropy);
+	ConnectNodeOutputToInput(Input::TransmissionWeight, OpenPBRSurfaceShaderNode, OpenPBRSurface::Parameters::TransmissionWeight.ToString(), DefaultValue::TransmissionWeight);
+
+	//Geometry
+	ConnectNodeOutputToInput(Input::GeometryCoatNormal, OpenPBRSurfaceShaderNode, OpenPBRSurface::Parameters::GeometryCoatNormal.ToString(), DefaultValue::GeometryCoatNormal);
+	ConnectNodeOutputToInput(Input::GeometryNormal, OpenPBRSurfaceShaderNode, OpenPBRSurface::Parameters::GeometryNormal.ToString(), DefaultValue::GeometryNormal);
+	ConnectNodeOutputToInput(Input::GeometryOpacity, OpenPBRSurfaceShaderNode, OpenPBRSurface::Parameters::GeometryOpacity.ToString(), DefaultValue::GeometryOpacity);
+	ConnectNodeOutputToInput(Input::GeometryTangent, OpenPBRSurfaceShaderNode, OpenPBRSurface::Parameters::GeometryTangent.ToString(), DefaultValue::GeometryTangent);
+	ConnectNodeOutputToInput(Input::GeometryThinWalled, OpenPBRSurfaceShaderNode, OpenPBRSurface::Parameters::GeometryThinWalled.ToString(), DefaultValue::GeometryThinWalled);
+	
+	//Fuzz
+	ConnectNodeOutputToInput(Input::FuzzColor, OpenPBRSurfaceShaderNode, OpenPBRSurface::Parameters::FuzzColor.ToString(), DefaultValue::FuzzColor);
+	ConnectNodeOutputToInput(Input::FuzzRoughness, OpenPBRSurfaceShaderNode, OpenPBRSurface::Parameters::FuzzRoughness.ToString(), DefaultValue::FuzzRoughness);
+	ConnectNodeOutputToInput(Input::FuzzWeight, OpenPBRSurfaceShaderNode, OpenPBRSurface::Parameters::FuzzWeight.ToString(), DefaultValue::FuzzWeight);
+
+	// We can't have Subsurface inputs if we have Transmission inputs
+	if(UInterchangeShaderPortsAPI::HasInput(OpenPBRSurfaceShaderNode, OpenPBRSurface::Parameters::TransmissionWeight))
+	{
+		OpenPBRSurfaceShaderNode->SetCustomMaterialFunction(bIsSubstrateEnabled ?
+															TEXT("/Engine/Functions/Substrate/MF_Substrate_OpenPBR_Translucent.MF_Substrate_OpenPBR_Translucent") :
+															TEXT("/Interchange/Functions/MX_OpenPBR_Translucent.MX_OpenPBR_Translucent"));
+		ShaderGraphNode->SetCustomBlendMode(EBlendMode::BLEND_TranslucentColoredTransmittance);
+	}
+	else
+	{
+		//Subsurface
+		ConnectNodeOutputToInput(Input::SubsurfaceAnisotropy, OpenPBRSurfaceShaderNode, OpenPBRSurface::Parameters::SubsurfaceAnisotropy.ToString(), DefaultValue::SubsurfaceAnisotropy);
+		ConnectNodeOutputToInput(Input::SubsurfaceColor, OpenPBRSurfaceShaderNode, OpenPBRSurface::Parameters::SubsurfaceColor.ToString(), DefaultValue::SubsurfaceColor);
+		ConnectNodeOutputToInput(Input::SubsurfaceRadius, OpenPBRSurfaceShaderNode, OpenPBRSurface::Parameters::SubsurfaceRadius.ToString(), DefaultValue::SubsurfaceRadius);
+		ConnectNodeOutputToInput(Input::SubsurfaceRadiusScale, OpenPBRSurfaceShaderNode, OpenPBRSurface::Parameters::SubsurfaceRadiusScale.ToString(), DefaultValue::SubsurfaceRadiusScale);
+		ConnectNodeOutputToInput(Input::SubsurfaceWeight, OpenPBRSurfaceShaderNode, OpenPBRSurface::Parameters::SubsurfaceWeight.ToString(), DefaultValue::SubsurfaceWeight);
+	}
+
+	if(bIsSubstrateEnabled)
+	{
+		UInterchangeShaderPortsAPI::ConnectOuputToInputByName(ShaderGraphNode, SubstrateMaterial::Parameters::FrontMaterial.ToString(), OpenPBRSurfaceShaderNode->GetUniqueID(), OpenPBRSurface::SubstrateMaterial::Outputs::FrontMaterial.ToString());
+
+		if(UInterchangeShaderPortsAPI::HasInput(OpenPBRSurfaceShaderNode, OpenPBRSurface::Parameters::GeometryOpacity))
+		{
+			UInterchangeShaderPortsAPI::ConnectOuputToInputByName(ShaderGraphNode, SubstrateMaterial::Parameters::OpacityMask.ToString(), OpenPBRSurfaceShaderNode->GetUniqueID(), OpenPBRSurface::SubstrateMaterial::Outputs::OpacityMask.ToString());
+			ShaderGraphNode->SetCustomBlendMode(EBlendMode::BLEND_Masked);
+		}
+	}
+	else
+	{	// Outputs
+		UInterchangeShaderPortsAPI::ConnectOuputToInputByName(ShaderGraphNode, PBRMR::Parameters::BaseColor.ToString(), OpenPBRSurfaceShaderNode->GetUniqueID(), PBRMR::Parameters::BaseColor.ToString());
+		UInterchangeShaderPortsAPI::ConnectOuputToInputByName(ShaderGraphNode, PBRMR::Parameters::Metallic.ToString(), OpenPBRSurfaceShaderNode->GetUniqueID(), PBRMR::Parameters::Metallic.ToString());
+		UInterchangeShaderPortsAPI::ConnectOuputToInputByName(ShaderGraphNode, PBRMR::Parameters::Specular.ToString(), OpenPBRSurfaceShaderNode->GetUniqueID(), PBRMR::Parameters::Specular.ToString());
+		UInterchangeShaderPortsAPI::ConnectOuputToInputByName(ShaderGraphNode, PBRMR::Parameters::Roughness.ToString(), OpenPBRSurfaceShaderNode->GetUniqueID(), PBRMR::Parameters::Roughness.ToString());
+		UInterchangeShaderPortsAPI::ConnectOuputToInputByName(ShaderGraphNode, PBRMR::Parameters::EmissiveColor.ToString(), OpenPBRSurfaceShaderNode->GetUniqueID(), PBRMR::Parameters::EmissiveColor.ToString());
+		UInterchangeShaderPortsAPI::ConnectOuputToInputByName(ShaderGraphNode, PBRMR::Parameters::Anisotropy.ToString(), OpenPBRSurfaceShaderNode->GetUniqueID(), PBRMR::Parameters::Anisotropy.ToString());
+		UInterchangeShaderPortsAPI::ConnectOuputToInputByName(ShaderGraphNode, PBRMR::Parameters::Normal.ToString(), OpenPBRSurfaceShaderNode->GetUniqueID(), PBRMR::Parameters::Normal.ToString());
+		UInterchangeShaderPortsAPI::ConnectOuputToInputByName(ShaderGraphNode, PBRMR::Parameters::Tangent.ToString(), OpenPBRSurfaceShaderNode->GetUniqueID(), PBRMR::Parameters::Tangent.ToString());
+
+		// We can't have all shading models at once, so we have to make a choice here
+		if(UInterchangeShaderPortsAPI::HasInput(OpenPBRSurfaceShaderNode, OpenPBRSurface::Parameters::TransmissionWeight))
+		{
+			UInterchangeShaderPortsAPI::ConnectOuputToInputByName(ShaderGraphNode, PBRMR::Parameters::Opacity.ToString(), OpenPBRSurfaceShaderNode->GetUniqueID(), PBRMR::Parameters::Opacity.ToString());
+			UInterchangeShaderPortsAPI::ConnectOuputToInputByName(ShaderGraphNode, ThinTranslucent::Parameters::TransmissionColor.ToString(), OpenPBRSurfaceShaderNode->GetUniqueID(), ThinTranslucent::Parameters::TransmissionColor.ToString());
+			UInterchangeShaderPortsAPI::ConnectOuputToInputByName(ShaderGraphNode, Common::Parameters::Refraction.ToString(), OpenPBRSurfaceShaderNode->GetUniqueID(), Common::Parameters::Refraction.ToString());
+		}
+		else if(UInterchangeShaderPortsAPI::HasInput(OpenPBRSurfaceShaderNode, OpenPBRSurface::Parameters::FuzzWeight))
+		{
+			UInterchangeShaderPortsAPI::ConnectOuputToInputByName(ShaderGraphNode, Sheen::Parameters::SheenColor.ToString(), OpenPBRSurfaceShaderNode->GetUniqueID(), Sheen::Parameters::SheenColor.ToString());
+			UInterchangeShaderPortsAPI::ConnectOuputToInputByName(ShaderGraphNode, Sheen::Parameters::SheenRoughness.ToString(), OpenPBRSurfaceShaderNode->GetUniqueID(), Sheen::Parameters::SheenRoughness.ToString());
+		}
+		else if(UInterchangeShaderPortsAPI::HasInput(OpenPBRSurfaceShaderNode, OpenPBRSurface::Parameters::CoatWeight))
+		{
+			UInterchangeShaderPortsAPI::ConnectOuputToInputByName(ShaderGraphNode, ClearCoat::Parameters::ClearCoat.ToString(), OpenPBRSurfaceShaderNode->GetUniqueID(), ClearCoat::Parameters::ClearCoat.ToString());
+			UInterchangeShaderPortsAPI::ConnectOuputToInputByName(ShaderGraphNode, ClearCoat::Parameters::ClearCoatRoughness.ToString(), OpenPBRSurfaceShaderNode->GetUniqueID(), ClearCoat::Parameters::ClearCoatRoughness.ToString());
+			UInterchangeShaderPortsAPI::ConnectOuputToInputByName(ShaderGraphNode, ClearCoat::Parameters::ClearCoatNormal.ToString(), OpenPBRSurfaceShaderNode->GetUniqueID(), ClearCoat::Parameters::ClearCoatNormal.ToString());
+		}
+		else if(UInterchangeShaderPortsAPI::HasInput(OpenPBRSurfaceShaderNode, OpenPBRSurface::Parameters::SubsurfaceWeight))
+		{
+			UInterchangeShaderPortsAPI::ConnectOuputToInputByName(ShaderGraphNode, PBRMR::Parameters::Opacity.ToString(), OpenPBRSurfaceShaderNode->GetUniqueID(), PBRMR::Parameters::Opacity.ToString());
+			UInterchangeShaderPortsAPI::ConnectOuputToInputByName(ShaderGraphNode, Subsurface::Parameters::SubsurfaceColor.ToString(), OpenPBRSurfaceShaderNode->GetUniqueID(), Subsurface::Parameters::SubsurfaceColor.ToString());
+		}
+	}
+
+	// Outputs
+}
+
+#endif //WITH_EDITOR
diff --git a/Engine/Plugins/Interchange/Runtime/Source/Import/Private/MaterialX/MaterialXUtils/MaterialXOpenPBRSurfaceShader.h b/Engine/Plugins/Interchange/Runtime/Source/Import/Private/MaterialX/MaterialXUtils/MaterialXOpenPBRSurfaceShader.h
new file mode 100644
--- /dev/null
+++ b/Engine/Plugins/Interchange/Runtime/Source/Import/Private/MaterialX/MaterialXUtils/MaterialXOpenPBRSurfaceShader.h
@@ -0,0 +1,24 @@
+// Copyright Epic Games, Inc. All Rights Reserved.
+
+#pragma once
+
+#if WITH_EDITOR
+#include "MaterialX/MaterialXUtils/MaterialXSurfaceShaderAbstract.h"
+
+class FMaterialXOpenPBRSurfaceShader : public FMaterialXSurfaceShaderAbstract
+{
+protected:
+
+	template<typename T, ESPMode>
+	friend class SharedPointerInternals::TIntrusiveReferenceController;
+
+	FMaterialXOpenPBRSurfaceShader(UInterchangeBaseNodeContainer& BaseNodeContainer);
+
+public:
+
+	static TSharedRef<FMaterialXBase> MakeInstance(UInterchangeBaseNodeContainer& BaseNodeContainer);
+
+	virtual void Translate(MaterialX::NodePtr OpenPBRSurfaceNode) override;
+};
+
+#endif //WITH_EDITOR
diff --git a/Engine/Plugins/Interchange/Runtime/Source/Import/Private/MaterialX/MaterialXUtils/MaterialXStandardSurfaceShader.cpp b/Engine/Plugins/Interchange/Runtime/Source/Import/Private/MaterialX/MaterialXUtils/MaterialXStandardSurfaceShader.cpp
--- a/Engine/Plugins/Interchange/Runtime/Source/Import/Private/MaterialX/MaterialXUtils/MaterialXStandardSurfaceShader.cpp
+++ b/Engine/Plugins/Interchange/Runtime/Source/Import/Private/MaterialX/MaterialXUtils/MaterialXStandardSurfaceShader.cpp
@@ -183,7 +183,7 @@ void FMaterialXStandardSurfaceShader::ConnectToSubstrateStandardSurface()
 	using namespace mx::StandardSurface;
 	constexpr bool bInputInTangentSpace = true;
 
-	UInterchangeFunctionCallShaderNode* StandardSurfaceShaderNode = CreateFunctionCallShaderNode(SurfaceShaderNode->getName().c_str(), TEXT("/Engine/Functions/Strata/Substrate-StandardSurface-Opaque.Substrate-StandardSurface-Opaque"));
+	UInterchangeFunctionCallShaderNode* StandardSurfaceShaderNode = CreateFunctionCallShaderNode(SurfaceShaderNode->getName().c_str(), TEXT("/Engine/Functions/Substrate/Substrate-StandardSurface-Opaque.Substrate-StandardSurface-Opaque"));
 
 	// Inputs
 	//Base
@@ -285,7 +285,7 @@ void FMaterialXStandardSurfaceShader::ConnectToSubstrateStandardSurface()
 	// Outputs
 	if(UInterchangeShaderPortsAPI::HasInput(StandardSurfaceShaderNode, StandardSurface::Parameters::Transmission))
 	{
-		StandardSurfaceShaderNode->SetCustomMaterialFunction(TEXT("/Engine/Functions/Strata/Substrate-StandardSurface-Translucent.Substrate-StandardSurface-Translucent"));
+		StandardSurfaceShaderNode->SetCustomMaterialFunction(TEXT("/Engine/Functions/Substrate/Substrate-StandardSurface-Translucent.Substrate-StandardSurface-Translucent"));
 		ShaderGraphNode->SetCustomBlendMode(EBlendMode::BLEND_TranslucentColoredTransmittance);
 		UInterchangeShaderPortsAPI::ConnectOuputToInputByName(ShaderGraphNode, PBRMR::Parameters::Opacity.ToString(), StandardSurfaceShaderNode->GetUniqueID(), PBRMR::Parameters::Opacity.ToString());
 		UInterchangeShaderPortsAPI::ConnectOuputToInputByName(ShaderGraphNode, SubstrateMaterial::Parameters::FrontMaterial.ToString(), StandardSurfaceShaderNode->GetUniqueID(), StandardSurface::SubstrateMaterial::Outputs::Translucent.ToString());
diff --git a/Engine/Plugins/Interchange/Runtime/Source/Import/Private/MaterialX/MaterialXUtils/MaterialXSurfaceShaderAbstract.cpp b/Engine/Plugins/Interchange/Runtime/Source/Import/Private/MaterialX/MaterialXUtils/MaterialXSurfaceShaderAbstract.cpp
--- a/Engine/Plugins/Interchange/Runtime/Source/Import/Private/MaterialX/MaterialXUtils/MaterialXSurfaceShaderAbstract.cpp
+++ b/Engine/Plugins/Interchange/Runtime/Source/Import/Private/MaterialX/MaterialXUtils/MaterialXSurfaceShaderAbstract.cpp
@@ -33,7 +33,11 @@ bool FMaterialXSurfaceShaderAbstract::AddAttribute(MaterialX::InputPtr Input, co
 	using namespace UE::Interchange::Materials::Standard::Nodes;
 	if(Input)
 	{
-		if(Input->getType() == mx::Type::Float)
+		if(Input->getType() == mx::Type::Boolean)
+		{
+			return AddBooleanAttribute(Input, InputChannelName, ShaderNode);
+		}
+		else if(Input->getType() == mx::Type::Float)
 		{
 			return AddFloatAttribute(Input, InputChannelName, ShaderNode);
 		}
@@ -97,6 +101,26 @@ bool FMaterialXSurfaceShaderAbstract::AddAttributeFromValueOrInterface(MaterialX
 	return bAttribute;
 }
 
+bool FMaterialXSurfaceShaderAbstract::AddBooleanAttribute(MaterialX::InputPtr Input, const FString& InputChannelName, UInterchangeShaderNode* ShaderNode)
+{
+	using namespace UE::Interchange::Materials::Standard::Nodes;
+	if(Input && Input->hasValue())
+	{
+		bool Value = mx::fromValueString<bool>(Input->getValueString());
+
+		// The parent is either a node, or it's an interfacename and we just take the name of the input
+		mx::NodePtr Node = Input->getParent()->asA<mx::Node>();
+		FString NodeName = Node ? Node->getName().c_str() + FString{ TEXT("_") } : FString{};
+		NodeName += Input->getName().c_str();
+
+		UInterchangeShaderNode* StaticBoolParameterNode = CreateShaderNode(NodeName, StaticBoolParameter::Name.ToString());
+		StaticBoolParameterNode->AddBooleanAttribute(UInterchangeShaderPortsAPI::MakeInputParameterKey(StaticBoolParameter::Attributes::DefaultValue.ToString()), Value);
+		return UInterchangeShaderPortsAPI::ConnectDefaultOuputToInput(ShaderNode, InputChannelName, StaticBoolParameterNode->GetUniqueID());
+	}
+
+	return false;
+}
+
 bool FMaterialXSurfaceShaderAbstract::AddFloatAttribute(MaterialX::InputPtr Input, const FString& InputChannelName, UInterchangeShaderNode* ShaderNode, float DefaultValue)
 {
 	using namespace UE::Interchange::Materials::Standard::Nodes;
@@ -112,7 +136,7 @@ bool FMaterialXSurfaceShaderAbstract::AddFloatAttribute(MaterialX::InputPtr Inpu
 			NodeName += Input->getName().c_str();
 
 			UInterchangeShaderNode* ScalarParameterNode = CreateShaderNode(NodeName, ScalarParameter::Name.ToString());
-			ScalarParameterNode->AddFloatAttribute(UInterchangeShaderPortsAPI::MakeInputParameterKey(ScalarParameter::Attributes::DefaultValue.ToString()), mx::fromValueString<float>(Input->getValueString()));
+			ScalarParameterNode->AddFloatAttribute(UInterchangeShaderPortsAPI::MakeInputParameterKey(ScalarParameter::Attributes::DefaultValue.ToString()), Value);
 			return UInterchangeShaderPortsAPI::ConnectDefaultOuputToInput(ShaderNode, InputChannelName, ScalarParameterNode->GetUniqueID());
 		}
 	}
@@ -638,16 +662,13 @@ void FMaterialXSurfaceShaderAbstract::ConnectIfGreaterInputToOutput(const FConne
 	mx::InputPtr Input = Connect.UpstreamNode->getInput("in2");
 	if(Input)
 	{
-		bool bHasValue = AddAttributeFromValueOrInterface(Input, GetInputName(Input), NodeIf);
+		AddAttributeFromValueOrInterface(Input, GetInputName(Input), NodeIf);
 		AddAttributeFromValueOrInterface(Input, TEXT("AEqualsB"), NodeIf);
 
-		if(bHasValue)
-		{
-			//Let's add a new input that is a copy of in2 to connect it to the equal input
-			mx::InputPtr Input3 = Connect.UpstreamNode->addInput("in3");
-			Input3->copyContentFrom(Input);
-			SetAttributeNewName(Input3, "AEqualsB");
-		}
+		//Let's add a new input that is a copy of in2 to connect it to the equal input
+		mx::InputPtr Input3 = Connect.UpstreamNode->addInput("in3");
+		Input3->copyContentFrom(Input);
+		SetAttributeNewName(Input3, "AEqualsB");
 	}
 }
 
@@ -675,16 +696,13 @@ void FMaterialXSurfaceShaderAbstract::ConnectIfGreaterEqInputToOutput(const FCon
 	mx::InputPtr Input = Connect.UpstreamNode->getInput("in1");
 	if(Input)
 	{
-		bool bHasValue = AddAttributeFromValueOrInterface(Input, GetInputName(Input), NodeIf);
+		AddAttributeFromValueOrInterface(Input, GetInputName(Input), NodeIf);
 		AddAttributeFromValueOrInterface(Input, TEXT("AEqualsB"), NodeIf);
 
-		if(bHasValue)
-		{
-			//Let's add a new input that is a copy of in2 to connect it to the equal input
-			mx::InputPtr Input3 = Connect.UpstreamNode->addInput("in3");
-			Input3->copyContentFrom(Input);
-			SetAttributeNewName(Input3, "AEqualsB");
-		}
+		//Let's add a new input that is a copy of in2 to connect it to the equal input
+		mx::InputPtr Input3 = Connect.UpstreamNode->addInput("in3");
+		Input3->copyContentFrom(Input);
+		SetAttributeNewName(Input3, "AEqualsB");
 	}
 }
 
@@ -693,17 +711,17 @@ void FMaterialXSurfaceShaderAbstract::ConnectIfEqualInputToOutput(const FConnect
 	UInterchangeShaderNode* NodeIf = CreateShaderNode(Connect.UpstreamNode->getName().c_str(), TEXT("If"));
 	UInterchangeShaderPortsAPI::ConnectDefaultOuputToInput(Connect.ParentShaderNode, Connect.InputChannelName, NodeIf->GetUniqueID());
 
-	if(mx::InputPtr Input = Connect.UpstreamNode->getInput("value1"); Input && Input->hasValue())
+	if(mx::InputPtr Input = Connect.UpstreamNode->getInput("value1"))
 	{
 		AddAttributeFromValueOrInterface(Input, GetInputName(Input), NodeIf);
 	}
 
-	if(mx::InputPtr Input = Connect.UpstreamNode->getInput("value2"); Input && Input->hasValue())
+	if(mx::InputPtr Input = Connect.UpstreamNode->getInput("value2"))
 	{
 		AddAttributeFromValueOrInterface(Input, GetInputName(Input), NodeIf);
 	}
 
-	if(mx::InputPtr Input = Connect.UpstreamNode->getInput("in1"); Input && Input->hasValue())
+	if(mx::InputPtr Input = Connect.UpstreamNode->getInput("in1"))
 	{
 		AddAttributeFromValueOrInterface(Input, GetInputName(Input), NodeIf);
 	}
@@ -712,16 +730,13 @@ void FMaterialXSurfaceShaderAbstract::ConnectIfEqualInputToOutput(const FConnect
 	mx::InputPtr Input = Connect.UpstreamNode->getInput("in2");
 	if(Input)
 	{
-		bool bHasValue = AddAttributeFromValueOrInterface(Input, GetInputName(Input), NodeIf);
+		AddAttributeFromValueOrInterface(Input, GetInputName(Input), NodeIf);
 		AddAttributeFromValueOrInterface(Input, TEXT("AGreaterThanB"), NodeIf);
 
-		if(bHasValue)
-		{
-			//Let's add a new input that is a copy of in2 to connect it to the equal input
-			mx::InputPtr Input3 = Connect.UpstreamNode->addInput("in3");
-			Input3->copyContentFrom(Input);
-			SetAttributeNewName(Input3, "AGreaterThanB");
-		}
+		//Let's add a new input that is a copy of in2 to connect it to the equal input
+		mx::InputPtr Input3 = Connect.UpstreamNode->addInput("in3");
+		Input3->copyContentFrom(Input);
+		SetAttributeNewName(Input3, "AGreaterThanB");
 	}
 }
 
diff --git a/Engine/Plugins/Interchange/Runtime/Source/Import/Public/MaterialX/MaterialXUtils/MaterialXSurfaceShaderAbstract.h b/Engine/Plugins/Interchange/Runtime/Source/Import/Public/MaterialX/MaterialXUtils/MaterialXSurfaceShaderAbstract.h
--- a/Engine/Plugins/Interchange/Runtime/Source/Import/Public/MaterialX/MaterialXUtils/MaterialXSurfaceShaderAbstract.h
+++ b/Engine/Plugins/Interchange/Runtime/Source/Import/Public/MaterialX/MaterialXUtils/MaterialXSurfaceShaderAbstract.h
@@ -62,6 +62,18 @@ protected:
 	 */
 	bool AddAttributeFromValueOrInterface(MaterialX::InputPtr Input, const FString& InputChannelName, UInterchangeShaderNode* ShaderNode);
 
+	/**
+	 * Add a bool attribute to a shader node only if its value taken from the input is not equal to its default value. Return false if the attribute does not exist or if we cannot add it
+	 *
+	 * @param Input - The MaterialX input to retrieve and add the value from, must be of type bool
+	 * @param InputChannelName - The name of the shader node's input to add the attribute
+	 * @param ShaderNode - The shader node to which we want to add the attribute
+	 * @param DefaultValue - the default value to test the input against
+	 *
+	 * @return true if the attribute was successfully added
+	 */
+	bool AddBooleanAttribute(MaterialX::InputPtr Input, const FString& InputChannelName, UInterchangeShaderNode* ShaderNode);
+
 	/**
 	 * Add a float attribute to a shader node only if its value taken from the input is not equal to its default value. Return false if the attribute does not exist or if we cannot add it
 	 *
@@ -105,7 +117,7 @@ protected:
 	 * @param ShaderNode - The Interchange shader node to connect the MaterialX's node or node graph to
 	 * @param InputShaderName - The name of the input of the shader node to connect to
 	 * @param DefaultValue - The default value of the MaterialX input
-	 * @param OptionalTextureCompression - Set the texture compression for all textures along the path of an input
+	 * @param bIsTangentSpaceInput - Set the tangent space along the path of an input
 	 */
 	template<typename T>
 	bool ConnectNodeOutputToInput(const char* InputName, UInterchangeShaderNode* ShaderNode, const FString& InputShaderName, T DefaultValue, bool bIsTangentSpaceInput = false)
diff --git a/Engine/Plugins/Interchange/Runtime/Source/Nodes/Public/InterchangeMaterialDefinitions.h b/Engine/Plugins/Interchange/Runtime/Source/Nodes/Public/InterchangeMaterialDefinitions.h
--- a/Engine/Plugins/Interchange/Runtime/Source/Nodes/Public/InterchangeMaterialDefinitions.h
+++ b/Engine/Plugins/Interchange/Runtime/Source/Nodes/Public/InterchangeMaterialDefinitions.h
@@ -336,6 +336,65 @@ namespace Interchange
 			}
 		}
 
+		namespace OpenPBRSurface
+		{
+			const FName Name = TEXT("open_pbr_surface");
+
+			namespace Parameters
+			{
+				const FName BaseWeight = TEXT("base_weight");
+				const FName BaseColor = TEXT("base_color");
+				const FName BaseRoughness = TEXT("base_roughness");
+				const FName BaseMetalness = TEXT("base_metalness");
+				const FName SpecularWeight = TEXT("specular_weight");
+				const FName SpecularColor = TEXT("specular_color");
+				const FName SpecularRoughness = TEXT("specular_roughness");
+				const FName SpecularIOR = TEXT("specular_ior");
+				const FName SpecularIORLevel = TEXT("specular_ior_level");
+				const FName SpecularAnisotropy = TEXT("specular_anisotropy");
+				const FName SpecularRotation = TEXT("specular_rotation");
+				const FName TransmissionWeight = TEXT("transmission_weight");
+				const FName TransmissionColor = TEXT("transmission_color");
+				const FName TransmissionDepth = TEXT("transmission_depth");
+				const FName TransmissionScatter = TEXT("transmission_scatter");
+				const FName TransmissionScatterAnisotropy = TEXT("transmission_scatter_anisotropy");
+				const FName TransmissionDispersion = TEXT("transmission_dispersion");
+				const FName SubsurfaceWeight = TEXT("subsurface_weight");
+				const FName SubsurfaceColor = TEXT("subsurface_color");
+				const FName SubsurfaceRadius = TEXT("subsurface_radius");
+				const FName SubsurfaceRadiusScale = TEXT("subsurface_radius_scale");
+				const FName SubsurfaceAnisotropy = TEXT("subsurface_anisotropy");
+				const FName FuzzWeight = TEXT("fuzz_weight");
+				const FName FuzzColor = TEXT("fuzz_color");
+				const FName FuzzRoughness = TEXT("fuzz_roughness");
+				const FName CoatWeight = TEXT("coat_weight");
+				const FName CoatColor = TEXT("coat_color");
+				const FName CoatRoughness = TEXT("coat_roughness");
+				const FName CoatAnisotropy = TEXT("coat_anisotropy");
+				const FName CoatRotation = TEXT("coat_rotation");
+				const FName CoatIOR = TEXT("coat_ior");
+				const FName CoatIORLevel = TEXT("coat_ior_level");
+				const FName ThinFilmThickness = TEXT("thin_film_thickness");
+				const FName ThinFilmIOR = TEXT("thin_film_ior");
+				const FName EmissionLuminance = TEXT("emission_luminance");
+				const FName EmissionColor = TEXT("emission_color");
+				const FName GeometryOpacity = TEXT("geometry_opacity");
+				const FName GeometryThinWalled = TEXT("geometry_thin_walled");
+				const FName GeometryNormal = TEXT("geometry_normal");
+				const FName GeometryCoatNormal = TEXT("geometry_coat_normal");
+				const FName GeometryTangent = TEXT("geometry_tangent");
+			}
+
+			namespace SubstrateMaterial
+			{
+				namespace Outputs
+				{
+					const FName FrontMaterial = TEXT("OpenPBR_FrontMaterial");
+					const FName OpacityMask = TEXT("OpacityMask");
+				}
+			}
+		}
+
 		namespace Phong
 		{
 			namespace Parameters
diff --git a/Engine/Plugins/Interchange/Runtime/Source/Pipelines/Private/InterchangeMaterialXPipeline.cpp b/Engine/Plugins/Interchange/Runtime/Source/Pipelines/Private/InterchangeMaterialXPipeline.cpp
--- a/Engine/Plugins/Interchange/Runtime/Source/Pipelines/Private/InterchangeMaterialXPipeline.cpp
+++ b/Engine/Plugins/Interchange/Runtime/Source/Pipelines/Private/InterchangeMaterialXPipeline.cpp
@@ -16,14 +16,19 @@
 	constexpr const TCHAR* Name##SubstratePath = TEXT("/Interchange/Substrate/") TEXT("MX_") TEXT(#Name) TEXT(".") TEXT("MX_") TEXT(#Name);
 
 #define MATERIALX_MATERIALFUNCTION_PATH(Name) \
-	MaterialXSettings->bIsSubstrateEnabled ? Name##FunctionsPath : Name##SubstratePath
+	!MaterialXSettings->bIsSubstrateEnabled ? Name##FunctionsPath : Name##SubstratePath
 
 namespace
 {
+	constexpr const TCHAR* OpenPBRSurfaceFunctionsPath = TEXT("/Interchange/Functions/MX_OpenPBR_Opaque.MX_OpenPBR_Opaque");
+	constexpr const TCHAR* OpenPBRSurfaceSubstratePath = TEXT("/Engine/Functions/Substrate/MF_Substrate_OpenPBR_Opaque.MF_Substrate_OpenPBR_Opaque");
+	constexpr const TCHAR* OpenPBRTransmissionSurfaceFunctionsPath = TEXT("/Interchange/Functions/MX_OpenPBR_Translucent.MX_OpenPBR_Translucent");
+	constexpr const TCHAR* OpenPBRTransmissionSurfaceSubstratePath = TEXT("/Engine/Functions/Substrate/MF_Substrate_OpenPBR_Translucent.MF_Substrate_OpenPBR_Translucent");
+
 	constexpr const TCHAR* StandardSurfaceFunctionsPath = TEXT("/Interchange/Functions/MX_StandardSurface.MX_StandardSurface");
-	constexpr const TCHAR* StandardSurfaceSubstratePath = TEXT("/Engine/Functions/Strata/Substrate-StandardSurface-Opaque.Substrate-StandardSurface-Opaque");
+	constexpr const TCHAR* StandardSurfaceSubstratePath = TEXT("/Engine/Functions/Substrate/Substrate-StandardSurface-Opaque.Substrate-StandardSurface-Opaque");
 	constexpr const TCHAR* TransmissionSurfaceFunctionsPath = TEXT("/Interchange/Functions/MX_TransmissionSurface.MX_TransmissionSurface");
-	constexpr const TCHAR* TransmissionSurfaceSubstratePath = TEXT("/Engine/Functions/Strata/Substrate-StandardSurface-Translucent.Substrate-StandardSurface-Translucent");
+	constexpr const TCHAR* TransmissionSurfaceSubstratePath = TEXT("/Engine/Functions/Substrate/Substrate-StandardSurface-Translucent.Substrate-StandardSurface-Translucent");
 	MATERIALX_FUNCTIONS_SUBSTRATE_PATH(SurfaceUnlit);
 	MATERIALX_FUNCTIONS_SUBSTRATE_PATH(Surface);
 	MATERIALX_FUNCTIONS_SUBSTRATE_PATH(UsdPreviewSurface);
@@ -60,6 +65,133 @@ UMaterialXPipelineSettings::UMaterialXPipelineSettings()
 
 		SettingsInputsOutputs = {
 			//Surface Shaders
+						{
+				UMaterialXPipelineSettings::ToEnumKey(EInterchangeMaterialXShaders::OpenPBRSurface),
+				{
+				// OpenPBRSurface Inputs
+				TSet<FName>{
+					UE::Interchange::Materials::OpenPBRSurface::Parameters::BaseWeight,
+					UE::Interchange::Materials::OpenPBRSurface::Parameters::BaseColor,
+					UE::Interchange::Materials::OpenPBRSurface::Parameters::BaseRoughness,
+					UE::Interchange::Materials::OpenPBRSurface::Parameters::BaseMetalness,
+					UE::Interchange::Materials::OpenPBRSurface::Parameters::SpecularWeight,
+					UE::Interchange::Materials::OpenPBRSurface::Parameters::SpecularColor,
+					UE::Interchange::Materials::OpenPBRSurface::Parameters::SpecularRoughness,
+					UE::Interchange::Materials::OpenPBRSurface::Parameters::SpecularIOR,
+					UE::Interchange::Materials::OpenPBRSurface::Parameters::SpecularIORLevel,
+					UE::Interchange::Materials::OpenPBRSurface::Parameters::SpecularAnisotropy,
+					UE::Interchange::Materials::OpenPBRSurface::Parameters::SpecularRotation,
+					UE::Interchange::Materials::OpenPBRSurface::Parameters::SubsurfaceWeight,
+					UE::Interchange::Materials::OpenPBRSurface::Parameters::SubsurfaceColor,
+					UE::Interchange::Materials::OpenPBRSurface::Parameters::SubsurfaceRadius,
+					UE::Interchange::Materials::OpenPBRSurface::Parameters::SubsurfaceRadiusScale,
+					UE::Interchange::Materials::OpenPBRSurface::Parameters::SubsurfaceAnisotropy,
+					UE::Interchange::Materials::OpenPBRSurface::Parameters::FuzzWeight,
+					UE::Interchange::Materials::OpenPBRSurface::Parameters::FuzzColor,
+					UE::Interchange::Materials::OpenPBRSurface::Parameters::FuzzRoughness,
+					UE::Interchange::Materials::OpenPBRSurface::Parameters::CoatWeight,
+					UE::Interchange::Materials::OpenPBRSurface::Parameters::CoatColor,
+					UE::Interchange::Materials::OpenPBRSurface::Parameters::CoatRoughness,
+					UE::Interchange::Materials::OpenPBRSurface::Parameters::CoatAnisotropy,
+					UE::Interchange::Materials::OpenPBRSurface::Parameters::CoatRotation,
+					UE::Interchange::Materials::OpenPBRSurface::Parameters::CoatIOR,
+					UE::Interchange::Materials::OpenPBRSurface::Parameters::CoatIORLevel,
+					UE::Interchange::Materials::OpenPBRSurface::Parameters::GeometryCoatNormal,
+					UE::Interchange::Materials::OpenPBRSurface::Parameters::ThinFilmThickness,
+					UE::Interchange::Materials::OpenPBRSurface::Parameters::ThinFilmIOR,
+					UE::Interchange::Materials::OpenPBRSurface::Parameters::EmissionLuminance,
+					UE::Interchange::Materials::OpenPBRSurface::Parameters::EmissionColor,
+					UE::Interchange::Materials::OpenPBRSurface::Parameters::GeometryNormal,
+					UE::Interchange::Materials::OpenPBRSurface::Parameters::GeometryTangent,
+					UE::Interchange::Materials::OpenPBRSurface::Parameters::GeometryOpacity,
+				},
+				// OpenPBRSurface Outputs
+				!bIsSubstrateEnabled ?
+				TSet<FName>{
+					UE::Interchange::Materials::PBRMR::Parameters::BaseColor,
+						UE::Interchange::Materials::PBRMR::Parameters::Metallic,
+						UE::Interchange::Materials::PBRMR::Parameters::Specular,
+						UE::Interchange::Materials::PBRMR::Parameters::Roughness,
+						UE::Interchange::Materials::PBRMR::Parameters::Anisotropy,
+						UE::Interchange::Materials::PBRMR::Parameters::EmissiveColor,
+						UE::Interchange::Materials::PBRMR::Parameters::Opacity,
+						UE::Interchange::Materials::PBRMR::Parameters::Normal,
+						UE::Interchange::Materials::PBRMR::Parameters::Tangent,
+						UE::Interchange::Materials::Sheen::Parameters::SheenRoughness,
+						UE::Interchange::Materials::Sheen::Parameters::SheenColor,
+						UE::Interchange::Materials::Subsurface::Parameters::SubsurfaceColor,
+						UE::Interchange::Materials::ClearCoat::Parameters::ClearCoat,
+						UE::Interchange::Materials::ClearCoat::Parameters::ClearCoatRoughness,
+						UE::Interchange::Materials::ClearCoat::Parameters::ClearCoatNormal
+					}	:
+					TSet<FName>{
+						UE::Interchange::Materials::OpenPBRSurface::SubstrateMaterial::Outputs::FrontMaterial,
+						UE::Interchange::Materials::OpenPBRSurface::SubstrateMaterial::Outputs::OpacityMask
+					}
+				}
+			},
+			{
+				UMaterialXPipelineSettings::ToEnumKey(EInterchangeMaterialXShaders::OpenPBRSurfaceTransmission),
+				{
+					// OpenPBRSurfaceTransmission Inputs
+					TSet<FName>{
+						UE::Interchange::Materials::OpenPBRSurface::Parameters::BaseWeight,
+						UE::Interchange::Materials::OpenPBRSurface::Parameters::BaseColor,
+						UE::Interchange::Materials::OpenPBRSurface::Parameters::BaseRoughness,
+						UE::Interchange::Materials::OpenPBRSurface::Parameters::BaseMetalness,
+						UE::Interchange::Materials::OpenPBRSurface::Parameters::SpecularWeight,
+						UE::Interchange::Materials::OpenPBRSurface::Parameters::SpecularColor,
+						UE::Interchange::Materials::OpenPBRSurface::Parameters::SpecularRoughness,
+						UE::Interchange::Materials::OpenPBRSurface::Parameters::SpecularIOR,
+						UE::Interchange::Materials::OpenPBRSurface::Parameters::SpecularIORLevel,
+						UE::Interchange::Materials::OpenPBRSurface::Parameters::SpecularAnisotropy,
+						UE::Interchange::Materials::OpenPBRSurface::Parameters::SpecularRotation,
+						UE::Interchange::Materials::OpenPBRSurface::Parameters::TransmissionWeight,
+						UE::Interchange::Materials::OpenPBRSurface::Parameters::TransmissionColor,
+						UE::Interchange::Materials::OpenPBRSurface::Parameters::TransmissionDepth,
+						UE::Interchange::Materials::OpenPBRSurface::Parameters::TransmissionDispersion,
+						UE::Interchange::Materials::OpenPBRSurface::Parameters::TransmissionScatter,
+						UE::Interchange::Materials::OpenPBRSurface::Parameters::TransmissionScatterAnisotropy,
+						UE::Interchange::Materials::OpenPBRSurface::Parameters::FuzzWeight,
+						UE::Interchange::Materials::OpenPBRSurface::Parameters::FuzzColor,
+						UE::Interchange::Materials::OpenPBRSurface::Parameters::FuzzRoughness,
+						UE::Interchange::Materials::OpenPBRSurface::Parameters::CoatWeight,
+						UE::Interchange::Materials::OpenPBRSurface::Parameters::CoatColor,
+						UE::Interchange::Materials::OpenPBRSurface::Parameters::CoatRoughness,
+						UE::Interchange::Materials::OpenPBRSurface::Parameters::CoatAnisotropy,
+						UE::Interchange::Materials::OpenPBRSurface::Parameters::CoatRotation,
+						UE::Interchange::Materials::OpenPBRSurface::Parameters::CoatIOR,
+						UE::Interchange::Materials::OpenPBRSurface::Parameters::CoatIORLevel,
+						UE::Interchange::Materials::OpenPBRSurface::Parameters::GeometryCoatNormal,
+						UE::Interchange::Materials::OpenPBRSurface::Parameters::ThinFilmThickness,
+						UE::Interchange::Materials::OpenPBRSurface::Parameters::ThinFilmIOR,
+						UE::Interchange::Materials::OpenPBRSurface::Parameters::EmissionLuminance,
+						UE::Interchange::Materials::OpenPBRSurface::Parameters::EmissionColor,
+						UE::Interchange::Materials::OpenPBRSurface::Parameters::GeometryNormal,
+						UE::Interchange::Materials::OpenPBRSurface::Parameters::GeometryTangent,
+						UE::Interchange::Materials::OpenPBRSurface::Parameters::GeometryOpacity,
+					},
+					// OpenPBRSurfaceTransmission Outputs
+					!bIsSubstrateEnabled ?
+					TSet<FName>{
+						UE::Interchange::Materials::PBRMR::Parameters::BaseColor,
+						UE::Interchange::Materials::PBRMR::Parameters::Metallic,
+						UE::Interchange::Materials::PBRMR::Parameters::Specular,
+						UE::Interchange::Materials::PBRMR::Parameters::Roughness,
+						UE::Interchange::Materials::PBRMR::Parameters::Anisotropy,
+						UE::Interchange::Materials::PBRMR::Parameters::EmissiveColor,
+						UE::Interchange::Materials::PBRMR::Parameters::Opacity,
+						UE::Interchange::Materials::PBRMR::Parameters::Normal,
+						UE::Interchange::Materials::PBRMR::Parameters::Tangent,
+						UE::Interchange::Materials::PBRMR::Parameters::Refraction,
+						UE::Interchange::Materials::ThinTranslucent::Parameters::TransmissionColor
+					}	:
+					TSet<FName>{
+						UE::Interchange::Materials::OpenPBRSurface::SubstrateMaterial::Outputs::FrontMaterial,
+						UE::Interchange::Materials::OpenPBRSurface::SubstrateMaterial::Outputs::OpacityMask
+					}
+				}
+			},
 			{
 				UMaterialXPipelineSettings::ToEnumKey(EInterchangeMaterialXShaders::StandardSurface),
 				{
@@ -497,28 +629,30 @@ UInterchangeMaterialXPipeline::UInterchangeMaterialXPipeline()
 	{
 #if WITH_EDITOR
 		PathToEnumMapping =	{
-			{MATERIALX_MATERIALFUNCTION_PATH(StandardSurface),		  UMaterialXPipelineSettings::ToEnumKey(EInterchangeMaterialXShaders::StandardSurface)},
-			{MATERIALX_MATERIALFUNCTION_PATH(TransmissionSurface),	  UMaterialXPipelineSettings::ToEnumKey(EInterchangeMaterialXShaders::StandardSurfaceTransmission)},
-			{MATERIALX_MATERIALFUNCTION_PATH(SurfaceUnlit),			  UMaterialXPipelineSettings::ToEnumKey(EInterchangeMaterialXShaders::SurfaceUnlit)},
-			{MATERIALX_MATERIALFUNCTION_PATH(Surface),				  UMaterialXPipelineSettings::ToEnumKey(EInterchangeMaterialXShaders::Surface)},
-			{MATERIALX_MATERIALFUNCTION_PATH(UsdPreviewSurface),	  UMaterialXPipelineSettings::ToEnumKey(EInterchangeMaterialXShaders::UsdPreviewSurface)},
-
-			{MATERIALX_MATERIALFUNCTION_PATH(OrenNayarBSDF),	      UMaterialXPipelineSettings::ToEnumKey(EInterchangeMaterialXBSDF::OrenNayarDiffuse)},
-			{MATERIALX_MATERIALFUNCTION_PATH(BurleyDiffuseBSDF),	  UMaterialXPipelineSettings::ToEnumKey(EInterchangeMaterialXBSDF::BurleyDiffuse)},
-			{MATERIALX_MATERIALFUNCTION_PATH(DielectricBSDF),		  UMaterialXPipelineSettings::ToEnumKey(EInterchangeMaterialXBSDF::Dielectric)},
-			{MATERIALX_MATERIALFUNCTION_PATH(ConductorBSDF),		  UMaterialXPipelineSettings::ToEnumKey(EInterchangeMaterialXBSDF::Conductor)},
-			{MATERIALX_MATERIALFUNCTION_PATH(SheenBSDF),			  UMaterialXPipelineSettings::ToEnumKey(EInterchangeMaterialXBSDF::Sheen)},
-			{MATERIALX_MATERIALFUNCTION_PATH(SubsurfaceBSDF),		  UMaterialXPipelineSettings::ToEnumKey(EInterchangeMaterialXBSDF::Subsurface)},
-			{MATERIALX_MATERIALFUNCTION_PATH(ThinFilmBSDF),			  UMaterialXPipelineSettings::ToEnumKey(EInterchangeMaterialXBSDF::ThinFilm)},
-			{MATERIALX_MATERIALFUNCTION_PATH(GeneralizedSchlickBSDF), UMaterialXPipelineSettings::ToEnumKey(EInterchangeMaterialXBSDF::GeneralizedSchlick)},
-			{MATERIALX_MATERIALFUNCTION_PATH(TranslucentBSDF),		  UMaterialXPipelineSettings::ToEnumKey(EInterchangeMaterialXBSDF::Translucent)},
-
-			{MATERIALX_MATERIALFUNCTION_PATH(UniformEDF),			  UMaterialXPipelineSettings::ToEnumKey(EInterchangeMaterialXEDF::Uniform)},
-			{MATERIALX_MATERIALFUNCTION_PATH(ConicalEDF),			  UMaterialXPipelineSettings::ToEnumKey(EInterchangeMaterialXEDF::Conical)},
-			{MATERIALX_MATERIALFUNCTION_PATH(MeasuredEDF),			  UMaterialXPipelineSettings::ToEnumKey(EInterchangeMaterialXEDF::Measured)},
-
-			{MATERIALX_MATERIALFUNCTION_PATH(AbsorptionVDF),		  UMaterialXPipelineSettings::ToEnumKey(EInterchangeMaterialXVDF::Absorption)},
-			{MATERIALX_MATERIALFUNCTION_PATH(AnisotropicVDF),		  UMaterialXPipelineSettings::ToEnumKey(EInterchangeMaterialXVDF::Anisotropic)},
+			{MATERIALX_MATERIALFUNCTION_PATH(OpenPBRSurface),		      UMaterialXPipelineSettings::ToEnumKey(EInterchangeMaterialXShaders::OpenPBRSurface)},
+			{MATERIALX_MATERIALFUNCTION_PATH(OpenPBRTransmissionSurface), UMaterialXPipelineSettings::ToEnumKey(EInterchangeMaterialXShaders::OpenPBRSurfaceTransmission)},
+			{MATERIALX_MATERIALFUNCTION_PATH(StandardSurface),			  UMaterialXPipelineSettings::ToEnumKey(EInterchangeMaterialXShaders::StandardSurface)},
+			{MATERIALX_MATERIALFUNCTION_PATH(TransmissionSurface),		  UMaterialXPipelineSettings::ToEnumKey(EInterchangeMaterialXShaders::StandardSurfaceTransmission)},
+			{MATERIALX_MATERIALFUNCTION_PATH(SurfaceUnlit),				  UMaterialXPipelineSettings::ToEnumKey(EInterchangeMaterialXShaders::SurfaceUnlit)},
+			{MATERIALX_MATERIALFUNCTION_PATH(Surface),					  UMaterialXPipelineSettings::ToEnumKey(EInterchangeMaterialXShaders::Surface)},
+			{MATERIALX_MATERIALFUNCTION_PATH(UsdPreviewSurface),		  UMaterialXPipelineSettings::ToEnumKey(EInterchangeMaterialXShaders::UsdPreviewSurface)},
+																		  
+			{MATERIALX_MATERIALFUNCTION_PATH(OrenNayarBSDF),			  UMaterialXPipelineSettings::ToEnumKey(EInterchangeMaterialXBSDF::OrenNayarDiffuse)},
+			{MATERIALX_MATERIALFUNCTION_PATH(BurleyDiffuseBSDF),		  UMaterialXPipelineSettings::ToEnumKey(EInterchangeMaterialXBSDF::BurleyDiffuse)},
+			{MATERIALX_MATERIALFUNCTION_PATH(DielectricBSDF),			  UMaterialXPipelineSettings::ToEnumKey(EInterchangeMaterialXBSDF::Dielectric)},
+			{MATERIALX_MATERIALFUNCTION_PATH(ConductorBSDF),			  UMaterialXPipelineSettings::ToEnumKey(EInterchangeMaterialXBSDF::Conductor)},
+			{MATERIALX_MATERIALFUNCTION_PATH(SheenBSDF),				  UMaterialXPipelineSettings::ToEnumKey(EInterchangeMaterialXBSDF::Sheen)},
+			{MATERIALX_MATERIALFUNCTION_PATH(SubsurfaceBSDF),			  UMaterialXPipelineSettings::ToEnumKey(EInterchangeMaterialXBSDF::Subsurface)},
+			{MATERIALX_MATERIALFUNCTION_PATH(ThinFilmBSDF),				  UMaterialXPipelineSettings::ToEnumKey(EInterchangeMaterialXBSDF::ThinFilm)},
+			{MATERIALX_MATERIALFUNCTION_PATH(GeneralizedSchlickBSDF),	  UMaterialXPipelineSettings::ToEnumKey(EInterchangeMaterialXBSDF::GeneralizedSchlick)},
+			{MATERIALX_MATERIALFUNCTION_PATH(TranslucentBSDF),			  UMaterialXPipelineSettings::ToEnumKey(EInterchangeMaterialXBSDF::Translucent)},
+																		  
+			{MATERIALX_MATERIALFUNCTION_PATH(UniformEDF),				  UMaterialXPipelineSettings::ToEnumKey(EInterchangeMaterialXEDF::Uniform)},
+			{MATERIALX_MATERIALFUNCTION_PATH(ConicalEDF),				  UMaterialXPipelineSettings::ToEnumKey(EInterchangeMaterialXEDF::Conical)},
+			{MATERIALX_MATERIALFUNCTION_PATH(MeasuredEDF),				  UMaterialXPipelineSettings::ToEnumKey(EInterchangeMaterialXEDF::Measured)},
+																		  
+			{MATERIALX_MATERIALFUNCTION_PATH(AbsorptionVDF),			  UMaterialXPipelineSettings::ToEnumKey(EInterchangeMaterialXVDF::Absorption)},
+			{MATERIALX_MATERIALFUNCTION_PATH(AnisotropicVDF),			  UMaterialXPipelineSettings::ToEnumKey(EInterchangeMaterialXVDF::Anisotropic)},
 		};
 
 		MaterialXSettings->InitPredefinedAssets();
@@ -657,6 +791,9 @@ void UMaterialXPipelineSettings::InitPredefinedAssets()
 	{
 		TArray<TTuple<EMaterialXSettings, FString, FString>> MappingToSubstrate
 		{
+			{UMaterialXPipelineSettings::ToEnumKey(EInterchangeMaterialXShaders::OpenPBRSurface), OpenPBRSurfaceFunctionsPath, OpenPBRSurfaceSubstratePath},
+			{UMaterialXPipelineSettings::ToEnumKey(EInterchangeMaterialXShaders::OpenPBRSurfaceTransmission), OpenPBRTransmissionSurfaceFunctionsPath, OpenPBRTransmissionSurfaceSubstratePath},
+			{UMaterialXPipelineSettings::ToEnumKey(EInterchangeMaterialXShaders::Surface), SurfaceFunctionsPath, SurfaceSubstratePath},
 			{UMaterialXPipelineSettings::ToEnumKey(EInterchangeMaterialXShaders::Surface), SurfaceFunctionsPath, SurfaceSubstratePath},
 			{UMaterialXPipelineSettings::ToEnumKey(EInterchangeMaterialXShaders::StandardSurface), StandardSurfaceFunctionsPath, StandardSurfaceSubstratePath},
 			{UMaterialXPipelineSettings::ToEnumKey(EInterchangeMaterialXShaders::StandardSurfaceTransmission), TransmissionSurfaceFunctionsPath, TransmissionSurfaceSubstratePath},
diff --git a/Engine/Plugins/Interchange/Runtime/Source/Pipelines/Public/InterchangeMaterialXPipeline.h b/Engine/Plugins/Interchange/Runtime/Source/Pipelines/Public/InterchangeMaterialXPipeline.h
--- a/Engine/Plugins/Interchange/Runtime/Source/Pipelines/Public/InterchangeMaterialXPipeline.h
+++ b/Engine/Plugins/Interchange/Runtime/Source/Pipelines/Public/InterchangeMaterialXPipeline.h
@@ -29,6 +29,12 @@ class UMaterialInterface;
 UENUM(BlueprintType)
 enum class EInterchangeMaterialXShaders : uint8
 {
+	/** Default settings for Open PBR Surface shader	*/
+	OpenPBRSurface,
+
+	/** Open PBR Surface shader	used for translucency*/
+	OpenPBRSurfaceTransmission,
+
 	/** Default settings for Autodesk's Standard Surface shader	*/
 	StandardSurface,
 
