Added logic to only disable FBX level import

diff --git a/Engine/Plugins/Interchange/Runtime/Source/Import/Private/Fbx/InterchangeFbxTranslator.cpp b/Engine/Plugins/Interchange/Runtime/Source/Import/Private/Fbx/InterchangeFbxTranslator.cpp
--- a/Engine/Plugins/Interchange/Runtime/Source/Import/Private/Fbx/InterchangeFbxTranslator.cpp
+++ b/Engine/Plugins/Interchange/Runtime/Source/Import/Private/Fbx/InterchangeFbxTranslator.cpp
@@ -33,6 +33,13 @@ static FAutoConsoleVariableRef CCvarInterchangeEnableFBXImport(
 	TEXT("Whether FBX support is enabled."),
 	ECVF_Default);
 
+static bool GInterchangeEnableFBXLevelImport = false;
+static FAutoConsoleVariableRef CCvarInterchangeEnableFBXLevelImport(
+	TEXT("Interchange.FeatureFlags.Import.FBX.ToLevel"),
+	GInterchangeEnableFBXLevelImport,
+	TEXT("Whether support for FBX level import is enabled."),
+	ECVF_Default);
+
 static bool GInterchangeFBXTranslatorUseWorker = false;
 static FAutoConsoleVariableRef CCvarInterchangeFBXTranslatorUseWorker(
 	TEXT("Interchange.FeatureFlags.Translator.UseWorker.FBX"),
@@ -138,7 +145,7 @@ bool UInterchangeFbxTranslator::IsThreadSafe() const
 
 EInterchangeTranslatorType UInterchangeFbxTranslator::GetTranslatorType() const
 {
-	return EInterchangeTranslatorType::Scenes;
+	return GInterchangeEnableFBXLevelImport ? EInterchangeTranslatorType::Scenes : EInterchangeTranslatorType::Assets;
 }
 
 EInterchangeTranslatorAssetType UInterchangeFbxTranslator::GetSupportedAssetTypes() const
diff --git a/Engine/Source/Developer/AssetTools/Private/AssetTools.cpp b/Engine/Source/Developer/AssetTools/Private/AssetTools.cpp
--- a/Engine/Source/Developer/AssetTools/Private/AssetTools.cpp
+++ b/Engine/Source/Developer/AssetTools/Private/AssetTools.cpp
@@ -3557,7 +3557,7 @@ TArray<UObject*> UAssetToolsImpl::ImportAssetsInternal(const TArray<FString>& Fi
 			{
 				UE::Interchange::FScopedSourceData ScopedSourceData(Filename);
 
-				if (!InterchangeManager.CanTranslateSourceData(ScopedSourceData.GetSourceData()))
+				if (!InterchangeManager.CanTranslateSourceData(ScopedSourceData.GetSourceData(), Params.bSceneImport))
 				{
 					bOnlyInterchangeImport = false;
 					break;
@@ -3640,7 +3640,7 @@ TArray<UObject*> UAssetToolsImpl::ImportAssetsInternal(const TArray<FString>& Fi
 		{
 			UE::Interchange::FScopedSourceData ScopedSourceData(Filename);
 
-			if (InterchangeManager.CanTranslateSourceData(ScopedSourceData.GetSourceData()))
+			if (InterchangeManager.CanTranslateSourceData(ScopedSourceData.GetSourceData(), Params.bSceneImport))
 			{
 				FImportAssetParameters ImportAssetParameters;
 				ImportAssetParameters.bIsAutomated = bAutomatedImport;
diff --git a/Engine/Source/Editor/UnrealEd/Private/Fbx/FbxFactory.cpp b/Engine/Source/Editor/UnrealEd/Private/Fbx/FbxFactory.cpp
--- a/Engine/Source/Editor/UnrealEd/Private/Fbx/FbxFactory.cpp
+++ b/Engine/Source/Editor/UnrealEd/Private/Fbx/FbxFactory.cpp
@@ -1196,7 +1196,8 @@ TArray<FString> UFbxFactory::GetFbxFormats(const UFactory* Factory)
 {
 	TArray<FString> FormatArray;
 	const IConsoleVariable* CVarFbx = IConsoleManager::Get().FindConsoleVariable(TEXT("Interchange.FeatureFlags.Import.FBX"));
-	const bool bUseLegacyFbx = (!CVarFbx || !CVarFbx->GetBool());
+	const IConsoleVariable* CVarFbxLevel = IConsoleManager::Get().FindConsoleVariable(TEXT("Interchange.FeatureFlags.Import.FBX.ToLevel"));
+	const bool bUseLegacyFbx = (!CVarFbx || !CVarFbx->GetBool()) || (!CVarFbxLevel || !CVarFbxLevel->GetBool());
 	const IConsoleVariable* CVarObj = IConsoleManager::Get().FindConsoleVariable(TEXT("Interchange.FeatureFlags.Import.OBJ"));
 	const bool bUseLegacyObj = (!CVarObj || !CVarObj->GetBool());
 
diff --git a/Engine/Source/Runtime/Interchange/Engine/Private/InterchangeManager.cpp b/Engine/Source/Runtime/Interchange/Engine/Private/InterchangeManager.cpp
--- a/Engine/Source/Runtime/Interchange/Engine/Private/InterchangeManager.cpp
+++ b/Engine/Source/Runtime/Interchange/Engine/Private/InterchangeManager.cpp
@@ -1220,7 +1220,7 @@ TArray<FString> UInterchangeManager::GetSupportedFormatsForObject(const UObject*
 	return FileExtensions;
 }
 
-bool UInterchangeManager::CanTranslateSourceData(const UInterchangeSourceData* SourceData) const
+bool UInterchangeManager::CanTranslateSourceData(const UInterchangeSourceData* SourceData, bool bSceneImportOnly) const
 {
 	if (!IsInterchangeImportEnabled())
 	{
@@ -1238,7 +1238,7 @@ bool UInterchangeManager::CanTranslateSourceData(const UInterchangeSourceData* S
 	if (UInterchangeTranslatorBase* Translator = GetTranslatorForSourceData(SourceData))
 	{
 		Translator->ReleaseSource();
-		return true;
+		return bSceneImportOnly ? Translator->GetTranslatorType() == EInterchangeTranslatorType::Scenes : true;
 	}
 
 	return false;
diff --git a/Engine/Source/Runtime/Interchange/Engine/Public/InterchangeManager.h b/Engine/Source/Runtime/Interchange/Engine/Public/InterchangeManager.h
--- a/Engine/Source/Runtime/Interchange/Engine/Public/InterchangeManager.h
+++ b/Engine/Source/Runtime/Interchange/Engine/Public/InterchangeManager.h
@@ -524,9 +524,12 @@ public:
 	 * Look if there is a registered translator for this source data.
 	 * This allow us to by pass the original asset tools system to import supported asset.
 	 * @Param SourceData - The source data input we want to translate to Uod
-	 * @return True if there is a registered translator that can handle handle this source data, false otherwise.
+	 * @Param bSceneImportOnly - Enforces a check on the type of translator
+	 * @return True if there is a registered translator that can handle handle this source data, false otherwise, when bSceneImportOnly is false.
+	 * Otherwise, returns true only if the translator supports level import
+	 * @Note: Temporary until FBX level import is production ready
 	 */
-	INTERCHANGEENGINE_API bool CanTranslateSourceData(const UInterchangeSourceData* SourceData) const;
+	INTERCHANGEENGINE_API bool CanTranslateSourceData(const UInterchangeSourceData* SourceData, bool bSceneImportOnly = false) const;
 
 	/**
 	 * Returns true if Interchange can create that type of assets and is able to translate its source file.
