[Interchange] Unexpected weights on skeletal mesh when using t0 as a ref pose

#rnx

diff --git a/Engine/Plugins/Interchange/Runtime/Source/Import/Private/Mesh/InterchangeSkeletalMeshFactory.cpp b/Engine/Plugins/Interchange/Runtime/Source/Import/Private/Mesh/InterchangeSkeletalMeshFactory.cpp
--- a/Engine/Plugins/Interchange/Runtime/Source/Import/Private/Mesh/InterchangeSkeletalMeshFactory.cpp
+++ b/Engine/Plugins/Interchange/Runtime/Source/Import/Private/Mesh/InterchangeSkeletalMeshFactory.cpp
@@ -357,12 +357,9 @@ namespace UE
 						continue;
 					}
 
-					//Get the mesh transform in local relative to the bind pose transform
- 					const FTransform MeshTransformRelativeToBindPoseTransform = MeshGlobalTransform * JointBindPoseGlobalTransform.Inverse();
-					//Get the time zero pose transform in local relative to the mesh transform
- 					const FTransform TimeZeroTransformRelativeToMeshTransform = JointTimeZeroGlobalTransform * MeshGlobalTransformInverse;
+					//JointBindPoseGlobalTransform Contains the GlobalTransform as the GlobalTransform calculation fallsback onto Transform (from BindPoseTransform or T0Transforms if they are not present)
 					//Multiply both transform to get a matrix that will transform the mesh vertices from the bind pose skinning to the time zero skinning
- 					const FMatrix VertexTransformMatrix = (MeshTransformRelativeToBindPoseTransform * TimeZeroTransformRelativeToMeshTransform).ToMatrixWithScale();
+					const FMatrix VertexTransformMatrix = (JointBindPoseGlobalTransform.Inverse() * JointTimeZeroGlobalTransform).ToMatrixWithScale();
 
 					//Iterate all bone vertices
 					for (FVertexID VertexID : MeshDescription.Vertices().GetElementIDs())
