[Interchange] Scene import - RE-Import - Crash with GLTF file ( EXCEPTION_ACCESS_VIOLATION)

#rnx

diff --git a/Engine/Plugins/Interchange/Runtime/Source/Import/Private/Mesh/InterchangeStaticMeshFactory.cpp b/Engine/Plugins/Interchange/Runtime/Source/Import/Private/Mesh/InterchangeStaticMeshFactory.cpp
--- a/Engine/Plugins/Interchange/Runtime/Source/Import/Private/Mesh/InterchangeStaticMeshFactory.cpp
+++ b/Engine/Plugins/Interchange/Runtime/Source/Import/Private/Mesh/InterchangeStaticMeshFactory.cpp
@@ -190,8 +190,9 @@ UInterchangeFactoryBase::FImportAssetResult UInterchangeStaticMeshFactory::Begin
 			FStaticMeshComponentRecreateRenderStateContext RecreateRenderStateContext(StaticMesh, bInvalidateLighting, bRefreshBounds);
 			StaticMesh->ReleaseResources();
 			StaticMesh->ReleaseResourcesFence.Wait();
-
+			
 			StaticMesh->SetRenderData(nullptr);
+			StaticMesh->SetBodySetup(nullptr);
 		}
 	}
 	
