Added source name to messages logged in MessageLog window
  - Added AssetFriendlyName on UInterchangeResult to pass Actor's label which can differ from its name
  - Added GetMessageLogText method to UInterchangeResult to centralize message formatting
  - Updated UE::Interchange::InterchangeEditorModule::LogError to use new GetMessageLogText method
  - Added SourceData to FImportSceneObjectsParams to allow actor factories to log complete message for warning and errors
  - Updated FTaskCreateSceneObjects::DoTask to pass SourceData in its parameters
  - Modified UInterchangeActorFactory::ProcessActor signature accordingly
  - Added utility methods to UInterchangeActorFactory to log error or warning messages
  - Updated actor's factories according to signature change
  - Updated TaskCompletion to add actor's label to successful message

diff --git a/Engine/Plugins/Experimental/Enterprise/DatasmithInterchange/Source/DatasmithInterchange/Private/InterchangeDatasmithAreaLightFactory.cpp b/Engine/Plugins/Experimental/Enterprise/DatasmithInterchange/Source/DatasmithInterchange/Private/InterchangeDatasmithAreaLightFactory.cpp
--- a/Engine/Plugins/Experimental/Enterprise/DatasmithInterchange/Source/DatasmithInterchange/Private/InterchangeDatasmithAreaLightFactory.cpp
+++ b/Engine/Plugins/Experimental/Enterprise/DatasmithInterchange/Source/DatasmithInterchange/Private/InterchangeDatasmithAreaLightFactory.cpp
@@ -29,7 +29,7 @@ UClass* UInterchangeDatasmithAreaLightFactory::GetFactoryClass() const
 	return LightShapeBlueprint->GeneratedClass;
 }
 
-UObject* UInterchangeDatasmithAreaLightFactory::ProcessActor(AActor& SpawnedActor, const UInterchangeActorFactoryNode& FactoryNode, const UInterchangeBaseNodeContainer& NodeContainer)
+UObject* UInterchangeDatasmithAreaLightFactory::ProcessActor(AActor& SpawnedActor, const UInterchangeActorFactoryNode& FactoryNode, const UInterchangeBaseNodeContainer& NodeContainer, const FImportSceneObjectsParams& Params)
 {
 	//using namespace UE::Interchange::ActorHelper;
 
@@ -82,5 +82,5 @@ UObject* UInterchangeDatasmithAreaLightFactory::ProcessActor(AActor& SpawnedActo
 #endif //WITH_EDITOR
 	}
 
-	return Super::ProcessActor(SpawnedActor, FactoryNode, NodeContainer);
+	return Super::ProcessActor(SpawnedActor, FactoryNode, NodeContainer, Params);
 };
\ No newline at end of file
diff --git a/Engine/Plugins/Experimental/Enterprise/DatasmithInterchange/Source/DatasmithInterchange/Public/InterchangeDatasmithAreaLightFactory.h b/Engine/Plugins/Experimental/Enterprise/DatasmithInterchange/Source/DatasmithInterchange/Public/InterchangeDatasmithAreaLightFactory.h
--- a/Engine/Plugins/Experimental/Enterprise/DatasmithInterchange/Source/DatasmithInterchange/Public/InterchangeDatasmithAreaLightFactory.h
+++ b/Engine/Plugins/Experimental/Enterprise/DatasmithInterchange/Source/DatasmithInterchange/Public/InterchangeDatasmithAreaLightFactory.h
@@ -28,7 +28,7 @@ protected:
 	//////////////////////////////////////////////////////////////////////////
 	// Interchange actor factory interface begin
 
-	virtual UObject* ProcessActor(class AActor& SpawnedActor, const UInterchangeActorFactoryNode& FactoryNode, const UInterchangeBaseNodeContainer& NodeContainer) override;
+	virtual UObject* ProcessActor(class AActor& SpawnedActor, const UInterchangeActorFactoryNode& FactoryNode, const UInterchangeBaseNodeContainer& NodeContainer, const FImportSceneObjectsParams& Params) override;
 
 	// Interchange actor factory interface end
 	//////////////////////////////////////////////////////////////////////////
diff --git a/Engine/Plugins/Interchange/Editor/Source/InterchangeEditor/Private/InterchangeEditorModule.cpp b/Engine/Plugins/Interchange/Editor/Source/InterchangeEditor/Private/InterchangeEditorModule.cpp
--- a/Engine/Plugins/Interchange/Editor/Source/InterchangeEditor/Private/InterchangeEditorModule.cpp
+++ b/Engine/Plugins/Interchange/Editor/Source/InterchangeEditor/Private/InterchangeEditorModule.cpp
@@ -56,7 +56,7 @@ namespace UE::Interchange::InterchangeEditorModule
 				const EInterchangeResultType ResultType = Result->GetResultType();
 				if (ResultType != EInterchangeResultType::Success)
 				{
-					TokenizedMessages.Add(FTokenizedMessage::Create(ResultType == EInterchangeResultType::Error ? EMessageSeverity::Error : EMessageSeverity::Warning, Result->GetText()));
+					TokenizedMessages.Add(FTokenizedMessage::Create(ResultType == EInterchangeResultType::Error ? EMessageSeverity::Error : EMessageSeverity::Warning, Result->GetMessageLogText()));
 					bHasErrors |= ResultType == EInterchangeResultType::Error;
 				}
 			}
diff --git a/Engine/Plugins/Interchange/Runtime/Source/Import/Private/Scene/InterchangeActorFactory.cpp b/Engine/Plugins/Interchange/Runtime/Source/Import/Private/Scene/InterchangeActorFactory.cpp
--- a/Engine/Plugins/Interchange/Runtime/Source/Import/Private/Scene/InterchangeActorFactory.cpp
+++ b/Engine/Plugins/Interchange/Runtime/Source/Import/Private/Scene/InterchangeActorFactory.cpp
@@ -29,7 +29,7 @@ UObject* UInterchangeActorFactory::ImportSceneObject_GameThread(const UInterchan
 
 	if (SpawnedActor)
 	{
-		if (UObject* ObjectToUpdate = ProcessActor(*SpawnedActor, *FactoryNode, *CreateSceneObjectsParams.NodeContainer))
+		if (UObject* ObjectToUpdate = ProcessActor(*SpawnedActor, *FactoryNode, *CreateSceneObjectsParams.NodeContainer, CreateSceneObjectsParams))
 		{
 			if (USceneComponent* RootComponent = SpawnedActor->GetRootComponent())
 			{
@@ -51,8 +51,7 @@ UObject* UInterchangeActorFactory::ImportSceneObject_GameThread(const UInterchan
 	return SpawnedActor;
 }
 
-UObject* UInterchangeActorFactory::ProcessActor(AActor& SpawnedActor, const UInterchangeActorFactoryNode& /*FactoryNode*/, const UInterchangeBaseNodeContainer& /*NodeContainer*/)
+UObject* UInterchangeActorFactory::ProcessActor(AActor& SpawnedActor, const UInterchangeActorFactoryNode& /*FactoryNode*/, const UInterchangeBaseNodeContainer& /*NodeContainer*/, const FImportSceneObjectsParams& Params)
 {
 	return SpawnedActor.GetRootComponent();
 }
-
diff --git a/Engine/Plugins/Interchange/Runtime/Source/Import/Private/Scene/InterchangeCameraActorFactory.cpp b/Engine/Plugins/Interchange/Runtime/Source/Import/Private/Scene/InterchangeCameraActorFactory.cpp
--- a/Engine/Plugins/Interchange/Runtime/Source/Import/Private/Scene/InterchangeCameraActorFactory.cpp
+++ b/Engine/Plugins/Interchange/Runtime/Source/Import/Private/Scene/InterchangeCameraActorFactory.cpp
@@ -19,7 +19,7 @@ UClass* UInterchangeCineCameraActorFactory::GetFactoryClass() const
 	return ACineCameraActor::StaticClass();
 }
 
-UObject* UInterchangeCineCameraActorFactory::ProcessActor(AActor& SpawnedActor, const UInterchangeActorFactoryNode& /*FactoryNode*/, const UInterchangeBaseNodeContainer& /*NodeContainer*/)
+UObject* UInterchangeCineCameraActorFactory::ProcessActor(AActor& SpawnedActor, const UInterchangeActorFactoryNode& /*FactoryNode*/, const UInterchangeBaseNodeContainer& /*NodeContainer*/, const FImportSceneObjectsParams& /*Params*/)
 {
 	ACineCameraActor* CineCameraActor = Cast<ACineCameraActor>(&SpawnedActor);
 
@@ -32,7 +32,7 @@ UClass* UInterchangeCameraActorFactory::GetFactoryClass() const
 	return ACameraActor::StaticClass();
 }
 
-UObject* UInterchangeCameraActorFactory::ProcessActor(AActor& SpawnedActor, const UInterchangeActorFactoryNode& /*FactoryNode*/, const UInterchangeBaseNodeContainer& /*NodeContainer*/)
+UObject* UInterchangeCameraActorFactory::ProcessActor(AActor& SpawnedActor, const UInterchangeActorFactoryNode& /*FactoryNode*/, const UInterchangeBaseNodeContainer& /*NodeContainer*/, const FImportSceneObjectsParams& /*Params*/)
 {
 	ACameraActor* CameraActor = Cast<ACameraActor>(&SpawnedActor);
 
diff --git a/Engine/Plugins/Interchange/Runtime/Source/Import/Private/Scene/InterchangeDecalActorFactory.cpp b/Engine/Plugins/Interchange/Runtime/Source/Import/Private/Scene/InterchangeDecalActorFactory.cpp
--- a/Engine/Plugins/Interchange/Runtime/Source/Import/Private/Scene/InterchangeDecalActorFactory.cpp
+++ b/Engine/Plugins/Interchange/Runtime/Source/Import/Private/Scene/InterchangeDecalActorFactory.cpp
@@ -17,7 +17,7 @@ UClass* UInterchangeDecalActorFactory::GetFactoryClass() const
 	return ADecalActor::StaticClass();
 }
 
-UObject* UInterchangeDecalActorFactory::ProcessActor(AActor& SpawnedActor, const UInterchangeActorFactoryNode& FactoryNode, const UInterchangeBaseNodeContainer& NodeContainer)
+UObject* UInterchangeDecalActorFactory::ProcessActor(AActor& SpawnedActor, const UInterchangeActorFactoryNode& FactoryNode, const UInterchangeBaseNodeContainer& NodeContainer, const FImportSceneObjectsParams& Params)
 {
 	ADecalActor* DecalActor = Cast<ADecalActor>(&SpawnedActor);
 	if (DecalActor)
@@ -51,7 +51,8 @@ UObject* UInterchangeDecalActorFactory::ProcessActor(AActor& SpawnedActor, const
 				}
 				else
 				{
-					UE_LOG(LogInterchangeImport, Error, TEXT("No valid decal material found. Make sure that the DecalMaterialPath is valid: %s"), *DecalMaterialPath);
+					const FText Message = FText::Format(NSLOCTEXT("DecalActorImport", "NoMaterial", "No valid decal material found. Make sure that the DecalMaterialPath is valid: %s"), FText::FromString(DecalMaterialPath));
+					LogMessage<UInterchangeResultError_Generic>(Params, Message, SpawnedActor.GetActorLabel());
 				}
 			}
 		}
diff --git a/Engine/Plugins/Interchange/Runtime/Source/Import/Private/Scene/InterchangeLightActorFactory.cpp b/Engine/Plugins/Interchange/Runtime/Source/Import/Private/Scene/InterchangeLightActorFactory.cpp
--- a/Engine/Plugins/Interchange/Runtime/Source/Import/Private/Scene/InterchangeLightActorFactory.cpp
+++ b/Engine/Plugins/Interchange/Runtime/Source/Import/Private/Scene/InterchangeLightActorFactory.cpp
@@ -18,7 +18,7 @@ UClass* UInterchangeLightActorFactory::GetFactoryClass() const
 	return ALight::StaticClass();
 }
 
-UObject* UInterchangeLightActorFactory::ProcessActor(AActor& SpawnedActor, const UInterchangeActorFactoryNode& FactoryNode, const UInterchangeBaseNodeContainer& NodeContainer)
+UObject* UInterchangeLightActorFactory::ProcessActor(AActor& SpawnedActor, const UInterchangeActorFactoryNode& FactoryNode, const UInterchangeBaseNodeContainer& NodeContainer, const FImportSceneObjectsParams& /*Params*/)
 {
 	if (ALight* LightActor = Cast<ALight>(&SpawnedActor))
 	{
diff --git a/Engine/Plugins/Interchange/Runtime/Source/Import/Private/Scene/InterchangeSkeletalMeshActorFactory.cpp b/Engine/Plugins/Interchange/Runtime/Source/Import/Private/Scene/InterchangeSkeletalMeshActorFactory.cpp
--- a/Engine/Plugins/Interchange/Runtime/Source/Import/Private/Scene/InterchangeSkeletalMeshActorFactory.cpp
+++ b/Engine/Plugins/Interchange/Runtime/Source/Import/Private/Scene/InterchangeSkeletalMeshActorFactory.cpp
@@ -15,7 +15,7 @@
 #include UE_INLINE_GENERATED_CPP_BY_NAME(InterchangeSkeletalMeshActorFactory)
 
 
-UObject* UInterchangeSkeletalMeshActorFactory::ProcessActor(AActor& SpawnedActor, const UInterchangeActorFactoryNode& /*FactoryNode*/, const UInterchangeBaseNodeContainer& /*NodeContainer*/)
+UObject* UInterchangeSkeletalMeshActorFactory::ProcessActor(AActor& SpawnedActor, const UInterchangeActorFactoryNode& /*FactoryNode*/, const UInterchangeBaseNodeContainer& /*NodeContainer*/, const FImportSceneObjectsParams& /*Params*/)
 {
 	ASkeletalMeshActor* SkeletalMeshActor = Cast<ASkeletalMeshActor>(&SpawnedActor);
 
diff --git a/Engine/Plugins/Interchange/Runtime/Source/Import/Private/Scene/InterchangeStaticMeshActorFactory.cpp b/Engine/Plugins/Interchange/Runtime/Source/Import/Private/Scene/InterchangeStaticMeshActorFactory.cpp
--- a/Engine/Plugins/Interchange/Runtime/Source/Import/Private/Scene/InterchangeStaticMeshActorFactory.cpp
+++ b/Engine/Plugins/Interchange/Runtime/Source/Import/Private/Scene/InterchangeStaticMeshActorFactory.cpp
@@ -16,7 +16,7 @@ UClass* UInterchangeStaticMeshActorFactory::GetFactoryClass() const
 	return AStaticMeshActor::StaticClass();
 }
 
-UObject* UInterchangeStaticMeshActorFactory::ProcessActor(AActor& SpawnedActor, const UInterchangeActorFactoryNode& FactoryNode, const UInterchangeBaseNodeContainer& NodeContainer)
+UObject* UInterchangeStaticMeshActorFactory::ProcessActor(AActor& SpawnedActor, const UInterchangeActorFactoryNode& FactoryNode, const UInterchangeBaseNodeContainer& NodeContainer, const FImportSceneObjectsParams& /*Params*/)
 {
 	using namespace UE::Interchange;
 
diff --git a/Engine/Plugins/Interchange/Runtime/Source/Import/Public/Scene/InterchangeActorFactory.h b/Engine/Plugins/Interchange/Runtime/Source/Import/Public/Scene/InterchangeActorFactory.h
--- a/Engine/Plugins/Interchange/Runtime/Source/Import/Public/Scene/InterchangeActorFactory.h
+++ b/Engine/Plugins/Interchange/Runtime/Source/Import/Public/Scene/InterchangeActorFactory.h
@@ -4,6 +4,8 @@
 
 #include "CoreMinimal.h"
 #include "InterchangeFactoryBase.h"
+#include "InterchangeSourceData.h"
+//#include "InterchangeResult.h"
 #include "UObject/Object.h"
 #include "UObject/ObjectMacros.h"
 
@@ -36,7 +38,28 @@ protected:
 	 * child class to complete the creation of the actor.
 	 * This method is expected to return the UOBject to apply the factory node's custom attributes to.
 	 */
-	virtual UObject* ProcessActor(AActor& SpawnedActor, const UInterchangeActorFactoryNode& FactoryNode, const UInterchangeBaseNodeContainer& NodeContainer);
+	virtual UObject* ProcessActor(AActor& SpawnedActor, const UInterchangeActorFactoryNode& FactoryNode, const UInterchangeBaseNodeContainer& NodeContainer, const FImportSceneObjectsParams& Params);
+
+	template <typename T>
+	void LogMessage(const FImportSceneObjectsParams& Params, const FText& Message, const FString& ActorLabel)
+	{
+		T* Item = /*Super::*/AddMessage<T>();
+		FillMessage(Params, Message, ActorLabel, *Item);
+	}
+
+	template <typename T>
+	void FillMessage(const FImportSceneObjectsParams& Params, const FText& Message, const FString& ActorLabel, T& Result)
+	{
+		Result.SourceAssetName = Params.SourceData ? Params.SourceData->GetFilename() : TEXT("Unknown file");
+		Result.DestinationAssetName = Params.ObjectName;
+		Result.AssetType = GetFactoryClass();
+		Result.Text = Message;
+
+		if (!ActorLabel.IsEmpty())
+		{
+			Result.AssetFriendlyName = ActorLabel;
+		}
+	}
 };
 
 
diff --git a/Engine/Plugins/Interchange/Runtime/Source/Import/Public/Scene/InterchangeCameraActorFactory.h b/Engine/Plugins/Interchange/Runtime/Source/Import/Public/Scene/InterchangeCameraActorFactory.h
--- a/Engine/Plugins/Interchange/Runtime/Source/Import/Public/Scene/InterchangeCameraActorFactory.h
+++ b/Engine/Plugins/Interchange/Runtime/Source/Import/Public/Scene/InterchangeCameraActorFactory.h
@@ -26,7 +26,7 @@ public:
 	//////////////////////////////////////////////////////////////////////////
 	// Interchange actor factory interface begin
 
-	virtual UObject* ProcessActor(class AActor& SpawnedActor, const UInterchangeActorFactoryNode& FactoryNode, const UInterchangeBaseNodeContainer& NodeContainer) override;
+	virtual UObject* ProcessActor(class AActor& SpawnedActor, const UInterchangeActorFactoryNode& FactoryNode, const UInterchangeBaseNodeContainer& NodeContainer, const FImportSceneObjectsParams& Params) override;
 
 
 	// Interchange actor factory interface end
@@ -50,7 +50,7 @@ public:
 	//////////////////////////////////////////////////////////////////////////
 	// Interchange actor factory interface begin
 
-	virtual UObject* ProcessActor(class AActor& SpawnedActor, const UInterchangeActorFactoryNode& FactoryNode, const UInterchangeBaseNodeContainer& NodeContainer) override;
+	virtual UObject* ProcessActor(class AActor& SpawnedActor, const UInterchangeActorFactoryNode& FactoryNode, const UInterchangeBaseNodeContainer& NodeContainer, const FImportSceneObjectsParams& Params) override;
 
 
 	// Interchange actor factory interface end
diff --git a/Engine/Plugins/Interchange/Runtime/Source/Import/Public/Scene/InterchangeDecalActorFactory.h b/Engine/Plugins/Interchange/Runtime/Source/Import/Public/Scene/InterchangeDecalActorFactory.h
--- a/Engine/Plugins/Interchange/Runtime/Source/Import/Public/Scene/InterchangeDecalActorFactory.h
+++ b/Engine/Plugins/Interchange/Runtime/Source/Import/Public/Scene/InterchangeDecalActorFactory.h
@@ -25,7 +25,7 @@ public:
 	//////////////////////////////////////////////////////////////////////////
 	// Interchange actor factory interface begin
 
-	virtual UObject* ProcessActor(class AActor& SpawnedActor, const UInterchangeActorFactoryNode& FactoryNode, const UInterchangeBaseNodeContainer& NodeContainer) override;
+	virtual UObject* ProcessActor(class AActor& SpawnedActor, const UInterchangeActorFactoryNode& FactoryNode, const UInterchangeBaseNodeContainer& NodeContainer, const FImportSceneObjectsParams& Params) override;
 
 	// Interchange actor factory interface end
 	//////////////////////////////////////////////////////////////////////////
diff --git a/Engine/Plugins/Interchange/Runtime/Source/Import/Public/Scene/InterchangeLightActorFactory.h b/Engine/Plugins/Interchange/Runtime/Source/Import/Public/Scene/InterchangeLightActorFactory.h
--- a/Engine/Plugins/Interchange/Runtime/Source/Import/Public/Scene/InterchangeLightActorFactory.h
+++ b/Engine/Plugins/Interchange/Runtime/Source/Import/Public/Scene/InterchangeLightActorFactory.h
@@ -24,7 +24,7 @@ protected:
 	//////////////////////////////////////////////////////////////////////////
 	// Interchange actor factory interface begin
 
-	virtual UObject* ProcessActor(class AActor& SpawnedActor, const UInterchangeActorFactoryNode& FactoryNode, const UInterchangeBaseNodeContainer& NodeContainer) override;
+	virtual UObject* ProcessActor(class AActor& SpawnedActor, const UInterchangeActorFactoryNode& FactoryNode, const UInterchangeBaseNodeContainer& NodeContainer, const FImportSceneObjectsParams& Params) override;
 
 
 	// Interchange actor factory interface end
diff --git a/Engine/Plugins/Interchange/Runtime/Source/Import/Public/Scene/InterchangeSkeletalMeshActorFactory.h b/Engine/Plugins/Interchange/Runtime/Source/Import/Public/Scene/InterchangeSkeletalMeshActorFactory.h
--- a/Engine/Plugins/Interchange/Runtime/Source/Import/Public/Scene/InterchangeSkeletalMeshActorFactory.h
+++ b/Engine/Plugins/Interchange/Runtime/Source/Import/Public/Scene/InterchangeSkeletalMeshActorFactory.h
@@ -33,7 +33,7 @@ protected:
 	//////////////////////////////////////////////////////////////////////////
 	// Interchange actor factory interface begin
 
-	virtual UObject* ProcessActor(class AActor& SpawnedActor, const UInterchangeActorFactoryNode& FactoryNode, const UInterchangeBaseNodeContainer& NodeContainer) override;
+	virtual UObject* ProcessActor(class AActor& SpawnedActor, const UInterchangeActorFactoryNode& FactoryNode, const UInterchangeBaseNodeContainer& NodeContainer, const FImportSceneObjectsParams& Params) override;
 
 	// Interchange actor factory interface end
 	//////////////////////////////////////////////////////////////////////////
diff --git a/Engine/Plugins/Interchange/Runtime/Source/Import/Public/Scene/InterchangeStaticMeshActorFactory.h b/Engine/Plugins/Interchange/Runtime/Source/Import/Public/Scene/InterchangeStaticMeshActorFactory.h
--- a/Engine/Plugins/Interchange/Runtime/Source/Import/Public/Scene/InterchangeStaticMeshActorFactory.h
+++ b/Engine/Plugins/Interchange/Runtime/Source/Import/Public/Scene/InterchangeStaticMeshActorFactory.h
@@ -31,7 +31,7 @@ protected:
 	//////////////////////////////////////////////////////////////////////////
 	// Interchange actor factory interface begin
 
-	virtual UObject* ProcessActor(class AActor& SpawnedActor, const UInterchangeActorFactoryNode& FactoryNode, const UInterchangeBaseNodeContainer& NodeContainer) override;
+	virtual UObject* ProcessActor(class AActor& SpawnedActor, const UInterchangeActorFactoryNode& FactoryNode, const UInterchangeBaseNodeContainer& NodeContainer, const FImportSceneObjectsParams& Params) override;
 
 	// Interchange actor factory interface end
 	//////////////////////////////////////////////////////////////////////////
diff --git a/Engine/Source/Runtime/Interchange/Core/Private/InterchangeResult.cpp b/Engine/Source/Runtime/Interchange/Core/Private/InterchangeResult.cpp
--- a/Engine/Source/Runtime/Interchange/Core/Private/InterchangeResult.cpp
+++ b/Engine/Source/Runtime/Interchange/Core/Private/InterchangeResult.cpp
@@ -3,6 +3,7 @@
 #include "InterchangeResult.h"
 #include "Dom/JsonValue.h"
 #include "Dom/JsonObject.h"
+#include "Misc/Paths.h"
 #include "Serialization/JsonReader.h"
 #include "Serialization/JsonWriter.h"
 #include "Serialization/JsonSerializer.h"
@@ -99,3 +100,29 @@ UInterchangeResult* UInterchangeResult::FromJson(const FString& JsonString)
 	return Result;
 }
 
+FText UInterchangeResult::GetMessageLogText() const
+{
+	FString SourceText = SourceAssetName;
+
+	// Make sure file path is not taking the whole line
+	const int32 MaximumSourceNameLength = 64;
+	if (SourceText.Len() > MaximumSourceNameLength)
+	{
+		SourceText = TEXT("...") + SourceText.RightChop(MaximumSourceNameLength-3);
+	}
+
+	const FText AssetTypeName = FText::FromString(AssetType ? AssetType->GetName() : TEXT("Unknown"));
+	FString AssetName = AssetFriendlyName;
+	
+	if (AssetName.IsEmpty())
+	{
+		FSoftObjectPath AssetSoftPath(DestinationAssetName);
+		AssetName = AssetSoftPath.GetAssetName();
+		if (AssetSoftPath.IsSubobject())
+		{
+			AssetName = FPaths::GetExtension(AssetSoftPath.GetSubPathString());
+		}
+	}
+
+	return FText::Format(NSLOCTEXT("InterchangeManager", "MessageLog", "[{0} : '{1}', {2}] {3}"), FText::FromString(SourceText), FText::FromString(AssetName), AssetTypeName, GetText());
+}
diff --git a/Engine/Source/Runtime/Interchange/Core/Public/InterchangeFactoryBase.h b/Engine/Source/Runtime/Interchange/Core/Public/InterchangeFactoryBase.h
--- a/Engine/Source/Runtime/Interchange/Core/Public/InterchangeFactoryBase.h
+++ b/Engine/Source/Runtime/Interchange/Core/Public/InterchangeFactoryBase.h
@@ -86,7 +86,7 @@ public:
 		/** The translator is use to retrieve the PayLoad data in case the factory need it */
 		const UInterchangeTranslatorBase* Translator = nullptr;
 
-		/** The source data, mainly use to set the asset import data file. TODO: we have to refactor UAssetImportData, the source data should be the base class for this now */
+		/** The source data. TODO: we have to refactor UAssetImportData, the source data should be the base class for this now */
 		const UInterchangeSourceData* SourceData = nullptr;
 
 		/** The node container associate with the current source index */
@@ -196,7 +196,10 @@ public:
 
 		/** The node container associated with the current source index */
 		const UInterchangeBaseNodeContainer* NodeContainer = nullptr;
-	
+
+		/** The source data. TODO: we have to refactor UAssetImportData, the source data should be the base class for this now */
+		const UInterchangeSourceData* SourceData = nullptr;
+
 		/**
 		 * If not null, the factory must perform a reimport of the scene node
 		 */
diff --git a/Engine/Source/Runtime/Interchange/Core/Public/InterchangeResult.h b/Engine/Source/Runtime/Interchange/Core/Public/InterchangeResult.h
--- a/Engine/Source/Runtime/Interchange/Core/Public/InterchangeResult.h
+++ b/Engine/Source/Runtime/Interchange/Core/Public/InterchangeResult.h
@@ -45,6 +45,11 @@ public:
 	 */
 	virtual FText GetText() const { return FText(); }
 
+	/**
+	 * Return a message to be displayed in the MessageLog window.
+	 */
+	INTERCHANGECORE_API FText GetMessageLogText() const;
+
 	/**
 	 * Serialize this UInterchangeResult object to Json.
 	 */
@@ -62,6 +67,9 @@ public:
 	UPROPERTY()
 	FString DestinationAssetName;
 
+	UPROPERTY()
+	FString AssetFriendlyName;
+
 	UPROPERTY()
 	TObjectPtr<const UClass> AssetType;
 
diff --git a/Engine/Source/Runtime/Interchange/Engine/Private/Tasks/InterchangeTaskCompletion.cpp b/Engine/Source/Runtime/Interchange/Engine/Private/Tasks/InterchangeTaskCompletion.cpp
--- a/Engine/Source/Runtime/Interchange/Engine/Private/Tasks/InterchangeTaskCompletion.cpp
+++ b/Engine/Source/Runtime/Interchange/Engine/Private/Tasks/InterchangeTaskCompletion.cpp
@@ -107,6 +107,7 @@ void UE::Interchange::FTaskPreCompletion::DoTask(ENamedThreads::Type CurrentThre
 			{
 				if (AActor* Actor = Cast<AActor>(ImportedObject))
 				{
+					Message->AssetFriendlyName = Actor->GetActorLabel();
 					Actor->RegisterAllComponents();
 				}
 				else if (UActorComponent* Component = Cast<UActorComponent>(ImportedObject))
