Fixed build issue

#issue 425381

diff --git a/Engine/Plugins/Interchange/Runtime/Source/Import/Private/Scene/InterchangeDecalActorFactory.cpp b/Engine/Plugins/Interchange/Runtime/Source/Import/Private/Scene/InterchangeDecalActorFactory.cpp
--- a/Engine/Plugins/Interchange/Runtime/Source/Import/Private/Scene/InterchangeDecalActorFactory.cpp
+++ b/Engine/Plugins/Interchange/Runtime/Source/Import/Private/Scene/InterchangeDecalActorFactory.cpp
@@ -52,7 +52,11 @@ UObject* UInterchangeDecalActorFactory::ProcessActor(AActor& SpawnedActor, const
 				else
 				{
 					const FText Message = FText::Format(NSLOCTEXT("DecalActorImport", "NoMaterial", "No valid decal material found. Make sure that the DecalMaterialPath is valid: %s"), FText::FromString(DecalMaterialPath));
+#if WITH_EDITOR
 					LogMessage<UInterchangeResultError_Generic>(Params, Message, SpawnedActor.GetActorLabel());
+#else
+					LogMessage<UInterchangeResultError_Generic>(Params, Message, TEXT(""));
+#endif
 				}
 			}
 		}
diff --git a/Engine/Source/Runtime/Interchange/Engine/Private/Tasks/InterchangeTaskCompletion.cpp b/Engine/Source/Runtime/Interchange/Engine/Private/Tasks/InterchangeTaskCompletion.cpp
--- a/Engine/Source/Runtime/Interchange/Engine/Private/Tasks/InterchangeTaskCompletion.cpp
+++ b/Engine/Source/Runtime/Interchange/Engine/Private/Tasks/InterchangeTaskCompletion.cpp
@@ -107,7 +107,9 @@ void UE::Interchange::FTaskPreCompletion::DoTask(ENamedThreads::Type CurrentThre
 			{
 				if (AActor* Actor = Cast<AActor>(ImportedObject))
 				{
+#if WITH_EDITOR
 					Message->AssetFriendlyName = Actor->GetActorLabel();
+#endif
 					Actor->RegisterAllComponents();
 				}
 				else if (UActorComponent* Component = Cast<UActorComponent>(ImportedObject))
