[Interchange] Removing errors trying to merge bones (Error: Incompatible skeleton)

#rnx

#changelist validated
#virtualized

diff --git a/Engine/Plugins/Interchange/Runtime/Source/Pipelines/Private/InterchangeGenericAnimationPipeline.cpp b/Engine/Plugins/Interchange/Runtime/Source/Pipelines/Private/InterchangeGenericAnimationPipeline.cpp
--- a/Engine/Plugins/Interchange/Runtime/Source/Pipelines/Private/InterchangeGenericAnimationPipeline.cpp
+++ b/Engine/Plugins/Interchange/Runtime/Source/Pipelines/Private/InterchangeGenericAnimationPipeline.cpp
@@ -696,7 +696,7 @@ void UInterchangeGenericAnimationPipeline::CreateAnimSequenceFactoryNode(UInterc
 		}
 		else
 		{
-			UInterchangeResultError_Generic* Message = AddMessage<UInterchangeResultError_Generic>();
+			UInterchangeResultDisplay_Generic* Message = AddMessage<UInterchangeResultDisplay_Generic>();
 			Message->Text = FText::Format(NSLOCTEXT("UInterchangeGenericAnimationPipeline", "IncompatibleSkeleton", "Incompatible skeleton {0} when importing AnimSequence {1}."),
 				FText::FromString(CommonSkeletalMeshesAndAnimationsProperties->Skeleton->GetName()),
 				FText::FromString(TrackNode.GetDisplayLabel()));
diff --git a/Engine/Plugins/Interchange/Runtime/Source/Pipelines/Private/InterchangeGenericSkeletalMeshPipeline.cpp b/Engine/Plugins/Interchange/Runtime/Source/Pipelines/Private/InterchangeGenericSkeletalMeshPipeline.cpp
--- a/Engine/Plugins/Interchange/Runtime/Source/Pipelines/Private/InterchangeGenericSkeletalMeshPipeline.cpp
+++ b/Engine/Plugins/Interchange/Runtime/Source/Pipelines/Private/InterchangeGenericSkeletalMeshPipeline.cpp
@@ -378,7 +378,7 @@ UInterchangeSkeletalMeshFactoryNode* UInterchangeGenericMeshPipeline::CreateSkel
 		}
 		else
 		{
-			UInterchangeResultError_Generic* Message = AddMessage<UInterchangeResultError_Generic>();
+			UInterchangeResultDisplay_Generic* Message = AddMessage<UInterchangeResultDisplay_Generic>();
 			Message->Text = FText::Format(NSLOCTEXT("UInterchangeGenericMeshPipeline", "IncompatibleSkeleton", "Incompatible skeleton {0} when importing skeletalmesh {1}."),
 				FText::FromString(CommonSkeletalMeshesAndAnimationsProperties->Skeleton->GetName()),
 				FText::FromString(DisplayLabel));
