[Interchange] Importing gltf assets in runtime results with invalid material instance

#rnx

#changelist validated
#virtualized

diff --git a/Engine/Plugins/Interchange/Runtime/Source/Pipelines/Private/InterchangeGenericSkeletalMeshPipeline.cpp b/Engine/Plugins/Interchange/Runtime/Source/Pipelines/Private/InterchangeGenericSkeletalMeshPipeline.cpp
--- a/Engine/Plugins/Interchange/Runtime/Source/Pipelines/Private/InterchangeGenericSkeletalMeshPipeline.cpp
+++ b/Engine/Plugins/Interchange/Runtime/Source/Pipelines/Private/InterchangeGenericSkeletalMeshPipeline.cpp
@@ -127,6 +127,21 @@ void UInterchangeGenericMeshPipeline::ExecutePreImportPipelineSkeletalMesh()
 		TMap<FString, TArray<FString>> MeshUidsPerSkeletonRootUid;
 		auto CreatePerSkeletonRootUidCombinedSkinnedMesh = [this, &MeshUidsPerSkeletonRootUid, &SetSkeletalMeshDependencies](const bool bUseInstanceMesh)
 		{
+			if (MeshUidsPerSkeletonRootUid.Num() > 0)
+			{
+#if !WITH_EDITOR || !WITH_EDITORONLY_DATA
+				UE_LOG(LogInterchangePipeline, Warning, TEXT("Cannot import skeletalMesh asset in runtime, this is an editor only feature."));
+				return false;
+#else
+				if (FApp::IsGame())
+				{
+					//PIE
+					UE_LOG(LogInterchangePipeline, Warning, TEXT("Cannot import skeletalMesh asset in runtime, this is an editor only feature."));
+					return false;
+				}
+#endif
+			}
+
 			bool bFoundInstances = false;
 			for (const TPair<FString, TArray<FString>>& SkeletonRootUidAndMeshUids : MeshUidsPerSkeletonRootUid)
 			{
@@ -165,10 +180,13 @@ void UInterchangeGenericMeshPipeline::ExecutePreImportPipelineSkeletalMesh()
 				if (MeshUidsPerLodIndex.Num() > 0)
 				{
 					UInterchangeSkeletalMeshFactoryNode* SkeletalMeshFactoryNode = CreateSkeletalMeshFactoryNode(SkeletonRootUid, MeshUidsPerLodIndex);
-					SetSkeletalMeshDependencies(SkeletonRootUid, SkeletalMeshFactoryNode);
-					SkeletonFactoryNodes.Add(SkeletonFactoryNode);
-					SkeletalMeshFactoryNodes.Add(SkeletalMeshFactoryNode);
-					bFoundInstances = true;
+					if (SkeletalMeshFactoryNode != nullptr)
+					{
+						SetSkeletalMeshDependencies(SkeletonRootUid, SkeletalMeshFactoryNode);
+						SkeletonFactoryNodes.Add(SkeletonFactoryNode);
+						SkeletalMeshFactoryNodes.Add(SkeletalMeshFactoryNode);
+						bFoundInstances = true;
+					}
 				}
 			}
 			return bFoundInstances;
@@ -186,6 +204,21 @@ void UInterchangeGenericMeshPipeline::ExecutePreImportPipelineSkeletalMesh()
 		TArray<FString> MeshUids;
 		auto CreatePerSkeletonRootUidSkinnedMesh = [this, &MeshUids, &SetSkeletalMeshDependencies](const bool bUseInstanceMesh)
 		{
+			if (MeshUids.Num() > 0)
+			{
+#if !WITH_EDITOR || !WITH_EDITORONLY_DATA
+				UE_LOG(LogInterchangePipeline, Warning, TEXT("Cannot import skeletalMesh asset in runtime, this is an editor only feature."));
+				return false;
+#else
+				if (FApp::IsGame())
+				{
+					//PIE
+					UE_LOG(LogInterchangePipeline, Warning, TEXT("Cannot import skeletalMesh asset in runtime, this is an editor only feature."));
+					return false;
+				}
+#endif
+			}
+
 			bool bFoundInstances = false;
 			for (const FString& MeshUid : MeshUids)
 			{
@@ -235,10 +268,14 @@ void UInterchangeGenericMeshPipeline::ExecutePreImportPipelineSkeletalMesh()
 				if (MeshUidsPerLodIndex.Num() > 0)
 				{
 					UInterchangeSkeletalMeshFactoryNode* SkeletalMeshFactoryNode = CreateSkeletalMeshFactoryNode(SkeletonRootUid, MeshUidsPerLodIndex);
-					SetSkeletalMeshDependencies(SkeletonRootUid, SkeletalMeshFactoryNode);
-					SkeletonFactoryNodes.Add(SkeletonFactoryNode);
-					SkeletalMeshFactoryNodes.Add(SkeletalMeshFactoryNode);
-					bFoundInstances = true;
+					if (SkeletalMeshFactoryNode != nullptr)
+					{
+						SetSkeletalMeshDependencies(SkeletonRootUid, SkeletalMeshFactoryNode);
+						SkeletonFactoryNodes.Add(SkeletonFactoryNode);
+						SkeletalMeshFactoryNodes.Add(SkeletalMeshFactoryNode);
+						bFoundInstances = true;
+					}
+					
 				}
 			}
 			return bFoundInstances;
