[Interchange] Reimport animation sequence asset with new files fails fixes.

#rnx

#changelist validated
#virtualized

diff --git a/Engine/Plugins/Interchange/Runtime/Source/Import/Private/Animation/InterchangeAnimSequenceFactory.cpp b/Engine/Plugins/Interchange/Runtime/Source/Import/Private/Animation/InterchangeAnimSequenceFactory.cpp
--- a/Engine/Plugins/Interchange/Runtime/Source/Import/Private/Animation/InterchangeAnimSequenceFactory.cpp
+++ b/Engine/Plugins/Interchange/Runtime/Source/Import/Private/Animation/InterchangeAnimSequenceFactory.cpp
@@ -907,6 +907,21 @@ UInterchangeFactoryBase::FImportAssetResult UInterchangeAnimSequenceFactory::Beg
 		return ImportAssetResult;
 	}
 
+	UObject* ExistingAsset = Arguments.ReimportObject;
+	if (!ExistingAsset)
+	{
+		FSoftObjectPath ReferenceObject;
+		if (AnimSequenceFactoryNode->GetCustomReferenceObject(ReferenceObject))
+		{
+			ExistingAsset = ReferenceObject.TryLoad();
+		}
+	}
+	if (ExistingAsset)
+	{
+		//This is a reimport, we are just re-updating the source data
+		NewAnimSequence = Cast<UAnimSequence>(ExistingAsset);
+	}
+
 	const IInterchangeAnimationPayloadInterface* AnimSequenceTranslatorPayloadInterface = Cast<IInterchangeAnimationPayloadInterface>(Arguments.Translator);
 	if (!AnimSequenceTranslatorPayloadInterface)
 	{
@@ -936,8 +951,9 @@ UInterchangeFactoryBase::FImportAssetResult UInterchangeAnimSequenceFactory::Beg
 		FSoftObjectPath SkeletalMesh;
 		SkeletalMeshFactoryNode->GetCustomReferenceObject(SkeletalMesh);
 
-		if (!SkeletalMesh.IsValid())
+		if (NewAnimSequence == nullptr && !SkeletalMesh.IsValid())
 		{
+			//For AnimSequence (single) Re-Imports the SkeletalMesh is not created in the pipeline.
 			UE_LOG(LogInterchangeImport, Error, TEXT("SkeletalMesh does not exist, the skeleton and skeletal mesh is obligatory to import this animsequence [%s]!"), *Arguments.AssetName);
 			return ImportAssetResult;
 		}
@@ -982,26 +998,11 @@ UInterchangeFactoryBase::FImportAssetResult UInterchangeAnimSequenceFactory::Beg
 		return ImportAssetResult;
 	}
 
-	UObject* ExistingAsset = Arguments.ReimportObject;
-	if (!ExistingAsset)
-	{
-		FSoftObjectPath ReferenceObject;
-		if (AnimSequenceFactoryNode->GetCustomReferenceObject(ReferenceObject))
-		{
-			ExistingAsset = ReferenceObject.TryLoad();
-		}
-	}
-
 	// create a new material or overwrite existing asset, if possible
 	if (!ExistingAsset)
 	{
 		NewAnimSequence = NewObject<UAnimSequence>(Arguments.Parent, *Arguments.AssetName, RF_Public | RF_Standalone);
 	}
-	else
-	{
-		//This is a reimport, we are just re-updating the source data
-		NewAnimSequence = Cast<UAnimSequence>(ExistingAsset);
-	}
 
 	if (!NewAnimSequence)
 	{
