[Interchange] SkeletalMeshMeshPipeline Static Analysis warning fix

#rnx

#changelist validated
#virtualized

diff --git a/Engine/Plugins/Interchange/Runtime/Source/Pipelines/Private/InterchangeGenericSkeletalMeshPipeline.cpp b/Engine/Plugins/Interchange/Runtime/Source/Pipelines/Private/InterchangeGenericSkeletalMeshPipeline.cpp
--- a/Engine/Plugins/Interchange/Runtime/Source/Pipelines/Private/InterchangeGenericSkeletalMeshPipeline.cpp
+++ b/Engine/Plugins/Interchange/Runtime/Source/Pipelines/Private/InterchangeGenericSkeletalMeshPipeline.cpp
@@ -127,69 +127,75 @@ void UInterchangeGenericMeshPipeline::ExecutePreImportPipelineSkeletalMesh()
 		TMap<FString, TArray<FString>> MeshUidsPerSkeletonRootUid;
 		auto CreatePerSkeletonRootUidCombinedSkinnedMesh = [this, &MeshUidsPerSkeletonRootUid, &SetSkeletalMeshDependencies](const bool bUseInstanceMesh)
 		{
-			if (MeshUidsPerSkeletonRootUid.Num() > 0)
-			{
 #if !WITH_EDITOR || !WITH_EDITORONLY_DATA
-				UE_LOG(LogInterchangePipeline, Warning, TEXT("Cannot import skeletalMesh asset in runtime, this is an editor only feature."));
-				return false;
-#else
-				if (FApp::IsGame())
+				if (MeshUidsPerSkeletonRootUid.Num() > 0)
 				{
-					//PIE
 					UE_LOG(LogInterchangePipeline, Warning, TEXT("Cannot import skeletalMesh asset in runtime, this is an editor only feature."));
-					return false;
 				}
-#endif
-			}
-
-			bool bFoundInstances = false;
-			for (const TPair<FString, TArray<FString>>& SkeletonRootUidAndMeshUids : MeshUidsPerSkeletonRootUid)
-			{
-				const FString& SkeletonRootUid = SkeletonRootUidAndMeshUids.Key;
-				//Every iteration is a skeletalmesh asset that combine all MeshInstances sharing the same skeleton root node
-				UInterchangeSkeletonFactoryNode* SkeletonFactoryNode = CommonSkeletalMeshesAndAnimationsProperties->CreateSkeletonFactoryNode(BaseNodeContainer, SkeletonRootUid);
-				//The MeshUids can represent a SceneNode pointing on a MeshNode or directly a MeshNode;
-				TMap<int32, TArray<FString>> MeshUidsPerLodIndex;
-				const TArray<FString>& MeshUids = SkeletonRootUidAndMeshUids.Value;
-				for (const FString& MeshUid : MeshUids)
+				return false;
+#else
+				if (MeshUidsPerSkeletonRootUid.Num() > 0)
 				{
-					if (bUseInstanceMesh)
+					if (FApp::IsGame())
 					{
-						const FInterchangeMeshInstance& MeshInstance = PipelineMeshesUtilities->GetMeshInstanceByUid(MeshUid);
-						for (const TPair<int32, FInterchangeLodSceneNodeContainer>& LodIndexAndSceneNodeContainer : MeshInstance.SceneNodePerLodIndex)
-						{
-							const int32 LodIndex = LodIndexAndSceneNodeContainer.Key;
-							const FInterchangeLodSceneNodeContainer& SceneNodeContainer = LodIndexAndSceneNodeContainer.Value;
-							TArray<FString>& TranslatedNodes = MeshUidsPerLodIndex.FindOrAdd(LodIndex);
-							for (const UInterchangeSceneNode* SceneNode : SceneNodeContainer.SceneNodes)
-							{
-								TranslatedNodes.Add(SceneNode->GetUniqueID());
-							}
-						}
+						//PIE
+						UE_LOG(LogInterchangePipeline, Warning, TEXT("Cannot import skeletalMesh asset in runtime, this is an editor only feature."));
+						return false;
 					}
 					else
 					{
-						//MeshGeometry cannot have Lod since LODs are define in the scene node
-						const FInterchangeMeshGeometry& MeshGeometry = PipelineMeshesUtilities->GetMeshGeometryByUid(MeshUid);
-						const int32 LodIndex = 0;
-						TArray<FString>& TranslatedNodes = MeshUidsPerLodIndex.FindOrAdd(LodIndex);
-						TranslatedNodes.Add(MeshGeometry.MeshUid);
-					}
-				}
+						bool bFoundInstances = false;
+						for (const TPair<FString, TArray<FString>>& SkeletonRootUidAndMeshUids : MeshUidsPerSkeletonRootUid)
+						{
+							const FString& SkeletonRootUid = SkeletonRootUidAndMeshUids.Key;
+							//Every iteration is a skeletalmesh asset that combine all MeshInstances sharing the same skeleton root node
+							UInterchangeSkeletonFactoryNode* SkeletonFactoryNode = CommonSkeletalMeshesAndAnimationsProperties->CreateSkeletonFactoryNode(BaseNodeContainer, SkeletonRootUid);
+							//The MeshUids can represent a SceneNode pointing on a MeshNode or directly a MeshNode;
+							TMap<int32, TArray<FString>> MeshUidsPerLodIndex;
+							const TArray<FString>& MeshUids = SkeletonRootUidAndMeshUids.Value;
+							for (const FString& MeshUid : MeshUids)
+							{
+								if (bUseInstanceMesh)
+								{
+									const FInterchangeMeshInstance& MeshInstance = PipelineMeshesUtilities->GetMeshInstanceByUid(MeshUid);
+									for (const TPair<int32, FInterchangeLodSceneNodeContainer>& LodIndexAndSceneNodeContainer : MeshInstance.SceneNodePerLodIndex)
+									{
+										const int32 LodIndex = LodIndexAndSceneNodeContainer.Key;
+										const FInterchangeLodSceneNodeContainer& SceneNodeContainer = LodIndexAndSceneNodeContainer.Value;
+										TArray<FString>& TranslatedNodes = MeshUidsPerLodIndex.FindOrAdd(LodIndex);
+										for (const UInterchangeSceneNode* SceneNode : SceneNodeContainer.SceneNodes)
+										{
+											TranslatedNodes.Add(SceneNode->GetUniqueID());
+										}
+									}
+								}
+								else
+								{
+									//MeshGeometry cannot have Lod since LODs are define in the scene node
+									const FInterchangeMeshGeometry& MeshGeometry = PipelineMeshesUtilities->GetMeshGeometryByUid(MeshUid);
+									const int32 LodIndex = 0;
+									TArray<FString>& TranslatedNodes = MeshUidsPerLodIndex.FindOrAdd(LodIndex);
+									TranslatedNodes.Add(MeshGeometry.MeshUid);
+								}
+							}
 
-				if (MeshUidsPerLodIndex.Num() > 0)
-				{
-					UInterchangeSkeletalMeshFactoryNode* SkeletalMeshFactoryNode = CreateSkeletalMeshFactoryNode(SkeletonRootUid, MeshUidsPerLodIndex);
-					if (SkeletalMeshFactoryNode != nullptr)
-					{
-						SetSkeletalMeshDependencies(SkeletonRootUid, SkeletalMeshFactoryNode);
-						SkeletonFactoryNodes.Add(SkeletonFactoryNode);
-						SkeletalMeshFactoryNodes.Add(SkeletalMeshFactoryNode);
-						bFoundInstances = true;
+							if (MeshUidsPerLodIndex.Num() > 0)
+							{
+								UInterchangeSkeletalMeshFactoryNode* SkeletalMeshFactoryNode = CreateSkeletalMeshFactoryNode(SkeletonRootUid, MeshUidsPerLodIndex);
+								if (SkeletalMeshFactoryNode != nullptr)
+								{
+									SetSkeletalMeshDependencies(SkeletonRootUid, SkeletalMeshFactoryNode);
+									SkeletonFactoryNodes.Add(SkeletonFactoryNode);
+									SkeletalMeshFactoryNodes.Add(SkeletalMeshFactoryNode);
+									bFoundInstances = true;
+								}
+							}
+						}
+						return bFoundInstances;
 					}
 				}
-			}
-			return bFoundInstances;
+				return false;
+#endif
 		};
 
 
@@ -204,81 +210,87 @@ void UInterchangeGenericMeshPipeline::ExecutePreImportPipelineSkeletalMesh()
 		TArray<FString> MeshUids;
 		auto CreatePerSkeletonRootUidSkinnedMesh = [this, &MeshUids, &SetSkeletalMeshDependencies](const bool bUseInstanceMesh)
 		{
+#if !WITH_EDITOR || !WITH_EDITORONLY_DATA
 			if (MeshUids.Num() > 0)
 			{
-#if !WITH_EDITOR || !WITH_EDITORONLY_DATA
 				UE_LOG(LogInterchangePipeline, Warning, TEXT("Cannot import skeletalMesh asset in runtime, this is an editor only feature."));
-				return false;
+			}
+			return false;
 #else
+			if (MeshUids.Num() > 0)
+			{
 				if (FApp::IsGame())
 				{
 					//PIE
 					UE_LOG(LogInterchangePipeline, Warning, TEXT("Cannot import skeletalMesh asset in runtime, this is an editor only feature."));
 					return false;
 				}
-#endif
-			}
-
-			bool bFoundInstances = false;
-			for (const FString& MeshUid : MeshUids)
-			{
-				//Every iteration is a skeletalmesh asset that combine all MeshInstances sharing the same skeleton root node
-				//The MeshUids can represent a SceneNode pointing on a MeshNode or directly a MeshNode;
-				TMap<int32, TArray<FString>> MeshUidsPerLodIndex;
-				FString SkeletonRootUid;
-				if (!(bUseInstanceMesh ? PipelineMeshesUtilities->IsValidMeshInstanceUid(MeshUid) : PipelineMeshesUtilities->IsValidMeshGeometryUid(MeshUid)))
-				{
-					continue;
-				}
-				SkeletonRootUid = (bUseInstanceMesh ? PipelineMeshesUtilities->GetMeshInstanceSkeletonRootUid(MeshUid) : PipelineMeshesUtilities->GetMeshGeometrySkeletonRootUid(MeshUid));
-				if (SkeletonRootUid.IsEmpty())
-				{
-					if (CommonSkeletalMeshesAndAnimationsProperties->bConvertStaticsWithMorphTargetsToSkeletals)
-					{
-						SkeletonRootUid = MeshUid;
-					}
-					else
-					{
-						//Log an error
-						continue;
-					}
-				}
-				UInterchangeSkeletonFactoryNode* SkeletonFactoryNode = CommonSkeletalMeshesAndAnimationsProperties->CreateSkeletonFactoryNode(BaseNodeContainer, SkeletonRootUid);
-				if (bUseInstanceMesh)
+				else
 				{
-					const FInterchangeMeshInstance& MeshInstance = PipelineMeshesUtilities->GetMeshInstanceByUid(MeshUid);
-					for (const TPair<int32, FInterchangeLodSceneNodeContainer>& LodIndexAndSceneNodeContainer : MeshInstance.SceneNodePerLodIndex)
+					bool bFoundInstances = false;
+					for (const FString& MeshUid : MeshUids)
 					{
-						const int32 LodIndex = LodIndexAndSceneNodeContainer.Key;
-						const FInterchangeLodSceneNodeContainer& SceneNodeContainer = LodIndexAndSceneNodeContainer.Value;
-						TArray<FString>& TranslatedNodes = MeshUidsPerLodIndex.FindOrAdd(LodIndex);
-						for (const UInterchangeSceneNode* SceneNode : SceneNodeContainer.SceneNodes)
+						//Every iteration is a skeletalmesh asset that combine all MeshInstances sharing the same skeleton root node
+						//The MeshUids can represent a SceneNode pointing on a MeshNode or directly a MeshNode;
+						TMap<int32, TArray<FString>> MeshUidsPerLodIndex;
+						FString SkeletonRootUid;
+						if (!(bUseInstanceMesh ? PipelineMeshesUtilities->IsValidMeshInstanceUid(MeshUid) : PipelineMeshesUtilities->IsValidMeshGeometryUid(MeshUid)))
 						{
-							TranslatedNodes.Add(SceneNode->GetUniqueID());
+							continue;
+						}
+						SkeletonRootUid = (bUseInstanceMesh ? PipelineMeshesUtilities->GetMeshInstanceSkeletonRootUid(MeshUid) : PipelineMeshesUtilities->GetMeshGeometrySkeletonRootUid(MeshUid));
+						if (SkeletonRootUid.IsEmpty())
+						{
+							if (CommonSkeletalMeshesAndAnimationsProperties->bConvertStaticsWithMorphTargetsToSkeletals)
+							{
+								SkeletonRootUid = MeshUid;
+							}
+							else
+							{
+								//Log an error
+								continue;
+							}
+						}
+						UInterchangeSkeletonFactoryNode* SkeletonFactoryNode = CommonSkeletalMeshesAndAnimationsProperties->CreateSkeletonFactoryNode(BaseNodeContainer, SkeletonRootUid);
+						if (bUseInstanceMesh)
+						{
+							const FInterchangeMeshInstance& MeshInstance = PipelineMeshesUtilities->GetMeshInstanceByUid(MeshUid);
+							for (const TPair<int32, FInterchangeLodSceneNodeContainer>& LodIndexAndSceneNodeContainer : MeshInstance.SceneNodePerLodIndex)
+							{
+								const int32 LodIndex = LodIndexAndSceneNodeContainer.Key;
+								const FInterchangeLodSceneNodeContainer& SceneNodeContainer = LodIndexAndSceneNodeContainer.Value;
+								TArray<FString>& TranslatedNodes = MeshUidsPerLodIndex.FindOrAdd(LodIndex);
+								for (const UInterchangeSceneNode* SceneNode : SceneNodeContainer.SceneNodes)
+								{
+									TranslatedNodes.Add(SceneNode->GetUniqueID());
+								}
+							}
+						}
+						else
+						{
+							const FInterchangeMeshGeometry& MeshGeometry = PipelineMeshesUtilities->GetMeshGeometryByUid(MeshUid);
+							const int32 LodIndex = 0;
+							TArray<FString>& TranslatedNodes = MeshUidsPerLodIndex.FindOrAdd(LodIndex);
+							TranslatedNodes.Add(MeshGeometry.MeshUid);
+						}
+						if (MeshUidsPerLodIndex.Num() > 0)
+						{
+							UInterchangeSkeletalMeshFactoryNode* SkeletalMeshFactoryNode = CreateSkeletalMeshFactoryNode(SkeletonRootUid, MeshUidsPerLodIndex);
+							if (SkeletalMeshFactoryNode != nullptr)
+							{
+								SetSkeletalMeshDependencies(SkeletonRootUid, SkeletalMeshFactoryNode);
+								SkeletonFactoryNodes.Add(SkeletonFactoryNode);
+								SkeletalMeshFactoryNodes.Add(SkeletalMeshFactoryNode);
+								bFoundInstances = true;
+							}
+
 						}
 					}
-				}
-				else
-				{
-					const FInterchangeMeshGeometry& MeshGeometry = PipelineMeshesUtilities->GetMeshGeometryByUid(MeshUid);
-					const int32 LodIndex = 0;
-					TArray<FString>& TranslatedNodes = MeshUidsPerLodIndex.FindOrAdd(LodIndex);
-					TranslatedNodes.Add(MeshGeometry.MeshUid);
-				}
-				if (MeshUidsPerLodIndex.Num() > 0)
-				{
-					UInterchangeSkeletalMeshFactoryNode* SkeletalMeshFactoryNode = CreateSkeletalMeshFactoryNode(SkeletonRootUid, MeshUidsPerLodIndex);
-					if (SkeletalMeshFactoryNode != nullptr)
-					{
-						SetSkeletalMeshDependencies(SkeletonRootUid, SkeletalMeshFactoryNode);
-						SkeletonFactoryNodes.Add(SkeletonFactoryNode);
-						SkeletalMeshFactoryNodes.Add(SkeletalMeshFactoryNode);
-						bFoundInstances = true;
-					}
-					
+					return bFoundInstances;
 				}
 			}
-			return bFoundInstances;
+			return false;
+#endif
 		};
 		
 		PipelineMeshesUtilities->GetAllSkinnedMeshInstance(MeshUids);
