[Interchange] Import FBX asset - Bone creation problem with M_MED_PrivateJet.fbx file

#rnx

diff --git a/Engine/Plugins/Interchange/Runtime/Source/Parsers/Fbx/Private/FbxMesh.cpp b/Engine/Plugins/Interchange/Runtime/Source/Parsers/Fbx/Private/FbxMesh.cpp
--- a/Engine/Plugins/Interchange/Runtime/Source/Parsers/Fbx/Private/FbxMesh.cpp
+++ b/Engine/Plugins/Interchange/Runtime/Source/Parsers/Fbx/Private/FbxMesh.cpp
@@ -2003,6 +2003,21 @@ bool FFbxMesh::GetGlobalJointBindPoseTransform(FbxScene* SDKScene, FbxNode* Join
 		}
 	}
 
+	auto AcquireBindPoseMatrix = [](FbxPose* CurrentPose, FbxAMatrix& GlobalBindPoseJointMatrix, FbxNode* Joint)
+	{
+		int32 PoseLinkIndex = CurrentPose->Find(Joint);
+		if (PoseLinkIndex >= 0)
+		{
+			if (!CurrentPose->IsLocalMatrix(PoseLinkIndex))
+			{
+				FbxMatrix NoneAffineMatrix = CurrentPose->GetMatrix(PoseLinkIndex);
+				GlobalBindPoseJointMatrix = *(FbxAMatrix*)(double*)&NoneAffineMatrix;
+				return true;
+			}
+		}
+		return false;
+	};
+
 	const int32 PoseCount = SDKScene->GetPoseCount();
 	for (int32 PoseIndex = 0; PoseIndex < PoseCount; PoseIndex++)
 	{
@@ -2015,16 +2030,64 @@ bool FFbxMesh::GetGlobalJointBindPoseTransform(FbxScene* SDKScene, FbxNode* Join
 			// all error report status
 			FbxStatus Status;
 
-			int32 PoseLinkIndex = CurrentPose->Find(Joint);
-			if (PoseLinkIndex >= 0)
+			FbxArray<FbxNode*> pMissingAncestors, pMissingDeformers, pMissingDeformersAncestors, pWrongMatrices;
+
+			if (CurrentPose->IsValidBindPoseVerbose(Joint, pMissingAncestors, pMissingDeformers, pMissingDeformersAncestors, pWrongMatrices, 0.0001, &Status))
 			{
-				if (!CurrentPose->IsLocalMatrix(PoseLinkIndex))
+				if (AcquireBindPoseMatrix(CurrentPose, GlobalBindPoseJointMatrix, Joint))
 				{
-					FbxMatrix NoneAffineMatrix = CurrentPose->GetMatrix(PoseLinkIndex);
-					GlobalBindPoseJointMatrix = *(FbxAMatrix*)(double*)&NoneAffineMatrix;
 					return true;
 				}
 			}
+			else
+			{
+				// first try to fix up
+				// add missing ancestors
+				for (int i = 0; i < pMissingAncestors.GetCount(); i++)
+				{
+					FbxAMatrix mat = pMissingAncestors.GetAt(i)->EvaluateGlobalTransform(FBXSDK_TIME_ZERO);
+					CurrentPose->Add(pMissingAncestors.GetAt(i), mat);
+				}
+
+				pMissingAncestors.Clear();
+				pMissingDeformers.Clear();
+				pMissingDeformersAncestors.Clear();
+				pWrongMatrices.Clear();
+
+				// check it again
+				if (CurrentPose->IsValidBindPose(Joint))
+				{
+					if (AcquireBindPoseMatrix(CurrentPose, GlobalBindPoseJointMatrix, Joint))
+					{
+						return true;
+					}
+				}
+				else
+				{
+					// first try to find parent who is null group and see if you can try test it again
+					FbxNode* ParentNode = Joint->GetParent();
+					while (ParentNode)
+					{
+						FbxNodeAttribute* Attr = ParentNode->GetNodeAttribute();
+						if (Attr && Attr->GetAttributeType() == FbxNodeAttribute::eNull)
+						{
+							// found it 
+							break;
+						}
+
+						// find next parent
+						ParentNode = ParentNode->GetParent();
+					}
+
+					if (ParentNode && CurrentPose->IsValidBindPose(ParentNode))
+					{
+						if (AcquireBindPoseMatrix(CurrentPose, GlobalBindPoseJointMatrix, Joint))
+						{
+							return true;
+						}
+					}
+				}
+			}
 		}
 	}
 
