Clean up typos and phrasing in Interchange modules - contribution from Robb Surridge

#rnx

diff --git a/Engine/Plugins/Enterprise/DatasmithImporter/Source/DatasmithImporter/Private/Utility/DatasmithImporterUtils.cpp b/Engine/Plugins/Enterprise/DatasmithImporter/Source/DatasmithImporter/Private/Utility/DatasmithImporterUtils.cpp
--- a/Engine/Plugins/Enterprise/DatasmithImporter/Source/DatasmithImporter/Private/Utility/DatasmithImporterUtils.cpp
+++ b/Engine/Plugins/Enterprise/DatasmithImporter/Source/DatasmithImporter/Private/Utility/DatasmithImporterUtils.cpp
@@ -1441,7 +1441,7 @@ void FScopedLogger::Dump(bool bClearPrevious)
 		}
 
 		LogListing->AddMessages(TokenizedMessages);
-		LogListing->NotifyIfAnyMessages(NSLOCTEXT("DatasmithLoggerNotification", "Log", "There was some issues with the import."), EMessageSeverity::Info);
+		LogListing->NotifyIfAnyMessages(NSLOCTEXT("DatasmithLoggerNotification", "Log", "There were issues with the import."), EMessageSeverity::Info);
 		ClearPending();
 	}
 }
diff --git a/Engine/Plugins/Experimental/Enterprise/DatasmithInterchange/Source/DatasmithInterchange/Private/InterchangeDatasmithMaterialNode.h b/Engine/Plugins/Experimental/Enterprise/DatasmithInterchange/Source/DatasmithInterchange/Private/InterchangeDatasmithMaterialNode.h
--- a/Engine/Plugins/Experimental/Enterprise/DatasmithInterchange/Source/DatasmithInterchange/Private/InterchangeDatasmithMaterialNode.h
+++ b/Engine/Plugins/Experimental/Enterprise/DatasmithInterchange/Source/DatasmithInterchange/Private/InterchangeDatasmithMaterialNode.h
@@ -30,7 +30,7 @@ public:
 	}
 
 	/**
-	 * Return the node type name of the class, we use this when reporting error
+	 * Return the node type name of the class. This is used when reporting errors.
 	 */
 	virtual FString GetTypeName() const override
 	{
diff --git a/Engine/Plugins/Experimental/Enterprise/DatasmithInterchange/Source/DatasmithInterchange/Public/InterchangeDatasmithAreaLightFactoryNode.h b/Engine/Plugins/Experimental/Enterprise/DatasmithInterchange/Source/DatasmithInterchange/Public/InterchangeDatasmithAreaLightFactoryNode.h
--- a/Engine/Plugins/Experimental/Enterprise/DatasmithInterchange/Source/DatasmithInterchange/Public/InterchangeDatasmithAreaLightFactoryNode.h
+++ b/Engine/Plugins/Experimental/Enterprise/DatasmithInterchange/Source/DatasmithInterchange/Public/InterchangeDatasmithAreaLightFactoryNode.h
@@ -18,7 +18,7 @@ class DATASMITHINTERCHANGE_API UInterchangeDatasmithAreaLightFactoryNode : publi
 public:
 
 	/**
-	 * Return the node type name of the class, we use this when reporting error
+	 * Return the node type name of the class. This is used when reporting errors.
 	 */
 	virtual FString GetTypeName() const override;
 
diff --git a/Engine/Plugins/Interchange/Editor/InterchangeEditor.uplugin b/Engine/Plugins/Interchange/Editor/InterchangeEditor.uplugin
--- a/Engine/Plugins/Interchange/Editor/InterchangeEditor.uplugin
+++ b/Engine/Plugins/Interchange/Editor/InterchangeEditor.uplugin
@@ -3,7 +3,7 @@
 	"Version": 1,
 	"VersionName": "1.0",
 	"FriendlyName": "Interchange Editor",
-	"Description": "Interchange Editor plugin pull on Interchange plugin to allow us to force interchange plugin in editor.",
+	"Description": "The Interchange Editor plugin exposes the Interchange import framework and pipelines to Unreal Editor.",
 	"Category": "Importers",
 	"CreatedBy": "Epic Games, Inc.",
 	"CreatedByURL": "https://epicgames.com",
diff --git a/Engine/Plugins/Interchange/Editor/Source/InterchangeEditor/Private/InterchangeEditorModule.cpp b/Engine/Plugins/Interchange/Editor/Source/InterchangeEditor/Private/InterchangeEditorModule.cpp
--- a/Engine/Plugins/Interchange/Editor/Source/InterchangeEditor/Private/InterchangeEditorModule.cpp
+++ b/Engine/Plugins/Interchange/Editor/Source/InterchangeEditor/Private/InterchangeEditorModule.cpp
@@ -62,7 +62,7 @@ namespace UE::Interchange::InterchangeEditorModule
 			}
 
 			LogListing->AddMessages(TokenizedMessages);
-			LogListing->NotifyIfAnyMessages(NSLOCTEXT("Interchange", "LogAndNotify", "There was some issues with the import."), EMessageSeverity::Info);
+			LogListing->NotifyIfAnyMessages(NSLOCTEXT("Interchange", "LogAndNotify", "There were issues with the import."), EMessageSeverity::Info);
 		}
 	}
 }
diff --git a/Engine/Plugins/Interchange/Editor/Source/InterchangeEditor/Public/InterchangeEditorModule.h b/Engine/Plugins/Interchange/Editor/Source/InterchangeEditor/Public/InterchangeEditorModule.h
--- a/Engine/Plugins/Interchange/Editor/Source/InterchangeEditor/Public/InterchangeEditorModule.h
+++ b/Engine/Plugins/Interchange/Editor/Source/InterchangeEditor/Public/InterchangeEditorModule.h
@@ -6,7 +6,7 @@
 #define INTERCHANGEEDITOR_MODULE_NAME TEXT("InterchangeEditor")
 
 /**
- * This module allow out of process interchange translator in case a third party SDK is not thread safe.
+ * This module allows out-of-process Interchange translators in case a third-party SDK is not thread-safe.
  */
 class FInterchangeEditorModule : public IModuleInterface
 {
diff --git a/Engine/Plugins/Interchange/Editor/Source/Pipelines/Private/InterchangePipelineFactories.cpp b/Engine/Plugins/Interchange/Editor/Source/Pipelines/Private/InterchangePipelineFactories.cpp
--- a/Engine/Plugins/Interchange/Editor/Source/Pipelines/Private/InterchangePipelineFactories.cpp
+++ b/Engine/Plugins/Interchange/Editor/Source/Pipelines/Private/InterchangePipelineFactories.cpp
@@ -36,7 +36,7 @@ UObject* UInterchangeBlueprintPipelineBaseFactory::FactoryCreateNew(UClass* Clas
 	{
 		FFormatNamedArguments Args;
 		Args.Add(TEXT("ClassName"), (ParentClass != nullptr) ? FText::FromString(ParentClass->GetName()) : NSLOCTEXT("UInterchangeBlueprintPipelineBaseFactory", "Null", "(null)"));
-		FMessageDialog::Open(EAppMsgType::Ok, FText::Format(NSLOCTEXT("UInterchangeBlueprintPipelineBaseFactory", "CannotCreateInterchangeBlueprintPipelineBase", "Cannot create a Interchange Blueprint Pipeline based on the class '{ClassName}'."), Args));
+		FMessageDialog::Open(EAppMsgType::Ok, FText::Format(NSLOCTEXT("UInterchangeBlueprintPipelineBaseFactory", "CannotCreateInterchangeBlueprintPipelineBase", "Cannot create an Interchange Blueprint Pipeline based on the class '{ClassName}'."), Args));
 		return nullptr;
 	}
 	else
diff --git a/Engine/Plugins/Interchange/Editor/Source/Pipelines/Private/SInterchangePipelineConfigurationDialog.cpp b/Engine/Plugins/Interchange/Editor/Source/Pipelines/Private/SInterchangePipelineConfigurationDialog.cpp
--- a/Engine/Plugins/Interchange/Editor/Source/Pipelines/Private/SInterchangePipelineConfigurationDialog.cpp
+++ b/Engine/Plugins/Interchange/Editor/Source/Pipelines/Private/SInterchangePipelineConfigurationDialog.cpp
@@ -612,11 +612,11 @@ void SInterchangePipelineConfigurationDialog::Construct(const FArguments& InArgs
 						{
 							if (bReimportClosure)
 							{
-								return LOCTEXT("SInterchangePipelineConfigurationDialog_ResetToPipelineAsset_TooltipReimport", "Reset the selected pipeline properties to the asset import data pipeline properties.");
+								return LOCTEXT("SInterchangePipelineConfigurationDialog_ResetToPipelineAsset_TooltipReimport", "Reset the selected pipeline to the values used the last time this asset was imported.");
 							}
 							else
 							{
-								return LOCTEXT("SInterchangePipelineConfigurationDialog_ResetToPipelineAsset_Tooltip", "Reset the selected pipeline properties to the stack pipeline properties.");
+								return LOCTEXT("SInterchangePipelineConfigurationDialog_ResetToPipelineAsset_Tooltip", "Reset the properties of the selected pipeline.");
 							}
 						})
 					.OnClicked(this, &SInterchangePipelineConfigurationDialog::OnResetToDefault)
diff --git a/Engine/Plugins/Interchange/Editor/Source/Pipelines/Public/InterchangeGraphInspectorPipeline.h b/Engine/Plugins/Interchange/Editor/Source/Pipelines/Public/InterchangeGraphInspectorPipeline.h
--- a/Engine/Plugins/Interchange/Editor/Source/Pipelines/Public/InterchangeGraphInspectorPipeline.h
+++ b/Engine/Plugins/Interchange/Editor/Source/Pipelines/Public/InterchangeGraphInspectorPipeline.h
@@ -7,13 +7,13 @@
 #include "InterchangeGraphInspectorPipeline.generated.h"
 
 /**
- * This pipeline is the generic pipeline option for all meshes type and should be call before specialized Mesh pipeline (like generic static mesh or skeletal mesh pipelines)
- * All shared import options between mesh type should be added here.
+ * This pipeline is the generic pipeline option for all types of meshes. It should be called before specialized mesh pipelines like the generic static mesh or skeletal mesh pipelines.
+ * All import options that are shared between mesh types should be added here.
  *
  * UPROPERTY possible meta values:
- * @meta ImportOnly - Boolean, the property is use only when we import not when we re-import. Cannot be mix with ReimportOnly!
- * @meta ReimportOnly - Boolean, the property is use only when we re-import not when we import. Cannot be mix with ImportOnly!
- * @meta MeshType - String, the property is for static or skeletal or both (static | skeletal) mesh type. If not specified it will apply to all mesh type.
+ * @meta ImportOnly - Boolean. The property is used only for import, not for reimport. Cannot be mixed with ReimportOnly.
+ * @meta ReimportOnly - Boolean. The property is used only for reimport, not for import. Cannot be mixed with ImportOnly.
+ * @meta MeshType - String. The property is for static mesh or skeletal mesh or both. If not specified, it will apply to all mesh types.
  */
 UCLASS(BlueprintType)
 class INTERCHANGEEDITORPIPELINES_API UInterchangeGraphInspectorPipeline : public UInterchangePipelineBase
@@ -29,7 +29,7 @@ public:
 	//////////////////////////////////////////////////////////////////////////
 
 	/*
-	 * This pipeline must never be save into any asset import data
+	 * This pipeline must never be saved into any asset import data.
 	 */
 	virtual bool SupportReimport() const override { return false; }
 
@@ -41,7 +41,7 @@ protected:
 
 	virtual bool CanExecuteOnAnyThread(EInterchangePipelineTask PipelineTask) override
 	{
-		//If a blueprint or python derived from this class, it will be execute on the game thread since we cannot currently execute script outside of the game thread, even if this function return true.
+		//If a Blueprint or Python class derived from this class, it will be executed on the game thread because we cannot currently execute script outside of the game thread, even if this function return true.
 		return false;
 	}
 
diff --git a/Engine/Plugins/Interchange/Runtime/Interchange.uplugin b/Engine/Plugins/Interchange/Runtime/Interchange.uplugin
--- a/Engine/Plugins/Interchange/Runtime/Interchange.uplugin
+++ b/Engine/Plugins/Interchange/Runtime/Interchange.uplugin
@@ -3,7 +3,7 @@
 	"Version": 1,
 	"VersionName": "1.0",
 	"FriendlyName": "Interchange Framework",
-	"Description": "Interchange framework plugin implement all engine importers and exporters.",
+	"Description": "The Interchange Framework plugin offers a customizable import and export system, with an extensible set of pipelines for handling common file types.",
 	"Category": "Importers",
 	"CreatedBy": "Epic Games, Inc.",
 	"CreatedByURL": "https://epicgames.com",
diff --git a/Engine/Plugins/Interchange/Runtime/Source/Dispatcher/Public/InterchangeDispatcherModule.h b/Engine/Plugins/Interchange/Runtime/Source/Dispatcher/Public/InterchangeDispatcherModule.h
--- a/Engine/Plugins/Interchange/Runtime/Source/Dispatcher/Public/InterchangeDispatcherModule.h
+++ b/Engine/Plugins/Interchange/Runtime/Source/Dispatcher/Public/InterchangeDispatcherModule.h
@@ -6,7 +6,7 @@
 #define INTERCHANGEDISPATCHER_MODULE_NAME TEXT("InterchangeDispatcher")
 
 /**
- * This module allow out of process interchange translator in case a third party SDK is not thread safe.
+ * This module allows out-of-process Interchange translators in case a third-party SDK is not thread-safe.
  */
 class FInterchangeDispatcherModule : public IModuleInterface
 {
diff --git a/Engine/Plugins/Interchange/Runtime/Source/FactoryNodes/Public/InterchangeAnimSequenceFactoryNode.h b/Engine/Plugins/Interchange/Runtime/Source/FactoryNodes/Public/InterchangeAnimSequenceFactoryNode.h
--- a/Engine/Plugins/Interchange/Runtime/Source/FactoryNodes/Public/InterchangeAnimSequenceFactoryNode.h
+++ b/Engine/Plugins/Interchange/Runtime/Source/FactoryNodes/Public/InterchangeAnimSequenceFactoryNode.h
@@ -39,15 +39,15 @@ public:
 	UInterchangeAnimSequenceFactoryNode();
 
 	/**
-	 * Initialize node data
-	 * @param UniqueID - The uniqueId for this node
-	 * @param DisplayLabel - The name of the node
+	 * Initialize node data.
+	 * @param UniqueID - The unique ID for this node.
+	 * @param DisplayLabel - The name of the node.
 	 */
 	UFUNCTION(BlueprintCallable, Category = "Interchange | Node | AnimSequence")
 	void InitializeAnimSequenceNode(const FString& UniqueID, const FString& DisplayLabel);
 
 	/**
-	 * Override serialize to restore SlotMaterialDependencies on load.
+	 * Override Serialize() to restore SlotMaterialDependencies on load.
 	 */
 	virtual void Serialize(FArchive& Ar) override
 	{
@@ -64,7 +64,7 @@ public:
 	}
 
 	/**
-	 * Return the node type name of the class, we use this when reporting error
+	 * Return the node type name of the class. This is used when reporting errors.
 	 */
 	virtual FString GetTypeName() const override;
 
@@ -74,16 +74,16 @@ public:
 	virtual bool ShouldHideAttribute(const UE::Interchange::FAttributeKey& NodeAttributeKey) const override;
 #endif //WITH_EDITOR
 
-	/** Get the class this node want to create */
+	/** Get the class this node creates. */
 	UFUNCTION(BlueprintCallable, Category = "Interchange | Node | AnimSequence")
 	virtual class UClass* GetObjectClass() const override;
 
 public:
-	/** Get the skeleton factory node unique id. Return false if the attribute is not set. */
+	/** Get the unique ID of the skeleton factory node. Return false if the attribute is not set. */
 	UFUNCTION(BlueprintCallable, Category = "Interchange | Node | AnimSequence")
 	bool GetCustomSkeletonFactoryNodeUid(FString& AttributeValue) const;
 
-	/** Set the skeleton factory node unique id. Return false if the attribute cannot be set. */
+	/** Set the unique ID of the skeleton factory node. Return false if the attribute cannot be set. */
 	UFUNCTION(BlueprintCallable, Category = "Interchange | Node | AnimSequence")
 	bool SetCustomSkeletonFactoryNodeUid(const FString& AttributeValue);
 
@@ -93,15 +93,15 @@ public:
 	 */
 
 	/**
-	 * Get the import bone tracks state. The attribute will be true if we need to import bone tracks.
-	 * False if we do not import bone tracks.
+	 * Get the import bone tracks state. If the attribute is true, bone tracks are imported. If the attribute 
+	 * is false, bone tracks are not imported.
 	 * 
-	 * Note - Return false if the attribute is not set. Return true if the attribute exist and can be query.
+	 * Return false if the attribute is not set. Return true if the attribute exists and can be queried.
 	 */
 	UFUNCTION(BlueprintCallable, Category = "Interchange | Node | AnimSequence")
 	bool GetCustomImportBoneTracks(bool& AttributeValue) const;
 
-	/** Set the import bone tracks state. Pass true to import bone tracks, false to not import bone tracks. */
+	/** Set the import bone tracks state. Pass true to import bone tracks, or false to not import bone tracks. */
 	UFUNCTION(BlueprintCallable, Category = "Interchange | Node | AnimSequence")
 	bool SetCustomImportBoneTracks(const bool& AttributeValue);
 
@@ -109,23 +109,23 @@ public:
 	UFUNCTION(BlueprintCallable, Category = "Interchange | Node | AnimSequence")
 	bool GetCustomImportBoneTracksSampleRate(double& AttributeValue) const;
 
-	/** Set the import bone tracks sample rate. Return false if the attribute cannot be set. */
+	/** Set the import bone tracks sample rate. Return false if the attribute could not be set. */
 	UFUNCTION(BlueprintCallable, Category = "Interchange | Node | AnimSequence")
 	bool SetCustomImportBoneTracksSampleRate(const double& AttributeValue);
 
-	/** Get the import bone tracks start time in second. Return false if the attribute is not set. */
+	/** Get the import bone tracks start time in seconds. Return false if the attribute is not set. */
 	UFUNCTION(BlueprintCallable, Category = "Interchange | Node | AnimSequence")
 	bool GetCustomImportBoneTracksRangeStart(double& AttributeValue) const;
 
-	/** Set the import bone tracks start time in second. Return false if the attribute cannot be set. */
+	/** Set the import bone tracks start time in seconds. Return false if the attribute could not be set. */
 	UFUNCTION(BlueprintCallable, Category = "Interchange | Node | AnimSequence")
 	bool SetCustomImportBoneTracksRangeStart(const double& AttributeValue);
 
-	/** Get the import bone tracks end time in second. Return false if the attribute is not set. */
+	/** Get the import bone tracks end time in seconds. Return false if the attribute is not set. */
 	UFUNCTION(BlueprintCallable, Category = "Interchange | Node | AnimSequence")
 	bool GetCustomImportBoneTracksRangeStop(double& AttributeValue) const;
 
-	/** Set the import bone tracks end time in second. Return false if the attribute cannot be set. */
+	/** Set the import bone tracks end time in seconds. Return false if the attribute could not be set. */
 	UFUNCTION(BlueprintCallable, Category = "Interchange | Node | AnimSequence")
 	bool SetCustomImportBoneTracksRangeStop(const double& AttributeValue);
 
@@ -140,58 +140,57 @@ public:
 
 
 	 /**
-	  * Get the import attribute curves state. If true this mean we want to import all user custom attributes
-	  * we can find on a node.
+	  * Get the import attribute curves state. If true, all user custom attributes on nodes are imported.
 	  * 
 	  * Return false if the attribute is not set.
 	  */
 	UFUNCTION(BlueprintCallable, Category = "Interchange | Node | AnimSequence")
 	bool GetCustomImportAttributeCurves(bool& AttributeValue) const;
 
-	/** Set the import attribute curves state. Return false if the attribute cannot be set. */
+	/** Set the import attribute curves state. Return false if the attribute could not be set. */
 	UFUNCTION(BlueprintCallable, Category = "Interchange | Node | AnimSequence")
 	bool SetCustomImportAttributeCurves(const bool& AttributeValue);
 
 	/**
-	 * Get the custom attribute DoNotImportCurveWithZero, return false if the attribute is not set.
+	 * Get the custom attribute DoNotImportCurveWithZero. Return false if the attribute is not set.
 	 * 
-	 * Note - If value is true, do not import if it doesn't have any value other than zero. This is to avoid adding extra curves to evaluate.
+	 * Note - If this attribute is enabled, only curves that have a value other than zero will be imported. This is to avoid adding extra curves to evaluate.
 	 */
 	UFUNCTION(BlueprintCallable, Category = "Interchange | Node | AnimSequence")
 	bool GetCustomDoNotImportCurveWithZero(bool& AttributeValue) const;
 
 	/**
-	 * Set the custom attribute DoNotImportCurveWithZero. Return false if the attribute cannot be set.
+	 * Set the custom attribute DoNotImportCurveWithZero. Return false if the attribute could not be set.
 	 * 
-	 * Note - If value is true, do not import if it doesn't have any value other than zero. This is to avoid adding extra curves to evaluate.
+	 * Note - If this attribute is enabled, only curves that have a value other than zero will be imported. This is to avoid adding extra curves to evaluate.
 	 */
 	UFUNCTION(BlueprintCallable, Category = "Interchange | Node | AnimSequence")
 	bool SetCustomDoNotImportCurveWithZero(const bool& AttributeValue);
 
 	/**
-	 * Get the custom attribute AddCurveMetadataToSkeleton, return false if the attribute is not set.
+	 * Get the custom attribute AddCurveMetadataToSkeleton. Return false if the attribute is not set.
 	 * 
-	 * Note - If value is true, do not import if it doesn't have any value other than zero. This is to avoid dirtying the skeleton on import.
+	 * Note - If this setting is disabled, curve metadata will be added to skeletal meshes for morph targets, but no metadata entry will be created for general curves.
 	 */
 	UFUNCTION(BlueprintCallable, Category = "Interchange | Node | AnimSequence")
 	bool GetCustomAddCurveMetadataToSkeleton(bool& AttributeValue) const;
 
 	/**
-	 * Set the custom attribute AddCurveMetadataToSkeleton. Return false if the attribute cannot be set.
+	 * Set the custom attribute AddCurveMetadataToSkeleton. Return false if the attribute could not be set.
 	 * 
-	 * Note - If value is true, do not import if it doesn't have any value other than zero. This is to avoid dirtying the skeleton on import.
+	 * Note - If this setting is disabled, curve metadata will be added to skeletal meshes for morph targets, but no metadata entry will be created for general curves.
 	 */
 	UFUNCTION(BlueprintCallable, Category = "Interchange | Node | AnimSequence")
 	bool SetCustomAddCurveMetadataToSkeleton(const bool& AttributeValue);
 	
 	/**
-	 * Get the custom attribute RemoveCurveRedundantKeys, return false if the attribute is not set.
+	 * Get the custom attribute RemoveCurveRedundantKeys. Return false if the attribute is not set.
 	 */
 	UFUNCTION(BlueprintCallable, Category = "Interchange | Node | AnimSequence")
 	bool GetCustomRemoveCurveRedundantKeys(bool& AttributeValue) const;
 
 	/**
-	 * Set the custom attribute RemoveCurveRedundantKeys. Return false if the attribute cannot be set.
+	 * Set the custom attribute RemoveCurveRedundantKeys. Return false if the attribute could not be set.
 	 */
 	UFUNCTION(BlueprintCallable, Category = "Interchange | Node | AnimSequence")
 	bool SetCustomRemoveCurveRedundantKeys(const bool& AttributeValue);
@@ -202,21 +201,21 @@ public:
 	/*********************************************************************************************************
 	 * Morph target curve API Begin
 	 *
-	 * Note - Morpgh target curve payload is FRichCurve.
+	 * Note: Morph target curve payload is FRichCurve.
 	 */
 
 	/**
-	 * Get the custom attribute DeleteExistingMorphTargetCurves, return false if the attribute is not set.
+	 * Get the custom attribute DeleteExistingMorphTargetCurves. Return false if the attribute is not set.
 	 * 
-	 * Note - If true, all previous moprh target curves will be deleted when doing a re-import.
+	 * Note: If true, all previous morph target curves are deleted if you reimport.
 	 */
 	UFUNCTION(BlueprintCallable, Category = "Interchange | Node | AnimSequence")
 	bool GetCustomDeleteExistingMorphTargetCurves(bool& AttributeValue) const;
 
 	/**
-	 * Set the custom attribute DeleteExistingMorphTargetCurves. Return false if the attribute cannot be set.
+	 * Set the custom attribute DeleteExistingMorphTargetCurves. Return false if the attribute could not be set.
 	 * 
-	 * Note - If true, all previous moprh target curves will be deleted when doing a re-import.
+	 * Note: If true, all previous morph target curves are deleted if you reimport.
 	 */
 	UFUNCTION(BlueprintCallable, Category = "Interchange | Node | AnimSequence")
 	bool SetCustomDeleteExistingMorphTargetCurves(const bool& AttributeValue);
@@ -234,7 +233,7 @@ public:
 	 *        Attribute curves are import has float FRichCurve
 	 */
 
-	 /** Return how many animated attribute curve names this anim sequence drive (curve are FRichCurve of type float). */
+	 /** Return the number of animated attribute curve names this anim sequence drives. Curves are FRichCurve of type float. */
 	UFUNCTION(BlueprintCallable, Category = "Interchange | Node | AnimSequence")
 	int32 GetAnimatedAttributeCurveNamesCount() const;
 
@@ -242,7 +241,7 @@ public:
 	UFUNCTION(BlueprintCallable, Category = "Interchange | Node | AnimSequence")
 	void GetAnimatedAttributeCurveNames(TArray<FString>& OutAttributeCurveNames) const;
 
-	/** Get an animated attribute curve name point by the specified index. */
+	/** Get the animated attribute curve name at the specified index. */
 	UFUNCTION(BlueprintCallable, Category = "Interchange | Node | AnimSequence")
 	void GetAnimatedAttributeCurveName(const int32 Index, FString& OutAttributeCurveName) const;
 
@@ -250,7 +249,7 @@ public:
 	UFUNCTION(BlueprintCallable, Category = "Interchange | Node | AnimSequence")
 	bool SetAnimatedAttributeCurveName(const FString& AttributeCurveName);
 
-	/** Remove one animated attribute curve name. */
+	/** Remove the specified animated attribute curve name. */
 	UFUNCTION(BlueprintCallable, Category = "Interchange | Node | AnimSequence")
 	bool RemoveAnimatedAttributeCurveName(const FString& AttributeCurveName);
 
@@ -263,26 +262,26 @@ public:
 	 /*********************************************************************************************************
 	  * Material curve API Begin
 	  *
-	  * Note - Material curve are attribute curve that can animate a material parameter.
+	  * Note - Material curves are attribute curves that can animate a material parameter.
 	  */
 
 	/**
-	 * Get the custom attribute MaterialDriveParameterOnCustomAttribute, return false if the attribute is not set.
+	 * Get the custom attribute MaterialDriveParameterOnCustomAttribute. Return false if the attribute is not set.
 	 *
-	 * Note - If true, Set Material Curve Type for all custom attributes that exists.
+	 * Note: If true, sets Material Curve Type for all custom attributes.
 	 */
 	UFUNCTION(BlueprintCallable, Category = "Interchange | Node | AnimSequence")
 	bool GetCustomMaterialDriveParameterOnCustomAttribute(bool& AttributeValue) const;
 
 	/**
-	 * Set the custom attribute MaterialDriveParameterOnCustomAttribute. Return false if the attribute cannot be set.
+	 * Set the custom attribute MaterialDriveParameterOnCustomAttribute. Return false if the attribute could not be set.
 	 *
-	 * Note - If true, Set Material Curve Type for all custom attributes that exists.
+	 * Note: If true, sets Material Curve Type for all custom attributes.
 	 */
 	UFUNCTION(BlueprintCallable, Category = "Interchange | Node | AnimSequence")
 	bool SetCustomMaterialDriveParameterOnCustomAttribute(const bool& AttributeValue);
 
-	/** Return how many animated material curve suffixes this anim sequence drive (curve are FRichCurve of type float). */
+	/** Return the number of animated material curve suffixes this anim sequence drives. Curves are FRichCurve of type float. */
 	UFUNCTION(BlueprintCallable, Category = "Interchange | Node | AnimSequence")
 	int32 GetAnimatedMaterialCurveSuffixesCount() const;
 
@@ -290,15 +289,15 @@ public:
 	UFUNCTION(BlueprintCallable, Category = "Interchange | Node | AnimSequence")
 	void GetAnimatedMaterialCurveSuffixes(TArray<FString>& OutMaterialCurveSuffixes) const;
 
-	/** Get an animated material curve suffixe point by the specified index. */
+	/** Get the animated material curve suffix with the specified index. */
 	UFUNCTION(BlueprintCallable, Category = "Interchange | Node | AnimSequence")
 	void GetAnimatedMaterialCurveSuffixe(const int32 Index, FString& OutMaterialCurveSuffixe) const;
 
-	/** Add an animated material curve suffixe. */
+	/** Add an animated material curve suffix. */
 	UFUNCTION(BlueprintCallable, Category = "Interchange | Node | AnimSequence")
 	bool SetAnimatedMaterialCurveSuffixe(const FString& MaterialCurveSuffixe);
 
-	/** Remove one animated material curve suffixe. */
+	/** Remove the specified animated material curve suffix. */
 	UFUNCTION(BlueprintCallable, Category = "Interchange | Node | AnimSequence")
 	bool RemoveAnimatedMaterialCurveSuffixe(const FString& MaterialCurveSuffixe);
 
@@ -311,12 +310,12 @@ public:
 	/*********************************************************************************************************
 	 * Attribute step curve API Begin
 	 * 
-	 * Note - Attribute step curve payload information can be retrieve via the UInterchangeUserDefinedAttributesAPI.
-	 *        Attribute step curves are import has a TArray<float> for key time and a TArray<ValueType> for the value.
+	 * Note - Attribute step curve payload information can be retrieved via the UInterchangeUserDefinedAttributesAPI.
+	 *        Imported attribute step curves have a TArray<float> for the key times and a TArray<ValueType> for the values.
 	 *        Supported value type are: int32, float, FString.
 	 */
 
-	 /** Return how many animated attribute step curve names this anim sequence drive. */
+	 /** Return the number of animated attribute step curve names this anim sequence drives. */
 	UFUNCTION(BlueprintCallable, Category = "Interchange | Node | AnimSequence")
 	int32 GetAnimatedAttributeStepCurveNamesCount() const;
 
@@ -324,7 +323,7 @@ public:
 	UFUNCTION(BlueprintCallable, Category = "Interchange | Node | AnimSequence")
 	void GetAnimatedAttributeStepCurveNames(TArray<FString>& OutAttributeStepCurveNames) const;
 
-	/** Get an animated attribute step curve name point by the specified index. */
+	/** Get the animated attribute step curve name at the specified index. */
 	UFUNCTION(BlueprintCallable, Category = "Interchange | Node | AnimSequence")
 	void GetAnimatedAttributeStepCurveName(const int32 Index, FString& OutAttributeStepCurveName) const;
 
@@ -332,7 +331,7 @@ public:
 	UFUNCTION(BlueprintCallable, Category = "Interchange | Node | AnimSequence")
 	bool SetAnimatedAttributeStepCurveName(const FString& AttributeStepCurveName);
 
-	/** Remove one animated attribute step curve name. */
+	/** Remove the specified animated attribute step curve name. */
 	UFUNCTION(BlueprintCallable, Category = "Interchange | Node | AnimSequence")
 	bool RemoveAnimatedAttributeStepCurveName(const FString& AttributeStepCurveName);
 
@@ -343,17 +342,17 @@ public:
 
 
 	/**
-	 * Get the custom attribute DeleteExistingCustomAttributeCurves, return false if the attribute is not set.
+	 * Get the custom attribute DeleteExistingCustomAttributeCurves. Return false if the attribute is not set.
 	 * 
-	 * Note - If true, all previous custom attribute curves will be deleted when doing a re-import.
+	 * Note - If true, all previous custom attribute curves are deleted if you reimport.
 	 */
 	UFUNCTION(BlueprintCallable, Category = "Interchange | Node | AnimSequence")
 	bool GetCustomDeleteExistingCustomAttributeCurves(bool& AttributeValue) const;
 
 	/**
-	 * Set the custom attribute DeleteExistingCustomAttributeCurves. Return false if the attribute cannot be set.
+	 * Set the custom attribute DeleteExistingCustomAttributeCurves. Return false if the attribute could not be set.
 	 * 
-	 * Note - If true, all previous custom attribute curves will be deleted when doing a re-import.
+	 * Note - If true, all previous custom attribute curves are deleted if you reimport.
 	 */
 	UFUNCTION(BlueprintCallable, Category = "Interchange | Node | AnimSequence")
 	bool SetCustomDeleteExistingCustomAttributeCurves(const bool& AttributeValue);
@@ -364,34 +363,34 @@ public:
 
 
 	/**
-	 * Get the custom attribute DeleteExistingNonCurveCustomAttributes, return false if the attribute is not set.
+	 * Get the custom attribute DeleteExistingNonCurveCustomAttributes. Return false if the attribute is not set.
 	 * 
-	 * Note - If true, all previous non-curve custom attributes will be deleted when doing a re-import.
+	 * Note - If true, all previous non-curve custom attributes are deleted if you reimport.
 	 */
 	UFUNCTION(BlueprintCallable, Category = "Interchange | Node | AnimSequence")
 	bool GetCustomDeleteExistingNonCurveCustomAttributes(bool& AttributeValue) const;
 
 	/**
-	 * Set the custom attribute DeleteExistingNonCurveCustomAttributes. Return false if the attribute cannot be set.
+	 * Set the custom attribute DeleteExistingNonCurveCustomAttributes. Return false if the attribute could not be set.
 	 * 
-	 * Note - If true, all previous non-curve custom attributes will be deleted when doing a re-import.
+	 * Note - If true, all previous non-curve custom attributes are deleted if you reimport.
 	 */
 	UFUNCTION(BlueprintCallable, Category = "Interchange | Node | AnimSequence")
 	bool SetCustomDeleteExistingNonCurveCustomAttributes(const bool& AttributeValue);
 
 	/**
-	 * Query the optional existing USkeleton this anim must use. The anim sequence factory will use this skeleton instead of the imported one
-	 * (from GetCustomSkeletonFactoryNodeUid) if this attribute is set and the skeleton pointer is valid.
-	 * Pipeline set this attribute in case the user want to specify an existing skeleton.
+	 * Query the optional existing USkeleton this animation must use. If this attribute is set and the skeleton is valid, 
+	 * the AnimSequence factory uses this skeleton instead of the one imported from GetCustomSkeletonFactoryNodeUid.
+	 * Pipelines set this attribute when the user wants to specify an existing skeleton.
 	 * Return false if the attribute was not set.
 	 */
 	UFUNCTION(BlueprintCallable, Category = "Interchange | Node | AnimSequence")
 	bool GetCustomSkeletonSoftObjectPath(FSoftObjectPath& AttributeValue) const;
 
 	/**
-	 * Set the optional existing USkeleton this anim must use. The AnimSequence factory will use this skeleton instead of the imported one
-	 * (from GetCustomSkeletonFactoryNodeUid) if this attribute is set and the skeleton pointer is valid.
-	 * Pipeline set this attribute in case the user want to specify an existing skeleton.
+	 * Set the optional existing USkeleton this animation must use. If this attribute is set and the skeleton is valid, 
+	 * the AnimSequence factory uses this skeleton instead of the one imported from GetCustomSkeletonFactoryNodeUid.
+	 * Pipelines set this attribute when the user wants to specify an existing skeleton.
 	 */
 	UFUNCTION(BlueprintCallable, Category = "Interchange | Node | AnimSequence")
 	bool SetCustomSkeletonSoftObjectPath(const FSoftObjectPath& AttributeValue);
diff --git a/Engine/Plugins/Interchange/Runtime/Source/FactoryNodes/Public/InterchangeCommonPipelineDataFactoryNode.h b/Engine/Plugins/Interchange/Runtime/Source/FactoryNodes/Public/InterchangeCommonPipelineDataFactoryNode.h
--- a/Engine/Plugins/Interchange/Runtime/Source/FactoryNodes/Public/InterchangeCommonPipelineDataFactoryNode.h
+++ b/Engine/Plugins/Interchange/Runtime/Source/FactoryNodes/Public/InterchangeCommonPipelineDataFactoryNode.h
@@ -9,32 +9,32 @@
 class UInterchangeBaseNodeContainer;
 
 
-/* This factory node is the place where pipeline can set global data that can be use by factories. */
+/* This factory node is where pipelines can set global data that can be used by factories. */
 UCLASS(BlueprintType)
 class INTERCHANGEFACTORYNODES_API UInterchangeCommonPipelineDataFactoryNode : public UInterchangeFactoryBaseNode
 {
 	GENERATED_BODY()
 
 public:
-	/* The pipelines that want to modify the common data shoiuld ensure they create the unique common pipeline node. */
+	/* Pipelines that want to modify the common data should ensure they create the unique common pipeline node. */
 	static UInterchangeCommonPipelineDataFactoryNode* FindOrCreateUniqueInstance(UInterchangeBaseNodeContainer* NodeContainer);
 	
-	/* If the unique instance doesn't exist it will return nullptr. This function should be use by the factories, to avoid creating a node. */
+	/* If the unique instance doesn't exist, this will return nullptr. This function should be use by the factories, to avoid creating a node. */
 	static UInterchangeCommonPipelineDataFactoryNode* GetUniqueInstance(const UInterchangeBaseNodeContainer* NodeContainer);
 
 	/** Return the global offset transform set by the pipelines. */
 	UFUNCTION(BlueprintCallable, Category = "Interchange | Node | Common Pipeline Data")
 	bool GetCustomGlobalOffsetTransform(FTransform& AttributeValue) const;
 
-	/** Pipeline can set a global transform, factories will use this global offset when importing asset. */
+	/** Pipelines can set a global transform. Factories will use this global offset when importing assets. */
 	UFUNCTION(BlueprintCallable, Category = "Interchange | Node | Common Pipeline Data")
 	bool SetCustomGlobalOffsetTransform(const UInterchangeBaseNodeContainer* NodeContainer, const FTransform& AttributeValue);
 
-	/** Return the Bake Meshes set by the pipelines. */
+	/** Return the value of the Bake Meshes setting set by the pipelines. */
 	UFUNCTION(BlueprintCallable, Category = "Interchange | Node | Common Pipeline Data")
 	bool GetBakeMeshes(bool& AttributeValue) const;
 
-	/** Pipeline can set Bake Meshes, factories will use this to identify if Global transforms should be applied to Meshes/Skeletals. */
+	/** Pipelines can set this Bake Meshes setting. Factories use this to identify whether they should apply global transforms to static meshes and skeletal meshes. */
 	UFUNCTION(BlueprintCallable, Category = "Interchange | Node | Common Pipeline Data")
 	bool SetBakeMeshes(const UInterchangeBaseNodeContainer* NodeContainer, const bool& AttributeValue);
 
diff --git a/Engine/Plugins/Interchange/Runtime/Source/FactoryNodes/Public/InterchangeMeshActorFactoryNode.h b/Engine/Plugins/Interchange/Runtime/Source/FactoryNodes/Public/InterchangeMeshActorFactoryNode.h
--- a/Engine/Plugins/Interchange/Runtime/Source/FactoryNodes/Public/InterchangeMeshActorFactoryNode.h
+++ b/Engine/Plugins/Interchange/Runtime/Source/FactoryNodes/Public/InterchangeMeshActorFactoryNode.h
@@ -15,7 +15,7 @@ public:
 	UInterchangeMeshActorFactoryNode();
 
 	/**
-	 * Override serialize to restore SlotMaterialDependencies on load.
+	 * Override Serialize() to restore SlotMaterialDependencies on load.
 	 */
 	virtual void Serialize(FArchive& Ar) override
 	{
@@ -28,41 +28,41 @@ public:
 	}
 
 	/**
-	 * Allow to retrieve the correspondence table between slot names and assigned materials for this object.
+	 * Retrieve the correspondence table between slot names and assigned materials for this object.
 	 */
 	UFUNCTION(BlueprintCallable, Category = "Interchange | Node | StaticMesh")
 	void GetSlotMaterialDependencies(TMap<FString, FString>& OutMaterialDependencies) const;
 
 	/**
-	 * Allow to retrieve one Material dependency for a given slot of this object.
+	 * Retrieve the Material dependency for the specified slot of this object.
 	 */
 	UFUNCTION(BlueprintCallable, Category = "Interchange | Node | StaticMesh")
 	bool GetSlotMaterialDependencyUid(const FString& SlotName, FString& OutMaterialDependency) const;
 
 	/**
-	 * Add one Material dependency to a specific slot name of this object.
+	 * Add a Material dependency to the specified slot of this object.
 	 */
 	UFUNCTION(BlueprintCallable, Category = "Interchange | Node | StaticMesh")
 	bool SetSlotMaterialDependencyUid(const FString& SlotName, const FString& MaterialDependencyUid);
 
 	/**
-	 * Remove the Material dependency associated with the given slot name from this object.
+	 * Remove the Material dependency associated with the specified slot name from this object.
 	 */
 	UFUNCTION(BlueprintCallable, Category = "Interchange | Node | StaticMesh")
 	bool RemoveSlotMaterialDependencyUid(const FString& SlotName);
 
-	/** Set the Animation Asset To Play by this Scene Node. (only relevant for SkeletalMeshActors (SceneNodes that are instantiating Skeletal Meshes)) */
+	/** Set the animation asset for this scene node to play. (This is only relevant for SkeletalMeshActors: scene nodes that are instantiating skeletal meshes.) */
 	UFUNCTION(BlueprintCallable, Category = "Interchange | Node | SkeletalMesh")
 	bool SetCustomAnimationAssetUidToPlay(const FString& AttributeValue);
-	/** Get the Animation Asset To Play by this Scene Node. */
+	/** Get the animation asset set for this scene node to play. */
 	UFUNCTION(BlueprintCallable, Category = "Interchange | Node | SkeletalMesh")
 	bool GetCustomAnimationAssetUidToPlay(FString& AttributeValue) const;
 
-	/** Get the geometric offset. Any mesh attach to this scene node will be offset using this transform. */
+	/** Get the geometric offset. Any mesh attached to this scene node will be offset using this transform. */
 	UFUNCTION(BlueprintCallable, Category = "Interchange | Node | StaticMesh")
 	bool GetCustomGeometricTransform(FTransform& AttributeValue) const;
 
-	/** Set the geometric offset. Any mesh attach to this scene node will be offset using this transform. */
+	/** Set the geometric offset. Any mesh attached to this scene node will be offset using this transform. */
 	UFUNCTION(BlueprintCallable, Category = "Interchange | Node | StaticMesh")
 	bool SetCustomGeometricTransform(const FTransform& AttributeValue);
 
diff --git a/Engine/Plugins/Interchange/Runtime/Source/FactoryNodes/Public/InterchangeMeshFactoryNode.h b/Engine/Plugins/Interchange/Runtime/Source/FactoryNodes/Public/InterchangeMeshFactoryNode.h
--- a/Engine/Plugins/Interchange/Runtime/Source/FactoryNodes/Public/InterchangeMeshFactoryNode.h
+++ b/Engine/Plugins/Interchange/Runtime/Source/FactoryNodes/Public/InterchangeMeshFactoryNode.h
@@ -27,7 +27,7 @@ public:
 	UInterchangeMeshFactoryNode();
 
 	/**
-	 * Override serialize to restore SlotMaterialDependencies on load.
+	 * Override Serialize() to restore SlotMaterialDependencies on load.
 	 */
 	virtual void Serialize(FArchive& Ar) override
 	{
@@ -50,7 +50,7 @@ public:
 #endif //WITH_EDITOR
 
 public:
-	/** Return The number of LOD this static mesh has.*/
+	/** Return the number of LODs this static mesh has.*/
 	UFUNCTION(BlueprintCallable, Category = "Interchange | FactoryNode | Mesh")
 	int32 GetLodDataCount() const;
 
@@ -63,27 +63,27 @@ public:
 	UFUNCTION(BlueprintCallable, Category = "Interchange | FactoryNode | Mesh")
 	bool RemoveLodDataUniqueId(const FString& LodDataUniqueId);
 
-	/** Query whether the static mesh factory should replace the vertex color. Return false if the attribute was not set.*/
+	/** Query whether the static mesh factory should replace the vertex color. Return false if the attribute was not set. */
 	UFUNCTION(BlueprintCallable, Category = "Interchange | FactoryNode | Mesh")
 	bool GetCustomVertexColorReplace(bool& AttributeValue) const;
 
-	/** Set whether the static mesh factory should replace the vertex color. Return false if the attribute cannot be set.*/
+	/** Set whether the static mesh factory should replace the vertex color. Return false if the attribute could not be set. */
 	UFUNCTION(BlueprintCallable, Category = "Interchange | FactoryNode | Mesh")
 	bool SetCustomVertexColorReplace(const bool& AttributeValue);
 
-	/** Query whether the static mesh factory should ignore the vertex color. Return false if the attribute was not set.*/
+	/** Query whether the static mesh factory should ignore the vertex color. Return false if the attribute was not set. */
 	UFUNCTION(BlueprintCallable, Category = "Interchange | FactoryNode | Mesh")
 	bool GetCustomVertexColorIgnore(bool& AttributeValue) const;
 
-	/** Set whether the static mesh factory should ignore the vertex color. Return false if the attribute cannot be set.*/
+	/** Set whether the static mesh factory should ignore the vertex color. Return false if the attribute could not be set. */
 	UFUNCTION(BlueprintCallable, Category = "Interchange | FactoryNode | Mesh")
 	bool SetCustomVertexColorIgnore(const bool& AttributeValue);
 
-	/** Query whether the static mesh factory should override the vertex color. Return false if the attribute was not set.*/
+	/** Query whether the static mesh factory should override the vertex color. Return false if the attribute was not set. */
 	UFUNCTION(BlueprintCallable, Category = "Interchange | FactoryNode | Mesh")
 	bool GetCustomVertexColorOverride(FColor& AttributeValue) const;
 
-	/** Set whether the static mesh factory should override the vertex color. Return false if the attribute cannot be set.*/
+	/** Set whether the static mesh factory should override the vertex color. Return false if the attribute could not be set. */
 	UFUNCTION(BlueprintCallable, Category = "Interchange | FactoryNode | Mesh")
 	bool SetCustomVertexColorOverride(const FColor& AttributeValue);
 
@@ -95,19 +95,19 @@ public:
 	UFUNCTION(BlueprintCallable, Category = "Interchange | FactoryNode | Mesh")
 	bool SetCustomKeepSectionsSeparate(const bool& AttributeValue);
 
-	/** Allow to retrieve the correspondence table between slot names and assigned materials for this object. */
+	/** Retrieve the correspondence table between slot names and assigned materials for this object. */
 	UFUNCTION(BlueprintCallable, Category = "Interchange | FactoryNode | Mesh")
 	void GetSlotMaterialDependencies(TMap<FString, FString>& OutMaterialDependencies) const;
 
-	/** Allow to retrieve one Material dependency for a given slot of this object. */
+	/** Retrieve the Material dependency for the specified slot of this object. */
 	UFUNCTION(BlueprintCallable, Category = "Interchange | FactoryNode | Mesh")
 	bool GetSlotMaterialDependencyUid(const FString& SlotName, FString& OutMaterialDependency) const;
 
-	/** Add one Material dependency to a specific slot name of this object. */
+	/** Add a Material dependency to the specified slot of this object. */
 	UFUNCTION(BlueprintCallable, Category = "Interchange | FactoryNode | Mesh")
 	bool SetSlotMaterialDependencyUid(const FString& SlotName, const FString& MaterialDependencyUid);
 
-	/** Remove the Material dependency associated with the given slot name from this object. */
+	/** Remove the Material dependency associated with the specfied slot name of this object. */
 	UFUNCTION(BlueprintCallable, Category = "Interchange | FactoryNode | Mesh")
 	bool RemoveSlotMaterialDependencyUid(const FString& SlotName);
 
@@ -115,19 +115,19 @@ public:
 	UFUNCTION(BlueprintCallable, Category = "Interchange | FactoryNode | Mesh")
 	bool ResetSlotMaterialDependencies();
 
-	/** Query whether normals in the imported mesh are ignored and recomputed. When normals are recomputed the tangents are also recomputed. */
+	/** Query whether a custom LOD group is set for the mesh. */
 	UFUNCTION(BlueprintCallable, Category = "Interchange | FactoryNode | Mesh")
 	bool GetCustomLODGroup(FName& AttributeValue) const;
 
-	/** Set whether normals in the imported mesh are ignored and recomputed. When normals are recomputed the tangents are also recomputed. */
+	/** Set a custom LOD group for the mesh. */
 	UFUNCTION(BlueprintCallable, Category = "Interchange | FactoryNode | Mesh")
 	bool SetCustomLODGroup(const FName& AttributeValue, bool bAddApplyDelegate = true);
 
-	/** Query whether normals in the imported mesh are ignored and recomputed. When normals are recomputed the tangents are also recomputed. */
+	/** Query whether normals in the imported mesh are ignored and recomputed. When normals are recomputed, the tangents are also recomputed. */
 	UFUNCTION(BlueprintCallable, Category = "Interchange | FactoryNode | Mesh")
 	bool GetCustomRecomputeNormals(bool& AttributeValue) const;
 
-	/** Set whether normals in the imported mesh are ignored and recomputed. When normals are recomputed the tangents are also recomputed. */
+	/** Set whether normals in the imported mesh are ignored and recomputed. When normals are recomputed, the tangents are also recomputed. */
 	UFUNCTION(BlueprintCallable, Category = "Interchange | FactoryNode | Mesh")
 	bool SetCustomRecomputeNormals(const bool& AttributeValue, bool bAddApplyDelegate = true);
 	
@@ -139,51 +139,51 @@ public:
 	UFUNCTION(BlueprintCallable, Category = "Interchange | FactoryNode | Mesh")
 	bool SetCustomRecomputeTangents(const bool& AttributeValue, bool bAddApplyDelegate = true);
 
-	/** Query whether tangents recompute will use mikkt space if tangents are recomputed. */
+	/** Query whether tangents are recomputed using MikkTSpace when they need to be recomputed. */
 	UFUNCTION(BlueprintCallable, Category = "Interchange | FactoryNode | Mesh")
 	bool GetCustomUseMikkTSpace(bool& AttributeValue) const;
 
-	/** Set whether tangents recompute will use mikkt space if tangents are recomputed. */
+	/** Set whether tangents are recomputed using MikkTSpace when they need to be recomputed. */
 	UFUNCTION(BlueprintCallable, Category = "Interchange | FactoryNode | Mesh")
 	bool SetCustomUseMikkTSpace(const bool& AttributeValue, bool bAddApplyDelegate = true);
 
-	/** Query whether normals recompute will use the surface area and the corner angle of the triangle as a ratio when computing the normals. */
+	/** Query whether normals are recomputed by weighting the surface area and the corner angle of the triangle as a ratio. */
 	UFUNCTION(BlueprintCallable, Category = "Interchange | FactoryNode | Mesh")
 	bool GetCustomComputeWeightedNormals(bool& AttributeValue) const;
 
-	/** Set whether normals recompute will use the surface area and the corner angle of the triangle as a ratio when computing the normals. */
+	/** Set whether normals are recomputed by weighting the surface area and the corner angle of the triangle as a ratio. */
 	UFUNCTION(BlueprintCallable, Category = "Interchange | FactoryNode | Mesh")
 	bool SetCustomComputeWeightedNormals(const bool& AttributeValue, bool bAddApplyDelegate = true);
 
-	/** Query whether Tangents will be stored at 16 bit vs 8 bit precision. */
+	/** Query whether tangents are stored at 16-bit precision instead of the default 8-bit precision. */
 	UFUNCTION(BlueprintCallable, Category = "Interchange | FactoryNode | Mesh")
 	bool GetCustomUseHighPrecisionTangentBasis(bool& AttributeValue) const;
 
-	/** Set whether Tangents will be stored at 16 bit vs 8 bit precision. */
+	/** Set whether tangents are stored at 16-bit precision instead of the default 8-bit precision. */
 	UFUNCTION(BlueprintCallable, Category = "Interchange | FactoryNode | Mesh")
 	bool SetCustomUseHighPrecisionTangentBasis(const bool& AttributeValue, bool bAddApplyDelegate = true);
 
-	/** Query whether UVs will be stored at full floating point precision. */
+	/** Query whether UVs are stored at full floating point precision. */
 	UFUNCTION(BlueprintCallable, Category = "Interchange | FactoryNode | Mesh")
 	bool GetCustomUseFullPrecisionUVs(bool& AttributeValue) const;
 
-	/** Set whether UVs will be stored at full floating point precision. */
+	/** Set whether UVs are stored at full floating point precision. */
 	UFUNCTION(BlueprintCallable, Category = "Interchange | FactoryNode | Mesh")
 	bool SetCustomUseFullPrecisionUVs(const bool& AttributeValue, bool bAddApplyDelegate = true);
 
-	/** Query whether UVs will be stored at full floating point precision. */
+	/** Query whether UVs are converted to 16-bit by a legacy truncation process instead of the default rounding process. This may avoid differences when reimporting older content. */
 	UFUNCTION(BlueprintCallable, Category = "Interchange | FactoryNode | Mesh")
 	bool GetCustomUseBackwardsCompatibleF16TruncUVs(bool& AttributeValue) const;
 
-	/** Set whether UVs will be stored at full floating point precision. */
+	/** Set whether UVs are converted to 16-bit by a legacy truncation process instead of the default rounding process. This may avoid differences when reimporting older content. */
 	UFUNCTION(BlueprintCallable, Category = "Interchange | FactoryNode | Mesh")
 	bool SetCustomUseBackwardsCompatibleF16TruncUVs(const bool& AttributeValue, bool bAddApplyDelegate = true);
 
-	/** Query whether degenerate triangles will be removed. */
+	/** Query whether degenerate triangles are removed. */
 	UFUNCTION(BlueprintCallable, Category = "Interchange | FactoryNode | Mesh")
 	bool GetCustomRemoveDegenerates(bool& AttributeValue) const;
 
-	/** Set whether degenerate triangles will be removed. */
+	/** Set whether degenerate triangles are removed. */
 	UFUNCTION(BlueprintCallable, Category = "Interchange | FactoryNode | Mesh")
 	bool SetCustomRemoveDegenerates(const bool& AttributeValue, bool bAddApplyDelegate = true);
 
diff --git a/Engine/Plugins/Interchange/Runtime/Source/FactoryNodes/Public/InterchangePhysicsAssetFactoryNode.h b/Engine/Plugins/Interchange/Runtime/Source/FactoryNodes/Public/InterchangePhysicsAssetFactoryNode.h
--- a/Engine/Plugins/Interchange/Runtime/Source/FactoryNodes/Public/InterchangePhysicsAssetFactoryNode.h
+++ b/Engine/Plugins/Interchange/Runtime/Source/FactoryNodes/Public/InterchangePhysicsAssetFactoryNode.h
@@ -19,9 +19,9 @@ class INTERCHANGEFACTORYNODES_API UInterchangePhysicsAssetFactoryNode : public U
 
 public:
 	/**
-	 * Initialize node data
-	 * @param: UniqueID - The uniqueId for this node
-	 * @param DisplayLabel - The name of the node
+	 * Initialize node data.
+	 * @param: UniqueID - The unique ID for this node.
+	 * @param DisplayLabel - The name of the node.
 	 * @param InAssetClass - The class the Skeleton factory will create for this node.
 	 */
 	UFUNCTION(BlueprintCallable, Category = "Interchange | Node | Skeleton")
@@ -48,7 +48,7 @@ public:
 	}
 
 	/**
-	 * Return the node type name of the class, we use this when reporting error
+	 * Return the node type name of the class. This is used when reporting errors.
 	 */
 	virtual FString GetTypeName() const override
 	{
@@ -56,7 +56,7 @@ public:
 		return TypeName;
 	}
 
-	/** Get the class this node want to create */
+	/** Get the class this node creates. */
 	UFUNCTION(BlueprintCallable, Category = "Interchange | Node | Skeleton")
 	virtual class UClass* GetObjectClass() const override
 	{
@@ -87,13 +87,13 @@ public:
 	}
 
 public:
-	/** Get skeletalMesh asset Uid used to create the data in the post pipeline step.*/
+	/** Get the Skeletal Mesh asset UID used to create the data in the post-pipeline step. */
 	UFUNCTION(BlueprintCallable, Category = "Interchange | Node | Skeleton")
 	bool GetCustomSkeletalMeshUid(FString& AttributeValue) const
 	{
 		IMPLEMENT_NODE_ATTRIBUTE_GETTER(SkeletalMeshUid, FString);
 	}
-	/** Set skeletalMesh asset Uid used to create the data in the post pipeline step.*/
+	/** Set the Skeletal Mesh asset UID used to create the data in the post-pipeline step. */
 	UFUNCTION(BlueprintCallable, Category = "Interchange | Node | Skeleton")
 	bool SetCustomSkeletalMeshUid(const FString& AttributeValue)
 	{
diff --git a/Engine/Plugins/Interchange/Runtime/Source/FactoryNodes/Public/InterchangeSceneVariantSetsFactoryNode.h b/Engine/Plugins/Interchange/Runtime/Source/FactoryNodes/Public/InterchangeSceneVariantSetsFactoryNode.h
--- a/Engine/Plugins/Interchange/Runtime/Source/FactoryNodes/Public/InterchangeSceneVariantSetsFactoryNode.h
+++ b/Engine/Plugins/Interchange/Runtime/Source/FactoryNodes/Public/InterchangeSceneVariantSetsFactoryNode.h
@@ -21,7 +21,7 @@ public:
 	}
 
 	/**
-	 * Return the node type name of the class, we use this when reporting errors
+	 * Return the node type name of the class. This is used when reporting errors.
 	 */
 	virtual FString GetTypeName() const override
 	{
@@ -29,36 +29,36 @@ public:
 		return TypeName;
 	}
 
-	/** Get the class this node want to create */
+	/** Get the class this node creates. */
 	UFUNCTION(BlueprintCallable, Category = "Interchange | Node | SceneVariantSetFactory")
 	virtual class UClass* GetObjectClass() const override;
 
 	/**
-	 * This function allow to retrieve the number of translated VariantSets' unique ids for this object.
+	 * Retrieve the number of unique IDs of all translated VariantSets for this object.
 	 */
 	UFUNCTION(BlueprintCallable, Category = "Interchange | Node | SceneVariantSetFactory")
 	int32 GetCustomVariantSetUidCount() const;
 
 	/**
-	 * This function allow to retrieve all the translated VariantSets' unique ids for this object.
+	 * Retrieve the unique IDs of all translated VariantSets for this object.
 	 */
 	UFUNCTION(BlueprintCallable, Category = "Interchange | Node | SceneVariantSetFactory")
 	void GetCustomVariantSetUids(TArray<FString>& OutVariantUids) const;
 
 	/**
-	 * This function allow to retrieve one translated VariantSet's unique id for this object.
+	 * Retrieve the UID of the VariantSet with the specified index.
 	 */
 	UFUNCTION(BlueprintCallable, Category = "Interchange | Node | SceneVariantSetFactory")
 	void GetCustomVariantSetUid(const int32 Index, FString& OutVariantUid) const;
 
 	/**
-	 * Add one translated VariantSet's unique id to this object.
+	 * Add a unique id of a translated VariantSet for this object.
 	 */
 	UFUNCTION(BlueprintCallable, Category = "Interchange | Node | SceneVariantSetFactory")
 	bool AddCustomVariantSetUid(const FString& VariantUid);
 
 	/**
-	 * Remove one translated VariantSet's unique id from this object.
+	 * Remove the specified unique ID of a translated VariantSet from this object.
 	 */
 	UFUNCTION(BlueprintCallable, Category = "Interchange | Node | SceneVariantSetFactory")
 	bool RemoveCustomVariantSetUid(const FString& VariantUid);
diff --git a/Engine/Plugins/Interchange/Runtime/Source/FactoryNodes/Public/InterchangeSkeletalMeshFactoryNode.h b/Engine/Plugins/Interchange/Runtime/Source/FactoryNodes/Public/InterchangeSkeletalMeshFactoryNode.h
--- a/Engine/Plugins/Interchange/Runtime/Source/FactoryNodes/Public/InterchangeSkeletalMeshFactoryNode.h
+++ b/Engine/Plugins/Interchange/Runtime/Source/FactoryNodes/Public/InterchangeSkeletalMeshFactoryNode.h
@@ -15,9 +15,9 @@
 UENUM(BlueprintType)
 enum class EInterchangeSkeletalMeshContentType : uint8
 {
-	All UMETA(DisplayName = "Geometry and Skinning Weights.", ToolTip = "Import all skelatl mesh content: geometry, skinning and weights."),
-	Geometry UMETA(DisplayName = "Geometry Only", ToolTip = "Import the skeletal mesh geometry only (will create a default skeleton, or map the geometry to the existing one). Morph and LOD can be imported with it."),
-	SkinningWeights UMETA(DisplayName = "Skinning Weights Only", ToolTip = "Import the skeletal mesh skinning and weights only (no geometry will be imported). Morph and LOD will not be imported with this settings."),
+	All UMETA(DisplayName = "Geometry and Skin Weights", ToolTip = "Imports all skeletal mesh content: geometry and skin weights."),
+	Geometry UMETA(DisplayName = "Geometry Only", ToolTip = "Imports the skeletal mesh geometry only. This creates a default skeleton, or maps the geometry to the existing one. You can import morph targets and LODs with the mesh."),
+	SkinningWeights UMETA(DisplayName = "Skin Weights Only", ToolTip = "Imports the skeletal mesh skin weights only. No geometry, morph targets, or LODs are imported."),
 	MAX,
 };
 
@@ -30,53 +30,53 @@ public:
 	UInterchangeSkeletalMeshFactoryNode();
 
 	/**
-	 * Initialize node data
-	 * @param: UniqueID - The uniqueId for this node
-	 * @param DisplayLabel - The name of the node
+	 * Initialize node data.
+	 * @param: UniqueID - The unique ID for this node.
+	 * @param DisplayLabel - The name of the node.
 	 * @param InAssetClass - The class the SkeletalMesh factory will create for this node.
 	 */
 	UFUNCTION(BlueprintCallable, Category = "Interchange | Node | SkeletalMesh")
 	void InitializeSkeletalMeshNode(const FString& UniqueID, const FString& DisplayLabel, const FString& InAssetClass);
 
 	/**
-	 * Return the node type name of the class, we use this when reporting error
+	 * Return the node type name of the class. This is used when reporting errors.
 	 */
 	virtual FString GetTypeName() const override;
 
-	/** Get the class this node want to create */
+	/** Get the class this node creates. */
 	UFUNCTION(BlueprintCallable, Category = "Interchange | Node | SkeletalMesh")
 	virtual class UClass* GetObjectClass() const override;
 
 public:
-	/** Query the skeletal mesh factory skeleton UObject. Return false if the attribute was not set.*/
+	/** Query the skeletal mesh factory skeleton UObject. Return false if the attribute was not set. */
 	UFUNCTION(BlueprintCallable, Category = "Interchange | Node | SkeletalMesh")
 	bool GetCustomSkeletonSoftObjectPath(FSoftObjectPath& AttributeValue) const;
 
-	/** Set the skeletal mesh factory skeleton UObject. Return false if the attribute cannot be set.*/
+	/** Set the skeletal mesh factory skeleton UObject. Return false if the attribute could not be set. */
 	UFUNCTION(BlueprintCallable, Category = "Interchange | Node | SkeletalMesh")
 	bool SetCustomSkeletonSoftObjectPath(const FSoftObjectPath& AttributeValue);
 
-	/** Query weather the skeletal mesh factory should create the morph target. Return false if the attribute was not set.*/
+	/** Query whether the skeletal mesh factory should create morph targets. Return false if the attribute was not set. */
 	UFUNCTION(BlueprintCallable, Category = "Interchange | Node | SkeletalMesh")
 	bool GetCustomImportMorphTarget(bool& AttributeValue) const;
 
-	/** Set weather the skeletal mesh factory should create the morph target. Return false if the attribute cannot be set.*/
+	/** Set whether the skeletal mesh factory should create morph targets. Return false if the attribute could not be set. */
 	UFUNCTION(BlueprintCallable, Category = "Interchange | Node | SkeletalMesh")
 	bool SetCustomImportMorphTarget(const bool& AttributeValue);
 
-	/** Query weather the skeletal mesh factory should import the vertex attributes. Return false if the attribute was not set.*/
+	/** Query whether the skeletal mesh factory should import vertex attributes. Return false if the attribute was not set.*/
 	UFUNCTION(BlueprintCallable, Category = "Interchange | Node | SkeletalMesh")
 	bool GetCustomImportVertexAttributes(bool& AttributeValue) const;
 
-	/** Set weather the skeletal mesh factory should import the vertex attributes. Return false if the attribute cannot be set.*/
+	/** Set whether the skeletal mesh factory should import vertex attributes. Return false if the attribute could not be set.*/
 	UFUNCTION(BlueprintCallable, Category = "Interchange | Node | SkeletalMesh")
 	bool SetCustomImportVertexAttributes(const bool& AttributeValue);
 
-	/** Query weather the skeletal mesh factory should create a physics asset. Return false if the attribute was not set.*/
+	/** Query whether the skeletal mesh factory should create a physics asset. Return false if the attribute was not set.*/
 	UFUNCTION(BlueprintCallable, Category = "Interchange | Node | SkeletalMesh")
 	bool GetCustomCreatePhysicsAsset(bool& AttributeValue) const;
 
-	/** Set weather the skeletal mesh factory should create a physics asset. Return false if the attribute cannot be set.*/
+	/** Set whether the skeletal mesh factory should create a physics asset. Return false if the attribute could not be set.*/
 	UFUNCTION(BlueprintCallable, Category = "Interchange | Node | SkeletalMesh")
 	bool SetCustomCreatePhysicsAsset(const bool& AttributeValue);
 
@@ -84,60 +84,60 @@ public:
 	UFUNCTION(BlueprintCallable, Category = "Interchange | Node | SkeletalMesh")
 	bool GetCustomPhysicAssetSoftObjectPath(FSoftObjectPath& AttributeValue) const;
 
-	/** Set a physics asset the skeletal mesh factory should use. Return false if the attribute cannot be set.*/
+	/** Set a physics asset the skeletal mesh factory should use. Return false if the attribute could not be set.*/
 	UFUNCTION(BlueprintCallable, Category = "Interchange | Node | SkeletalMesh")
 	bool SetCustomPhysicAssetSoftObjectPath(const FSoftObjectPath& AttributeValue);
 
-	/** Query the skeletal mesh import content type. The content type is use by the factory to import partial or full translated content. Return false if the attribute was not set.*/
+	/** Query the skeletal mesh import content type. This content type determines whether the factory imports partial or full translated content. Return false if the attribute was not set.*/
 	UFUNCTION(BlueprintCallable, Category = "Interchange | Node | SkeletalMesh")
 	bool GetCustomImportContentType(EInterchangeSkeletalMeshContentType& AttributeValue) const;
 
-	/** Set the skeletal mesh import content type. The content type is use by the factory to import partial or full translated content. Return false if the attribute cannot be set.*/
+	/** Set the skeletal mesh import content type. This content type determines whether the factory imports partial or full translated content. Return false if the attribute could not be set.*/
 	UFUNCTION(BlueprintCallable, Category = "Interchange | Node | SkeletalMesh")
 	bool SetCustomImportContentType(const EInterchangeSkeletalMeshContentType& AttributeValue);
 
-	/** Query the skeletal mesh UseHighPrecisionSkinWeights. */
+	/** Query the skeletal mesh UseHighPrecisionSkinWeights setting. */
 	UFUNCTION(BlueprintCallable, Category = "Interchange | Node | SkeletalMesh")
 	bool GetCustomUseHighPrecisionSkinWeights(bool& AttributeValue) const;
 
-	/** Set the skeletal mesh UseHighPrecisionSkinWeights. */
+	/** Set the skeletal mesh UseHighPrecisionSkinWeights setting. */
 	UFUNCTION(BlueprintCallable, Category = "Interchange | Node | SkeletalMesh")
 	bool SetCustomUseHighPrecisionSkinWeights(const bool& AttributeValue, bool bAddApplyDelegate = true);
 		
-	/** Query the skeletal mesh threshold use to decide if two vertex position are equal. */
+	/** Query the skeletal mesh threshold value that is used to decide whether two vertex positions are equal. */
 	UFUNCTION(BlueprintCallable, Category = "Interchange | Node | SkeletalMesh")
 	bool GetCustomThresholdPosition(float& AttributeValue) const;
 
-	/** Set the skeletal mesh threshold use to decide if two vertex position are equal. */
+	/** Set the skeletal mesh threshold value that is used to decide whether two vertex positions are equal. */
 	UFUNCTION(BlueprintCallable, Category = "Interchange | Node | SkeletalMesh")
 	bool SetCustomThresholdPosition(const float& AttributeValue, bool bAddApplyDelegate = true);
 
-	/** Query the skeletal mesh threshold use to decide if two normal, tangents or bi-normals are equal. */
+	/** Query the skeletal mesh threshold value that is used to decide whether two normals, tangents, or bi-normals are equal. */
 	UFUNCTION(BlueprintCallable, Category = "Interchange | Node | SkeletalMesh")
 	bool GetCustomThresholdTangentNormal(float& AttributeValue) const;
 
-	/** Set the skeletal mesh threshold use to decide if two normal, tangents or bi-normals are equal. */
+	/** Set the skeletal mesh threshold value that is used to decide whether two normals, tangents, or bi-normals are equal. */
 	UFUNCTION(BlueprintCallable, Category = "Interchange | Node | SkeletalMesh")
 	bool SetCustomThresholdTangentNormal(const float& AttributeValue, bool bAddApplyDelegate = true);
 
-	/** Query the skeletal mesh threshold use to decide if two UVs are equal. */
+	/** Query the skeletal mesh threshold value that is used to decide whether two UVs are equal. */
 	UFUNCTION(BlueprintCallable, Category = "Interchange | Node | SkeletalMesh")
 	bool GetCustomThresholdUV(float& AttributeValue) const;
 
-	/** Set the skeletal mesh threshold use to decide if two UVs are equal. */
+	/** Set the skeletal mesh threshold value that is used to decide whether two UVs are equal. */
 	UFUNCTION(BlueprintCallable, Category = "Interchange | Node | SkeletalMesh")
 	bool SetCustomThresholdUV(const float& AttributeValue, bool bAddApplyDelegate = true);
 
-	/** Query the skeletal mesh threshold to compare vertex position equality when computing morph target deltas. */
+	/** Query the skeletal mesh threshold value that is used to compare vertex position equality when computing morph target deltas. */
 	UFUNCTION(BlueprintCallable, Category = "Interchange | Node | SkeletalMesh")
 	bool GetCustomMorphThresholdPosition(float& AttributeValue) const;
 
-	/** Set the skeletal mesh threshold to compare vertex position equality when computing morph target deltas. */
+	/** Set the skeletal mesh threshold value that is used to compare vertex position equality when computing morph target deltas. */
 	UFUNCTION(BlueprintCallable, Category = "Interchange | Node | SkeletalMesh")
 	bool SetCustomMorphThresholdPosition(const float& AttributeValue, bool bAddApplyDelegate = true);
 
 	/**
-	 * Query The maximum number of bone influences to allow each vertex in this mesh to use.
+	 * Query the maximum number of bone influences to allow each vertex in this mesh to use.
 	 * If set higher than the limit determined by the project settings, it has no effect.
 	 * If set to 0, the value is taken from the DefaultBoneInfluenceLimit project setting.
 	 */
@@ -145,7 +145,7 @@ public:
 	bool GetCustomBoneInfluenceLimit(int32& AttributeValue) const;
 
 	/**
-	 * Set The maximum number of bone influences to allow each vertex in this mesh to use.
+	 * Set the maximum number of bone influences to allow each vertex in this mesh to use.
 	 * If set higher than the limit determined by the project settings, it has no effect.
 	 * If set to 0, the value is taken from the DefaultBoneInfluenceLimit project setting.
 	 */
@@ -153,7 +153,7 @@ public:
 	bool SetCustomBoneInfluenceLimit(const int32& AttributeValue, bool bAddApplyDelegate = true);
 
 	/**
-	 * The skeletal mesh thumbnail can have an overlay if the last re-import was geometry only. This thumbnail overlay feature use the metadata to find out if the last import was geometry only.
+	 * The skeletal mesh thumbnail can have an overlay if the last reimport was geometry only. This thumbnail overlay feature uses the metadata to find out if the last import was geometry only.
 	 */
 	virtual void AppendAssetRegistryTags(FAssetRegistryTagsContext Context) const override;
 	UE_DEPRECATED(5.4, "Implement the version that takes FAssetRegistryTagsContext instead.")
diff --git a/Engine/Plugins/Interchange/Runtime/Source/FactoryNodes/Public/InterchangeSkeletalMeshLodDataNode.h b/Engine/Plugins/Interchange/Runtime/Source/FactoryNodes/Public/InterchangeSkeletalMeshLodDataNode.h
--- a/Engine/Plugins/Interchange/Runtime/Source/FactoryNodes/Public/InterchangeSkeletalMeshLodDataNode.h
+++ b/Engine/Plugins/Interchange/Runtime/Source/FactoryNodes/Public/InterchangeSkeletalMeshLodDataNode.h
@@ -30,7 +30,7 @@ public:
 	UInterchangeSkeletalMeshLodDataNode();
 
 	/**
-	 * Return the node type name of the class, we use this when reporting error
+	 * Return the node type name of the class. This is used when reporting errors.
 	 */
 	virtual FString GetTypeName() const override;
 
@@ -41,31 +41,31 @@ public:
 #endif //WITH_EDITOR
 
 public:
-	/** Query the LOD skeletal mesh factory skeleton reference. Return false if the attribute was not set.*/
+	/** Query the LOD skeletal mesh factory skeleton reference. Return false if the attribute was not set. */
 	UFUNCTION(BlueprintCallable, Category = "Interchange | Node | SkeletalMeshLodData")
 	bool GetCustomSkeletonUid(FString& AttributeValue) const;
 
-	/** Set the LOD skeletal mesh factory skeleton reference. Return false if the attribute cannot be set.*/
+	/** Set the LOD skeletal mesh factory skeleton reference. Return false if the attribute could not be set. */
 	UFUNCTION(BlueprintCallable, Category = "Interchange | Node | SkeletalMeshLodData")
 	bool SetCustomSkeletonUid(const FString& AttributeValue);
 
-	/* Return the number of mesh geometry this LOD will be made of. Mesh uids can be either a scene or a mesh node. If its a scene it mean we want the mesh factory to bake the geo payload with the global transform of the scene node. */
+	/* Return the number of mesh geometries this LOD will be made from. A mesh UID can represent either a scene node or a mesh node. If it is a scene node, the mesh factory bakes the geometry payload with the global transform of the scene node. */
 	UFUNCTION(BlueprintCallable, Category = "Interchange | Node | SkeletalMeshLodData")
 	int32 GetMeshUidsCount() const;
 
-	/* Query all mesh geometry this LOD will be made of. Mesh uids can be either a scene or a mesh node. If its a scene it mean we want the mesh factory to bake the geo payload with the global transform of the scene node. */
+	/* Query all mesh geometry this LOD will be made from. A mesh UID can represent either a scene node or a mesh node. If it is a scene node, the mesh factory bakes the geometry payload with the global transform of the scene node. */
 	UFUNCTION(BlueprintCallable, Category = "Interchange | Node | SkeletalMeshLodData")
 	void GetMeshUids(TArray<FString>& OutMeshNames) const;
 
-	/* Add one mesh geometry use to create this LOD geometry. Mesh uids can be either a scene or a mesh node. If its a scene it mean we want the mesh factory to bake the geo payload with the global transform of the scene node. */
+	/* Add a mesh geometry used to create this LOD geometry. A mesh UID can represent either a scene node or a mesh node. If it is a scene node, the mesh factory bakes the geometry payload with the global transform of the scene node. */
 	UFUNCTION(BlueprintCallable, Category = "Interchange | Node | SkeletalMeshLodData")
 	bool AddMeshUid(const FString& MeshName);
 
-	/* Remove one mesh geometry use to create this LOD geometry. Mesh uids can be either a scene or a mesh node. If its a scene it mean we want the mesh factory to bake the geo payload with the global transform of the scene node. */
+	/* Remove a mesh geometry used to create this LOD geometry. A mesh UID can represent either a scene node or a mesh node. If it is a scene node, the mesh factory bakes the geometry payload with the global transform of the scene node. */
 	UFUNCTION(BlueprintCallable, Category = "Interchange | Node | SkeletalMeshLodData")
 	bool RemoveMeshUid(const FString& MeshName);
 
-	/* Remove all mesh geometry use to create this LOD geometry. Mesh uids can be either a scene or a mesh node. If its a scene it mean we want the mesh factory to bake the geo payload with the global transform of the scene node. */
+	/* Remove all mesh geometry used to create this LOD geometry. A mesh UID can represent either a scene node or a mesh node. If it is a scene node, the mesh factory bakes the geometry payload with the global transform of the scene node. */
 	UFUNCTION(BlueprintCallable, Category = "Interchange | Node | SkeletalMeshLodData")
 	bool RemoveAllMeshes();
 
diff --git a/Engine/Plugins/Interchange/Runtime/Source/FactoryNodes/Public/InterchangeSkeletonFactoryNode.h b/Engine/Plugins/Interchange/Runtime/Source/FactoryNodes/Public/InterchangeSkeletonFactoryNode.h
--- a/Engine/Plugins/Interchange/Runtime/Source/FactoryNodes/Public/InterchangeSkeletonFactoryNode.h
+++ b/Engine/Plugins/Interchange/Runtime/Source/FactoryNodes/Public/InterchangeSkeletonFactoryNode.h
@@ -18,9 +18,9 @@ class INTERCHANGEFACTORYNODES_API UInterchangeSkeletonFactoryNode : public UInte
 
 public:
 	/**
-	 * Initialize node data
-	 * @param: UniqueID - The uniqueId for this node
-	 * @param DisplayLabel - The name of the node
+	 * Initialize node data.
+	 * @param: UniqueID - The unique ID for this node.
+	 * @param DisplayLabel - The name of the node.
 	 * @param InAssetClass - The class the Skeleton factory will create for this node.
 	 */
 	UFUNCTION(BlueprintCallable, Category = "Interchange | Node | Skeleton")
@@ -47,7 +47,7 @@ public:
 	}
 
 	/**
-	 * Return the node type name of the class, we use this when reporting error
+	 * Return the node type name of the class. This is used when reporting errors.
 	 */
 	virtual FString GetTypeName() const override
 	{
@@ -55,7 +55,7 @@ public:
 		return TypeName;
 	}
 
-	/** Get the class this node want to create */
+	/** Get the class this node creates. */
 	UFUNCTION(BlueprintCallable, Category = "Interchange | Node | Skeleton")
 	virtual class UClass* GetObjectClass() const override
 	{
@@ -73,7 +73,7 @@ public:
 	}
 
 public:
-	/** Return false if the Attribute was not set previously.*/
+	/** Return false if the attribute was not set previously. */
 	UFUNCTION(BlueprintCallable, Category = "Interchange | Node | Skeleton")
 	bool GetCustomRootJointUid(FString& AttributeValue) const
 	{
@@ -86,14 +86,14 @@ public:
 		IMPLEMENT_NODE_ATTRIBUTE_SETTER_NODELEGATE(RootJointID, FString)
 	}
 
-	/** Query if this skeleton should replace joint transform with time zero evaluation instead of bind pose.*/
+	/** Query whether this skeleton should replace joint transforms with time-zero evaluation instead of bind pose. */
 	UFUNCTION(BlueprintCallable, Category = "Interchange | Node | Skeleton")
 	bool GetCustomUseTimeZeroForBindPose(bool& AttributeValue) const
 	{
 		IMPLEMENT_NODE_ATTRIBUTE_GETTER(UseTimeZeroForBindPose, bool);
 	}
 
-	/** If AttributeValue is true, force this skeleton bind pose to use time zero instead of bind pose.*/
+	/** If AttributeValue is true, force this skeleton to use time-zero evaluation instead of its bind pose. */
 	UFUNCTION(BlueprintCallable, Category = "Interchange | Node | Skeleton")
 	bool SetCustomUseTimeZeroForBindPose(const bool& AttributeValue)
 	{
diff --git a/Engine/Plugins/Interchange/Runtime/Source/FactoryNodes/Public/InterchangeStaticMeshFactoryNode.h b/Engine/Plugins/Interchange/Runtime/Source/FactoryNodes/Public/InterchangeStaticMeshFactoryNode.h
--- a/Engine/Plugins/Interchange/Runtime/Source/FactoryNodes/Public/InterchangeStaticMeshFactoryNode.h
+++ b/Engine/Plugins/Interchange/Runtime/Source/FactoryNodes/Public/InterchangeStaticMeshFactoryNode.h
@@ -34,20 +34,20 @@ public:
 	UInterchangeStaticMeshFactoryNode();
 
 	/**
-	 * Initialize node data
-	 * @param UniqueID - The uniqueId for this node
-	 * @param DisplayLabel - The name of the node
+	 * Initialize node data.
+	 * @param UniqueID - The unique ID for this node.
+	 * @param DisplayLabel - The name of the node.
 	 * @param InAssetClass - The class the StaticMesh factory will create for this node.
 	 */
 	UFUNCTION(BlueprintCallable, Category = "Interchange | Node | StaticMesh")
 	void InitializeStaticMeshNode(const FString& UniqueID, const FString& DisplayLabel, const FString& InAssetClass);
 
 	/**
-	 * Return the node type name of the class, we use this when reporting error
+	 * Return the node type name of the class. This is used when reporting errors.
 	 */
 	virtual FString GetTypeName() const override;
 
-	/** Get the class this node want to create */
+	/** Get the class this node creates. */
 	UFUNCTION(BlueprintCallable, Category = "Interchange | Node | StaticMesh")
 	virtual class UClass* GetObjectClass() const override;
 
@@ -59,15 +59,15 @@ public:
 #endif //WITH_EDITOR
 
 public:
-	/** Get weather the static mesh factory should set the nanite build settings. Return false if the attribute was not set.*/
+	/** Get whether the static mesh factory should set the Nanite build setting. Return false if the attribute was not set. */
 	UFUNCTION(BlueprintCallable, Category = "Interchange | Node | StaticMesh")
 	bool GetCustomBuildNanite(bool& AttributeValue) const;
 
-	/** Set weather the static mesh factory should set the nanite build settings. Return false if the attribute was not set.*/
+	/** Set whether the static mesh factory should set the Nanite build setting. Return false if the attribute was not set. */
 	UFUNCTION(BlueprintCallable, Category = "Interchange | Node | StaticMesh")
 	bool SetCustomBuildNanite(const bool& AttributeValue, bool bAddApplyDelegate = true);
 
-	/** Return The number of socket UIDs this static mesh has.*/
+	/** Return the number of socket UIDs this static mesh has. */
 	UFUNCTION(BlueprintCallable, Category = "Interchange | Node | StaticMesh")
 	int32 GetSocketUidCount() const;
 
@@ -83,112 +83,112 @@ public:
 	UFUNCTION(BlueprintCallable, Category = "Interchange | Node | StaticMesh")
 	bool RemoveSocketUd(const FString& SocketUid);
 
-	/** Get weather the static mesh should build build a reverse index buffer. */
+	/** Get whether the static mesh should build a reversed index buffer. */
 	UFUNCTION(BlueprintCallable, Category = "Interchange | Node | StaticMesh")
 	bool GetCustomBuildReversedIndexBuffer(bool& AttributeValue) const;
 
-	/** Set weather the static mesh should build build a reverse index buffer. */
+	/** Set whether the static mesh should build a reversed index buffer. */
 	UFUNCTION(BlueprintCallable, Category = "Interchange | Node | StaticMesh")
 	bool SetCustomBuildReversedIndexBuffer(const bool& AttributeValue, bool bAddApplyDelegate = true);
 
-	/** Get weather the static mesh should build build a reverse index buffer. */
+	/** Get whether the static mesh should generate lightmap UVs. */
 	UFUNCTION(BlueprintCallable, Category = "Interchange | Node | StaticMesh")
 	bool GetCustomGenerateLightmapUVs(bool& AttributeValue) const;
 
-	/** Set weather the static mesh should build build a reverse index buffer. */
+	/** Set whether the static mesh should generate lightmap UVs. */
 	UFUNCTION(BlueprintCallable, Category = "Interchange | Node | StaticMesh")
 	bool SetCustomGenerateLightmapUVs(const bool& AttributeValue, bool bAddApplyDelegate = true);
 
 	/**
-	 * Get Whether to generate the distance field treating every triangle hit as a front face.  
-	 * When enabled prevents the distance field from being discarded due to the mesh being open, but also lowers Distance Field AO quality.
+	 * Get whether to generate the distance field by treating every triangle hit as a front face.  
+	 * This prevents the distance field from being discarded due to the mesh being open, but also lowers distance field ambient occlusion quality.
 	 */
 	UFUNCTION(BlueprintCallable, Category = "Interchange | Node | StaticMesh")
 	bool GetCustomGenerateDistanceFieldAsIfTwoSided(bool& AttributeValue) const;
 
 	/**
-	 * Set Whether to generate the distance field treating every triangle hit as a front face.
-	 * When enabled prevents the distance field from being discarded due to the mesh being open, but also lowers Distance Field AO quality.
+	 * Set whether to generate the distance field by treating every triangle hit as a front face.
+	 * This prevents the distance field from being discarded due to the mesh being open, but also lowers distance field ambient occlusion quality.
 	 */
 	UFUNCTION(BlueprintCallable, Category = "Interchange | Node | StaticMesh")
 	bool SetCustomGenerateDistanceFieldAsIfTwoSided(const bool& AttributeValue, bool bAddApplyDelegate = true);
 
-	/** Get weather the static mesh ... (TODO add documentation) */
+	/** Get whether the static mesh is set up for use with physical material masks. */
 	UFUNCTION(BlueprintCallable, Category = "Interchange | Node | StaticMesh")
 	bool GetCustomSupportFaceRemap(bool& AttributeValue) const;
 
-	/** Set weather the static mesh ... (TODO add documentation) */
+	/** Set whether the static mesh is set up for use with physical material masks. */
 	UFUNCTION(BlueprintCallable, Category = "Interchange | Node | StaticMesh")
 	bool SetCustomSupportFaceRemap(const bool& AttributeValue, bool bAddApplyDelegate = true);
 
-	/** Get weather the static mesh ... (TODO add documentation) */
+	/** Get the amount of padding used to pack UVs for the static mesh. */
 	UFUNCTION(BlueprintCallable, Category = "Interchange | Node | StaticMesh")
 	bool GetCustomMinLightmapResolution(int32& AttributeValue) const;
 
-	/** Set weather the static mesh ... (TODO add documentation) */
+	/** Set the amount of padding used to pack UVs for the static mesh. */
 	UFUNCTION(BlueprintCallable, Category = "Interchange | Node | StaticMesh")
 	bool SetCustomMinLightmapResolution(const int32& AttributeValue, bool bAddApplyDelegate = true);
 
-	/** Get weather the static mesh ... (TODO add documentation) */
+	/** Get the index of the UV that is used as the source for generating lightmaps for the static mesh.  */
 	UFUNCTION(BlueprintCallable, Category = "Interchange | Node | StaticMesh")
 	bool GetCustomSrcLightmapIndex(int32& AttributeValue) const;
 
-	/** Set weather the static mesh ... (TODO add documentation) */
+	/** Set the index of the UV that is used as the source for generating lightmaps for the static mesh. */
 	UFUNCTION(BlueprintCallable, Category = "Interchange | Node | StaticMesh")
 	bool SetCustomSrcLightmapIndex(const int32& AttributeValue, bool bAddApplyDelegate = true);
 
-	/** Get weather the static mesh ... (TODO add documentation) */
+	/** Get the index of the UV that is used to store generated lightmaps for the static mesh. */
 	UFUNCTION(BlueprintCallable, Category = "Interchange | Node | StaticMesh")
 	bool GetCustomDstLightmapIndex(int32& AttributeValue) const;
 
-	/** Set weather the static mesh ... (TODO add documentation) */
+	/** Set the index of the UV that is used to store generated lightmaps for the static mesh. */
 	UFUNCTION(BlueprintCallable, Category = "Interchange | Node | StaticMesh")
 	bool SetCustomDstLightmapIndex(const int32& AttributeValue, bool bAddApplyDelegate = true);
 
-	/** Get the static mesh local scale applied when building the static mesh. */
+	/** Get the local scale that is applied when building the static mesh. */
 	UFUNCTION(BlueprintCallable, Category = "Interchange | Node | StaticMesh")
 	bool GetCustomBuildScale3D(FVector& AttributeValue) const;
 
-	/** Set the static mesh local scale applied when building the static mesh. */
+	/** Set the local scale that is applied when building the static mesh. */
 	UFUNCTION(BlueprintCallable, Category = "Interchange | Node | StaticMesh")
 	bool SetCustomBuildScale3D(const FVector& AttributeValue, bool bAddApplyDelegate = true);
 
 	/**
-	 * Get the Scale to apply to the mesh when allocating the distance field volume texture.
-	 * The default scale is 1, which is assuming that the mesh will be placed unscaled in the world.
+	 * Get the scale to apply to the mesh when allocating the distance field volume texture.
+	 * The default scale is 1, which assumes that the mesh will be placed unscaled in the world.
 	 */
 	UFUNCTION(BlueprintCallable, Category = "Interchange | Node | StaticMesh")
 	bool GetCustomDistanceFieldResolutionScale(float& AttributeValue) const;
 
 	/**
-	 * Set the Scale to apply to the mesh when allocating the distance field volume texture.
-	 * The default scale is 1, which is assuming that the mesh will be placed unscaled in the world.
+	 * Set the scale to apply to the mesh when allocating the distance field volume texture.
+	 * The default scale is 1, which assumes that the mesh will be placed unscaled in the world.
 	 */
 	UFUNCTION(BlueprintCallable, Category = "Interchange | Node | StaticMesh")
 	bool SetCustomDistanceFieldResolutionScale(const float& AttributeValue, bool bAddApplyDelegate = true);
 
-	/** Get the static mesh local scale applied when building the static mesh. */
+	/** Get the static mesh asset whose distance field will be used as the distance field for the imported mesh. */
 	UFUNCTION(BlueprintCallable, Category = "Interchange | Node | StaticMesh")
 	bool GetCustomDistanceFieldReplacementMesh(FSoftObjectPath& AttributeValue) const;
 
-	/** Set the static mesh local scale applied when building the static mesh. */
+	/** Set the static mesh asset whose distance field will be used as the distance field for the imported mesh. */
 	UFUNCTION(BlueprintCallable, Category = "Interchange | Node | StaticMesh")
 	bool SetCustomDistanceFieldReplacementMesh(const FSoftObjectPath& AttributeValue, bool bAddApplyDelegate = true);
 
 	/**
-	 * Max Lumen mesh cards to generate for this mesh.
-	 * More cards means that surface will have better coverage, but will result in increased runtime overhead.
-	 * Set to 0 in order to disable mesh card generation for this mesh.
-	 * Default is 12.
+	 * Get the maximum number of Lumen mesh cards to generate for this mesh.
+	 * More cards means that the surface will have better coverage, but will result in increased runtime overhead.
+	 * Set this to 0 to disable mesh card generation for this mesh.
+	 * The default is 12.
 	 */
 	UFUNCTION(BlueprintCallable, Category = "Interchange | Node | StaticMesh")
 	bool GetCustomMaxLumenMeshCards(int32& AttributeValue) const;
 
 	/**
-	 * Max Lumen mesh cards to generate for this mesh.
-	 * More cards means that surface will have better coverage, but will result in increased runtime overhead.
-	 * Set to 0 in order to disable mesh card generation for this mesh.
-	 * Default is 12.
+	 * Set the maximum number of Lumen mesh cards to generate for this mesh.
+	 * More cards means that the surface will have better coverage, but will result in increased runtime overhead.
+	 * Set this to 0 to disable mesh card generation for this mesh.
+	 * The default is 12.
 	 */
 	UFUNCTION(BlueprintCallable, Category = "Interchange | Node | StaticMesh")
 	bool SetCustomMaxLumenMeshCards(const int32& AttributeValue, bool bAddApplyDelegate = true);
diff --git a/Engine/Plugins/Interchange/Runtime/Source/FactoryNodes/Public/InterchangeStaticMeshLodDataNode.h b/Engine/Plugins/Interchange/Runtime/Source/FactoryNodes/Public/InterchangeStaticMeshLodDataNode.h
--- a/Engine/Plugins/Interchange/Runtime/Source/FactoryNodes/Public/InterchangeStaticMeshLodDataNode.h
+++ b/Engine/Plugins/Interchange/Runtime/Source/FactoryNodes/Public/InterchangeStaticMeshLodDataNode.h
@@ -34,7 +34,7 @@ public:
 	UInterchangeStaticMeshLodDataNode();
 
 	/**
-	 * Return the node type name of the class, we use this when reporting error
+	 * Return the node type name of the class. This is used when reporting errors.
 	 */
 	virtual FString GetTypeName() const override;
 
@@ -44,7 +44,7 @@ public:
 #endif //WITH_EDITOR
 
 public:
-	/* Mesh Uids: It can be either a scene or a mesh node uid. If its a scene it mean we want the mesh factory to bake the geo payload with the global transform of the scene node. */
+	/* Mesh UIDs can be either a scene node or a mesh node UID. If it is a scene node, the mesh factory bakes the geometry payload with the global transform of the scene node. */
 
 	UFUNCTION(BlueprintCallable, Category = "Interchange | Node | StaticMeshLodData")
 	int32 GetMeshUidsCount() const;
diff --git a/Engine/Plugins/Interchange/Runtime/Source/FactoryNodes/Public/InterchangeTexture2DFactoryNode.h b/Engine/Plugins/Interchange/Runtime/Source/FactoryNodes/Public/InterchangeTexture2DFactoryNode.h
--- a/Engine/Plugins/Interchange/Runtime/Source/FactoryNodes/Public/InterchangeTexture2DFactoryNode.h
+++ b/Engine/Plugins/Interchange/Runtime/Source/FactoryNodes/Public/InterchangeTexture2DFactoryNode.h
@@ -71,11 +71,11 @@ public:
 
 	//////////////////////////////////////////////////////////////////////////
 	// UDIMs begin here
-	// UDIM base texture use a different model for the source data
+	// UDIM base textures use a different model for the source data.
 
 	/**
-	 * Get the source blocks for the texture
-	 * If the map is empty then the texture will be simply be imported as normal texture using the payload key
+	 * Get the source blocks for the texture.
+	 * If the map is empty, the texture is imported as a normal texture using the payload key.
 	 */
 	UFUNCTION(BlueprintCallable, Category = "Interchange | Node | Texture | UDIMs")
 	TMap<int32, FString> GetSourceBlocks() const
@@ -84,11 +84,11 @@ public:
 	}
 
 	/**
-	 * Set the source blocks
-	 * Using this will suggest the pipeline to consider this texture as UDIM and it can chose to pass or not these block to the texture factory node.
+	 * Set the source blocks for the texture.
+	 * Using this suggests to the pipeline to consider this texture as a UDIM. The pipeline can choose whether to pass these blocks to the texture factory node.
 	 * @param InSourceBlocks The blocks and their source image that compose the whole texture.
-	 * The textures must be of the same format and use the same pixel format
-	 * The first block in the map is used to determine the accepted texture format and pixel format
+	 * The textures must be of the same format and use the same pixel format.
+	 * The first block in the map is used to determine the accepted texture format and pixel format.
 	 */
 	UFUNCTION(BlueprintCallable, Category = "Interchange | Node | Texture | UDIMs")
 	void SetSourceBlocks(const TMap<int32, FString>& InSourceBlocks)
@@ -98,11 +98,11 @@ public:
 
 	
 	/**
-	 * Set the source blocks
-	 * Using this will suggest the pipeline to consider this texture as UDIM and it can chose to pass or not these block to the texture factory node.
+	 * Set the source blocks for the texture.
+	 * Using this suggests to the pipeline to consider this texture as a UDIM. The pipeline can choose whether to pass these blocks to the texture factory node.
 	 * @param InSourceBlocks The blocks and their source image that compose the whole texture.
-	 * The textures must be of the same format and use the same pixel format
-	 * The first block in the map is used to determine the accepted texture format and pixel format
+	 * The textures must be of the same format and use the same pixel format.
+	 * The first block in the map is used to determine the accepted texture format and pixel format.
 	 */
 	void SetSourceBlocks(TMap<int32, FString>&& InSourceBlocks)
 	{
@@ -111,12 +111,12 @@ public:
 
 
 	/**
-	 * Get a source block from the texture
+	 * Get a source block from the texture.
 	 *
-	 * @param X The X coordinate of the block
-	 * @param Y The Y coordinate of the block
-	 * @param OutSourceFile The source file for that block if found
-	 * @return True if the source file for the block was found
+	 * @param X The X coordinate of the block.
+	 * @param Y The Y coordinate of the block.
+	 * @param OutSourceFile The source file for that block if found.
+	 * @return True if the source file for the block was found.
 	 */
 	UFUNCTION(BlueprintCallable, Category = "Interchange | Node | Texture | UDIMs")
 	bool GetSourceBlockByCoordinates(int32 X, int32 Y, FString& OutSourceFile) const
@@ -125,11 +125,11 @@ public:
 	}
 
 	/**
-	 * Get a source block from the texture
+	 * Get a source block from the texture.
 	 *
-	 * @param BlockIndex The UDIM Index of the block
-	 * @param OutSourceFile The source file for that block if found
-	 * @return True if the source file for the block was found
+	 * @param BlockIndex The UDIM Index of the block.
+	 * @param OutSourceFile The source file for that block if found.
+	 * @return True if the source file for the block was found.
 	 */
 	UFUNCTION(BlueprintCallable, Category = "Interchange | Node | Texture | UDIMs")
 	bool GetSourceBlock(int32 BlockIndex, FString& OutSourceFile) const
@@ -139,14 +139,14 @@ public:
 
 
 	/**
-	 * Set a source block for the texture
+	 * Set a source block for the texture.
 	 *
-	 * @param X The X coordinate of the block
-	 * @param Y The Y coordinate of the block
-	 * @param InSourceFile The source file for that block
+	 * @param X The X coordinate of the block.
+	 * @param Y The Y coordinate of the block.
+	 * @param InSourceFile The source file for that block.
 	 * 
-	 * The textures must be of the same format and use the same pixel format
-	 * The first block in the map is used to determine the accepted texture format and pixel format
+	 * The textures must be of the same format and use the same pixel format.
+	 * The first block in the map is used to determine the accepted texture format and pixel format.
 	 */
 	UFUNCTION(BlueprintCallable, Category = "Interchange | Node | Texture | UDIMs")
 	void SetSourceBlockByCoordinates(int32 X, int32 Y,  const FString& InSourceFile)
@@ -155,13 +155,13 @@ public:
 	}
 
 	/**
-	 * Set a source block for the texture
+	 * Set a source block for the texture.
 	 *
-	 * @param BlockIndex The UDIM Index of the block
-	 * @param InSourceFile The source file for that block
+	 * @param BlockIndex The UDIM index of the block.
+	 * @param InSourceFile The source file for that block.
 	 * 
-	 * The textures must be of the same format and use the same pixel format
-	 * The first block in the map is used to determine the accepted texture format and pixel format
+	 * The textures must be of the same format and use the same pixel format.
+	 * The first block in the map is used to determine the accepted texture format and pixel format.
 	 */
 	UFUNCTION(BlueprintCallable, Category = "Interchange | Node | Texture | UDIMs")
 	void SetSourceBlock(int32 BlockIndex, const FString& InSourceFile)
diff --git a/Engine/Plugins/Interchange/Runtime/Source/FactoryNodes/Public/InterchangeTextureFactoryNode.h b/Engine/Plugins/Interchange/Runtime/Source/FactoryNodes/Public/InterchangeTextureFactoryNode.h
--- a/Engine/Plugins/Interchange/Runtime/Source/FactoryNodes/Public/InterchangeTextureFactoryNode.h
+++ b/Engine/Plugins/Interchange/Runtime/Source/FactoryNodes/Public/InterchangeTextureFactoryNode.h
@@ -17,9 +17,9 @@ class INTERCHANGEFACTORYNODES_API UInterchangeTextureFactoryNode : public UInter
 
 public:
 	/**
-	 * Initialize node data
-	 * @param: UniqueID - The uniqueId for this node
-	 * @param DisplayLabel - The name of the node
+	 * Initialize node data.
+	 * @param: UniqueID - The unique ID for this node.
+	 * @param DisplayLabel - The name of the node.
 	 * @param InAssetClass - The class the texture factory will create for this node.
 	 */
 	UFUNCTION(BlueprintCallable, Category = "Interchange | Node | Texture")
@@ -34,7 +34,7 @@ public:
 	}
 
 	/**
-	 * Return the node type name of the class, we use this when reporting error
+	 * Return the node type name of the class. This is used when reporting errors.
 	 */
 	virtual FString GetTypeName() const override
 	{
@@ -42,7 +42,7 @@ public:
 		return TypeName;
 	}
 
-	/** Get the class this node want to create */
+	/** Get the class this node creates. */
 	UFUNCTION(BlueprintCallable, Category = "Interchange | Node | Texture")
 	virtual UClass* GetObjectClass() const override
 	{
@@ -63,7 +63,7 @@ public:
 	}
 
 	/**
-	 * Get the translated texture node unique ID.
+	 * Get the unique ID of the translated texture node.
 	 */
 	UFUNCTION(BlueprintCallable, Category = "Interchange | Node | Texture")
 	bool GetCustomTranslatedTextureNodeUid(FString& AttributeValue) const
@@ -72,7 +72,7 @@ public:
 	}
 
 	/**
-	 * Set the translated texture node unique ID. This is the reference to the node that was create by the translator and this node is needed to get the texture payload.
+	 * Set the unique ID of the translated texture node. This is a reference to the node that was created by the translator. It is needed to get the texture payload.
 	 */
 	UFUNCTION(BlueprintCallable, Category = "Interchange | Node | Texture")
 	bool SetCustomTranslatedTextureNodeUid(const FString& AttributeValue)
@@ -683,8 +683,8 @@ public:
 	}
 
 	/*
-	 * Should the factory tell the translator to provide a compressed source data payload when available.
-	 * This will generally result in smaller assets, but some operations like the texture build might be slower because the source data will need to be uncompressed.
+	 * Determines whether the factory should tell the translator to provide a compressed source data payload when available.
+	 * This will generally result in smaller assets. However, some operations like the texture build might be slower, because the source data first needs to be decompressed.
 	 */
 	UFUNCTION(BlueprintCallable, Category = "Interchange | Node | Texture | Import Settings")
 	bool SetCustomPreferCompressedSourceData(const bool& AttributeValue)
