Clean up typos and phrasing in Interchange modules - contribution from Robb Surridge

diff --git a/Engine/Plugins/Interchange/Runtime/Source/Import/Private/Animation/InterchangeAnimSequenceFactory.cpp b/Engine/Plugins/Interchange/Runtime/Source/Import/Private/Animation/InterchangeAnimSequenceFactory.cpp
--- a/Engine/Plugins/Interchange/Runtime/Source/Import/Private/Animation/InterchangeAnimSequenceFactory.cpp
+++ b/Engine/Plugins/Interchange/Runtime/Source/Import/Private/Animation/InterchangeAnimSequenceFactory.cpp
@@ -506,7 +506,7 @@ namespace UE::Interchange::Private
 						UInterchangeResultWarning_Generic* Message = Factory->AddMessage<UInterchangeResultWarning_Generic>();
 						Message->DestinationAssetName = AssetName;
 						Message->AssetType = UAnimSequence::StaticClass();
-						Message->Text = FText::Format(NSLOCTEXT("UInterchangeAnimSequenceFactory", "AnimationPayloadBakeFrequencyNotTheExpected", "Animation Payload [{0}]'s BakeFrequency, RangeStartTime and RangeEndTime does not equal with the provided one."), FText::FromString(PayloadKey));
+						Message->Text = FText::Format(NSLOCTEXT("UInterchangeAnimSequenceFactory", "AnimationPayloadBakeFrequencyNotTheExpected", "The BakeFrequency, RangeStartTime and RangeEndTime of Animation Payload [{0}] are not the same as the values provided."), FText::FromString(PayloadKey));
 					}
 				}
 
@@ -742,7 +742,7 @@ namespace UE::Interchange::Private
 							UInterchangeResultWarning_Generic* Message = Factory->AddMessage<UInterchangeResultWarning_Generic>();
 							Message->DestinationAssetName = AssetName;
 							Message->AssetType = UAnimSequence::StaticClass();
-							Message->Text = FText::Format(NSLOCTEXT("UInterchangeAnimSequenceFactory", "AnimationPayloadInvalidCurve", "Invalid animation curve payload key [{0}] AnimSequence asset {1}"), FText::FromString(CurveName), FText::FromString(AssetName));
+							Message->Text = FText::Format(NSLOCTEXT("UInterchangeAnimSequenceFactory", "AnimationPayloadInvalidCurve", "Invalid animation material curve payload key [{0}] in AnimSequence asset {1}"), FText::FromString(CurveName), FText::FromString(AssetName));
 							continue;
 						}
 
@@ -839,7 +839,7 @@ namespace UE::Interchange::Private
 							UInterchangeResultWarning_Generic* Message = Factory->AddMessage<UInterchangeResultWarning_Generic>();
 							Message->DestinationAssetName = AssetName;
 							Message->AssetType = UAnimSequence::StaticClass();
-							Message->Text = FText::Format(NSLOCTEXT("UInterchangeAnimSequenceFactory", "AnimationPayloadInvalidCurvePayloadKey", "Invalid animation curve payload key [{0}] AnimSequence asset {1}"), FText::FromString(CurveNameAndPayload.Key), FText::FromString(AssetName));
+							Message->Text = FText::Format(NSLOCTEXT("UInterchangeAnimSequenceFactory", "AnimationPayloadInvalidCurvePayloadKey", "Invalid animation material curve payload key [{0}] in AnimSequence asset {1}"), FText::FromString(CurveNameAndPayload.Key), FText::FromString(AssetName));
 							continue;
 						}
 						FAnimationPayloadData& AnimationPayloadData = AnimationCurvePayload.GetValue();
@@ -855,7 +855,7 @@ namespace UE::Interchange::Private
 							UInterchangeResultWarning_Generic* Message = Factory->AddMessage<UInterchangeResultWarning_Generic>();
 							Message->DestinationAssetName = AssetName;
 							Message->AssetType = UAnimSequence::StaticClass();
-							Message->Text = FText::Format(NSLOCTEXT("UInterchangeAnimSequenceFactory", "AnimationPayloadInvalidStepCurvePayloadKey", "Invalid animation curve payload key [{0}] AnimSequence asset {1}"), FText::FromString(StepCurveNameAndPayload.Key), FText::FromString(AssetName));
+							Message->Text = FText::Format(NSLOCTEXT("UInterchangeAnimSequenceFactory", "AnimationPayloadInvalidStepCurvePayloadKey", "Invalid animation curve payload key [{0}] in AnimSequence asset {1}"), FText::FromString(StepCurveNameAndPayload.Key), FText::FromString(AssetName));
 							continue;
 						}
 
@@ -892,7 +892,7 @@ UInterchangeFactoryBase::FImportAssetResult UInterchangeAnimSequenceFactory::Beg
 	UInterchangeFactoryBase::FImportAssetResult ImportAssetResult;
 #if !WITH_EDITOR || !WITH_EDITORONLY_DATA
 
-	UE_LOG(LogInterchangeImport, Error, TEXT("Cannot import animsequence asset in runtime, this is an editor only feature."));
+	UE_LOG(LogInterchangeImport, Error, TEXT("Cannot import AnimSequence asset at runtime. This is an editor-only feature."));
 	return ImportAssetResult;
 #else
 	UAnimSequence* NewAnimSequence = nullptr;
@@ -925,7 +925,7 @@ UInterchangeFactoryBase::FImportAssetResult UInterchangeAnimSequenceFactory::Beg
 	const IInterchangeAnimationPayloadInterface* AnimSequenceTranslatorPayloadInterface = Cast<IInterchangeAnimationPayloadInterface>(Arguments.Translator);
 	if (!AnimSequenceTranslatorPayloadInterface)
 	{
-		UE_LOG(LogInterchangeImport, Error, TEXT("Cannot import AnimSequence, the translator do not implement the IInterchangeAnimationPayloadInterface."));
+		UE_LOG(LogInterchangeImport, Error, TEXT("Cannot import AnimSequence. The translator does not implement IInterchangeAnimationPayloadInterface."));
 		return ImportAssetResult;
 	}
 
@@ -987,7 +987,7 @@ UInterchangeFactoryBase::FImportAssetResult UInterchangeAnimSequenceFactory::Beg
 
 		if (!ensure(Skeleton))
 		{
-			UE_LOG(LogInterchangeImport, Error, TEXT("Invalid Skeleton when importing animation sequence asset %s"), *Arguments.AssetName);
+			UE_LOG(LogInterchangeImport, Error, TEXT("Invalid Skeleton when importing animation sequence asset %s."), *Arguments.AssetName);
 			return ImportAssetResult;
 		}
 	}
@@ -1030,7 +1030,7 @@ UInterchangeFactoryBase::FImportAssetResult UInterchangeAnimSequenceFactory::Imp
 	UInterchangeFactoryBase::FImportAssetResult ImportAssetResult;
 #if !WITH_EDITOR || !WITH_EDITORONLY_DATA
 
-	UE_LOG(LogInterchangeImport, Error, TEXT("Cannot import animsequence asset in runtime, this is an editor only feature."));
+	UE_LOG(LogInterchangeImport, Error, TEXT("Cannot import AnimSequence asset at runtime. This is an editor-only feature."));
 	return ImportAssetResult;
 #else
 
@@ -1101,7 +1101,7 @@ UInterchangeFactoryBase::FImportAssetResult UInterchangeAnimSequenceFactory::Imp
 			if (!OptionalAnimationTransformPayload.IsSet())
 			{
 				FString PayloadKey = PayloadKeys[AnimationPayload.Key->GetUniqueID()].UniqueId;
-				UE_LOG(LogInterchangeImport, Warning, TEXT("Invalid animation transform payload key [%s] AnimSequence asset %s"), *PayloadKey, *Arguments.AssetName);
+				UE_LOG(LogInterchangeImport, Warning, TEXT("Invalid animation transform payload key [%s] in AnimSequence asset %s."), *PayloadKey, *Arguments.AssetName);
 				continue;
 			}
 
@@ -1183,7 +1183,7 @@ UInterchangeFactoryBase::FImportAssetResult UInterchangeAnimSequenceFactory::Imp
 			TOptional<UE::Interchange::FAnimationPayloadData> AnimationCurvePayload = CurveNameAndPayload.Value.Get();
 			if (!AnimationCurvePayload.IsSet())
 			{
-				UE_LOG(LogInterchangeImport, Warning, TEXT("Invalid animation morph target curve payload key [%s] AnimSequence asset %s"), *CurveNameAndPayload.Key, *Arguments.AssetName);
+				UE_LOG(LogInterchangeImport, Warning, TEXT("Invalid animation morph target curve payload key [%s] in AnimSequence asset %s."), *CurveNameAndPayload.Key, *Arguments.AssetName);
 				continue;
 			}
 			UE::Interchange::FAnimationPayloadData& CurvePayload = AnimationCurvePayload.GetValue();
@@ -1206,7 +1206,7 @@ UInterchangeFactoryBase::FImportAssetResult UInterchangeAnimSequenceFactory::End
 	UInterchangeFactoryBase::FImportAssetResult ImportAssetResult;
 #if !WITH_EDITOR || !WITH_EDITORONLY_DATA
 
-	UE_LOG(LogInterchangeImport, Error, TEXT("Cannot import animsequence asset in runtime, this is an editor only feature."));
+	UE_LOG(LogInterchangeImport, Error, TEXT("Cannot import AnimSequence asset at runtime. This is an editor-only feature."));
 	return ImportAssetResult;
 #else
 	if (!AnimSequence)
@@ -1235,7 +1235,7 @@ UInterchangeFactoryBase::FImportAssetResult UInterchangeAnimSequenceFactory::End
 	const UInterchangeSkeletonFactoryNode* SkeletonFactoryNode = Cast<UInterchangeSkeletonFactoryNode>(Arguments.NodeContainer->GetNode(SkeletonUid));
 	if (!SkeletonFactoryNode)
 	{
-		UE_LOG(LogInterchangeImport, Warning, TEXT("Invalid skeleton factory node, the skeleton factory node is obligatory to import this animsequence [%s]!"), *Arguments.AssetName);
+		UE_LOG(LogInterchangeImport, Warning, TEXT("Invalid skeleton factory node. The skeleton factory node is obligatory to import AnimSequence [%s]."), *Arguments.AssetName);
 		return ImportAssetResult;
 	}
 
@@ -1267,7 +1267,7 @@ UInterchangeFactoryBase::FImportAssetResult UInterchangeAnimSequenceFactory::End
 
 		if (!ensure(Skeleton))
 		{
-			UE_LOG(LogInterchangeImport, Warning, TEXT("Invalid Skeleton when importing animation sequence asset %s"), *Arguments.AssetName);
+			UE_LOG(LogInterchangeImport, Warning, TEXT("Invalid Skeleton when importing animation sequence asset %s."), *Arguments.AssetName);
 			return ImportAssetResult;
 		}
 	}
@@ -1275,7 +1275,7 @@ UInterchangeFactoryBase::FImportAssetResult UInterchangeAnimSequenceFactory::End
 	const IInterchangeAnimationPayloadInterface* AnimSequenceTranslatorPayloadInterface = Cast<IInterchangeAnimationPayloadInterface>(Arguments.Translator);
 	if (!AnimSequenceTranslatorPayloadInterface)
 	{
-		UE_LOG(LogInterchangeImport, Error, TEXT("Cannot import AnimSequence, the translator do not implement the IInterchangeAnimationPayloadInterface."));
+		UE_LOG(LogInterchangeImport, Error, TEXT("Cannot import AnimSequence. The translator does not implement IInterchangeAnimationPayloadInterface."));
 		return ImportAssetResult;
 	}
 
@@ -1431,7 +1431,7 @@ bool UInterchangeAnimSequenceFactory::IsBoneTrackAnimationValid(const UInterchan
 			Message->SourceAssetName = Arguments.SourceData->GetFilename();
 			Message->DestinationAssetName = Arguments.AssetName;
 			Message->AssetType = UAnimSequence::StaticClass();
-			Message->Text = FText::Format(NSLOCTEXT("UInterchangeAnimSequenceFactory", "WrongSequenceLength", "Animation length {0} is not compatible with import frame-rate {1} (sub frame {2}), animation has to be frame-border aligned."),
+			Message->Text = FText::Format(NSLOCTEXT("UInterchangeAnimSequenceFactory", "WrongSequenceLength", "Animation length {0} is not compatible with import frame-rate {1} (sub frame {2}). The animation must be frame-border aligned."),
 				FText::AsNumber(SequenceLength), FrameRate.ToPrettyText(), FText::AsNumber(SubFrame));
 			bResult = false;
 		}
diff --git a/Engine/Plugins/Interchange/Runtime/Source/Import/Private/Gltf/InterchangeGltfAnimation.cpp b/Engine/Plugins/Interchange/Runtime/Source/Import/Private/Gltf/InterchangeGltfAnimation.cpp
--- a/Engine/Plugins/Interchange/Runtime/Source/Import/Private/Gltf/InterchangeGltfAnimation.cpp
+++ b/Engine/Plugins/Interchange/Runtime/Source/Import/Private/Gltf/InterchangeGltfAnimation.cpp
@@ -364,7 +364,7 @@ namespace UE::Interchange::Gltf::Private
 			}
 			case GLTF::FAnimation::EPath::Weights:
 			{
-				UE_LOG(LogInterchangeImport, Warning, TEXT("Animation[%d]:Channel[%d] : Morph Animation (Weight path) type not supported for Transform like Animations. (Morph Animations should be handled via \"GetMorphTargetAnimationPayloadData\" function."), AnimationIndex, ChannelIndex);
+				UE_LOG(LogInterchangeImport, Warning, TEXT("Animation[%d]:Channel[%d]: Morph Animation (weight path) type not supported for transform-like animations. Morph animations should be handled through the \"GetMorphTargetAnimationPayloadData\" function."), AnimationIndex, ChannelIndex);
 				break;
 			}
 			default:
@@ -660,7 +660,7 @@ namespace UE::Interchange::Gltf::Private
 				break;
 			}
 			default:
-				UE_LOG(LogInterchangeImport, Warning, TEXT("Animation type not supported"));
+				UE_LOG(LogInterchangeImport, Warning, TEXT("Animation type not supported."));
 				break;
 			}
 		}
@@ -859,7 +859,7 @@ namespace UE::Interchange::Gltf::Private
 							if (!ValidateChannelSampler(Channel, Sampler))
 							{
 								// if any of the channels are corrupt the joint will not receive any of the  animation data
-								UE_LOG(LogInterchangeImport, Warning, TEXT("GLTF Sampler Corrupt. Input and Output not meeting expectations."));
+								UE_LOG(LogInterchangeImport, Warning, TEXT("glTF sampler corrupt. Input and Output do not meet expectations."));
 								break;
 							}
 
diff --git a/Engine/Plugins/Interchange/Runtime/Source/Import/Private/Material/InterchangeMaterialFactory.cpp b/Engine/Plugins/Interchange/Runtime/Source/Import/Private/Material/InterchangeMaterialFactory.cpp
--- a/Engine/Plugins/Interchange/Runtime/Source/Import/Private/Material/InterchangeMaterialFactory.cpp
+++ b/Engine/Plugins/Interchange/Runtime/Source/Import/Private/Material/InterchangeMaterialFactory.cpp
@@ -794,7 +794,7 @@ UInterchangeFactoryBase::FImportAssetResult UInterchangeMaterialFactory::BeginIm
 		ImportAssetResult.bIsFactorySkipAsset = true;
 	};
 	
-	const FText MissMatchClassText = LOCTEXT("MatFactory_CouldNotCreateMat_MissMatchClass", "Missmatch between interchange material factory node class and factory class.");
+	const FText MissMatchClassText = LOCTEXT("MatFactory_CouldNotCreateMat_MissMatchClass", "Mismatch between Interchange material factory node class and factory class.");
 
 	if (!Arguments.AssetNode || !Arguments.AssetNode->GetObjectClass()->IsChildOf(GetFactoryClass()))
 	{
@@ -805,7 +805,7 @@ UInterchangeFactoryBase::FImportAssetResult UInterchangeMaterialFactory::BeginIm
 	const UInterchangeBaseMaterialFactoryNode* MaterialFactoryNode = Cast<UInterchangeBaseMaterialFactoryNode>(Arguments.AssetNode);
 	if (MaterialFactoryNode == nullptr)
 	{
-		CouldNotCreateMaterialLog(LOCTEXT("MatFactory_CouldNotCreateMat_CannotCastFactoryNode", "Cannot cast interchange factory node to UInterchangeBaseMaterialFactoryNode."));
+		CouldNotCreateMaterialLog(LOCTEXT("MatFactory_CouldNotCreateMat_CannotCastFactoryNode", "Cannot cast Interchange factory node to UInterchangeBaseMaterialFactoryNode."));
 		return ImportAssetResult;
 	}
 
@@ -867,7 +867,7 @@ UInterchangeFactoryBase::FImportAssetResult UInterchangeMaterialFactory::BeginIm
 	{
 		if (bCanImportMaterial)
 		{
-			CouldNotCreateMaterialLog(LOCTEXT("MatFactory_CouldNotCreateMat_MaterialCreationFail", "Material creation fail."));
+			CouldNotCreateMaterialLog(LOCTEXT("MatFactory_CouldNotCreateMat_MaterialCreationFail", "Material creation failed"));
 		}
 		return ImportAssetResult;
 	}
@@ -1016,14 +1016,14 @@ UInterchangeFactoryBase::FImportAssetResult UInterchangeMaterialFactory::ImportA
 
 	if (!MaterialObject)
 	{
-		UE_LOG(LogInterchangeImport, Error, TEXT("Could not import the Material asset %s, because the asset do not exist."), *Arguments.AssetName);
+		UE_LOG(LogInterchangeImport, Error, TEXT("Could not import the Material asset %s because the asset does not exist."), *Arguments.AssetName);
 		return ImportAssetResult;
 	}
 
 	UMaterialInterface* MaterialInterface = Cast<UMaterialInterface>(MaterialObject);
 	if (!ensure(MaterialInterface))
 	{
-		UE_LOG(LogInterchangeImport, Error, TEXT("Could not cast to Material asset %s"), *Arguments.AssetName);
+		UE_LOG(LogInterchangeImport, Error, TEXT("Could not cast to Material asset %s."), *Arguments.AssetName);
 		return ImportAssetResult;
 	}
 
@@ -1886,7 +1886,7 @@ UInterchangeFactoryBase::FImportAssetResult UInterchangeMaterialFunctionFactory:
 		ImportAssetResult.bIsFactorySkipAsset = true;
 	};
 
-	const FText MissMatchClassText = LOCTEXT("MatFunc_CouldNotCreateMat_MissMatchClass", "Missmatch between interchange material factory node class and factory class.");
+	const FText MissMatchClassText = LOCTEXT("MatFunc_CouldNotCreateMat_MissMatchClass", "Mismatch between Interchange material factory node class and factory class.");
 
 	if (!Arguments.AssetNode || !Arguments.AssetNode->GetObjectClass()->IsChildOf(GetFactoryClass()))
 	{
@@ -1897,7 +1897,7 @@ UInterchangeFactoryBase::FImportAssetResult UInterchangeMaterialFunctionFactory:
 	const UInterchangeMaterialFunctionFactoryNode* MaterialFactoryNode = Cast<UInterchangeMaterialFunctionFactoryNode>(Arguments.AssetNode);
 	if (!ensure(MaterialFactoryNode))
 	{
-		CouldNotCreateMaterialLog(LOCTEXT("MatFunc_CouldNotCreateMat_CannotCastFactoryNode", "Cannot cast interchange factory node to UInterchangeBaseMaterialFactoryNode."));
+		CouldNotCreateMaterialLog(LOCTEXT("MatFunc_CouldNotCreateMat_CannotCastFactoryNode", "Cannot cast Interchange factory node to UInterchangeBaseMaterialFactoryNode."));
 		return ImportAssetResult;
 	}
 
@@ -1940,7 +1940,7 @@ UInterchangeFactoryBase::FImportAssetResult UInterchangeMaterialFunctionFactory:
 
 	if (!Material)
 	{
-		CouldNotCreateMaterialLog(LOCTEXT("MatFunc_CouldNotCreateMat_MaterialCreationFail", "Material creation fail."));
+		CouldNotCreateMaterialLog(LOCTEXT("MatFunc_CouldNotCreateMat_MaterialCreationFail", "Material creation failed."));
 		return ImportAssetResult;
 	}
 
@@ -1985,7 +1985,7 @@ UInterchangeFactoryBase::FImportAssetResult UInterchangeMaterialFunctionFactory:
 
 	if (!MaterialObject)
 	{
-		UE_LOG(LogInterchangeImport, Error, TEXT("Could not import the Material asset %s, because the asset do not exist."), *Arguments.AssetName);
+		UE_LOG(LogInterchangeImport, Error, TEXT("Could not import the Material asset %s because the asset does not exist."), *Arguments.AssetName);
 		return ImportAssetResult;
 	}
 
diff --git a/Engine/Plugins/Interchange/Runtime/Source/Import/Private/MaterialX/MaterialXUtils/MaterialXSurfaceShaderAbstract.cpp b/Engine/Plugins/Interchange/Runtime/Source/Import/Private/MaterialX/MaterialXUtils/MaterialXSurfaceShaderAbstract.cpp
--- a/Engine/Plugins/Interchange/Runtime/Source/Import/Private/MaterialX/MaterialXUtils/MaterialXSurfaceShaderAbstract.cpp
+++ b/Engine/Plugins/Interchange/Runtime/Source/Import/Private/MaterialX/MaterialXUtils/MaterialXSurfaceShaderAbstract.cpp
@@ -204,7 +204,7 @@ bool FMaterialXSurfaceShaderAbstract::ConnectNodeGraphOutputToInput(MaterialX::I
 
 		if(!Output)
 		{
-			UE_LOG(LogInterchangeImport, Warning, TEXT("Couldn't find a connected output to (%s)"), *GetInputName(InputToNodeGraph));
+			UE_LOG(LogInterchangeImport, Warning, TEXT("Couldn't find a connected output to (%s)."), *GetInputName(InputToNodeGraph));
 			return false;
 		}
 
@@ -346,7 +346,7 @@ void FMaterialXSurfaceShaderAbstract::ConnectNodeCategoryOutputToInput(const Mat
 		{
 			if(!ConnectMatchingNodeOutputToInput({ UpstreamNode, ParentShaderNode, InputChannelName, OutputChannelName }))
 			{
-				UE_LOG(LogInterchangeImport, Warning, TEXT("<%s> is not supported yet"), ANSI_TO_TCHAR(UpstreamNode->getCategory().c_str()));
+				UE_LOG(LogInterchangeImport, Warning, TEXT("<%s> is not supported."), ANSI_TO_TCHAR(UpstreamNode->getCategory().c_str()));
 			}
 		}
 	}
@@ -407,7 +407,7 @@ void FMaterialXSurfaceShaderAbstract::ConnectExtractInputToOutput(const FConnect
 		case 2: MaskShaderNode->AddBooleanAttribute(Mask::Attributes::B.ToString(), true); break;
 		case 3: MaskShaderNode->AddBooleanAttribute(Mask::Attributes::A.ToString(), true); break;
 		default:
-			UE_LOG(LogInterchangeImport, Warning, TEXT("Wrong index number for extract node, values are from [0-3]"));
+			UE_LOG(LogInterchangeImport, Warning, TEXT("Wrong index number for extracted node. Values are from [0-3]."));
 			break;
 		}
 	}
@@ -528,7 +528,7 @@ void FMaterialXSurfaceShaderAbstract::ConnectConvertInputToOutput(const FConnect
 
 	if(!bIsNodeTypeSupported || !bIsInputTypeSupported)
 	{
-		UE_LOG(LogInterchangeImport, Warning, TEXT("<convert> node has non supported types"));
+		UE_LOG(LogInterchangeImport, Warning, TEXT("<convert> node has unsupported types."));
 		return;
 	}
 
@@ -1060,7 +1060,7 @@ void FMaterialXSurfaceShaderAbstract::ConnectBlurInputToOutput(const FConnectNod
 					}
 					else
 					{
-						UE_LOG(LogInterchangeImport, Warning, TEXT("<%s>: input 'size' must have a value"), ANSI_TO_TCHAR(Connect.UpstreamNode->getName().c_str()));
+						UE_LOG(LogInterchangeImport, Warning, TEXT("<%s>: input 'size' must have a value."), ANSI_TO_TCHAR(Connect.UpstreamNode->getName().c_str()));
 					}
 				}
 			}
@@ -1129,7 +1129,7 @@ void FMaterialXSurfaceShaderAbstract::ConnectSeparateInputToOutput(const FConnec
 	}
 	else
 	{
-		UE_LOG(LogInterchangeImport, Warning, TEXT("output <%s> not defined in <%s>"), *Connect.OutputName, ANSI_TO_TCHAR(Connect.UpstreamNode->getCategory().c_str()));
+		UE_LOG(LogInterchangeImport, Warning, TEXT("output <%s> not defined in <%s>."), *Connect.OutputName, ANSI_TO_TCHAR(Connect.UpstreamNode->getCategory().c_str()));
 	}
 }
 
@@ -1284,7 +1284,7 @@ FLinearColor FMaterialXSurfaceShaderAbstract::GetVector(MaterialX::InputPtr Inpu
 	}
 	else
 	{
-		ensureMsgf(false, TEXT("input type can only be a vectorN"));
+		ensureMsgf(false, TEXT("Input type can only be a vectorN."));
 	}
 
 	return LinearColor;
diff --git a/Engine/Plugins/Interchange/Runtime/Source/Import/Private/Mesh/InterchangePhysicsAssetFactory.cpp b/Engine/Plugins/Interchange/Runtime/Source/Import/Private/Mesh/InterchangePhysicsAssetFactory.cpp
--- a/Engine/Plugins/Interchange/Runtime/Source/Import/Private/Mesh/InterchangePhysicsAssetFactory.cpp
+++ b/Engine/Plugins/Interchange/Runtime/Source/Import/Private/Mesh/InterchangePhysicsAssetFactory.cpp
@@ -75,7 +75,7 @@ UInterchangeFactoryBase::FImportAssetResult UInterchangePhysicsAssetFactory::Imp
 	FImportAssetResult ImportAssetResult;
 #if !WITH_EDITORONLY_DATA
 
-	UE_LOG(LogInterchangeImport, Error, TEXT("Cannot import PhysicsAsset asset in runtime, this is an editor only feature."));
+	UE_LOG(LogInterchangeImport, Error, TEXT("Cannot import PhysicsAsset asset at runtime. This is an editor-only feature."));
 	return ImportAssetResult;
 
 #else
@@ -95,14 +95,14 @@ UInterchangeFactoryBase::FImportAssetResult UInterchangePhysicsAssetFactory::Imp
 
 	if (!PhysicsAssetObject)
 	{
-		UE_LOG(LogInterchangeImport, Error, TEXT("Could not import the PhysicsAsset asset %s, because the asset do not exist."), *Arguments.AssetName);
+		UE_LOG(LogInterchangeImport, Error, TEXT("Could not import the PhysicsAsset asset %s because the asset does not exist."), *Arguments.AssetName);
 		return ImportAssetResult;
 	}
 
 	UPhysicsAsset* PhysicsAsset = Cast<UPhysicsAsset>(PhysicsAssetObject);
 	if (!ensure(PhysicsAsset))
 	{
-		UE_LOG(LogInterchangeImport, Error, TEXT("Could not cast to PhysicsAsset asset %s"), *Arguments.AssetName);
+		UE_LOG(LogInterchangeImport, Error, TEXT("Could not cast to PhysicsAsset asset %s."), *Arguments.AssetName);
 		return ImportAssetResult;
 	}
 
diff --git a/Engine/Plugins/Interchange/Runtime/Source/Import/Private/Mesh/InterchangeSkeletalMeshFactory.cpp b/Engine/Plugins/Interchange/Runtime/Source/Import/Private/Mesh/InterchangeSkeletalMeshFactory.cpp
--- a/Engine/Plugins/Interchange/Runtime/Source/Import/Private/Mesh/InterchangeSkeletalMeshFactory.cpp
+++ b/Engine/Plugins/Interchange/Runtime/Source/Import/Private/Mesh/InterchangeSkeletalMeshFactory.cpp
@@ -102,7 +102,7 @@ namespace UE
 						TOptional<FInterchangeMeshPayLoadKey> OptionalPayLoadKey = MorphTargetMeshNode->GetPayLoadKey();
 						if (!OptionalPayLoadKey.IsSet())
 						{
-							UE_LOG(LogInterchangeImport, Warning, TEXT("Empty LOD morph target mesh reference payload when importing SkeletalMesh asset %s"), *AssetName);
+							UE_LOG(LogInterchangeImport, Warning, TEXT("Empty LOD morph target mesh reference payload when importing SkeletalMesh asset %s."), *AssetName);
 							continue;
 						}
 						FInterchangeMeshPayLoadKey& PayLoadKey = OptionalPayLoadKey.GetValue();
@@ -133,7 +133,7 @@ namespace UE
 						TOptional<UE::Interchange::FMeshPayloadData> MorphTargetMeshPayload = TempMorphTargetMeshDescriptionsPerMorphTargetName.FindChecked(MorphTargetPayloadKeyString).Get();
 						if (!MorphTargetMeshPayload.IsSet())
 						{
-							UE_LOG(LogInterchangeImport, Warning, TEXT("Invalid Skeletal mesh morph target payload key [%s] SkeletalMesh asset %s"), *MorphTargetPayloadKeyString, *AssetName);
+							UE_LOG(LogInterchangeImport, Warning, TEXT("Invalid skeletal mesh morph target payload key [%s] for SkeletalMesh asset %s."), *MorphTargetPayloadKeyString, *AssetName);
 							continue;
 						}
 						MorphTargetMeshPayload->VertexOffset = VertexOffset;
@@ -162,7 +162,7 @@ namespace UE
 					const TOptional<UE::Interchange::FMeshPayloadData>& MorphTargetPayloadData = Pair.Value;
 					if (!MorphTargetPayloadData.IsSet())
 					{
-						UE_LOG(LogInterchangeImport, Error, TEXT("Empty morph target optional payload data [%s]"), *MorphTargetName);
+						UE_LOG(LogInterchangeImport, Error, TEXT("Empty morph target optional payload data [%s]."), *MorphTargetName);
 						continue;
 					}
 
@@ -172,7 +172,7 @@ namespace UE
 					const int32 DestinationVertexIndexMax = VertexOffset + SourceMeshVertexCount;
 					if (!DestinationSkeletalMeshImportData.Points.IsValidIndex(DestinationVertexIndexMax-1))
 					{
-						UE_LOG(LogInterchangeImport, Error, TEXT("Corrupted morph target optional payload data [%s]"), *MorphTargetName);
+						UE_LOG(LogInterchangeImport, Error, TEXT("Corrupted morph target optional payload data [%s]."), *MorphTargetName);
 						continue;
 					}
 					Keys.Add(Pair.Key);
@@ -493,7 +493,7 @@ namespace UE
 					TOptional<UE::Interchange::FMeshPayloadData> LodMeshPayload = MeshNodeContextAndFuture.Value.Get();
 					if (!LodMeshPayload.IsSet())
 					{
-						UE_LOG(LogInterchangeImport, Warning, TEXT("Invalid Skeletal mesh payload key [%s] SkeletalMesh asset %s"), *MeshNodeContext.TranslatorPayloadKey.UniqueId, *Arguments.AssetName);
+						UE_LOG(LogInterchangeImport, Warning, TEXT("Invalid skeletal mesh payload key [%s] for SkeletalMesh asset %s."), *MeshNodeContext.TranslatorPayloadKey.UniqueId, *Arguments.AssetName);
 						continue;
 					}
 					
@@ -1010,7 +1010,7 @@ UInterchangeFactoryBase::FImportAssetResult UInterchangeSkeletalMeshFactory::Beg
 
 #if !WITH_EDITOR || !WITH_EDITORONLY_DATA
 
-	UE_LOG(LogInterchangeImport, Error, TEXT("Cannot import skeletalMesh asset in runtime, this is an editor only feature."));
+	UE_LOG(LogInterchangeImport, Error, TEXT("Cannot import SkeletalMesh asset at runtime. This is an editor-only feature."));
 	return ImportAssetResult;
 
 #else
@@ -1052,7 +1052,7 @@ UInterchangeFactoryBase::FImportAssetResult UInterchangeSkeletalMeshFactory::Beg
 	{
 		if (!Arguments.ReimportObject)
 		{
-			UE_LOG(LogInterchangeImport, Warning, TEXT("Could not create or find a SkeletalMesh asset name %s"), *Arguments.AssetName);
+			UE_LOG(LogInterchangeImport, Warning, TEXT("Could not create or find a SkeletalMesh asset named %s."), *Arguments.AssetName);
 		}
 		return ImportAssetResult;
 	}
@@ -1121,10 +1121,10 @@ UInterchangeFactoryBase::FImportAssetResult UInterchangeSkeletalMeshFactory::Beg
 
 	for (int32 LodIndex = 0; LodIndex < LodCount; ++LodIndex)
 	{
-		FText WarningMessage_InvalidSkeleton = FText::Format(NSLOCTEXT("InterchangeSkeletalMeshFactory", "BeginImportAsset_GameThread_InvalidSkeletonLOD", "Invalid Skeleton LOD {0} when importing SkeletalMesh asset {1}")
+		FText WarningMessage_InvalidSkeleton = FText::Format(NSLOCTEXT("InterchangeSkeletalMeshFactory", "BeginImportAsset_GameThread_InvalidSkeletonLOD", "Invalid Skeleton LOD {0} when importing SkeletalMesh asset {1}.")
 			, FText::AsNumber(LodIndex)
 			, FText::FromString(Arguments.AssetName));
-		FText WarningMessage_InvalidRootJoint = FText::Format(NSLOCTEXT("InterchangeSkeletalMeshFactory", "BeginImportAsset_GameThread_InvalidSkeletonRootJoint", "Invalid Skeleton LOD {0} Root Joint when importing SkeletalMesh asset {1}")
+		FText WarningMessage_InvalidRootJoint = FText::Format(NSLOCTEXT("InterchangeSkeletalMeshFactory", "BeginImportAsset_GameThread_InvalidSkeletonRootJoint", "Invalid Skeleton LOD {0}'s Root Joint when importing SkeletalMesh asset {1}.")
 			, FText::AsNumber(LodIndex)
 			, FText::FromString(Arguments.AssetName));
 
@@ -1133,7 +1133,7 @@ UInterchangeFactoryBase::FImportAssetResult UInterchangeSkeletalMeshFactory::Beg
 		if (!LodDataNode)
 		{
 			UInterchangeResultWarning_Generic* Message = AddMessage<UInterchangeResultWarning_Generic>();
-			Message->Text = FText::Format(NSLOCTEXT("InterchangeSkeletalMeshFactory", "BeginImportAsset_GameThread_InvalidLOD", "Invalid LOD {0} when importing SkeletalMesh asset {1}")
+			Message->Text = FText::Format(NSLOCTEXT("InterchangeSkeletalMeshFactory", "BeginImportAsset_GameThread_InvalidLOD", "Invalid LOD {0} when importing SkeletalMesh asset {1}.")
 				, FText::AsNumber(LodIndex)
 				, FText::FromString(Arguments.AssetName));
 			continue;
@@ -1244,7 +1244,7 @@ UInterchangeFactoryBase::FImportAssetResult UInterchangeSkeletalMeshFactory::Imp
 
 #if !WITH_EDITOR || !WITH_EDITORONLY_DATA
 
-	UE_LOG(LogInterchangeImport, Error, TEXT("Cannot import skeletalMesh asset in runtime, this is an editor only feature."));
+	UE_LOG(LogInterchangeImport, Error, TEXT("Cannot import SkeletalMesh asset at runtime. This is an editor-only feature."));
 	return ImportAssetResult;
 
 #else
@@ -1263,7 +1263,7 @@ UInterchangeFactoryBase::FImportAssetResult UInterchangeSkeletalMeshFactory::Imp
 	if (!MeshTranslatorPayloadInterface)
 	{
 		UInterchangeResultError_Generic* Message = AddMessage<UInterchangeResultError_Generic>();
-		Message->Text = FText::Format(NSLOCTEXT("InterchangeSkeletalMeshFactory", "ImportAsset_Async_TranslatorInterfaceMissing", "Cannot import skeletalMesh {0}, the translator {1} do not implement the IInterchangeSkeletalMeshPayloadInterface.")
+		Message->Text = FText::Format(NSLOCTEXT("InterchangeSkeletalMeshFactory", "ImportAsset_Async_TranslatorInterfaceMissing", "Cannot import skeletalMesh {0}, the translator {1} does not implement the IInterchangeSkeletalMeshPayloadInterface.")
 			, FText::FromString(Arguments.AssetName)
 			, FText::FromString(Arguments.Translator->GetName()));
 		return ImportAssetResult;
@@ -1402,7 +1402,7 @@ UInterchangeFactoryBase::FImportAssetResult UInterchangeSkeletalMeshFactory::Imp
 		if (!LodDataNode)
 		{
 			UInterchangeResultWarning_Generic* Message = AddMessage<UInterchangeResultWarning_Generic>();
-			Message->Text = FText::Format(NSLOCTEXT("InterchangeSkeletalMeshFactory", "ImportAsset_Async_InvalidLOD", "Invalid LOD {0} when importing SkeletalMesh asset {1}")
+			Message->Text = FText::Format(NSLOCTEXT("InterchangeSkeletalMeshFactory", "ImportAsset_Async_InvalidLOD", "Invalid LOD {0} when importing SkeletalMesh asset {1}.")
 				, FText::AsNumber(LodIndex)
 				, FText::FromString(Arguments.AssetName));
 			continue;
@@ -1520,7 +1520,7 @@ UInterchangeFactoryBase::FImportAssetResult UInterchangeSkeletalMeshFactory::Imp
 				else
 				{
 					UInterchangeResultWarning_Generic* Message = AddMessage<UInterchangeResultWarning_Generic>();
-					Message->Text = FText::Format(NSLOCTEXT("InterchangeSkeletalMeshFactory", "ImportAsset_Async_EmptyLODMeshReference", "Empty LOD {0} mesh reference payload when importing SkeletalMesh asset {1}")
+					Message->Text = FText::Format(NSLOCTEXT("InterchangeSkeletalMeshFactory", "ImportAsset_Async_EmptyLODMeshReference", "Empty LOD {0} mesh reference payload when importing SkeletalMesh asset {1}.")
 						, FText::AsNumber(LodIndex)
 						, FText::FromString(Arguments.AssetName));
 					continue;
@@ -1712,7 +1712,7 @@ UInterchangeFactoryBase::FImportAssetResult UInterchangeSkeletalMeshFactory::End
 
 #if !WITH_EDITOR || !WITH_EDITORONLY_DATA
 
-	UE_LOG(LogInterchangeImport, Error, TEXT("Cannot import skeletalMesh asset in runtime, this is an editor only feature."));
+	UE_LOG(LogInterchangeImport, Error, TEXT("Cannot import SkeletalMesh asset at runtime. This is an editor-only feature."));
 	return ImportAssetResult;
 
 #else
@@ -1746,11 +1746,11 @@ UInterchangeFactoryBase::FImportAssetResult UInterchangeSkeletalMeshFactory::End
 	{
 		if (Arguments.ReimportObject == nullptr)
 		{
-			UE_LOG(LogInterchangeImport, Error, TEXT("Could not create skeletalMesh asset %s"), *Arguments.AssetName);
+			UE_LOG(LogInterchangeImport, Error, TEXT("Could not create SkeletalMesh asset %s."), *Arguments.AssetName);
 		}
 		else
 		{
-			UE_LOG(LogInterchangeImport, Error, TEXT("Could not find reimported skeletalMesh asset %s"), *Arguments.AssetName);
+			UE_LOG(LogInterchangeImport, Error, TEXT("Could not find reimported SkeletalMesh asset %s."), *Arguments.AssetName);
 		}
 		return ImportAssetResult;
 	}
@@ -1824,7 +1824,7 @@ UInterchangeFactoryBase::FImportAssetResult UInterchangeSkeletalMeshFactory::End
 
 				EAppReturnType::Type MergeBonesChoice = FMessageDialog::Open(EAppMsgType::YesNo
 					, EAppReturnType::No
-					, NSLOCTEXT("InterchangeSkeletalMeshFactory", "SkeletonFailed_BoneMerge", "FAILED TO MERGE BONES:\n\n This could happen if significant hierarchical changes have been made\ne.g. inserting a bone between nodes.\nWould you like to regenerate the Skeleton from this mesh?\n\n***WARNING: THIS MAY INVALIDATE OR REQUIRE RECOMPRESSION OF ANIMATION DATA.***\n"));
+					, NSLOCTEXT("InterchangeSkeletalMeshFactory", "SkeletonFailed_BoneMerge", "Failed to merge bones.\n\n This can happen if significant hierarchical changes have been made,\nsuch as inserting a bone between nodes.\n\nWould you like to regenerate the Skeleton from this mesh? This may invalidate or require recompression of animation data.\n"));
 				if (MergeBonesChoice == EAppReturnType::Yes)
 				{
 					//Allow this thread scope to read and write skeletalmesh locked properties
diff --git a/Engine/Plugins/Interchange/Runtime/Source/Import/Private/Mesh/InterchangeStaticMeshFactory.cpp b/Engine/Plugins/Interchange/Runtime/Source/Import/Private/Mesh/InterchangeStaticMeshFactory.cpp
--- a/Engine/Plugins/Interchange/Runtime/Source/Import/Private/Mesh/InterchangeStaticMeshFactory.cpp
+++ b/Engine/Plugins/Interchange/Runtime/Source/Import/Private/Mesh/InterchangeStaticMeshFactory.cpp
@@ -200,7 +200,7 @@ UInterchangeFactoryBase::FImportAssetResult UInterchangeStaticMeshFactory::Begin
 	{
 		if (!Arguments.ReimportObject)
 		{
-			UE_LOG(LogInterchangeImport, Warning, TEXT("Could not create StaticMesh asset %s"), *Arguments.AssetName);
+			UE_LOG(LogInterchangeImport, Warning, TEXT("Could not create StaticMesh asset %s."), *Arguments.AssetName);
 		}
 		return ImportAssetResult;
 	}
@@ -240,14 +240,14 @@ UInterchangeFactoryBase::FImportAssetResult UInterchangeStaticMeshFactory::Impor
 
 	if (!StaticMeshObject)
 	{
-		UE_LOG(LogInterchangeImport, Error, TEXT("Could not import the StaticMesh asset %s, because the asset do not exist."), *Arguments.AssetName);
+		UE_LOG(LogInterchangeImport, Error, TEXT("Could not import the StaticMesh asset %s because the asset does not exist."), *Arguments.AssetName);
 		return ImportAssetResult;
 	}
 
 	UStaticMesh* StaticMesh = Cast<UStaticMesh>(StaticMeshObject);
 	if (!ensure(StaticMesh))
 	{
-		UE_LOG(LogInterchangeImport, Error, TEXT("Could not cast to StaticMesh asset %s"), *Arguments.AssetName);
+		UE_LOG(LogInterchangeImport, Error, TEXT("Could not cast to StaticMesh asset %s."), *Arguments.AssetName);
 		return ImportAssetResult;
 	}
 
@@ -257,7 +257,7 @@ UInterchangeFactoryBase::FImportAssetResult UInterchangeStaticMeshFactory::Impor
 	if (LodCount != StaticMeshFactoryNode->GetLodDataCount())
 	{
 		const int32 LodCountDiff = StaticMeshFactoryNode->GetLodDataCount() - MAX_STATIC_MESH_LODS;
-		UE_LOG(LogInterchangeImport, Warning, TEXT("Reached the maximum number of LODs for a Static Mesh(%d) - discarding %d LOD meshes."), MAX_STATIC_MESH_LODS, LodCountDiff);
+		UE_LOG(LogInterchangeImport, Warning, TEXT("Reached the maximum number of LODs for a Static Mesh (%d) - discarding %d LOD meshes."), MAX_STATIC_MESH_LODS, LodCountDiff);
 	}
 #if WITH_EDITOR
 	const int32 PrevLodCount = StaticMesh->GetNumSourceModels();
@@ -341,7 +341,7 @@ UInterchangeFactoryBase::FImportAssetResult UInterchangeStaticMeshFactory::Impor
 		const UInterchangeStaticMeshLodDataNode* LodDataNode = Cast<UInterchangeStaticMeshLodDataNode>(Arguments.NodeContainer->GetNode(LodUniqueId));
 		if (!LodDataNode)
 		{
-			UE_LOG(LogInterchangeImport, Warning, TEXT("Invalid LOD when importing StaticMesh asset %s"), *Arguments.AssetName);
+			UE_LOG(LogInterchangeImport, Warning, TEXT("Invalid LOD when importing StaticMesh asset %s."), *Arguments.AssetName);
 			continue;
 		}
 
@@ -368,7 +368,7 @@ UInterchangeFactoryBase::FImportAssetResult UInterchangeStaticMeshFactory::Impor
 			const TOptional<UE::Interchange::FMeshPayloadData>& LodMeshPayload = MeshPayload.PayloadData.Get();
 			if (!LodMeshPayload.IsSet())
 			{
-				UE_LOG(LogInterchangeImport, Warning, TEXT("Invalid static mesh payload key, StaticMesh asset %s"), *Arguments.AssetName);
+				UE_LOG(LogInterchangeImport, Warning, TEXT("Invalid static mesh payload key for StaticMesh asset %s."), *Arguments.AssetName);
 				continue;
 			}
 
@@ -525,7 +525,7 @@ UInterchangeFactoryBase::FImportAssetResult UInterchangeStaticMeshFactory::EndIm
 	UStaticMesh* StaticMesh = Cast<UStaticMesh>(ExistingAsset);
 	if (!ensure(StaticMesh))
 	{
-		UE_LOG(LogInterchangeImport, Error, TEXT("Could not create StaticMesh asset %s"), *Arguments.AssetName);
+		UE_LOG(LogInterchangeImport, Error, TEXT("Could not create StaticMesh asset %s."), *Arguments.AssetName);
 		return ImportAssetResult;
 	}
 
@@ -676,8 +676,8 @@ UInterchangeFactoryBase::FImportAssetResult UInterchangeStaticMeshFactory::EndIm
 			if (HiresMeshDescription && Lod0MeshDescription)
 			{
 				StaticMesh->ModifyHiResMeshDescription();
-				FString MaterialNameConflictMsg = TEXT("[Asset ") + StaticMesh->GetPathName() + TEXT("] Nanite hi - res import have some material name that differ from the LOD 0 material name.Your nanite hi - res should use the same material names the LOD 0 use to ensure we can remap the section in the same order.");
-				FString MaterialCountConflictMsg = TEXT("[Asset ") + StaticMesh->GetPathName() + TEXT("] Nanite hi-res import dont have the same material count then LOD 0. Your nanite hi-res should have equal number of material.");
+				FString MaterialNameConflictMsg = TEXT("[Asset ") + StaticMesh->GetPathName() + TEXT("] Nanite high-resolution import has material names that differ from the LOD 0 material name. Your Nanite high-resolution mesh should use the same material names the LOD 0 uses to ensure the sections can be remapped in the same order.");
+				FString MaterialCountConflictMsg = TEXT("[Asset ") + StaticMesh->GetPathName() + TEXT("] Nanite high-resolution import doesn't have the same material count as LOD 0. Your Nanite high-resolution mesh should have the same number of materials as LOD 0.");
 				FStaticMeshOperations::ReorderMeshDescriptionPolygonGroups(*Lod0MeshDescription, *HiresMeshDescription, MaterialNameConflictMsg, MaterialCountConflictMsg);
 				StaticMesh->CommitHiResMeshDescription();
 			}
@@ -858,7 +858,7 @@ TArray<UInterchangeStaticMeshFactory::FMeshPayload> UInterchangeStaticMeshFactor
 	const IInterchangeMeshPayloadInterface* MeshTranslatorPayloadInterface = Cast<IInterchangeMeshPayloadInterface>(Arguments.Translator);
 	if (!MeshTranslatorPayloadInterface)
 	{
-		UE_LOG(LogInterchangeImport, Error, TEXT("Cannot import static mesh, the translator does not implement the IInterchangeMeshPayloadInterface."));
+		UE_LOG(LogInterchangeImport, Error, TEXT("Cannot import static mesh. The translator does not implement IInterchangeMeshPayloadInterface."));
 		return Payloads;
 	}
 
@@ -882,7 +882,7 @@ TArray<UInterchangeStaticMeshFactory::FMeshPayload> UInterchangeStaticMeshFactor
 			const UInterchangeSceneNode* SceneNode = Cast<const UInterchangeSceneNode>(Node);
 			if (!ensure(SceneNode))
 			{
-				UE_LOG(LogInterchangeImport, Warning, TEXT("Invalid LOD mesh reference when importing StaticMesh asset %s"), *Arguments.AssetName);
+				UE_LOG(LogInterchangeImport, Warning, TEXT("Invalid LOD mesh reference when importing StaticMesh asset %s."), *Arguments.AssetName);
 				continue;
 			}
 
@@ -912,14 +912,14 @@ TArray<UInterchangeStaticMeshFactory::FMeshPayload> UInterchangeStaticMeshFactor
 
 		if (!ensure(MeshNode))
 		{
-			UE_LOG(LogInterchangeImport, Warning, TEXT("Invalid LOD mesh reference when importing StaticMesh asset %s"), *Arguments.AssetName);
+			UE_LOG(LogInterchangeImport, Warning, TEXT("Invalid LOD mesh reference when importing StaticMesh asset %s."), *Arguments.AssetName);
 			continue;
 		}
 
 		TOptional<FInterchangeMeshPayLoadKey> OptionalPayLoadKey = MeshNode->GetPayLoadKey();
 		if (!ensure(OptionalPayLoadKey.IsSet()))
 		{
-			UE_LOG(LogInterchangeImport, Warning, TEXT("Empty LOD mesh reference payload when importing StaticMesh asset %s"), *Arguments.AssetName);
+			UE_LOG(LogInterchangeImport, Warning, TEXT("Empty LOD mesh reference payload when importing StaticMesh asset %s."), *Arguments.AssetName);
 			continue;
 		}
 
diff --git a/Engine/Plugins/Interchange/Runtime/Source/Import/Private/Scene/InterchangeSceneVariantSetsFactory.cpp b/Engine/Plugins/Interchange/Runtime/Source/Import/Private/Scene/InterchangeSceneVariantSetsFactory.cpp
--- a/Engine/Plugins/Interchange/Runtime/Source/Import/Private/Scene/InterchangeSceneVariantSetsFactory.cpp
+++ b/Engine/Plugins/Interchange/Runtime/Source/Import/Private/Scene/InterchangeSceneVariantSetsFactory.cpp
@@ -101,7 +101,7 @@ namespace UE::Interchange::Private
 		const TOptional<UE::Interchange::FVariantSetPayloadData> PayloadData = PayloadInterface.GetVariantSetPayloadData(PayloadKey).Consume();
 		if (!PayloadData.IsSet())
 		{
-			UE_LOG(LogInterchangeImport, Warning, TEXT("No payload for variant set %s"), *DisplayText);
+			UE_LOG(LogInterchangeImport, Warning, TEXT("No payload for variant set %s."), *DisplayText);
 			return;
 		}
 
@@ -133,7 +133,7 @@ namespace UE::Interchange::Private
 		AActor* Actor = GetActor(Binding.TargetUid);
 		if (!Actor)
 		{
-			UE_LOG(LogInterchangeImport, Warning, TEXT("Cannot find actor for variant %s"), *Variant.GetDisplayText().ToString());
+			UE_LOG(LogInterchangeImport, Warning, TEXT("Cannot find actor for variant %s."), *Variant.GetDisplayText().ToString());
 			return;
 		}
 
@@ -261,12 +261,12 @@ namespace UE::Interchange::Private
 					}
 					else
 					{
-						UE_LOG(LogInterchangeImport, Warning, TEXT("Did not find material '%s' when creating variant asset."), *TargetFactoryNode->GetDisplayLabel());
+						UE_LOG(LogInterchangeImport, Warning, TEXT("Could not find material '%s' when creating variant asset."), *TargetFactoryNode->GetDisplayLabel());
 					}
 				}
 				else
 				{
-					UE_LOG(LogInterchangeImport, Warning, TEXT("Did not find factory when creating variant assets."));
+					UE_LOG(LogInterchangeImport, Warning, TEXT("Could not find factory when creating variant assets."));
 				}
 			}
 		}
@@ -300,7 +300,7 @@ UInterchangeFactoryBase::FImportAssetResult UInterchangeSceneVariantSetsFactory:
 	FImportAssetResult ImportAssetResult;
 #if !WITH_EDITOR || !WITH_EDITORONLY_DATA
 
-	UE_LOG(LogInterchangeImport, Error, TEXT("Cannot import levelsequence asset in runtime, this is an editor only feature."));
+	UE_LOG(LogInterchangeImport, Error, TEXT("Cannot import LevelSequence asset at runtime. This is an editor-only feature."));
 	return ImportAssetResult;
 #else
 	auto CannotReimportMessage = [&Arguments, this]()
@@ -309,7 +309,7 @@ UInterchangeFactoryBase::FImportAssetResult UInterchangeSceneVariantSetsFactory:
 		Message->SourceAssetName = Arguments.SourceData->GetFilename();
 		Message->DestinationAssetName = Arguments.AssetName;
 		Message->AssetType = ULevelVariantSets::StaticClass();
-		Message->Text = LOCTEXT("CreateEmptyAssetUnsupportedReimport", "Re-import of ULevelVariantSets not supported yet.");
+		Message->Text = LOCTEXT("CreateEmptyAssetUnsupportedReimport", "Reimport of ULevelVariantSets is not supported.");
 		Arguments.AssetNode->SetSkipNodeImport();
 	};
 
@@ -321,7 +321,7 @@ UInterchangeFactoryBase::FImportAssetResult UInterchangeSceneVariantSetsFactory:
 
 	if (!Cast<IInterchangeVariantSetPayloadInterface>(Arguments.Translator))
 	{
-		UE_LOG(LogInterchangeImport, Error, TEXT("Cannot import LevelVariantSets, the translator do not implement the IInterchangeVariantSetPayloadInterface interface."));
+		UE_LOG(LogInterchangeImport, Error, TEXT("Cannot import LevelVariantSets. The translator does not implement the IInterchangeVariantSetPayloadInterface interface."));
 		return ImportAssetResult;
 	}
 
@@ -363,7 +363,7 @@ UInterchangeFactoryBase::FImportAssetResult UInterchangeSceneVariantSetsFactory:
 
 	if (!LevelVariantSets)
 	{
-		UE_LOG(LogInterchangeImport, Warning, TEXT("Could not create LevelVariantSets asset %s"), *Arguments.AssetName);
+		UE_LOG(LogInterchangeImport, Warning, TEXT("Could not create LevelVariantSets asset %s."), *Arguments.AssetName);
 		return ImportAssetResult;
 	}
 
@@ -377,7 +377,7 @@ UObject* UInterchangeSceneVariantSetsFactory::ImportObjectSourceData(const FImpo
 	TRACE_CPUPROFILER_EVENT_SCOPE(UInterchangeSceneVariantSetsFactory::ImportObjectSourceData);
 #if !WITH_EDITOR || !WITH_EDITORONLY_DATA
 	// TODO: Can we import ULevelVariantSets at runtime
-	UE_LOG(LogInterchangeImport, Error, TEXT("Cannot import LevelVariantSets asset in runtime, this is an editor only feature."));
+	UE_LOG(LogInterchangeImport, Error, TEXT("Cannot import LevelVariantSets asset at runtime. This is an editor-only feature."));
 	return nullptr;
 
 #else
@@ -399,7 +399,7 @@ UObject* UInterchangeSceneVariantSetsFactory::ImportObjectSourceData(const FImpo
 	const IInterchangeVariantSetPayloadInterface* LevelVariantSetsPayloadInterface = Cast<IInterchangeVariantSetPayloadInterface>(Arguments.Translator);
 	if (!LevelVariantSetsPayloadInterface)
 	{
-		UE_LOG(LogInterchangeImport, Error, TEXT("Cannot import LevelVariantSets, the translator do not implement the IInterchangeVariantSetPayloadInterface interface."));
+		UE_LOG(LogInterchangeImport, Error, TEXT("Cannot import LevelVariantSets. The translator does not implement the IInterchangeVariantSetPayloadInterface interface."));
 		return nullptr;
 	}
 
diff --git a/Engine/Plugins/Interchange/Runtime/Source/Import/Private/Texture/InterchangeTextureFactory.cpp b/Engine/Plugins/Interchange/Runtime/Source/Import/Private/Texture/InterchangeTextureFactory.cpp
--- a/Engine/Plugins/Interchange/Runtime/Source/Import/Private/Texture/InterchangeTextureFactory.cpp
+++ b/Engine/Plugins/Interchange/Runtime/Source/Import/Private/Texture/InterchangeTextureFactory.cpp
@@ -247,10 +247,10 @@ namespace UE::Interchange::Private::InterchangeTextureFactory
 		}
 
 		ensureMsgf(false
-			, TEXT("Unknow factory node class (%s). To add support for a new texture type either update this factory or register a new factory to interchange.\n\
-				If the factory node class is used to add functions for some custom attributes of a pipeline please don't do that.\n\
-				Use some static functions libraries or use some helper structs/objects instead.\n\
-				Here is exemple of what this could look like in c++:\n\n\
+			, TEXT("Unknown factory node class (%s). To add support for a new texture type, either update this factory or register a new factory to Interchange.\n\
+				However, if the factory node class is used to add functions for some custom attributes of a pipeline,\n\
+				use static function libraries or use helper structs/objects instead.\n\
+				Here is an example of what this could look like in C++:\n\n\
 				\tif (UCustomAttributeInterface::HasInterface(TextureFactoryNode))\n\
 				\t{\n\
 					\t\tFString MyPath = UCustomAttributeInterface::GetCustomPath(TextureFactoryNode);\n\
@@ -577,7 +577,7 @@ namespace UE::Interchange::Private::InterchangeTextureFactory
 
 			if (bMismatchedGammaSpace || bMismatchedFormats)
 			{
-				UE_LOG(LogInterchangeImport, Display, TEXT("Mismatched UDIM image %s, converting all to %s/%s ..."), bMismatchedGammaSpace ? TEXT("gamma spaces") : TEXT("pixel formats"),
+				UE_LOG(LogInterchangeImport, Display, TEXT("Mismatched UDIM image %s. Converting all to %s/%s ..."), bMismatchedGammaSpace ? TEXT("gamma spaces") : TEXT("pixel formats"),
 					ERawImageFormat::GetName(FImageCoreUtils::ConvertToRawImageFormat(BlockedImage.Format)), BlockedImage.bSRGB ? TEXT("sRGB") : TEXT("Linear"));
 
 				for (UE::Interchange::FImportImage& Image : Images)
@@ -600,7 +600,7 @@ namespace UE::Interchange::Private::InterchangeTextureFactory
 
 						if ( Image.NumMips != 1 )
 						{
-							UE_LOG(LogInterchangeImport, Warning, TEXT("UDIM Image had existing mips; they were discarded in format change") );
+							UE_LOG(LogInterchangeImport, Warning, TEXT("UDIM Image had existing mips. They were discarded in the format change.") );
 							Image.NumMips = 1; // discard imported mips, they won't be used in UDIM VT anyway
 							Image.MipGenSettings.Reset();
 						}
@@ -620,7 +620,7 @@ namespace UE::Interchange::Private::InterchangeTextureFactory
 
 				if ( ! BlockedImage.InitBlockFromImage(BlockX, BlockY, Images[Index]) )
 				{
-					UE_LOG(LogInterchangeImport, Warning, TEXT("UDIM InitBlockFromImage failed") );
+					UE_LOG(LogInterchangeImport, Warning, TEXT("UDIM InitBlockFromImage failed.") );
 					return { };
 				}
 
@@ -632,7 +632,7 @@ namespace UE::Interchange::Private::InterchangeTextureFactory
 
 			if (! BlockedImage.MigrateDataFromImagesToRawData(Images) )
 			{
-				UE_LOG(LogInterchangeImport, Warning, TEXT("UDIM MigrateDataFromImagesToRawData failed") );
+				UE_LOG(LogInterchangeImport, Warning, TEXT("UDIM MigrateDataFromImagesToRawData failed.") );
 				return { };
 			}
 			
@@ -686,7 +686,7 @@ namespace UE::Interchange::Private::InterchangeTextureFactory
 						const bool bShouldImportRawCache = bShoudImportCompressedImage;
 						if (GConfig->GetBool(TEXT("TextureImporter"), TEXT("RetainJpegFormat"), bShoudImportCompressedImage, GEditorIni) && bShouldImportRawCache != bShoudImportCompressedImage)
 						{
-							UE_LOG(LogInterchangeImport, Log, TEXT("JPEG file [%s]: Pipeline setting 'bPreferCompressedSourceData' has been overridden by Editor setting 'RetainJpegFormat'"), *SourceData->GetFilename());
+							UE_LOG(LogInterchangeImport, Log, TEXT("JPEG file [%s]: Pipeline setting 'bPreferCompressedSourceData' has been overridden by Editor setting 'RetainJpegFormat'."), *SourceData->GetFilename());
 						}
 					}
 
@@ -1574,7 +1574,7 @@ UInterchangeFactoryBase::FImportAssetResult UInterchangeTextureFactory::BeginImp
 	const UClass* TextureClass = Arguments.AssetNode->GetObjectClass();
 	if (!TextureClass || !TextureClass->IsChildOf(UTexture::StaticClass()))
 	{
-		CouldNotCreateTextureLog(LOCTEXT("TextureFactory_NodeClassMissmatch", "Asset node parameter class doesnt derive from UTexture."));
+		CouldNotCreateTextureLog(LOCTEXT("TextureFactory_NodeClassMissmatch", "Asset node parameter class doesn't derive from UTexture."));
 		return ImportAssetResult;
 	}
 
@@ -1589,7 +1589,7 @@ UInterchangeFactoryBase::FImportAssetResult UInterchangeTextureFactory::BeginImp
 	FTextureNodeVariant TextureNodeVariant = GetTextureNodeVariantFromFactoryVariant(GetAsTextureFactoryNodeVariant(Arguments.AssetNode, SupportedFactoryNodeClass), Arguments.NodeContainer);
 	if (TextureNodeVariant.IsType<FEmptyVariantState>())
 	{
-		FText Info = FText::Format(LOCTEXT("TextureFactory_InvalidTextureTranslated", "Asset factory node (%s) do not reference a valid texture translated node.")
+		FText Info = FText::Format(LOCTEXT("TextureFactory_InvalidTextureTranslated", "Asset factory node (%s) doesn't reference a valid texture translated node.")
 			, FText::FromString(SupportedFactoryNodeClass->GetAuthoredName()));
 		CouldNotCreateTextureLog(Info);
 		return ImportAssetResult;
@@ -1597,7 +1597,7 @@ UInterchangeFactoryBase::FImportAssetResult UInterchangeTextureFactory::BeginImp
 
 	if (!HasPayloadKey(TextureNodeVariant))
 	{
-		CouldNotCreateTextureLog(LOCTEXT("TextureFactory_InvalidPayloadKey", "Texture translated node doesnt have a payload key."));
+		CouldNotCreateTextureLog(LOCTEXT("TextureFactory_InvalidPayloadKey", "Texture translated node doesn't have a payload key."));
 		return ImportAssetResult;
 	}
 
@@ -1642,7 +1642,7 @@ UInterchangeFactoryBase::FImportAssetResult UInterchangeTextureFactory::BeginImp
 
 	if (!Texture)
 	{
-		CouldNotCreateTextureLog(LOCTEXT("TextureFactory_TextureCreateFail", "Texture creation fail."));
+		CouldNotCreateTextureLog(LOCTEXT("TextureFactory_TextureCreateFail", "Texture creation failed."));
 		return ImportAssetResult;
 	}
 
@@ -1677,7 +1677,7 @@ UInterchangeFactoryBase::FImportAssetResult UInterchangeTextureFactory::ImportAs
 	const UClass* TextureClass = Arguments.AssetNode->GetObjectClass();
 	if (!TextureClass || !TextureClass->IsChildOf(UTexture::StaticClass()))
 	{
-		ImportTextureErrorLog(LOCTEXT("TextureFactory_Async_MissMatchClass", "UInterchangeTextureFactory: Asset node parameter class doesnt derive from UTexture."));
+		ImportTextureErrorLog(LOCTEXT("TextureFactory_Async_MissMatchClass", "UInterchangeTextureFactory: Asset node parameter class doesn't derive from UTexture."));
 		return ImportAssetResult;
 	}
 
@@ -1706,7 +1706,7 @@ UInterchangeFactoryBase::FImportAssetResult UInterchangeTextureFactory::ImportAs
 	FTextureNodeVariant TextureNodeVariant = GetTextureNodeVariantFromFactoryVariant(TextureFactoryNodeVariant, Arguments.NodeContainer);
 	if (TextureNodeVariant.IsType<FEmptyVariantState>())
 	{
-		FText Info = FText::Format(LOCTEXT("TextureFactory_Async_InvalidTextureTranslated", "UInterchangeTextureFactory: Asset factory node (%s) do not reference a valid texture translated node.")
+		FText Info = FText::Format(LOCTEXT("TextureFactory_Async_InvalidTextureTranslated", "UInterchangeTextureFactory: Asset factory node (%s) doesn't reference a valid texture translated node.")
 			, FText::FromString(SupportedFactoryNodeClass->GetAuthoredName()));
 		ImportTextureErrorLog(Info);
 		return ImportAssetResult;
@@ -1715,7 +1715,7 @@ UInterchangeFactoryBase::FImportAssetResult UInterchangeTextureFactory::ImportAs
 	const TOptional<FString>& PayLoadKey = GetPayloadKey(TextureNodeVariant);
 	if (!PayLoadKey.IsSet())
 	{
-		ImportTextureErrorLog(LOCTEXT("TextureFactory_Async_InvalidPayloadKey", "UInterchangeTextureFactory: Texture translated node (UInterchangeTexture2DNode) doesnt have a payload key."));
+		ImportTextureErrorLog(LOCTEXT("TextureFactory_Async_InvalidPayloadKey", "UInterchangeTextureFactory: Texture translated node (UInterchangeTexture2DNode) doesn't have a payload key."));
 		return ImportAssetResult;
 	}
 
@@ -1724,7 +1724,7 @@ UInterchangeFactoryBase::FImportAssetResult UInterchangeTextureFactory::ImportAs
 
 	if(TexturePayload.IsType<FEmptyVariantState>())
 	{
-		ImportTextureErrorLog(LOCTEXT("TextureFactory_Async_CannotRetrievePayload", "UInterchangeTextureFactory: Invalid translator couldn't retrive a payload."));
+		ImportTextureErrorLog(LOCTEXT("TextureFactory_Async_CannotRetrievePayload", "UInterchangeTextureFactory: Invalid translator couldn't retrieve a payload."));
 		return ImportAssetResult;
 	}
 
diff --git a/Engine/Plugins/Interchange/Runtime/Source/Import/Public/MaterialX/MaterialXUtils/MaterialXSurfaceShaderAbstract.h b/Engine/Plugins/Interchange/Runtime/Source/Import/Public/MaterialX/MaterialXUtils/MaterialXSurfaceShaderAbstract.h
--- a/Engine/Plugins/Interchange/Runtime/Source/Import/Public/MaterialX/MaterialXUtils/MaterialXSurfaceShaderAbstract.h
+++ b/Engine/Plugins/Interchange/Runtime/Source/Import/Public/MaterialX/MaterialXUtils/MaterialXSurfaceShaderAbstract.h
@@ -16,23 +16,23 @@ protected:
 
 	static FString EmptyString;
 
-	// MaterialX states the default output name of the different nodes is 'out'
+	// MaterialX states the default output name of the different nodes is 'out'.
 	static FString DefaultOutput;
 
 	friend class FMaterialXSurfaceMaterial;
 
 	struct FConnectNode
 	{
-		// The MaterialX node of a given type used to create the appropriate shader node
+		// The MaterialX node of a given type used to create the appropriate shader node.
 		MaterialX::NodePtr UpstreamNode;
 
-		// The shader node to connect to
+		// The shader node to connect to.
 		UInterchangeShaderNode* ParentShaderNode;
 
-		// The input of the ParentShaderNode to connect to
+		// The input of the ParentShaderNode to connect to.
 		FString InputChannelName;
 
-		// The output name of the MaterialX node, default name is 'out' as stated by the standard library
+		// The output name of the MaterialX node. The default name is 'out' as stated by the standard library.
 		FString OutputName{ DefaultOutput };
 	};
 		
@@ -41,24 +41,24 @@ protected:
 	FMaterialXSurfaceShaderAbstract(UInterchangeBaseNodeContainer & BaseNodeContainer);
 
 	/**
-	 * Add an attribute to a shader node from the given MaterialX input, only floats and linear colors are supported for the moment
+	 * Add an attribute to a shader node from the given MaterialX input. Only floats and linear colors are supported.
 	 *
-	 * @param Input - The MaterialX input to retrieve and add the value from, must be of type float/color/vector
-	 * @param InputChannelName - The name of the shader node's input to add the attribute
-	 * @param ShaderNode - The shader node to which we want to add the attribute
+	 * @param Input - The MaterialX input to retrieve and add the value from. Must be of type float/color/vector.
+	 * @param InputChannelName - The name of the shader node's input to add the attribute.
+	 * @param ShaderNode - The shader node to which we want to add the attribute.
 	 *
-	 * @return true if the attribute was successfully added
+	 * @return true if the attribute was successfully added.
 	 */
 	bool AddAttribute(MaterialX::InputPtr Input, const FString& InputChannelName, UInterchangeShaderNode* ShaderNode);
 
 	/**
-	 * Add an attribute to a shader node from the given MaterialX input if that input has either a value or an interface name
+	 * Add an attribute to a shader node from the given MaterialX input if that input has either a value or an interface name.
 	 *
-	 * @param Input - The MaterialX input to retrieve and add the value from, must be of type float/color/vector
-	 * @param InputChannelName - The name of the shader node's input to add the attribute
-	 * @param ShaderNode - The shader node to which we want to add the attribute
+	 * @param Input - The MaterialX input to retrieve and add the value from. Must be of type float/color/vector.
+	 * @param InputChannelName - The name of the shader node's input to add the attribute.
+	 * @param ShaderNode - The shader node to which we want to add the attribute.
 	 *
-	 * @return true if the attribute was successfully added
+	 * @return true if the attribute was successfully added.
 	 */
 	bool AddAttributeFromValueOrInterface(MaterialX::InputPtr Input, const FString& InputChannelName, UInterchangeShaderNode* ShaderNode);
 
@@ -77,41 +77,41 @@ protected:
 	/**
 	 * Add a float attribute to a shader node only if its value taken from the input is not equal to its default value. Return false if the attribute does not exist or if we cannot add it
 	 *
-	 * @param Input - The MaterialX input to retrieve and add the value from, must be of type float
-	 * @param InputChannelName - The name of the shader node's input to add the attribute
-	 * @param ShaderNode - The shader node to which we want to add the attribute
-	 * @param DefaultValue - the default value to test the input against
+	 * @param Input - The MaterialX input to retrieve and add the value from. Must be of type float.
+	 * @param InputChannelName - The name of the shader node's input to add the attribute.
+	 * @param ShaderNode - The shader node to which we want to add the attribute.
+	 * @param DefaultValue - the default value to test the input against.
 	 *
-	 * @return true if the attribute was successfully added
+	 * @return true if the attribute was successfully added.
 	 */
 	bool AddFloatAttribute(MaterialX::InputPtr Input, const FString& InputChannelName, UInterchangeShaderNode* ShaderNode, float DefaultValue = std::numeric_limits<float>::max());
 
 	/**
-	 * Add a FLinearColor attribute to a shader node only if its value taken from the input is not equal to its default value. Return false if the attribute does not exist or if we cannot add it
+	 * Add an FLinearColor attribute to a shader node only if its value taken from the input is not equal to its default value. Return false if the attribute does not exist or if we cannot add it.
 	 *
-	 * @param Input - The MaterialX input to retrieve and add the value from, must be of type color
-	 * @param InputChannelName - The name of the shader node's input to add the attribute
-	 * @param ShaderNode - The shader node to which we want to add the attribute
-	 * @param DefaultValue - the default value to test the input against
+	 * @param Input - The MaterialX input to retrieve and add the value from. Must be of type color.
+	 * @param InputChannelName - The name of the shader node's input to add the attribute.
+	 * @param ShaderNode - The shader node to which we want to add the attribute.
+	 * @param DefaultValue - the default value to test the input against.
 	 *
-	 * @return true if the attribute was successfully added
+	 * @return true if the attribute was successfully added.
 	 */
 	bool AddLinearColorAttribute(MaterialX::InputPtr Input, const FString& InputChannelName, UInterchangeShaderNode* ShaderNode, const FLinearColor& DefaultValue = FLinearColor{std::numeric_limits<float>::max(), std::numeric_limits<float>::max(), std::numeric_limits<float>::max()});
 
 	/**
-	 * Add a FLinearColor attribute to a shader node only if its value taken from the input is not equal to its default value. Return false if the attribute does not exist or if we cannot add it
+	 * Add an FLinearColor attribute to a shader node only if its value taken from the input is not equal to its default value. Return false if the attribute does not exist or if we cannot add it.
 	 *
-	 * @param Input - The MaterialX input to retrieve and add the value from, must be of type vector
-	 * @param InputChannelName - The name of the shader node's input to add the attribute
-	 * @param ShaderNode - The shader node to which we want to add the attribute
-	 * @param DefaultValue - the default value to test the input against
+	 * @param Input - The MaterialX input to retrieve and add the value from. Must be of type vector.
+	 * @param InputChannelName - The name of the shader node's input to add the attribute.
+	 * @param ShaderNode - The shader node to which we want to add the attribute.
+	 * @param DefaultValue - the default value to test the input against.
 	 *
-	 * @return true if the attribute was successfully added
+	 * @return true if the attribute was successfully added.
 	 */
 	bool AddVectorAttribute(MaterialX::InputPtr Input, const FString& InputChannelName, UInterchangeShaderNode* ShaderNode, const FVector4f& DefaultValue = FVector4f{ std::numeric_limits<float>::max(), std::numeric_limits<float>::max(), std::numeric_limits<float>::max() });
 
 	/**
-	 * Connect an output either from a node name or a node graph from a MaterialX input to the ShaderNode
+	 * Connect an output either from a node name or a node graph from a MaterialX input to the shader node.
 	 * 
 	 * @param InputName - The name of the input of the SurfaceShaderNode to retrieve
 	 * @param ShaderNode - The Interchange shader node to connect the MaterialX's node or node graph to
@@ -154,47 +154,47 @@ protected:
 	};
 
 	/**
-	 * Connect an ouput in the NodeGraph to the ShaderGraph
+	 * Connect an ouput in the NodeGraph to the ShaderGraph.
 	 *
-	 * @param InputToNodeGraph - The input from the standard surface to retrieve the output in the NodeGraph
-	 * @param ShaderNode - The Interchange shader node to connect the MaterialX's node graph to
-	 * @param ParentInputName - The name of the input of the shader node to which we want the node graph to be connected to
+	 * @param InputToNodeGraph - The input from the standard surface to retrieve the output in the NodeGraph.
+	 * @param ShaderNode - The Interchange shader node to connect the MaterialX's node graph to.
+	 * @param ParentInputName - The name of the input of the shader node that we want the node graph to be connected to.
 	 *
-	 * @return true if the given input is attached to one of the outputs of a node graph
+	 * @return true if the given input is attached to one of the outputs of a node graph.
 	 */
 	bool ConnectNodeGraphOutputToInput(MaterialX::InputPtr InputToNodeGraph, UInterchangeShaderNode* ShaderNode, const FString& ParentInputName);
 
 	/**
-	 * Create and Connect the output of a MaterialX node that has already a matching in UE to a shader node,
-	 * If not search for a registered delegate
+	 * Create and connect the output of a MaterialX node that has already a matching in UE to a shader node.
+	 * If not, search for a registered delegate.
 	 *
-	 * @param Node - The MaterialX node of a given type used to create the appropriate shader node
-	 * @param ParentShaderNode - The shader node to connect to
-	 * @param InputChannelName - The input of the ParentShaderNode to connect to
+	 * @param Node - The MaterialX node of a given type used to create the appropriate shader node.
+	 * @param ParentShaderNode - The shader node to connect to.
+	 * @param InputChannelName - The input of the ParentShaderNode to connect to.
 	 *
-	 * @return true if a shader node has been successfully created and is connected to the given input
+	 * @return true if a shader node has been successfully created and is connected to the given input.
 	 */
 	bool ConnectMatchingNodeOutputToInput(const FConnectNode& Connect);
 
 	/**
-	 * Create and Connect manually the output of a MaterialX node to a shader node
+	 * Create and connect manually the output of a MaterialX node to a shader node.
 	 *
-	 * @param Edge - The MaterialX edge that has the current node, its parent and the input bridge between the two
-	 * @param ParentShaderNode - The shader node to connect to
-	 * @param InputChannelName - The input of the ParentShaderNode to connect to
+	 * @param Edge - The MaterialX edge that has the current node, its parent and the input bridge between the two.
+	 * @param ParentShaderNode - The shader node to connect to.
+	 * @param InputChannelName - The input of the ParentShaderNode to connect to.
 	 *
-	 * @return true if a shader node has been successfully created and is connected to the given input
+	 * @return true if a shader node has been successfully created and is connected to the given input.
 	 */
 	void ConnectNodeCategoryOutputToInput(const MaterialX::Edge& Edge, UInterchangeShaderNode* ParentShaderNode, const FString& InputChannelName, const FString& OutputName = TEXT("out"));
 
 	/**
-	 * Create and Connect a node name directly connected from an input to a shader node
+	 * Create and connect a node name directly connected from an input to a shader node.
 	 *
-	 * @param Node - The MaterialX node of a given type used to create the appropriate shader node
-	 * @param ParentShaderNode - The shader node to connect to
-	 * @param InputChannelName - The input of the ParentShaderNode to connect to
+	 * @param Node - The MaterialX node of a given type used to create the appropriate shader node.
+	 * @param ParentShaderNode - The shader node to connect to.
+	 * @param InputChannelName - The input of the ParentShaderNode to connect to.
 	 *
-	 * @return true if a shader node has been successfully created and is connected to the given input
+	 * @return true if a shader node has been successfully created and is connected to the given input.
 	 */
 	bool ConnectNodeNameOutputToInput(MaterialX::InputPtr Input, UInterchangeShaderNode* ShaderNode, const FString& ParentInputName);
 
@@ -280,41 +280,41 @@ protected:
 	/**
 	 * Create a ComponentMask shader node.
 	 *
-	 * @param RGBA - The mask component, e.g: 0b1011 -> Only RBA are toggled
-	 * @param NodeName - the name of the shader node
-	 * @param OutputName - the name of the output of the MaterialX node, default name is 'out' as stated by the standard library
-	 * @return The ComponentMask node
+	 * @param RGBA - The mask component. For example: 0b1011 -> Only RBA are toggled
+	 * @param NodeName - the name of the shader node.
+	 * @param OutputName - the name of the output of the MaterialX node. The default name is 'out' as stated by the standard library.
+	 * @return The ComponentMask node.
 	 */
 	UInterchangeShaderNode* CreateMaskShaderNode(uint8 RGBA, const FString& NodeName, const FString& OutputName = TEXT("out"));
 
 	/**
-	 * Helper function to create an InterchangeShaderNode
+	 * Helper function to create an InterchangeShaderNode.
 	 *
-	 * @param NodeName - The name of the shader node
-	 * @param ShaderType - The shader node's type we want to create
-	 * @param OutputName - The output name of the MaterialX node, default name is 'out' as stated by the standard library
+	 * @param NodeName - The name of the shader node.
+	 * @param ShaderType - The type of shader node we want to create.
+	 * @param OutputName - The output name of the MaterialX node. The default name is 'out' as stated by the standard library.
 	 *
-	 * @return The shader node that was created
+	 * @return The shader node that was created.
 	 */
 	UInterchangeShaderNode* CreateShaderNode(const FString& NodeName, const FString& ShaderType, const FString& OutputName = TEXT("out"));
 
 	/**
-	 * Helper function to create an InterchangeFunctionCallShaderNode
+	 * Helper function to create an InterchangeFunctionCallShaderNode.
 	 *
-	 * @param NodeName - The name of the shader node
-	 * @param FunctionPath - The path to the Material Function we want to create
-	 * @param OutputName - The output name of the MaterialX node, default name is 'out' as stated by the standard library
+	 * @param NodeName - The name of the shader node.
+	 * @param FunctionPath - The path to the Material Function we want to create.
+	 * @param OutputName - The output name of the MaterialX node. The default name is 'out' as stated by the standard library.
 	 *
-	 * @return The shader node that was created
+	 * @return The shader node that was created.
 	 */
 	UInterchangeFunctionCallShaderNode* CreateFunctionCallShaderNode(const FString& NodeName, const FString& FunctionPath, const FString& OutputName = TEXT("out"));
 
 	/**
-	 * Helper function to create an InterchangeTextureNode
+	 * Helper function to create an InterchangeTextureNode.
 	 *
-	 * @param Node - The MaterialX node, it should be of the category <image> no test is done on it
+	 * @param Node - The MaterialX node. This should be of the category <image>. No test is done on it.
 	 *
-	 * @return The texture node that was created
+	 * @return The texture node that was created.
 	 */
 	template<typename TextureTypeNode>
 	UInterchangeTextureNode* CreateTextureNode(MaterialX::NodePtr Node) const
@@ -394,64 +394,64 @@ protected:
 	}
 
 	/**
-	 * Get the UE corresponding name of a MaterialX Node category and input for a material
+	 * Get the UE corresponding name of a MaterialX Node category and input for a material.
 	 *
-	 * @param Input - MaterialX input
+	 * @param Input - MaterialX input.
 	 *
-	 * @return The matched name of the Node/Input else empty string
+	 * @return The matched name of the Node/Input, or an empty string.
 	 */
 	const FString& GetMatchedInputName(MaterialX::NodePtr Node, MaterialX::InputPtr Input) const;
 
 	/**
-	 * Get the input name, use this function instead of getName, because a renaming may have occured and we ensure to have the proper name that will be used by UE inputs
+	 * Get the input name. Use this function instead of getName() because this returns the name that will be used by UE inputs even if a renaming has occurred.
 	 *
-	 * @param Input - The input to retrieve the name from
+	 * @param Input - The input to retrieve the name from.
 	 *
-	 * @return The input name
+	 * @return The input name.
 	 */
 	FString GetInputName(MaterialX::InputPtr Input) const;
 
 	/**
-	 * Return the innermost file prefix of an element in the current scope, if it has none, it will take the one from its parents
+	 * Return the innermost file prefix of an element in the current scope. If it has none, it will take the one from its parents.
 	 *
-	 * @param Element - the Element to retrieve the file prefix from (can be anything, an input, a node, a nodegraph, etc.)
+	 * @param Element - the Element to retrieve the file prefix from. This can be anything: an input, a node, a nodegraph, and so on.
 	 *
-	 * @return a file prefix or an empty string
+	 * @return a file prefix or an empty string.
 	 *
 	 */
 	FString GetFilePrefix(MaterialX::ElementPtr Element) const;
 
 	/**
-	 * Helper function that returns a vector, the function makes no assumption on the input, and it should have a value of vectorN type
+	 * Helper function that returns a vector. The function makes no assumption on the input, and it should have a value of vectorN type.
 	 *
-	 * @param Input - The input that has a vectorN value in it
+	 * @param Input - The input that has a vectorN value in it.
 	 *
-	 * @return The vector
+	 * @return The vector.
 	 */
 	FLinearColor GetVector(MaterialX::InputPtr Input) const;
 
 	/**
-	 * Retrieve the interchange parent name of a MaterialX node, useful when a node is a combination of several nodes connected to different inputs (look for example Noise3D)
+	 * Retrieve the Interchange parent name of a MaterialX node. Useful when a node is a combination of several nodes connected to different inputs, such as Noise3D.
 	 *
-	 * @param Node - The node to retrieve the parent name
-	 * @return The node of the parent
+	 * @param Node - The node whose parent name is retrieved.
+	 * @return The node of the parent.
 	 */
 	FString GetAttributeParentName(MaterialX::NodePtr Node) const;
 
 	virtual void RegisterConnectNodeOutputToInputDelegates();
 
 	/**
-	 * Set the matching inputs names of a node to correspond to the one used by UE, the matching name is stored under the attribute UE::NewName
+	 * Set the matching inputs names of a node to correspond to the one used by UE. The matching name is stored under the attribute UE::NewName.
 	 *
-	 * @param Node - Look up to all inputs of Node and set the matching name attribute
+	 * @param Node - Look up to all inputs of Node and set the matching name attribute.
 	 */
 	void SetMatchingInputsNames(MaterialX::NodePtr Node) const;
 
 	/**
 	 * Add the input new name under the attribute UE::NewName.
 	 *
-	 * @param Input - The input to add the proper attribute
-	 * @param NewName - the new name of the input
+	 * @param Input - The input to add the proper attribute.
+	 * @param NewName - the new name of the input.
 	 */
 	void SetAttributeNewName(MaterialX::InputPtr Input, const char* NewName) const;
 
diff --git a/Engine/Plugins/Interchange/Runtime/Source/Nodes/Public/InterchangeAnimationTrackSetNode.h b/Engine/Plugins/Interchange/Runtime/Source/Nodes/Public/InterchangeAnimationTrackSetNode.h
--- a/Engine/Plugins/Interchange/Runtime/Source/Nodes/Public/InterchangeAnimationTrackSetNode.h
+++ b/Engine/Plugins/Interchange/Runtime/Source/Nodes/Public/InterchangeAnimationTrackSetNode.h
@@ -12,7 +12,7 @@ enum class EInterchangeAnimationPayLoadType : uint8
 {
 	NONE = 0,
 	CURVE,
-	MORPHTARGETCURVE, //handles/generates the same properties as the CURVE variation, but they way it is acquired might be different (depending on the Translator)
+	MORPHTARGETCURVE, // Handles/generates the same properties as the CURVE variation, but the way it is acquired might be different (depending on the Translator).
 	STEPCURVE,
 	BAKED,
 	MORPHTARGETCURVEWEIGHTINSTANCE	//Handled within UInterchangeAnimSequenceFactory, contrary to the others which are handled in the Translators
@@ -70,21 +70,21 @@ UENUM(BlueprintType)
 enum class EInterchangeAnimatedProperty : uint8
 {
 	None UMETA(DisplayName = "No property.", ToolTip = "The associated animation track will be ignored."),
-	Visibility UMETA(DisplayName = "Visibility property.", ToolTip = "The associated animation track is applied to the visibility property of the actor"),
+	Visibility UMETA(DisplayName = "Visibility property.", ToolTip = "The associated animation track is applied to the visibility property of the actor."),
 	MAX,
 };
 /**
- * Enumeration specifying how to handle the state of the animated property at the end of an animation track
+ * Enumeration specifying how to handle the state of the animated property at the end of an animation track.
  */
 enum class EInterchangeAimationCompletionMode : uint8
 {
 	KeepState UMETA(DisplayName = "Keep State", ToolTip = "Keep the animated property at the state set at the end of the animation track."),
 	RestoreState UMETA(DisplayName = "Restore State", ToolTip = "Restore the animated property to its state before the start of the animation track."),
-	ProjectDefault UMETA(DisplayName = "Project Default", ToolTip = "Restore the animated property to the state set in the project for such property."),
+	ProjectDefault UMETA(DisplayName = "Project Default", ToolTip = "Restore the animated property to the state set in the project for that property."),
 };
 
 /**
- * Class to represent a set of animation track nodes sharing the same frame rate
+ * Class to represent a set of animation track nodes that share the same frame rate.
  */
 UCLASS(BlueprintType)
 class INTERCHANGENODES_API UInterchangeAnimationTrackSetNode : public UInterchangeBaseNode
@@ -100,7 +100,7 @@ public:
 	}
 
 	/**
-	 * Return the node type name of the class, we use this when reporting errors
+	 * Return the node type name of the class. This is used when reporting errors.
 	 */
 	virtual FString GetTypeName() const override
 	{
@@ -109,19 +109,19 @@ public:
 	}
 
 	/**
-	 * This function allow to retrieve the number of track dependencies for this object.
+	 * Retrieve the number of track dependencies for this object.
 	 */
 	UFUNCTION(BlueprintCallable, Category = "Interchange | Node | AnimationTrackSet")
 	int32 GetCustomAnimationTrackUidCount() const;
 
 	/**
-	 * This function allow to retrieve the track dependency for this object.
+	 * Retrieve the track dependency for this object.
 	 */
 	UFUNCTION(BlueprintCallable, Category = "Interchange | Node | AnimationTrackSet")
 	void GetCustomAnimationTrackUids(TArray<FString>& OutAnimationTrackUids) const;
 
 	/**
-	 * This function allow to retrieve one track dependency for this object.
+	 * Retrieve one track dependency for this object.
 	 */
 	UFUNCTION(BlueprintCallable, Category = "Interchange | Node | AnimationTrackSet")
 	void GetCustomAnimationTrackUid(const int32 Index, FString& OutAnimationTrackUid) const;
@@ -157,7 +157,7 @@ private:
 };
 
 /**
- * Abstract class providing the minimal services required for an animation track node
+ * Abstract class providing the minimal services required for an animation track node.
  */
 UCLASS(Abstract)
 class INTERCHANGENODES_API UInterchangeAnimationTrackBaseNode : public UInterchangeBaseNode
@@ -171,7 +171,7 @@ public:
 	}
 
 	/**
-	* Return the node type name of the class, we use this when reporting errors
+	* Return the node type name of the class. This is used when reporting errors.
 	*/
 	virtual FString GetTypeName() const override
 	{
@@ -187,7 +187,7 @@ public:
 
 	/**
 	 * Get how the actor's animated property behaves once this animation is complete.
-	 * The output value will be clamped to the range of values defined in EInterchangeAimationCompletionMode
+	 * The output value will be clamped to the range of values defined in EInterchangeAimationCompletionMode.
 	 */
 	UFUNCTION(BlueprintCallable, Category = "Interchange | Node | AnimationTrack")
 	bool GetCustomCompletionMode(int32& AttributeValue) const;
@@ -197,7 +197,7 @@ private:
 };
 
 /**
- * Class to represent an animation which instances another animation track set node
+ * Class to represent an animation that instances another animation track set node.
  */
 UCLASS(BlueprintType)
 class INTERCHANGENODES_API UInterchangeAnimationTrackSetInstanceNode : public UInterchangeAnimationTrackBaseNode
@@ -211,7 +211,7 @@ public:
 	}
 
 	/**
-	 * Return the node type name of the class, we use this when reporting errors
+	 * Return the node type name of the class. This is used when reporting errors.
 	 */
 	virtual FString GetTypeName() const override
 	{
@@ -275,8 +275,8 @@ private:
 };
 
 /**
- * Class to represent an animation on the property of a camera, light or scene node
- * The list of supported properties is enumerated in EInterchangeAnimatedProperty
+ * Class to represent an animation on the property of a camera, light, or scene node
+ * The list of supported properties is enumerated in EInterchangeAnimatedProperty.
  */
 UCLASS(BlueprintType)
 class INTERCHANGENODES_API UInterchangeAnimationTrackNode : public UInterchangeAnimationTrackBaseNode
@@ -290,7 +290,7 @@ public:
 	}
 
 	/**
-	 * Return the node type name of the class, we use this when reporting errors
+	 * Return the node type name of the class. This is used when reporting errors.
 	 */
 	virtual FString GetTypeName() const override
 	{
@@ -368,7 +368,7 @@ private:
 };
 
 /**
- * Class to represent an animation on the transform of a camera, light or scene node
+ * Class to represent an animation on the transform of a camera, light, or scene node.
  */
 UCLASS(BlueprintType)
 class INTERCHANGENODES_API UInterchangeTransformAnimationTrackNode : public UInterchangeAnimationTrackNode
@@ -382,7 +382,7 @@ public:
 	}
 
 	/**
-	 * Return the node type name of the class, we use this when reporting errors
+	 * Return the node type name of the class. This is used when reporting errors.
 	 */
 	virtual FString GetTypeName() const override
 	{
@@ -436,7 +436,7 @@ public:
 	UInterchangeSkeletalAnimationTrackNode();
 
 	/**
-	 * Override serialize to restore SlotMaterialDependencies on load.
+	 * Override Serialize() to restore SlotMaterialDependencies on load.
 	 */
 	virtual void Serialize(FArchive& Ar) override
 	{
@@ -458,7 +458,7 @@ public:
 	}
 
 	/**
-	* Return the node type name of the class, we use this when reporting errors
+	* Return the node type name of the class. This is used when reporting errors.
 	*/
 	virtual FString GetTypeName() const override
 	{
@@ -468,16 +468,16 @@ public:
 
 
 public:
-	/** Get the skeleton factory node unique id. Return false if the attribute is not set. */
+	/** Get the unique ID of the skeleton factory node. Return false if the attribute is not set. */
 	UFUNCTION(BlueprintCallable, Category = "Interchange | Node | SkeletalAnimationTrack")
 	bool GetCustomSkeletonNodeUid(FString& AttributeValue) const;
 
-	/** Set the skeleton factory node unique id. Return false if the attribute cannot be set. */
+	/** Set the unique ID of the skeleton factory node. Return false if the attribute could not be set. */
 	UFUNCTION(BlueprintCallable, Category = "Interchange | Node | SkeletalAnimationTrack")
 	bool SetCustomSkeletonNodeUid(const FString& AttributeValue);
 
 	/**
-	 * Set the animation sample rate. Return false if the attribute cannot be set.
+	 * Set the animation sample rate. Return false if the attribute could not be set.
 	 */
 	UFUNCTION(BlueprintCallable, Category = "Interchange | Node | SkeletalAnimationTrack")
 	bool SetCustomAnimationSampleRate(const double& SampleRate);
@@ -489,7 +489,7 @@ public:
 	bool GetCustomAnimationSampleRate(double& SampleRate) const;
 
 	/**
-	 * Set the animation start time. Return false if the attribute cannot be set.
+	 * Set the animation start time. Return false if the attribute could not be set.
 	 */
 	UFUNCTION(BlueprintCallable, Category = "Interchange | Node | SkeletalAnimationTrack")
 	bool SetCustomAnimationStartTime(const double& StartTime);
@@ -501,7 +501,7 @@ public:
 	bool GetCustomAnimationStartTime(double& StartTime) const;
 	
 	/**
-	 * Set the animation stop time. Return false if the attribute cannot be set.
+	 * Set the animation stop time. Return false if the attribute could not be set.
 	 */
 	UFUNCTION(BlueprintCallable, Category = "Interchange | Node | SkeletalAnimationTrack")
 	bool SetCustomAnimationStopTime(const double& StopTime);
diff --git a/Engine/Plugins/Interchange/Runtime/Source/Nodes/Public/InterchangeCameraNode.h b/Engine/Plugins/Interchange/Runtime/Source/Nodes/Public/InterchangeCameraNode.h
--- a/Engine/Plugins/Interchange/Runtime/Source/Nodes/Public/InterchangeCameraNode.h
+++ b/Engine/Plugins/Interchange/Runtime/Source/Nodes/Public/InterchangeCameraNode.h
@@ -19,7 +19,7 @@ public:
 	}
 
 	/**
-	 * Return the node type name of the class, we use this when reporting errors
+	 * Return the node type name of the class. This is used when reporting errors.
 	 */
 	virtual FString GetTypeName() const override
 	{
@@ -102,7 +102,7 @@ public:
 	}
 
 	/**
-	 * Return the node type name of the class, we use this when reporting errors
+	 * Return the node type name of the class. This is used when reporting errors.
 	 */
 	virtual FString GetTypeName() const override
 	{
diff --git a/Engine/Plugins/Interchange/Runtime/Source/Nodes/Public/InterchangeLightNode.h b/Engine/Plugins/Interchange/Runtime/Source/Nodes/Public/InterchangeLightNode.h
--- a/Engine/Plugins/Interchange/Runtime/Source/Nodes/Public/InterchangeLightNode.h
+++ b/Engine/Plugins/Interchange/Runtime/Source/Nodes/Public/InterchangeLightNode.h
@@ -7,7 +7,7 @@
 
 #include "InterchangeLightNode.generated.h"
 
-// This enum is used as a placeholder for ELightUnits, because InterchangeWorker is not compiled against Engine, the LightFactoryNode is not affected
+// This enum is used as a placeholder for ELightUnits. Because InterchangeWorker is not compiled against Engine, the LightFactoryNode is not affected.
 UENUM()
 enum class EInterchangeLightUnits : uint8
 {
@@ -26,7 +26,7 @@ public:
 	static FStringView StaticAssetTypeName();
 
 	/**
-	 * Return the node type name of the class, we use this when reporting errors
+	 * Return the node type name of the class. This is used when reporting errors.
 	 */
 	virtual FString GetTypeName() const override;
 
